Control Rig: Fix crash when selecting AdditiveControlRig for ControlRigComponent

diff --git a/Engine/Plugins/Animation/ControlRig/Source/ControlRig/Private/Rigs/AdditiveControlRig.cpp b/Engine/Plugins/Animation/ControlRig/Source/ControlRig/Private/Rigs/AdditiveControlRig.cpp
--- a/Engine/Plugins/Animation/ControlRig/Source/ControlRig/Private/Rigs/AdditiveControlRig.cpp
+++ b/Engine/Plugins/Animation/ControlRig/Source/ControlRig/Private/Rigs/AdditiveControlRig.cpp
@@ -61,9 +61,12 @@ bool UAdditiveControlRig::Execute_Internal(const FName& InEventName)
 	}
 	else if(InEventName == FRigUnit_PrepareForExecution::EventName)
 	{
-		if (USkeletalMeshComponent* SkeletalMeshComponent = Cast<USkeletalMeshComponent>(GetObjectBinding()->GetBoundObject()))
+		if (TSharedPtr<IControlRigObjectBinding> Binding = GetObjectBinding())
 		{
-			CreateRigElements(SkeletalMeshComponent->GetSkeletalMeshAsset());
+			if (USkeletalMeshComponent* SkeletalMeshComponent = Cast<USkeletalMeshComponent>(Binding->GetBoundObject()))
+			{
+				CreateRigElements(SkeletalMeshComponent->GetSkeletalMeshAsset());
+			}
 		}
 
 		// add units and initialize
