Control Rig: Python access to currently focused model / graph

diff --git a/Engine/Plugins/Animation/ControlRig/Source/ControlRigDeveloper/Private/ControlRigBlueprint.cpp b/Engine/Plugins/Animation/ControlRig/Source/ControlRigDeveloper/Private/ControlRigBlueprint.cpp
--- a/Engine/Plugins/Animation/ControlRig/Source/ControlRigDeveloper/Private/ControlRigBlueprint.cpp
+++ b/Engine/Plugins/Animation/ControlRig/Source/ControlRigDeveloper/Private/ControlRigBlueprint.cpp
@@ -1982,6 +1982,17 @@ bool UControlRigBlueprint::RemoveModel(FString InName, bool bSetupUndoRedo, bool
 	return RigVMClient.RemoveModel(InName, bSetupUndoRedo);
 }
 
+URigVMGraph* UControlRigBlueprint::GetFocusedModel() const
+{
+#if WITH_EDITOR
+	if(OnGetFocusedGraphDelegate.IsBound())
+	{
+		return OnGetFocusedGraphDelegate.Execute(); 
+	}
+#endif
+	return RigVMClient.GetDefaultModel();
+}
+
 URigVMController* UControlRigBlueprint::GetController(const URigVMGraph* InGraph) const
 {
 	return RigVMClient.GetController(InGraph);
diff --git a/Engine/Plugins/Animation/ControlRig/Source/ControlRigDeveloper/Public/ControlRigBlueprint.h b/Engine/Plugins/Animation/ControlRig/Source/ControlRigDeveloper/Public/ControlRigBlueprint.h
--- a/Engine/Plugins/Animation/ControlRig/Source/ControlRigDeveloper/Public/ControlRigBlueprint.h
+++ b/Engine/Plugins/Animation/ControlRig/Source/ControlRigDeveloper/Public/ControlRigBlueprint.h
@@ -46,6 +46,7 @@ DECLARE_DELEGATE_RetVal_OneParam(bool, FControlRigOnBreakLinksDialogRequestedDel
 DECLARE_DELEGATE_RetVal_OneParam(TRigVMTypeIndex, FControlRigOnPinTypeSelectionRequestedDelegate, const TArray<TRigVMTypeIndex>&);
 DECLARE_EVENT(UControlRigBlueprint, FOnBreakpointAdded);
 DECLARE_EVENT_OneParam(UControlRigBlueprint, FOnRequestInspectObject, const TArray<UObject*>& );
+DECLARE_DELEGATE_RetVal(URigVMGraph*, FControlRigBlueprintGetFocusedGraph);
 
 USTRUCT(meta = (Deprecated = "5.2"))
 struct CONTROLRIGDEVELOPER_API FControlRigPublicFunctionArg
@@ -279,6 +280,8 @@ public:
 
 #endif	// #if WITH_EDITOR
 
+	FControlRigBlueprintGetFocusedGraph& OnGetFocusedGraph() { return OnGetFocusedGraphDelegate; }
+
 	virtual bool ShouldBeMarkedDirtyUponTransaction() const override { return false; }
 
 	/** IInterface_PreviewMeshProvider interface */
@@ -335,6 +338,9 @@ public:
 	UFUNCTION(BlueprintCallable, Category = "Control Rig Blueprint")
 	bool RemoveModel(FString InName = TEXT("Rig Graph"), bool bSetupUndoRedo = true, bool bPrintPythonCommand = true);
 
+	UFUNCTION(BlueprintCallable, Category = "Control Rig Blueprint")
+	URigVMGraph* GetFocusedModel() const;
+
 	UFUNCTION(BlueprintCallable, Category = "Control Rig Blueprint")
 	URigVMController* GetController(const URigVMGraph* InGraph = nullptr) const;
 
@@ -747,7 +753,8 @@ private:
 	TArray<URigVMNode*> RigVMBreakpointNodes;
 
 	FOnRequestInspectObject OnRequestInspectObjectEvent;
-	
+	FControlRigBlueprintGetFocusedGraph OnGetFocusedGraphDelegate;
+		
 public:
 
 	/** Sets the execution mode. In Release mode the rig will ignore all breakpoints. */
diff --git a/Engine/Plugins/Animation/ControlRig/Source/ControlRigEditor/Private/Editor/ControlRigEditor.cpp b/Engine/Plugins/Animation/ControlRig/Source/ControlRigEditor/Private/Editor/ControlRigEditor.cpp
--- a/Engine/Plugins/Animation/ControlRig/Source/ControlRigEditor/Private/Editor/ControlRigEditor.cpp
+++ b/Engine/Plugins/Animation/ControlRig/Source/ControlRigEditor/Private/Editor/ControlRigEditor.cpp
@@ -275,6 +275,7 @@ FControlRigEditor::~FControlRigEditor()
 		RigBlueprint->SetDebugMode(false);
 		RigBlueprint->ClearBreakpoints();
 		SetHaltedNode(nullptr);
+		RigBlueprint->OnGetFocusedGraph().Unbind();
 #endif
 	}
 
@@ -542,6 +543,9 @@ void FControlRigEditor::InitControlRigEditor(const EToolkitMode::Type Mode, cons
 		InControlRigBlueprint->OnRequestBreakLinksDialog().BindSP(this, &FControlRigEditor::OnRequestBreakLinksDialog);
 		InControlRigBlueprint->OnRequestPinTypeSelectionDialog().BindSP(this, &FControlRigEditor::OnRequestPinTypeSelectionDialog);
 		InControlRigBlueprint->OnRequestJumpToHyperlink().BindSP(this, &FControlRigEditor::HandleJumpToHyperlink);
+#if WITH_EDITOR
+		InControlRigBlueprint->OnGetFocusedGraph().BindSP(this, &FControlRigEditor::GetFocusedModel);
+#endif
 	}
 
 	UpdateStaleWatchedPins();
