Sequencer: Curve Editor: Don't show the space or constrraint key bar spaces unless there is a space.

diff --git a/Engine/Plugins/Animation/ControlRig/Source/ControlRigEditor/Private/Sequencer/ControlRigSpaceChannelEditors.cpp b/Engine/Plugins/Animation/ControlRig/Source/ControlRigEditor/Private/Sequencer/ControlRigSpaceChannelEditors.cpp
--- a/Engine/Plugins/Animation/ControlRig/Source/ControlRigEditor/Private/Sequencer/ControlRigSpaceChannelEditors.cpp
+++ b/Engine/Plugins/Animation/ControlRig/Source/ControlRigEditor/Private/Sequencer/ControlRigSpaceChannelEditors.cpp
@@ -1390,9 +1390,17 @@ FReply FControlRigSpaceChannelHelpers::OpenBakeDialog(ISequencer* Sequencer, FMo
 	return FReply::Unhandled();
 }
 
-TUniquePtr<FCurveModel> CreateCurveEditorModel(const TMovieSceneChannelHandle<FMovieSceneControlRigSpaceChannel>& Channel, UMovieSceneSection* OwningSection, TSharedRef<ISequencer> InSequencer)
+TUniquePtr<FCurveModel> CreateCurveEditorModel(const TMovieSceneChannelHandle<FMovieSceneControlRigSpaceChannel>& ChannelHandle, UMovieSceneSection* OwningSection, TSharedRef<ISequencer> InSequencer)
 {
-	return MakeUnique<FControlRigSpaceChannelCurveModel>(Channel, OwningSection, InSequencer);
+	if (FMovieSceneControlRigSpaceChannel* Channel = ChannelHandle.Get())
+	{
+		if (Channel->GetNumKeys() > 0)
+		{
+			return MakeUnique<FControlRigSpaceChannelCurveModel>(ChannelHandle, OwningSection, InSequencer);
+		}
+
+	}
+	return nullptr;
 }
 
 TArray<FKeyBarCurveModel::FBarRange> FControlRigSpaceChannelHelpers::FindRanges(FMovieSceneControlRigSpaceChannel* Channel, const UMovieSceneSection* Section)
diff --git a/Engine/Source/Editor/MovieSceneTools/Private/Channels/ConstraintChannelEditor.cpp b/Engine/Source/Editor/MovieSceneTools/Private/Channels/ConstraintChannelEditor.cpp
--- a/Engine/Source/Editor/MovieSceneTools/Private/Channels/ConstraintChannelEditor.cpp
+++ b/Engine/Source/Editor/MovieSceneTools/Private/Channels/ConstraintChannelEditor.cpp
@@ -361,7 +361,14 @@ TSharedRef<SWidget> CreateKeyEditor(
 	return SNullWidget::NullWidget;
 }
 
-TUniquePtr<FCurveModel> CreateCurveEditorModel(const TMovieSceneChannelHandle<FMovieSceneConstraintChannel>& Channel, UMovieSceneSection* OwningSection, TSharedRef<ISequencer> InSequencer)
+TUniquePtr<FCurveModel> CreateCurveEditorModel(const TMovieSceneChannelHandle<FMovieSceneConstraintChannel>& ChannelHandle, UMovieSceneSection* OwningSection, TSharedRef<ISequencer> InSequencer)
 {
-	return MakeUnique<FConstraintChannelCurveModel>(Channel, OwningSection, InSequencer);
+	if (FMovieSceneConstraintChannel* Channel = ChannelHandle.Get())
+	{
+		if (Channel->GetNumKeys() > 0)
+		{
+			return MakeUnique<FConstraintChannelCurveModel>(ChannelHandle, OwningSection, InSequencer);
+		}
+	}
+	return nullptr;
 }
\ No newline at end of file
