Control Rig: Disable draw interface during play

diff --git a/Engine/Plugins/Animation/ControlRig/Source/ControlRig/Private/ControlRig.cpp b/Engine/Plugins/Animation/ControlRig/Source/ControlRig/Private/ControlRig.cpp
--- a/Engine/Plugins/Animation/ControlRig/Source/ControlRig/Private/ControlRig.cpp
+++ b/Engine/Plugins/Animation/ControlRig/Source/ControlRig/Private/ControlRig.cpp
@@ -53,6 +53,14 @@ static TAutoConsoleVariable<int32> CVarControlRigCreateFloatControlsForCurves(
 	TEXT("If nonzero we create a float control for each curve in the curve container, useful for debugging low level controls."),
 	ECVF_Default);
 
+//CVar to specify if we should allow debug drawing in game (during shipped game or PIE)
+//By default we don't for performance 
+static TAutoConsoleVariable<float> CVarControlRigEnableDrawInterfaceInGame(
+	TEXT("ControlRig.EnableDrawInterfaceInGame"),
+	0,
+	TEXT("If nonzero debug drawing will be enabled during play."),
+	ECVF_Default);
+
 UControlRig::UControlRig(const FObjectInitializer& ObjectInitializer)
 	: Super(ObjectInitializer)
 #if WITH_EDITOR
@@ -345,12 +353,29 @@ bool UControlRig::Execute(const FName& InEventName)
 
 	FRigUnitContext& Context = PublicContext.UnitContext;
 
-	// setup the draw interface for debug drawing
-	if(!bIsEventInQueue || bIsEventFirstInQueue)
+	bool bEnableDrawInterface = false;
+#if WITH_EDITOR
+	const bool bEnabledDuringGame = CVarControlRigEnableDrawInterfaceInGame->GetInt() != 0;
+	const bool bInGame = !(GIsEditor && !GEditor->GetPIEWorldContext());
+	if (bEnabledDuringGame || !bInGame)
+	{
+		bEnableDrawInterface = true;
+	}	
+#endif
+
+	if (bEnableDrawInterface)
+	{
+		// setup the draw interface for debug drawing
+		if(!bIsEventInQueue || bIsEventFirstInQueue)
+		{
+			DrawInterface.Reset();
+		}
+		PublicContext.SetDrawInterface(&DrawInterface);
+	}
+	else
 	{
-		DrawInterface.Reset();
+		PublicContext.SetDrawInterface(nullptr);
 	}
-	PublicContext.SetDrawInterface(&DrawInterface);
 
 	// setup the animation attribute container
 	Context.AnimAttributeContainer = ExternalAnimAttributeContainer;
