Control Rig: Fix crash during cook in ShooterGame

diff --git a/Engine/Plugins/Animation/ControlRig/Source/ControlRig/Private/ControlRig.cpp b/Engine/Plugins/Animation/ControlRig/Source/ControlRig/Private/ControlRig.cpp
--- a/Engine/Plugins/Animation/ControlRig/Source/ControlRig/Private/ControlRig.cpp
+++ b/Engine/Plugins/Animation/ControlRig/Source/ControlRig/Private/ControlRig.cpp
@@ -356,7 +356,7 @@ bool UControlRig::Execute(const FName& InEventName)
 	bool bEnableDrawInterface = false;
 #if WITH_EDITOR
 	const bool bEnabledDuringGame = CVarControlRigEnableDrawInterfaceInGame->GetInt() != 0;
-	const bool bInGame = !(GIsEditor && !GEditor->GetPIEWorldContext());
+	const bool bInGame = !(GEditor && !GEditor->GetPIEWorldContext());
 	if (bEnabledDuringGame || !bInGame)
 	{
 		bEnableDrawInterface = true;
