Control Rig: Fix GEditor undeclared identifier

diff --git a/Engine/Plugins/Animation/ControlRig/Source/ControlRig/Private/ControlRig.cpp b/Engine/Plugins/Animation/ControlRig/Source/ControlRig/Private/ControlRig.cpp
--- a/Engine/Plugins/Animation/ControlRig/Source/ControlRig/Private/ControlRig.cpp
+++ b/Engine/Plugins/Animation/ControlRig/Source/ControlRig/Private/ControlRig.cpp
@@ -27,6 +27,7 @@
 #include "EdGraphSchema_K2.h"
 #include "Kismet2/BlueprintEditorUtils.h"
 #include "RigVMDeveloperTypeUtils.h"
+#include "Editor.h"
 #endif// WITH_EDITOR
 #include "ControlRigComponent.h"
 #include "Constraints/ControlRigTransformableHandle.h"
