Control Rig: Fix crash when deleting variable and variable nodes are in use

diff --git a/Engine/Plugins/Animation/ControlRig/Source/ControlRigDeveloper/Private/ControlRigBlueprint.cpp b/Engine/Plugins/Animation/ControlRig/Source/ControlRigDeveloper/Private/ControlRigBlueprint.cpp
--- a/Engine/Plugins/Animation/ControlRig/Source/ControlRigDeveloper/Private/ControlRigBlueprint.cpp
+++ b/Engine/Plugins/Animation/ControlRig/Source/ControlRigDeveloper/Private/ControlRigBlueprint.cpp
@@ -4931,6 +4931,8 @@ void UControlRigBlueprint::OnPostVariableChange(UBlueprint* InBlueprint)
 		return;
 	}
 
+	TArray<FBPVariableDescription> LocalLastNewVariables = LastNewVariables;
+
 	TMap<FGuid, int32> NewVariablesByGuid;
 	for (int32 VarIndex = 0; VarIndex < NewVariables.Num(); VarIndex++)
 	{
@@ -4938,12 +4940,12 @@ void UControlRigBlueprint::OnPostVariableChange(UBlueprint* InBlueprint)
 	}
 
 	TMap<FGuid, int32> OldVariablesByGuid;
-	for (int32 VarIndex = 0; VarIndex < LastNewVariables.Num(); VarIndex++)
+	for (int32 VarIndex = 0; VarIndex < LocalLastNewVariables.Num(); VarIndex++)
 	{
-		OldVariablesByGuid.Add(LastNewVariables[VarIndex].VarGuid, VarIndex);
+		OldVariablesByGuid.Add(LocalLastNewVariables[VarIndex].VarGuid, VarIndex);
 	}
 
-	for (const FBPVariableDescription& OldVariable : LastNewVariables)
+	for (const FBPVariableDescription& OldVariable : LocalLastNewVariables)
 	{
 		if (!NewVariablesByGuid.Contains(OldVariable.VarGuid))
 		{
@@ -4961,7 +4963,7 @@ void UControlRigBlueprint::OnPostVariableChange(UBlueprint* InBlueprint)
 		}
 
 		int32 OldVarIndex = OldVariablesByGuid.FindChecked(NewVariable.VarGuid);
-		const FBPVariableDescription& OldVariable = LastNewVariables[OldVarIndex];
+		const FBPVariableDescription& OldVariable = LocalLastNewVariables[OldVarIndex];
 		if (OldVariable.VarName != NewVariable.VarName)
 		{
 			OnVariableRenamed(OldVariable.VarName, NewVariable.VarName);
