Control Rig: Fix load of ERigUnitDebugTransformMode

diff --git a/Engine/Plugins/Animation/ControlRig/Config/BaseControlRig.ini b/Engine/Plugins/Animation/ControlRig/Config/BaseControlRig.ini
--- a/Engine/Plugins/Animation/ControlRig/Config/BaseControlRig.ini
+++ b/Engine/Plugins/Animation/ControlRig/Config/BaseControlRig.ini
@@ -90,6 +90,7 @@
 +EnumRedirects=(OldName="EBoneGetterSetterMode",NewName="/Script/RigVM.ERigVMTransformSpace")
 +EnumRedirects=(OldName="EControlRigAnimEasingType",NewName="/Script/RigVM.ERigVMAnimEasingType")
 +EnumRedirects=(OldName="EControlRigClampSpatialMode",NewName="/Script/RigVM.ERigVMClampSpatialMode")
++EnumRedirects=(OldName="/Script/ControlRig.ERigUnitDebugTransformMode",NewName="/Script/RigVM.ERigUnitDebugTransformMode")
 +StructRedirects=(OldName="/Script/ControlRig.CRFourPointBezier",NewName="/Script/RigVM.RigVMFourPointBezier")
 +StructRedirects=(OldName="/Script/ControlRig.RigMirrorSettings",NewName="/Script/RigVM.RigVMMirrorSettings")
 +StructRedirects=(OldName="/Script/ControlRig.RigUnit_MathBase",NewName="/Script/RigVM.RigVMFunction_MathBase")
