Control Rig: Fix postload hashing not working properly

diff --git a/Engine/Plugins/Animation/ControlRig/Source/ControlRigDeveloper/Private/ControlRigBlueprint.cpp b/Engine/Plugins/Animation/ControlRig/Source/ControlRigDeveloper/Private/ControlRigBlueprint.cpp
--- a/Engine/Plugins/Animation/ControlRig/Source/ControlRigDeveloper/Private/ControlRigBlueprint.cpp
+++ b/Engine/Plugins/Animation/ControlRig/Source/ControlRigDeveloper/Private/ControlRigBlueprint.cpp
@@ -1321,14 +1321,16 @@ void UControlRigBlueprint::RefreshAllModels(EControlRigBlueprintLoadType InLoadT
 	const bool bIsPostLoad = InLoadType == EControlRigBlueprintLoadType::PostLoad;
 
 	// avoid any compute if the current structure hashes match with the serialized ones
-	if(RigVMClient.GetStructureHash() == RigVMClient.GetSerializedStructureHash())
+	const uint32 StructureHash = RigVMClient.GetStructureHash();
+	const uint32 SerializedHash = RigVMClient.GetSerializedStructureHash();
+	if(StructureHash == SerializedHash)
 	{
 		if(bIsPostLoad)
 		{
 			TArray<URigVMGraph*> ModelGraphs = RigVMClient.GetAllModels(true, true);
 			Algo::Reverse(ModelGraphs);
 			for (URigVMGraph* ModelGraph : ModelGraphs)
-			{	
+			{
 				URigVMController* Controller = GetOrCreateController(ModelGraph);
 				Controller->ReattachLinksToPinObjects(true /* follow redirectors */, nullptr, true, true);
 			}
diff --git a/Engine/Source/Developer/RigVMDeveloper/Private/RigVMModel/Nodes/RigVMFunctionEntryNode.cpp b/Engine/Source/Developer/RigVMDeveloper/Private/RigVMModel/Nodes/RigVMFunctionEntryNode.cpp
--- a/Engine/Source/Developer/RigVMDeveloper/Private/RigVMModel/Nodes/RigVMFunctionEntryNode.cpp
+++ b/Engine/Source/Developer/RigVMDeveloper/Private/RigVMModel/Nodes/RigVMFunctionEntryNode.cpp
@@ -5,6 +5,12 @@
 
 #include UE_INLINE_GENERATED_CPP_BY_NAME(RigVMFunctionEntryNode)
 
+uint32 URigVMFunctionEntryNode::GetStructureHash() const
+{
+	// Avoid hashing the template for library nodes
+	return URigVMNode::GetStructureHash();
+}
+
 FLinearColor URigVMFunctionEntryNode::GetNodeColor() const
 {
 	if(URigVMGraph* RootGraph = GetRootGraph())
diff --git a/Engine/Source/Developer/RigVMDeveloper/Private/RigVMModel/Nodes/RigVMFunctionReturnNode.cpp b/Engine/Source/Developer/RigVMDeveloper/Private/RigVMModel/Nodes/RigVMFunctionReturnNode.cpp
--- a/Engine/Source/Developer/RigVMDeveloper/Private/RigVMModel/Nodes/RigVMFunctionReturnNode.cpp
+++ b/Engine/Source/Developer/RigVMDeveloper/Private/RigVMModel/Nodes/RigVMFunctionReturnNode.cpp
@@ -5,6 +5,12 @@
 
 #include UE_INLINE_GENERATED_CPP_BY_NAME(RigVMFunctionReturnNode)
 
+uint32 URigVMFunctionReturnNode::GetStructureHash() const
+{
+	// Avoid hashing the template for library nodes
+	return URigVMNode::GetStructureHash();
+}
+
 FLinearColor URigVMFunctionReturnNode::GetNodeColor() const
 {
 	if(URigVMGraph* RootGraph = GetRootGraph())
diff --git a/Engine/Source/Developer/RigVMDeveloper/Private/RigVMModel/Nodes/RigVMLibraryNode.cpp b/Engine/Source/Developer/RigVMDeveloper/Private/RigVMModel/Nodes/RigVMLibraryNode.cpp
--- a/Engine/Source/Developer/RigVMDeveloper/Private/RigVMModel/Nodes/RigVMLibraryNode.cpp
+++ b/Engine/Source/Developer/RigVMDeveloper/Private/RigVMModel/Nodes/RigVMLibraryNode.cpp
@@ -35,6 +35,12 @@ bool URigVMLibraryNode::IsDefinedAsVarying() const
 	return false;
 }
 
+uint32 URigVMLibraryNode::GetStructureHash() const
+{
+	// Avoid hashing the template for library nodes
+	return URigVMNode::GetStructureHash();
+}
+
 FText URigVMLibraryNode::GetToolTipText() const
 {
 	return FText::FromName(GetFunctionIdentifier().HostObject.GetLongPackageFName());
diff --git a/Engine/Source/Developer/RigVMDeveloper/Public/RigVMModel/Nodes/RigVMFunctionEntryNode.h b/Engine/Source/Developer/RigVMDeveloper/Public/RigVMModel/Nodes/RigVMFunctionEntryNode.h
--- a/Engine/Source/Developer/RigVMDeveloper/Public/RigVMModel/Nodes/RigVMFunctionEntryNode.h
+++ b/Engine/Source/Developer/RigVMDeveloper/Public/RigVMModel/Nodes/RigVMFunctionEntryNode.h
@@ -20,6 +20,7 @@ public:
 	virtual UScriptStruct* GetScriptStruct() const override { return nullptr; }
 	virtual const FRigVMTemplate* GetTemplate() const override { return nullptr; }
 	virtual FName GetNotation() const override { return NAME_None; }
+	virtual uint32 GetStructureHash() const override;
 
 	// Override node functions
 	virtual FLinearColor GetNodeColor() const override;
diff --git a/Engine/Source/Developer/RigVMDeveloper/Public/RigVMModel/Nodes/RigVMFunctionReturnNode.h b/Engine/Source/Developer/RigVMDeveloper/Public/RigVMModel/Nodes/RigVMFunctionReturnNode.h
--- a/Engine/Source/Developer/RigVMDeveloper/Public/RigVMModel/Nodes/RigVMFunctionReturnNode.h
+++ b/Engine/Source/Developer/RigVMDeveloper/Public/RigVMModel/Nodes/RigVMFunctionReturnNode.h
@@ -20,6 +20,7 @@ public:
 	virtual UScriptStruct* GetScriptStruct() const override { return nullptr; }
 	virtual const FRigVMTemplate* GetTemplate() const override { return nullptr; }
 	virtual FName GetNotation() const override { return NAME_None; }
+	virtual uint32 GetStructureHash() const override;
 
 	// Override node functions
 	virtual FLinearColor GetNodeColor() const override;
diff --git a/Engine/Source/Developer/RigVMDeveloper/Public/RigVMModel/Nodes/RigVMLibraryNode.h b/Engine/Source/Developer/RigVMDeveloper/Public/RigVMModel/Nodes/RigVMLibraryNode.h
--- a/Engine/Source/Developer/RigVMDeveloper/Public/RigVMModel/Nodes/RigVMLibraryNode.h
+++ b/Engine/Source/Developer/RigVMDeveloper/Public/RigVMModel/Nodes/RigVMLibraryNode.h
@@ -34,6 +34,7 @@ public:
 	virtual UScriptStruct* GetScriptStruct() const override { return nullptr; }
 	virtual const FRigVMTemplate* GetTemplate() const override { return nullptr; }
 	virtual FName GetNotation() const override { return NAME_None; }
+	virtual uint32 GetStructureHash() const override;
 
 	// URigVMNode interface
 	virtual  FText GetToolTipText() const override;
