Fix out-of-bounds access in control rig curve I/O

Bulk operations were not correctly dealing with cases where the curve did not exist in the bulk data

diff --git a/Engine/Plugins/Animation/ControlRig/Source/ControlRig/Private/AnimNode_ControlRigBase.cpp b/Engine/Plugins/Animation/ControlRig/Source/ControlRig/Private/AnimNode_ControlRigBase.cpp
--- a/Engine/Plugins/Animation/ControlRig/Source/ControlRig/Private/AnimNode_ControlRigBase.cpp
+++ b/Engine/Plugins/Animation/ControlRig/Source/ControlRig/Private/AnimNode_ControlRigBase.cpp
@@ -248,10 +248,11 @@ void FAnimNode_ControlRigBase::UpdateInput(UControlRig* ControlRig, const FPoseC
 	{
 		// TODO: Do we need to 'unset' old values manually here? The previous code could do this because it knows the
 		// global set of curves, but we dont know that any more
-		UE::Anim::FCurveUtils::BulkGet(InOutput.Curve, CachedBulkCurves, [ControlRig](const UE::Anim::FNamedIndexElement& InBulkElement, float InValue)
+		URigHierarchy* Hierarchy = ControlRig->GetHierarchy();
+		UE::Anim::FCurveUtils::BulkGet(InOutput.Curve, CachedBulkCurves, [Hierarchy](const UE::Anim::FNamedIndexElement& InBulkElement, float InValue)
 		{
 			const FRigElementKey Key(InBulkElement.Name, ERigElementType::Curve);
-			ControlRig->GetHierarchy()->SetCurveValue(Key, InValue);
+			Hierarchy->SetCurveValue(Key, InValue);
 		});
 	}
 
@@ -364,10 +365,12 @@ void FAnimNode_ControlRigBase::UpdateOutput(UControlRig* ControlRig, FPoseContex
 
 	if (OutputSettings.bUpdateCurves)
 	{
-		UE::Anim::FCurveUtils::BulkSet(InOutput.Curve, CachedBulkCurves, [ControlRig](const UE::Anim::FNamedIndexElement& InBulkElement)
+		URigHierarchy* Hierarchy = ControlRig->GetHierarchy();
+		TArray<FRigCurveElement*> Curves = Hierarchy->GetCurves();
+		UE::Anim::FCurveUtils::BulkSet(InOutput.Curve, CachedBulkCurves, [Hierarchy, &Curves](const UE::Anim::FNamedIndexElement& InBulkElement)
 		{
-			FRigCurveElement* CurveElement = ControlRig->GetHierarchy()->GetCurves()[InBulkElement.Index];
-			return ControlRig->GetHierarchy()->GetCurveValue(CurveElement);
+			FRigCurveElement* CurveElement = Curves[InBulkElement.Index];
+			return Hierarchy->GetCurveValue(CurveElement);
 		});
 	}
 
diff --git a/Engine/Source/Runtime/Engine/Public/Animation/AnimCurveUtils.h b/Engine/Source/Runtime/Engine/Public/Animation/AnimCurveUtils.h
--- a/Engine/Source/Runtime/Engine/Public/Animation/AnimCurveUtils.h
+++ b/Engine/Source/Runtime/Engine/Public/Animation/AnimCurveUtils.h
@@ -447,7 +447,10 @@ public:
 		FNamedValueArrayUtils::Union(InCurve, InBulkCurves,
 			[&InValuePredicate](typename CurveType0::ElementType& InOutElement0, const typename CurveType1::ElementType& InElement1, UE::Anim::ENamedValueUnionFlags InFlags)
 			{
-				InOutElement0.Value = InValuePredicate(InElement1);
+				if(EnumHasAllFlags(InFlags, UE::Anim::ENamedValueUnionFlags::ValidArg1))
+				{
+					InOutElement0.Value = InValuePredicate(InElement1);
+				}
 			});
 	}
 };
