Moving implementations of FindItemAtPosition into SListView/SListViewBase to remove external dependencies on the Private header SListPanel.h

diff --git a/Engine/Plugins/Animation/ControlRig/Source/ControlRigEditor/Private/EditMode/SControlRigOutliner.cpp b/Engine/Plugins/Animation/ControlRig/Source/ControlRigEditor/Private/EditMode/SControlRigOutliner.cpp
--- a/Engine/Plugins/Animation/ControlRig/Source/ControlRigEditor/Private/EditMode/SControlRigOutliner.cpp
+++ b/Engine/Plugins/Animation/ControlRig/Source/ControlRigEditor/Private/EditMode/SControlRigOutliner.cpp
@@ -33,7 +33,6 @@
 #include "Widgets/Layout/SSpacer.h"
 #include "Widgets/Images/SImage.h"
 #include "Widgets/Input/SSearchBox.h"
-#include "Widgets/Views/SListPanel.h"
 #include "PropertyCustomizationHelpers.h"
 #include "Framework/Application/SlateApplication.h"
 #include "Editor/EditorEngine.h"
@@ -854,23 +853,6 @@ TArray<FMultiRigData> SMultiRigHierarchyTreeView::GetSelectedData() const
 	return Keys;
 }
 
-TSharedPtr<FMultiRigTreeElement> SMultiRigHierarchyTreeView::FindItemAtPosition(FVector2D InScreenSpacePosition) const
-{
-	if (ItemsPanel.IsValid() && HasValidItemsSource())
-	{
-		const FGeometry MyGeometry = ItemsPanel->GetCachedGeometry();
-		FArrangedChildren ArrangedChildren(EVisibility::Visible);
-		ItemsPanel->ArrangeChildren(MyGeometry, ArrangedChildren, true);
-
-		const int32 Index = ItemsPanel->FindChildUnderPosition(ArrangedChildren, InScreenSpacePosition);
-		if (GetItems().IsValidIndex(Index))
-		{
-			return GetItems()[Index];
-		}
-	}
-	return TSharedPtr<FMultiRigTreeElement>();
-}
-
 TArray<URigHierarchy*> SMultiRigHierarchyTreeView::GetHierarchy() const
 {
 	TArray<URigHierarchy*> RigHierarchy;
diff --git a/Engine/Plugins/Animation/ControlRig/Source/ControlRigEditor/Private/EditMode/SControlRigOutliner.h b/Engine/Plugins/Animation/ControlRig/Source/ControlRigEditor/Private/EditMode/SControlRigOutliner.h
--- a/Engine/Plugins/Animation/ControlRig/Source/ControlRigEditor/Private/EditMode/SControlRigOutliner.h
+++ b/Engine/Plugins/Animation/ControlRig/Source/ControlRigEditor/Private/EditMode/SControlRigOutliner.h
@@ -190,7 +190,6 @@ public:
 	const TArray<TSharedPtr<FMultiRigTreeElement>>& GetRootElements() const { return RootElements; }
 	FMultiRigTreeDelegates& GetTreeDelegates() { return Delegates; }
 
-	TSharedPtr<FMultiRigTreeElement> FindItemAtPosition(FVector2D InScreenSpacePosition) const;
 	TArray<URigHierarchy*> GetHierarchy() const;
 	void SetControlRigs(TArrayView < TWeakObjectPtr<UControlRig>>& InControlRigs);
 
diff --git a/Engine/Plugins/Animation/ControlRig/Source/ControlRigEditor/Private/Editor/SRigHierarchy.cpp b/Engine/Plugins/Animation/ControlRig/Source/ControlRigEditor/Private/Editor/SRigHierarchy.cpp
--- a/Engine/Plugins/Animation/ControlRig/Source/ControlRigEditor/Private/Editor/SRigHierarchy.cpp
+++ b/Engine/Plugins/Animation/ControlRig/Source/ControlRigEditor/Private/Editor/SRigHierarchy.cpp
@@ -42,7 +42,6 @@
 #include "Editor/ControlRigContextMenuContext.h"
 #include "Editor/SRigSpacePickerWidget.h"
 #include "Settings/ControlRigSettings.h"
-#include "Widgets/Views/SListPanel.h"
 #include "Framework/Notifications/NotificationManager.h"
 #include "Widgets/Notifications/SNotificationList.h"
 #include "Styling/AppStyle.h"
@@ -51,6 +50,7 @@
 #include "Algo/MinElement.h"
 #include "Algo/MaxElement.h"
 #include "RigVMFunctions/Math/RigVMMathLibrary.h"
+#include "Preferences/PersonaOptions.h"
 
 #include UE_INLINE_GENERATED_CPP_BY_NAME(SRigHierarchy)
 
@@ -2459,8 +2459,8 @@ bool SRigHierarchy::HandleVerifyNameChanged(const FRigElementKey& OldKey, const
 FReply SRigHierarchy::OnDrop(const FGeometry& MyGeometry, const FDragDropEvent& DragDropEvent)
 {
 	// only allow drops onto empty space of the widget (when there's no target item under the mouse)
-	TSharedPtr<FRigTreeElement> ItemAtMouse = TreeView->FindItemAtPosition(DragDropEvent.GetScreenSpacePosition());
-	if(ItemAtMouse.IsValid())
+	const TSharedPtr<FRigTreeElement>* ItemAtMouse = TreeView->FindItemAtPosition(DragDropEvent.GetScreenSpacePosition());
+	if (ItemAtMouse && ItemAtMouse->IsValid())
 	{
 		return FReply::Unhandled();
 	}
diff --git a/Engine/Plugins/Animation/ControlRig/Source/ControlRigEditor/Private/Editor/SRigHierarchyTreeView.cpp b/Engine/Plugins/Animation/ControlRig/Source/ControlRigEditor/Private/Editor/SRigHierarchyTreeView.cpp
--- a/Engine/Plugins/Animation/ControlRig/Source/ControlRigEditor/Private/Editor/SRigHierarchyTreeView.cpp
+++ b/Engine/Plugins/Animation/ControlRig/Source/ControlRigEditor/Private/Editor/SRigHierarchyTreeView.cpp
@@ -7,7 +7,6 @@
 #include "Widgets/Images/SImage.h"
 #include "Widgets/Input/SSearchBox.h"
 #include "Widgets/Text/STextBlock.h"
-#include "Widgets/Views/SListPanel.h"
 #include "PropertyCustomizationHelpers.h"
 #include "Framework/Application/SlateApplication.h"
 #include "Editor/EditorEngine.h"
@@ -750,36 +749,6 @@ TArray<FRigElementKey> SRigHierarchyTreeView::GetSelectedKeys() const
 	return Keys;
 }
 
-TSharedPtr<FRigTreeElement> SRigHierarchyTreeView::FindItemAtPosition(FVector2D InScreenSpacePosition) const
-{
-	if(ItemsPanel.IsValid() && HasValidItemsSource())
-	{
-		const FGeometry MyGeometry = ItemsPanel->GetCachedGeometry();
-		FArrangedChildren ArrangedChildren(EVisibility::Visible);
-		ItemsPanel->ArrangeChildren(MyGeometry, ArrangedChildren, true);
-
-		const int32 Index = ItemsPanel->FindChildUnderPosition(ArrangedChildren, InScreenSpacePosition); 
-		if(ArrangedChildren.IsValidIndex(Index))
-		{
-			TSharedRef<SRigHierarchyItem> ItemWidget = StaticCastSharedRef<SRigHierarchyItem>(ArrangedChildren[Index].Widget);
-			if(ItemWidget->WeakRigTreeElement.IsValid())
-			{
-				const FRigElementKey Key = ItemWidget->WeakRigTreeElement.Pin()->Key;
-				const TSharedPtr<FRigTreeElement>* ResultPtr = GetItems().FindByPredicate([Key](const TSharedPtr<FRigTreeElement>& Item) -> bool
-				{
-					return Item->Key == Key;
-				});
-				
-				if(ResultPtr)
-				{
-					return *ResultPtr;
-				}
-			}
-		}
-	}
-	return TSharedPtr<FRigTreeElement>();
-}
-
 bool SRigHierarchyItem::OnVerifyNameChanged(const FText& InText, FText& OutErrorMessage)
 {
 	const FString NewName = InText.ToString();
diff --git a/Engine/Plugins/Animation/ControlRig/Source/ControlRigEditor/Private/Editor/SRigHierarchyTreeView.h b/Engine/Plugins/Animation/ControlRig/Source/ControlRigEditor/Private/Editor/SRigHierarchyTreeView.h
--- a/Engine/Plugins/Animation/ControlRig/Source/ControlRigEditor/Private/Editor/SRigHierarchyTreeView.h
+++ b/Engine/Plugins/Animation/ControlRig/Source/ControlRigEditor/Private/Editor/SRigHierarchyTreeView.h
@@ -284,8 +284,6 @@ public:
 	const TArray<TSharedPtr<FRigTreeElement>>& GetRootElements() const { return RootElements; }
 	FRigTreeDelegates& GetRigTreeDelegates() { return Delegates; }
 
-	TSharedPtr<FRigTreeElement> FindItemAtPosition(FVector2D InScreenSpacePosition) const;
-
 private:
 
 	/** A temporary snapshot of the SparseItemInfos in STreeView, used during RefreshTreeView() */
diff --git a/Engine/Source/Runtime/Slate/Private/Widgets/Views/STableViewBase.cpp b/Engine/Source/Runtime/Slate/Private/Widgets/Views/STableViewBase.cpp
--- a/Engine/Source/Runtime/Slate/Private/Widgets/Views/STableViewBase.cpp
+++ b/Engine/Source/Runtime/Slate/Private/Widgets/Views/STableViewBase.cpp
@@ -1024,6 +1024,20 @@ void STableViewBase::NavigateToWidget(const uint32 UserIndex, const TSharedPtr<S
 	FSlateApplication::Get().NavigateToWidget(UserIndex, NavigationDestination, NavigationSource);
 }
 
+int32 STableViewBase::FindChildUnderPosition(const FVector2D& ArrangedSpacePosition) const
+{
+	if (ItemsPanel.IsValid())
+	{
+		const FGeometry MyGeometry = ItemsPanel->GetCachedGeometry();
+		FArrangedChildren ArrangedChildren(EVisibility::Visible);
+		ItemsPanel->ArrangeChildren(MyGeometry, ArrangedChildren, true);
+
+		return ItemsPanel->FindChildUnderPosition(ArrangedChildren, ArrangedSpacePosition);
+	}
+
+	return INDEX_NONE;
+}
+
 void STableViewBase::OnRightMouseButtonUp(const FPointerEvent& MouseEvent)
 {
 	FVector2f SummonLocation = MouseEvent.GetScreenSpacePosition();
diff --git a/Engine/Source/Runtime/Slate/Public/Widgets/Views/SListView.h b/Engine/Source/Runtime/Slate/Public/Widgets/Views/SListView.h
--- a/Engine/Source/Runtime/Slate/Public/Widgets/Views/SListView.h
+++ b/Engine/Source/Runtime/Slate/Public/Widgets/Views/SListView.h
@@ -1717,6 +1717,20 @@ public:
 		return Private_ItemFromWidget( WidgetToFind );
 	}
 
+	/** Given a position, return the item under that position. If nothing is there, return null. */
+	const ItemType* FindItemAtPosition(FVector2D InScreenSpacePosition) const
+	{
+		if (this->ItemsSource)
+		{
+			const int32 Index = this->FindChildUnderPosition(InScreenSpacePosition);
+			if (this->ItemsSource->IsValidIndex(Index))
+			{
+				return &this->ItemsSource->operator[](Index);
+			}
+		}
+		return nullptr;
+	}
+
 	/**
 	 * Test if the current item is selected.
 	 *
diff --git a/Engine/Source/Runtime/Slate/Public/Widgets/Views/STableViewBase.h b/Engine/Source/Runtime/Slate/Public/Widgets/Views/STableViewBase.h
--- a/Engine/Source/Runtime/Slate/Public/Widgets/Views/STableViewBase.h
+++ b/Engine/Source/Runtime/Slate/Public/Widgets/Views/STableViewBase.h
@@ -383,6 +383,9 @@ protected:
 	/** Util Function so templates classes don't need to include SlateApplication */
 	void NavigateToWidget(const uint32 UserIndex, const TSharedPtr<SWidget>& NavigationDestination, ENavigationSource NavigationSource = ENavigationSource::FocusedWidget) const;
 
+	/** Util function to find the child index under the given position. */
+	int32 FindChildUnderPosition(const FVector2D& ArrangedSpacePosition) const;
+
 	/** The panel which holds the visible widgets in this list */
 	TSharedPtr< SListPanel > ItemsPanel;
 
