Reverting some parts and fixing issues from moving FindItemAtPosition around

diff --git a/Engine/Plugins/Animation/ControlRig/Source/ControlRigEditor/Private/Editor/SRigHierarchyTreeView.cpp b/Engine/Plugins/Animation/ControlRig/Source/ControlRigEditor/Private/Editor/SRigHierarchyTreeView.cpp
--- a/Engine/Plugins/Animation/ControlRig/Source/ControlRigEditor/Private/Editor/SRigHierarchyTreeView.cpp
+++ b/Engine/Plugins/Animation/ControlRig/Source/ControlRigEditor/Private/Editor/SRigHierarchyTreeView.cpp
@@ -277,12 +277,12 @@ void SRigHierarchyTreeView::Tick(const FGeometry& AllottedGeometry, const double
 			}
 			else
 			{
-				const TSharedPtr<FRigTreeElement> Item = FindItemAtPosition(MousePosition);
-				if(Item.IsValid())
+				const TSharedPtr<FRigTreeElement>* Item = FindItemAtPosition(MousePosition);
+				if(Item && Item->IsValid())
 				{
-					if(!IsItemExpanded(Item))
+					if(!IsItemExpanded(*Item))
 					{
-						SetItemExpansion(Item, true);
+						SetItemExpansion(*Item, true);
 					}
 				}
 			}
@@ -749,6 +749,33 @@ TArray<FRigElementKey> SRigHierarchyTreeView::GetSelectedKeys() const
 	return Keys;
 }
 
+const TSharedPtr<FRigTreeElement>* SRigHierarchyTreeView::FindItemAtPosition(FVector2D InScreenSpacePosition) const
+{
+	if (ItemsPanel.IsValid() && HasValidItemsSource())
+	{
+		FArrangedChildren ArrangedChildren(EVisibility::Visible);
+		const int32 Index = FindChildUnderPosition(ArrangedChildren, InScreenSpacePosition);
+		if (ArrangedChildren.IsValidIndex(Index))
+		{
+			TSharedRef<SRigHierarchyItem> ItemWidget = StaticCastSharedRef<SRigHierarchyItem>(ArrangedChildren[Index].Widget);
+			if (ItemWidget->WeakRigTreeElement.IsValid())
+			{
+				const FRigElementKey Key = ItemWidget->WeakRigTreeElement.Pin()->Key;
+				const TSharedPtr<FRigTreeElement>* ResultPtr = GetItems().FindByPredicate([Key](const TSharedPtr<FRigTreeElement>& Item) -> bool
+					{
+						return Item->Key == Key;
+					});
+
+				if (ResultPtr)
+				{
+					return ResultPtr;
+				}
+			}
+		}
+	}
+	return nullptr;
+}
+
 bool SRigHierarchyItem::OnVerifyNameChanged(const FText& InText, FText& OutErrorMessage)
 {
 	const FString NewName = InText.ToString();
diff --git a/Engine/Plugins/Animation/ControlRig/Source/ControlRigEditor/Private/Editor/SRigHierarchyTreeView.h b/Engine/Plugins/Animation/ControlRig/Source/ControlRigEditor/Private/Editor/SRigHierarchyTreeView.h
--- a/Engine/Plugins/Animation/ControlRig/Source/ControlRigEditor/Private/Editor/SRigHierarchyTreeView.h
+++ b/Engine/Plugins/Animation/ControlRig/Source/ControlRigEditor/Private/Editor/SRigHierarchyTreeView.h
@@ -284,6 +284,9 @@ public:
 	const TArray<TSharedPtr<FRigTreeElement>>& GetRootElements() const { return RootElements; }
 	FRigTreeDelegates& GetRigTreeDelegates() { return Delegates; }
 
+	/** Given a position, return the item under that position. If nothing is there, return null. */
+	const TSharedPtr<FRigTreeElement>* FindItemAtPosition(FVector2D InScreenSpacePosition) const;
+
 private:
 
 	/** A temporary snapshot of the SparseItemInfos in STreeView, used during RefreshTreeView() */
diff --git a/Engine/Source/Runtime/Slate/Private/Widgets/Views/STableViewBase.cpp b/Engine/Source/Runtime/Slate/Private/Widgets/Views/STableViewBase.cpp
--- a/Engine/Source/Runtime/Slate/Private/Widgets/Views/STableViewBase.cpp
+++ b/Engine/Source/Runtime/Slate/Private/Widgets/Views/STableViewBase.cpp
@@ -1024,12 +1024,11 @@ void STableViewBase::NavigateToWidget(const uint32 UserIndex, const TSharedPtr<S
 	FSlateApplication::Get().NavigateToWidget(UserIndex, NavigationDestination, NavigationSource);
 }
 
-int32 STableViewBase::FindChildUnderPosition(const FVector2D& ArrangedSpacePosition) const
+int32 STableViewBase::FindChildUnderPosition(FArrangedChildren& ArrangedChildren, const FVector2D& ArrangedSpacePosition) const
 {
 	if (ItemsPanel.IsValid())
 	{
 		const FGeometry MyGeometry = ItemsPanel->GetCachedGeometry();
-		FArrangedChildren ArrangedChildren(EVisibility::Visible);
 		ItemsPanel->ArrangeChildren(MyGeometry, ArrangedChildren, true);
 
 		return ItemsPanel->FindChildUnderPosition(ArrangedChildren, ArrangedSpacePosition);
diff --git a/Engine/Source/Runtime/Slate/Public/Widgets/Views/SListView.h b/Engine/Source/Runtime/Slate/Public/Widgets/Views/SListView.h
--- a/Engine/Source/Runtime/Slate/Public/Widgets/Views/SListView.h
+++ b/Engine/Source/Runtime/Slate/Public/Widgets/Views/SListView.h
@@ -1717,20 +1717,6 @@ public:
 		return Private_ItemFromWidget( WidgetToFind );
 	}
 
-	/** Given a position, return the item under that position. If nothing is there, return null. */
-	const ItemType* FindItemAtPosition(FVector2D InScreenSpacePosition) const
-	{
-		if (this->ItemsSource)
-		{
-			const int32 Index = this->FindChildUnderPosition(InScreenSpacePosition);
-			if (this->ItemsSource->IsValidIndex(Index))
-			{
-				return &this->ItemsSource->operator[](Index);
-			}
-		}
-		return nullptr;
-	}
-
 	/**
 	 * Test if the current item is selected.
 	 *
diff --git a/Engine/Source/Runtime/Slate/Public/Widgets/Views/STableViewBase.h b/Engine/Source/Runtime/Slate/Public/Widgets/Views/STableViewBase.h
--- a/Engine/Source/Runtime/Slate/Public/Widgets/Views/STableViewBase.h
+++ b/Engine/Source/Runtime/Slate/Public/Widgets/Views/STableViewBase.h
@@ -384,7 +384,7 @@ protected:
 	void NavigateToWidget(const uint32 UserIndex, const TSharedPtr<SWidget>& NavigationDestination, ENavigationSource NavigationSource = ENavigationSource::FocusedWidget) const;
 
 	/** Util function to find the child index under the given position. */
-	int32 FindChildUnderPosition(const FVector2D& ArrangedSpacePosition) const;
+	int32 FindChildUnderPosition(FArrangedChildren& ArrangedChildren, const FVector2D& ArrangedSpacePosition) const;
 
 	/** The panel which holds the visible widgets in this list */
 	TSharedPtr< SListPanel > ItemsPanel;
