Control Rig: Fix a out-of-bounds memory access error in the Harmonics node.

#rnx

diff --git a/Engine/Plugins/Animation/ControlRig/Source/ControlRig/Private/Units/Highlevel/Harmonics/RigUnit_BoneHarmonics.cpp b/Engine/Plugins/Animation/ControlRig/Source/ControlRig/Private/Units/Highlevel/Harmonics/RigUnit_BoneHarmonics.cpp
--- a/Engine/Plugins/Animation/ControlRig/Source/ControlRig/Private/Units/Highlevel/Harmonics/RigUnit_BoneHarmonics.cpp
+++ b/Engine/Plugins/Animation/ControlRig/Source/ControlRig/Private/Units/Highlevel/Harmonics/RigUnit_BoneHarmonics.cpp
@@ -80,12 +80,11 @@ FRigUnit_ItemHarmonics_Execute()
 
 	if (CachedItems.Num() != Targets.Num())
 	{
-		CachedItems.Reset();
+		CachedItems.Reset(Targets.Num());
+		
 		WaveTime = FVector::ZeroVector;
-	}
 
 	if (CachedItems.IsEmpty())
-	{
 		for (int32 ItemIndex = 0; ItemIndex < Targets.Num(); ItemIndex++)
 		{
 			FCachedRigElement CachedItem(Targets[ItemIndex].Item, Hierarchy);
