Control Rig: Fix warning when deprecated unit structs have template names

diff --git a/Engine/Plugins/Animation/ControlRig/Source/ControlRig/Private/Units/Debug/RigUnit_VisualDebug.h b/Engine/Plugins/Animation/ControlRig/Source/ControlRig/Private/Units/Debug/RigUnit_VisualDebug.h
--- a/Engine/Plugins/Animation/ControlRig/Source/ControlRig/Private/Units/Debug/RigUnit_VisualDebug.h
+++ b/Engine/Plugins/Animation/ControlRig/Source/ControlRig/Private/Units/Debug/RigUnit_VisualDebug.h
@@ -6,7 +6,7 @@
 #include "RigVMFunctions/Debug/RigVMFunction_VisualDebug.h"
 #include "RigUnit_VisualDebug.generated.h"
 
-USTRUCT(meta=(DisplayName = "Visual Debug Vector", TemplateName="VisualDebug", Keywords = "Draw,Point", Deprecated = "4.25", Varying))
+USTRUCT(meta=(DisplayName = "Visual Debug Vector", Keywords = "Draw,Point", Deprecated = "4.25", Varying))
 struct CONTROLRIG_API FRigUnit_VisualDebugVector : public FRigUnit_DebugBase
 {
 	GENERATED_BODY()
@@ -53,7 +53,7 @@ struct CONTROLRIG_API FRigUnit_VisualDebugVector : public FRigUnit_DebugBase
 /**
  * Debug draw parameters for a Point or Vector given a vector
  */
-USTRUCT(meta=(DisplayName = "Visual Debug Vector", TemplateName="VisualDebug", Keywords = "Draw,Point", Deprecated = "5.2", Varying))
+USTRUCT(meta=(DisplayName = "Visual Debug Vector", Keywords = "Draw,Point", Deprecated = "5.2", Varying))
 struct CONTROLRIG_API FRigUnit_VisualDebugVectorItemSpace : public FRigUnit_DebugBase
 {
 	GENERATED_BODY()
@@ -102,7 +102,7 @@ struct CONTROLRIG_API FRigUnit_VisualDebugVectorItemSpace : public FRigUnit_Debu
 	virtual FRigVMStructUpgradeInfo GetUpgradeInfo() const override;
 };
 
-USTRUCT(meta = (DisplayName = "Visual Debug Quat", TemplateName="VisualDebug", Keywords = "Draw,Rotation", Deprecated = "4.25", Varying))
+USTRUCT(meta = (DisplayName = "Visual Debug Quat", Keywords = "Draw,Rotation", Deprecated = "4.25", Varying))
 struct CONTROLRIG_API FRigUnit_VisualDebugQuat : public FRigUnit_DebugBase
 {
 	GENERATED_BODY()
@@ -141,7 +141,7 @@ struct CONTROLRIG_API FRigUnit_VisualDebugQuat : public FRigUnit_DebugBase
 /**
  * Debug draw parameters for an Axis given a quaternion
  */
-USTRUCT(meta = (DisplayName = "Visual Debug Quat", TemplateName="VisualDebug", Keywords = "Draw,Rotation", Deprecated = "5.2", Varying))
+USTRUCT(meta = (DisplayName = "Visual Debug Quat", Keywords = "Draw,Rotation", Deprecated = "5.2", Varying))
 struct CONTROLRIG_API FRigUnit_VisualDebugQuatItemSpace : public FRigUnit_DebugBase
 {
 	GENERATED_BODY()
@@ -182,7 +182,7 @@ struct CONTROLRIG_API FRigUnit_VisualDebugQuatItemSpace : public FRigUnit_DebugB
 	virtual FRigVMStructUpgradeInfo GetUpgradeInfo() const override;
 };
 
-USTRUCT(meta=(DisplayName="Visual Debug Transform", TemplateName="VisualDebug", Keywords = "Draw,Axes", Deprecated = "4.25", Varying))
+USTRUCT(meta=(DisplayName="Visual Debug Transform", Keywords = "Draw,Axes", Deprecated = "4.25", Varying))
 struct CONTROLRIG_API FRigUnit_VisualDebugTransform : public FRigUnit_DebugBase
 {
 	GENERATED_BODY()
@@ -221,7 +221,7 @@ struct CONTROLRIG_API FRigUnit_VisualDebugTransform : public FRigUnit_DebugBase
 /**
  * Debug draw parameters for an Axis given a transform
  */
-USTRUCT(meta=(DisplayName="Visual Debug Transform", TemplateName="VisualDebug", Keywords = "Draw,Axes", Deprecated = "5.2", Varying))
+USTRUCT(meta=(DisplayName="Visual Debug Transform", Keywords = "Draw,Axes", Deprecated = "5.2", Varying))
 struct CONTROLRIG_API FRigUnit_VisualDebugTransformItemSpace : public FRigUnit_DebugBase
 {
 	GENERATED_BODY()
diff --git a/Engine/Plugins/Animation/ControlRig/Source/ControlRig/Private/Units/Hierarchy/RigUnit_SetControlTransform.h b/Engine/Plugins/Animation/ControlRig/Source/ControlRig/Private/Units/Hierarchy/RigUnit_SetControlTransform.h
--- a/Engine/Plugins/Animation/ControlRig/Source/ControlRig/Private/Units/Hierarchy/RigUnit_SetControlTransform.h
+++ b/Engine/Plugins/Animation/ControlRig/Source/ControlRig/Private/Units/Hierarchy/RigUnit_SetControlTransform.h
@@ -537,7 +537,7 @@ struct CONTROLRIG_API FRigUnit_SetMultiControlRotator : public FRigUnitMutable
 /**
  * SetControlTransform is used to perform a change in the hierarchy by setting a single control's transform.
  */
-USTRUCT(meta=(DisplayName="Set Control Transform", Category="Controls", DocumentationPolicy="Strict", Keywords = "SetControlTransform,SetGizmoTransform", TemplateName="SetControlValue", Deprecated = "4.25"))
+USTRUCT(meta=(DisplayName="Set Control Transform", Category="Controls", DocumentationPolicy="Strict", Keywords = "SetControlTransform,SetGizmoTransform", Deprecated = "4.25"))
 struct CONTROLRIG_API FRigUnit_SetControlTransform : public FRigUnitMutable
 {
 	GENERATED_BODY()
diff --git a/Engine/Source/Runtime/RigVM/Public/RigVMFunctions/Debug/RigVMFunction_VisualDebug.h b/Engine/Source/Runtime/RigVM/Public/RigVMFunctions/Debug/RigVMFunction_VisualDebug.h
--- a/Engine/Source/Runtime/RigVM/Public/RigVMFunctions/Debug/RigVMFunction_VisualDebug.h
+++ b/Engine/Source/Runtime/RigVM/Public/RigVMFunctions/Debug/RigVMFunction_VisualDebug.h
@@ -18,7 +18,7 @@ enum class ERigUnitVisualDebugPointMode : uint8
 	Max UMETA(Hidden),
 };
 
-USTRUCT(meta=(DisplayName = "Visual Debug Vector", TemplateName="VisualDebug", Keywords = "Draw,Point", Deprecated = "4.25", Varying))
+USTRUCT(meta=(DisplayName = "Visual Debug Vector", Keywords = "Draw,Point", Deprecated = "4.25", Varying))
 struct RIGVM_API FRigVMFunction_VisualDebugVector : public FRigVMFunction_DebugBase
 {
 	GENERATED_BODY()
@@ -102,7 +102,7 @@ struct RIGVM_API FRigVMFunction_VisualDebugVectorNoSpace : public FRigVMFunction
 	float Scale;
 };
 
-USTRUCT(meta = (DisplayName = "Visual Debug Quat", TemplateName="VisualDebug", Keywords = "Draw,Rotation", Deprecated = "4.25", Varying))
+USTRUCT(meta = (DisplayName = "Visual Debug Quat", Keywords = "Draw,Rotation", Deprecated = "4.25", Varying))
 struct RIGVM_API FRigVMFunction_VisualDebugQuat : public FRigVMFunction_DebugBase
 {
 	GENERATED_BODY()
@@ -170,7 +170,7 @@ struct RIGVM_API FRigVMFunction_VisualDebugQuatNoSpace : public FRigVMFunction_D
 	float Scale;
 };
 
-USTRUCT(meta=(DisplayName="Visual Debug Transform", TemplateName="VisualDebug", Keywords = "Draw,Axes", Deprecated = "4.25", Varying))
+USTRUCT(meta=(DisplayName="Visual Debug Transform", Keywords = "Draw,Axes", Deprecated = "4.25", Varying))
 struct RIGVM_API FRigVMFunction_VisualDebugTransform : public FRigVMFunction_DebugBase
 {
 	GENERATED_BODY()
diff --git a/Engine/Source/Runtime/RigVM/Public/RigVMFunctions/Math/RigVMFunction_MathFloat.h b/Engine/Source/Runtime/RigVM/Public/RigVMFunctions/Math/RigVMFunction_MathFloat.h
--- a/Engine/Source/Runtime/RigVM/Public/RigVMFunctions/Math/RigVMFunction_MathFloat.h
+++ b/Engine/Source/Runtime/RigVM/Public/RigVMFunctions/Math/RigVMFunction_MathFloat.h
@@ -748,7 +748,7 @@ struct RIGVM_API FRigVMFunction_MathFloatIsNearlyEqual : public FRigVMFunction_M
 /**
  * Return one of the two values based on the condition
  */
-USTRUCT(meta=(DisplayName="Select", TemplateName="Select", Keywords="Pick,If", Deprecated = "4.26.0"))
+USTRUCT(meta=(DisplayName="Select", Keywords="Pick,If", Deprecated = "4.26.0"))
 struct RIGVM_API FRigVMFunction_MathFloatSelectBool : public FRigVMFunction_MathFloatBase
 {
 	GENERATED_BODY()
diff --git a/Engine/Source/Runtime/RigVM/Public/RigVMFunctions/Math/RigVMFunction_MathQuaternion.h b/Engine/Source/Runtime/RigVM/Public/RigVMFunctions/Math/RigVMFunction_MathQuaternion.h
--- a/Engine/Source/Runtime/RigVM/Public/RigVMFunctions/Math/RigVMFunction_MathQuaternion.h
+++ b/Engine/Source/Runtime/RigVM/Public/RigVMFunctions/Math/RigVMFunction_MathQuaternion.h
@@ -484,7 +484,7 @@ struct RIGVM_API FRigVMFunction_MathQuaternionNotEquals : public FRigVMFunction_
 /**
  * Return one of the two values based on the condition
  */
-USTRUCT(meta=(DisplayName="Select", TemplateName="Select", Keywords="Pick,If", Deprecated = "4.26.0"))
+USTRUCT(meta=(DisplayName="Select", Keywords="Pick,If", Deprecated = "4.26.0"))
 struct RIGVM_API FRigVMFunction_MathQuaternionSelectBool : public FRigVMFunction_MathQuaternionBase
 {
 	GENERATED_BODY()
diff --git a/Engine/Source/Runtime/RigVM/Public/RigVMFunctions/Math/RigVMFunction_MathTransform.h b/Engine/Source/Runtime/RigVM/Public/RigVMFunctions/Math/RigVMFunction_MathTransform.h
--- a/Engine/Source/Runtime/RigVM/Public/RigVMFunctions/Math/RigVMFunction_MathTransform.h
+++ b/Engine/Source/Runtime/RigVM/Public/RigVMFunctions/Math/RigVMFunction_MathTransform.h
@@ -332,7 +332,7 @@ struct RIGVM_API FRigVMFunction_MathTransformLerp : public FRigVMFunction_MathTr
 /**
  * Return one of the two values based on the condition
  */
-USTRUCT(meta=(DisplayName="Select", TemplateName="Select", Keywords="Pick,If", Deprecated = "4.26.0"))
+USTRUCT(meta=(DisplayName="Select", Keywords="Pick,If", Deprecated = "4.26.0"))
 struct RIGVM_API FRigVMFunction_MathTransformSelectBool : public FRigVMFunction_MathTransformBase
 {
 	GENERATED_BODY()
diff --git a/Engine/Source/Runtime/RigVM/Public/RigVMFunctions/Math/RigVMFunction_MathVector.h b/Engine/Source/Runtime/RigVM/Public/RigVMFunctions/Math/RigVMFunction_MathVector.h
--- a/Engine/Source/Runtime/RigVM/Public/RigVMFunctions/Math/RigVMFunction_MathVector.h
+++ b/Engine/Source/Runtime/RigVM/Public/RigVMFunctions/Math/RigVMFunction_MathVector.h
@@ -567,7 +567,7 @@ struct RIGVM_API FRigVMFunction_MathVectorIsNearlyEqual : public FRigVMFunction_
 /**
  * Return one of the two values based on the condition
  */
-USTRUCT(meta=(DisplayName="Select", TemplateName="Select", Keywords="Pick,If", Deprecated = "4.26.0"))
+USTRUCT(meta=(DisplayName="Select", Keywords="Pick,If", Deprecated = "4.26.0"))
 struct RIGVM_API FRigVMFunction_MathVectorSelectBool : public FRigVMFunction_MathVectorBase
 {
 	GENERATED_BODY()
diff --git a/Engine/Source/Runtime/RigVM/Public/RigVMFunctions/Math/RigVMFunction_Noise.h b/Engine/Source/Runtime/RigVM/Public/RigVMFunctions/Math/RigVMFunction_Noise.h
--- a/Engine/Source/Runtime/RigVM/Public/RigVMFunctions/Math/RigVMFunction_Noise.h
+++ b/Engine/Source/Runtime/RigVM/Public/RigVMFunctions/Math/RigVMFunction_Noise.h
@@ -88,7 +88,7 @@ struct RIGVM_API FRigVMFunction_NoiseDouble : public FRigVMFunction_MathBase
 /**
  * Generates a vector through a noise fluctuation process between a min and a max through speed
  */
-USTRUCT(meta = (DisplayName = "Noise (Vector)", Category = "Math|Noise", TemplateName="Noise", Deprecated = "5.0.0"))
+USTRUCT(meta = (DisplayName = "Noise (Vector)", Category = "Math|Noise", Deprecated = "5.0.0"))
 struct RIGVM_API FRigVMFunction_NoiseVector : public FRigVMFunction_MathBase
 {
 	GENERATED_BODY()
