Control Rig: More surgical approach to ControlRigGraphNode::ReconstructNode_Internal

diff --git a/Engine/Plugins/Animation/ControlRig/Source/ControlRigDeveloper/Private/Graph/ControlRigGraph.cpp b/Engine/Plugins/Animation/ControlRig/Source/ControlRigDeveloper/Private/Graph/ControlRigGraph.cpp
--- a/Engine/Plugins/Animation/ControlRig/Source/ControlRigDeveloper/Private/Graph/ControlRigGraph.cpp
+++ b/Engine/Plugins/Animation/ControlRig/Source/ControlRigDeveloper/Private/Graph/ControlRigGraph.cpp
@@ -1154,6 +1154,13 @@ void UControlRigGraph::RemoveNode(UEdGraphNode* InNode)
 {
 	// Make sure EdGraph is not part of the transaction
 	TGuardValue<ITransaction*> TransactionGuard(GUndo, nullptr);
+
+	// clear out the pin relationships
+	for(UEdGraphPin* Pin : InNode->Pins)
+	{
+		Pin->MarkAsGarbage();
+	}
+	InNode->Pins.Reset();
 					
 	// Rename the soon to be deleted object to a unique name, so that other objects can use
 	// the old name
@@ -1163,12 +1170,12 @@ void UControlRigGraph::RemoveNode(UEdGraphNode* InNode)
 		static int32 DeletedIndex = FMath::Rand();
 		do
 		{
-			DeletedName = FString::Printf(TEXT("%s_Deleted_%d"), *InNode->GetName(), DeletedIndex++); 
+			DeletedName = FString::Printf(TEXT("ControlRigGraph_%s_Deleted_%d"), *InNode->GetName(), DeletedIndex++); 
 			ExistingObject = StaticFindObject(/*Class=*/ NULL, this, *DeletedName, true);						
 		}
 		while (ExistingObject);
 	}
-	InNode->Rename(*DeletedName, nullptr, REN_ForceNoResetLoaders | REN_DontCreateRedirectors);	
+	InNode->Rename(*DeletedName, GetTransientPackage(), REN_ForceNoResetLoaders | REN_DontCreateRedirectors);	
 	Super::RemoveNode(InNode);
 
 	NotifyGraphChanged(FEdGraphEditAction(EEdGraphActionType::GRAPHACTION_RemoveNode, this, InNode, false));
diff --git a/Engine/Plugins/Animation/ControlRig/Source/ControlRigDeveloper/Private/Graph/ControlRigGraphNode.cpp b/Engine/Plugins/Animation/ControlRig/Source/ControlRigDeveloper/Private/Graph/ControlRigGraphNode.cpp
--- a/Engine/Plugins/Animation/ControlRig/Source/ControlRigDeveloper/Private/Graph/ControlRigGraphNode.cpp
+++ b/Engine/Plugins/Animation/ControlRig/Source/ControlRigDeveloper/Private/Graph/ControlRigGraphNode.cpp
@@ -235,12 +235,21 @@ void UControlRigGraphNode::ReconstructNode_Internal(bool bForce)
 	// Clear previously set messages
 	ErrorMsg.Reset();
 
-	// Move the existing pins to a saved array
-	TArray<UEdGraphPin*> OldPins(Pins);
+	// Move the existing pins to a saved array.
+	// This way we can reuse them later
+	LastEdGraphPins = Pins;
+
+	// Also make sure to clean up any indirections
+	// on them.
+	for(UEdGraphPin* LastEdGraphPin : LastEdGraphPins)
+	{
+		LastEdGraphPin->ParentPin = nullptr;
+		LastEdGraphPin->SubPins.Reset();
+	}
 	Pins.Reset();
 
 	// Recreate the new pins
-	ReallocatePinsDuringReconstruction(OldPins);
+	ReallocatePinsDuringReconstruction(LastEdGraphPins);
 
 	// Maintain watches up to date
 	if (URigVMNode* Node = GetModelNode())
@@ -261,7 +270,10 @@ void UControlRigGraphNode::ReconstructNode_Internal(bool bForce)
 		}
 	}
 	
-	RewireOldPinsToNewPins(OldPins, Pins);
+	// LogBlueprint: Warning: UEdGraphPin::DestroyImpl BasicIKSetup_0.UpperBone? causing an issue?
+	// we may need to clean these pins in the opposite order? child pins first?
+	RewireOldPinsToNewPins(LastEdGraphPins, Pins);
+	LastEdGraphPins.Reset();
 
 	DrawAsCompactNodeCache.Reset();
 
@@ -688,7 +700,33 @@ bool UControlRigGraphNode::CreateGraphPinFromModelPin(const URigVMPin* InModelPi
 
 	auto CreatePinLambda = [this](const URigVMPin* InModelPin, EEdGraphPinDirection InDirection, UEdGraphPin* InParentPin) -> UEdGraphPin*
 	{
-		UEdGraphPin* GraphPin = CreatePin(InDirection, GetPinTypeForModelPin(InModelPin), FName(*InModelPin->GetPinPath()));
+		// check if we already have a pin like this in the last pins arrays
+		const FString PinPath = InModelPin->GetPinPath();
+		const FEdGraphPinType PinType = GetPinTypeForModelPin(InModelPin);
+
+		UEdGraphPin* GraphPin = nullptr;
+		UEdGraphPin** ExistingEdGraphPinPtr = LastEdGraphPins.FindByPredicate([PinPath, InDirection, PinType](const UEdGraphPin* ExistingPin) -> bool
+		{
+			return !ExistingPin->bWasTrashed &&
+				ExistingPin->GetName() == PinPath &&
+				ExistingPin->Direction == InDirection &&
+				ExistingPin->PinType == PinType;
+		});
+		if(ExistingEdGraphPinPtr)
+		{
+			GraphPin = *ExistingEdGraphPinPtr;
+			GraphPin->ParentPin = nullptr;
+			GraphPin->SubPins.Reset();
+			ExistingEdGraphPinPtr = nullptr;
+			LastEdGraphPins.Remove(GraphPin);
+			Pins.Add(GraphPin);
+		}
+
+		if(GraphPin == nullptr)
+		{
+			GraphPin = CreatePin(InDirection, GetPinTypeForModelPin(InModelPin), FName(*InModelPin->GetPinPath()));
+		}
+		
 		if (GraphPin)
 		{
 			ConfigurePin(GraphPin, InModelPin);
@@ -772,6 +810,8 @@ void UControlRigGraphNode::RemoveGraphSubPins(UEdGraphPin* InParentPin, const TA
 				{
 					RemoveGraphSubPins(SubPin);
 				}
+
+				SubPin->MarkAsGarbage();
 			};
 			
 			Traverse(PinPair->InputPin);
@@ -1033,6 +1073,7 @@ bool UControlRigGraphNode::ModelPinRemoved_Internal(const URigVMPin* InModelPin)
 		{
 			RemoveGraphSubPins(GraphPin);
 			Pins.Remove(GraphPin);
+			GraphPin->MarkAsGarbage();
 		}
 	};
 	RemoveGraphPin(InModelPin, true);
diff --git a/Engine/Plugins/Animation/ControlRig/Source/ControlRigDeveloper/Public/Graph/ControlRigGraphNode.h b/Engine/Plugins/Animation/ControlRig/Source/ControlRigDeveloper/Public/Graph/ControlRigGraphNode.h
--- a/Engine/Plugins/Animation/ControlRig/Source/ControlRigDeveloper/Public/Graph/ControlRigGraphNode.h
+++ b/Engine/Plugins/Animation/ControlRig/Source/ControlRigDeveloper/Public/Graph/ControlRigGraphNode.h
@@ -261,6 +261,7 @@ private:
 	TArray<URigVMPin*> InputPins;
 	TArray<URigVMPin*> OutputPins;
 	TArray<TSharedPtr<FRigVMExternalVariable>> ExternalVariables;
+	TArray<UEdGraphPin*> LastEdGraphPins;
 	
 	TMap<URigVMPin*, FPinPair> CachedPins;
 
diff --git a/Engine/Source/Runtime/Engine/Private/EdGraph/EdGraphPin.cpp b/Engine/Source/Runtime/Engine/Private/EdGraph/EdGraphPin.cpp
--- a/Engine/Source/Runtime/Engine/Private/EdGraph/EdGraphPin.cpp
+++ b/Engine/Source/Runtime/Engine/Private/EdGraph/EdGraphPin.cpp
@@ -1645,7 +1645,7 @@ void UEdGraphPin::DestroyImpl(bool bClearLinks)
 	for (int32 SubPinIndex = SubPins.Num() - 1; SubPinIndex >= 0; --SubPinIndex)
 	{
 		UEdGraphPin* SubPin = SubPins[SubPinIndex];
-		if (!SubPin->bWasTrashed)
+		if (SubPin && !SubPin->bWasTrashed)
 		{
 			SubPins[SubPinIndex]->DestroyImpl(bClearLinks);
 		}
