Minor dependency cleanup inside ControlRigEditor.
Adding IPersonaPreviewScene::UnregisterForUndo so that ControlRigEditor can opt it out of undos.

diff --git a/Engine/Plugins/Animation/ControlRig/Source/ControlRigEditor/Private/ControlRigBlueprintActions.cpp b/Engine/Plugins/Animation/ControlRig/Source/ControlRigEditor/Private/ControlRigBlueprintActions.cpp
--- a/Engine/Plugins/Animation/ControlRig/Source/ControlRigEditor/Private/ControlRigBlueprintActions.cpp
+++ b/Engine/Plugins/Animation/ControlRig/Source/ControlRigEditor/Private/ControlRigBlueprintActions.cpp
@@ -14,7 +14,6 @@
 
 #include "ToolMenus.h"
 #include "ContentBrowserMenuContexts.h"
-#include "AssetTools.h"
 #include "Engine/SkeletalMesh.h"
 #include "Animation/SkeletalMeshActor.h"
 #include "Components/SkeletalMeshComponent.h"
diff --git a/Engine/Plugins/Animation/ControlRig/Source/ControlRigEditor/Private/ControlRigGraphDetails.h b/Engine/Plugins/Animation/ControlRig/Source/ControlRigEditor/Private/ControlRigGraphDetails.h
--- a/Engine/Plugins/Animation/ControlRig/Source/ControlRigEditor/Private/ControlRigGraphDetails.h
+++ b/Engine/Plugins/Animation/ControlRig/Source/ControlRigEditor/Private/ControlRigGraphDetails.h
@@ -16,7 +16,6 @@
 #include "IDetailPropertyExtensionHandler.h"
 #include "Graph/ControlRigGraphSchema.h"
 #include "SAdvancedTransformInputBox.h"
-#include "PropertyNode.h"
 
 class IDetailLayoutBuilder;
 
diff --git a/Engine/Plugins/Animation/ControlRig/Source/ControlRigEditor/Private/Editor/ControlRigEditor.cpp b/Engine/Plugins/Animation/ControlRig/Source/ControlRigEditor/Private/Editor/ControlRigEditor.cpp
--- a/Engine/Plugins/Animation/ControlRig/Source/ControlRigEditor/Private/Editor/ControlRigEditor.cpp
+++ b/Engine/Plugins/Animation/ControlRig/Source/ControlRigEditor/Private/Editor/ControlRigEditor.cpp
@@ -26,7 +26,6 @@
 #include "Sequencer/ControlRigLayerInstance.h"
 #include "Animation/DebugSkelMeshComponent.h"
 #include "IPersonaPreviewScene.h"
-#include "AnimationEditorPreviewScene.h"
 #include "Animation/AnimData/BoneMaskFilter.h"
 #include "ControlRig.h"
 #include "Editor/ControlRigSkeletalMeshComponent.h"
@@ -105,6 +104,7 @@
 #include "MaterialDomain.h"
 #include "RigVMFunctions/RigVMFunction_ControlFlow.h"
 #include "RigVMModel/Nodes/RigVMAggregateNode.h"
+#include "AnimationEditorViewportClient.h"
 
 #define LOCTEXT_NAMESPACE "ControlRigEditor"
 
@@ -4362,15 +4362,8 @@ void FControlRigEditor::HandlePreviewSceneCreated(const TSharedRef<IPersonaPrevi
 		PreviewInstance = Cast<UAnimPreviewInstance>(EditorSkelComp->GetAnimInstance());
 	}
 
-	if (GEditor)
-	{
-		// remove the preview scene undo handling - it has unwanted side effects
-		FAnimationEditorPreviewScene* AnimationEditorPreviewScene = static_cast<FAnimationEditorPreviewScene*>(&InPersonaPreviewScene.Get());
-		if (AnimationEditorPreviewScene)
-		{
-			GEditor->UnregisterForUndo(AnimationEditorPreviewScene);
-		}
-	}
+	// remove the preview scene undo handling - it has unwanted side effects
+	InPersonaPreviewScene->UnregisterForUndo();
 }
 
 void FControlRigEditor::UpdateControlRig()
diff --git a/Engine/Plugins/Animation/ControlRig/Source/ControlRigEditor/Private/Editor/ControlRigEditorEditMode.cpp b/Engine/Plugins/Animation/ControlRig/Source/ControlRigEditor/Private/Editor/ControlRigEditorEditMode.cpp
--- a/Engine/Plugins/Animation/ControlRig/Source/ControlRigEditor/Private/Editor/ControlRigEditorEditMode.cpp
+++ b/Engine/Plugins/Animation/ControlRig/Source/ControlRigEditor/Private/Editor/ControlRigEditorEditMode.cpp
@@ -1,11 +1,13 @@
 // Copyright Epic Games, Inc. All Rights Reserved.
 
 #include "Editor/ControlRigEditorEditMode.h"
+#include "AnimationEditorViewportClient.h"
 #include "IPersonaPreviewScene.h"
 #include "AssetEditorModeManager.h"
 #include "ControlRig.h"
 #include "ControlRigBlueprint.h"
 #include "ControlRigGizmoActor.h"
+#include "Preferences/PersonaOptions.h"
 #include "SkeletalDebugRendering.h"
 
 FName FControlRigEditorEditMode::ModeName("EditMode.ControlRigEditor");
diff --git a/Engine/Plugins/Animation/ControlRig/Source/ControlRigEditor/Private/Editor/ControlRigEditorEditMode.h b/Engine/Plugins/Animation/ControlRig/Source/ControlRigEditor/Private/Editor/ControlRigEditorEditMode.h
--- a/Engine/Plugins/Animation/ControlRig/Source/ControlRigEditor/Private/Editor/ControlRigEditorEditMode.h
+++ b/Engine/Plugins/Animation/ControlRig/Source/ControlRigEditor/Private/Editor/ControlRigEditorEditMode.h
@@ -2,13 +2,9 @@
 
 #pragma once
 
-#include "CoreMinimal.h"
 #include "EditMode/ControlRigEditMode.h"
-#include "IPersonaEditMode.h"
-#include "Preferences/PersonaOptions.h"
-#include "AnimationEditorViewportClient.h"
 
-#pragma once
+class UPersonaOptions;
 
 class FControlRigEditorEditMode : public FControlRigEditMode
 {
diff --git a/Engine/Source/Editor/Persona/Private/AnimationEditorPreviewScene.cpp b/Engine/Source/Editor/Persona/Private/AnimationEditorPreviewScene.cpp
--- a/Engine/Source/Editor/Persona/Private/AnimationEditorPreviewScene.cpp
+++ b/Engine/Source/Editor/Persona/Private/AnimationEditorPreviewScene.cpp
@@ -126,6 +126,14 @@ FAnimationEditorPreviewScene::~FAnimationEditorPreviewScene()
 	}
 }
 
+void FAnimationEditorPreviewScene::UnregisterForUndo()
+{
+	if (GEditor)
+	{
+		GEditor->UnregisterForUndo(this);
+	}
+}
+
 TArray<UDebugSkelMeshComponent*> FAnimationEditorPreviewScene::GetAllPreviewMeshComponents() const
 {
 	TArray<UDebugSkelMeshComponent*> PreviewMeshComponents;
diff --git a/Engine/Source/Editor/Persona/Private/AnimationEditorPreviewScene.h b/Engine/Source/Editor/Persona/Private/AnimationEditorPreviewScene.h
--- a/Engine/Source/Editor/Persona/Private/AnimationEditorPreviewScene.h
+++ b/Engine/Source/Editor/Persona/Private/AnimationEditorPreviewScene.h
@@ -21,6 +21,7 @@ public:
 	~FAnimationEditorPreviewScene();
 
 	/** IPersonaPreviewScene interface */
+	virtual void UnregisterForUndo() override;
 	virtual TSharedRef<class IPersonaToolkit> GetPersonaToolkit() const override { return PersonaToolkit.Pin().ToSharedRef(); }
 	virtual void SetPreviewAnimationAsset(UAnimationAsset* AnimAsset, bool bEnablePreview = true) override;
 	virtual UAnimationAsset* GetPreviewAnimationAsset() const override;
diff --git a/Engine/Source/Editor/Persona/Public/IPersonaPreviewScene.h b/Engine/Source/Editor/Persona/Public/IPersonaPreviewScene.h
--- a/Engine/Source/Editor/Persona/Public/IPersonaPreviewScene.h
+++ b/Engine/Source/Editor/Persona/Public/IPersonaPreviewScene.h
@@ -71,6 +71,9 @@ public:
 		: FAdvancedPreviewScene(CVS)
 	{}
 
+	/** Unregister this scene from the Undo chain. Necessary for certain editor windows. */
+	virtual void UnregisterForUndo() = 0;
+
 	/** Get the persona toolkit we are associated with */
 	virtual TSharedRef<class IPersonaToolkit> GetPersonaToolkit() const = 0;
 
