Control Rig: Disable postload RigVM hash optimization by default, and add cvar to enable it

diff --git a/Engine/Plugins/Animation/ControlRig/Source/ControlRigDeveloper/Private/ControlRigBlueprint.cpp b/Engine/Plugins/Animation/ControlRig/Source/ControlRigDeveloper/Private/ControlRigBlueprint.cpp
--- a/Engine/Plugins/Animation/ControlRig/Source/ControlRigDeveloper/Private/ControlRigBlueprint.cpp
+++ b/Engine/Plugins/Animation/ControlRig/Source/ControlRigDeveloper/Private/ControlRigBlueprint.cpp
@@ -55,6 +55,11 @@
 
 #define LOCTEXT_NAMESPACE "ControlRigBlueprint"
 
+TAutoConsoleVariable<bool> CVarEnablePostLoadHashing(
+	TEXT("ControlRig.EnablePostLoadHashing"),
+	false,
+	TEXT("When true refreshing the RigVMGraphs will be skipped if the hash matches the serialized hash."));
+
 static TArray<UClass*> GetClassObjectsInPackage(UPackage* InPackage)
 {
 	TArray<UObject*> Objects;
@@ -1321,9 +1326,7 @@ void UControlRigBlueprint::RefreshAllModels(EControlRigBlueprintLoadType InLoadT
 	const bool bIsPostLoad = InLoadType == EControlRigBlueprintLoadType::PostLoad;
 
 	// avoid any compute if the current structure hashes match with the serialized ones
-	const uint32 StructureHash = RigVMClient.GetStructureHash();
-	const uint32 SerializedHash = RigVMClient.GetSerializedStructureHash();
-	if(StructureHash == SerializedHash)
+	if(CVarEnablePostLoadHashing->GetBool() && RigVMClient.GetStructureHash() == RigVMClient.GetSerializedStructureHash())
 	{
 		if(bIsPostLoad)
 		{
