Fixing ControlRig include paths.

diff --git a/Engine/Plugins/Animation/ControlRig/Source/ControlRig/Private/Units/Collision/RigUnit_WorldCollision.cpp b/Engine/Plugins/Animation/ControlRig/Source/ControlRig/Private/Units/Collision/RigUnit_WorldCollision.cpp
--- a/Engine/Plugins/Animation/ControlRig/Source/ControlRig/Private/Units/Collision/RigUnit_WorldCollision.cpp
+++ b/Engine/Plugins/Animation/ControlRig/Source/ControlRig/Private/Units/Collision/RigUnit_WorldCollision.cpp
@@ -1,6 +1,6 @@
 // Copyright Epic Games, Inc. All Rights Reserved.
 
-#include "RigUnit_WorldCollision.h"
+#include "Units/Collision/RigUnit_WorldCollision.h"
 #include "Engine/World.h"
 #include "Math/ControlRigMathLibrary.h"
 #include "Units/RigUnitContext.h"
diff --git a/Engine/Plugins/Animation/ControlRig/Source/ControlRig/Private/Units/Execution/RigUnit_BeginExecution.cpp b/Engine/Plugins/Animation/ControlRig/Source/ControlRig/Private/Units/Execution/RigUnit_BeginExecution.cpp
--- a/Engine/Plugins/Animation/ControlRig/Source/ControlRig/Private/Units/Execution/RigUnit_BeginExecution.cpp
+++ b/Engine/Plugins/Animation/ControlRig/Source/ControlRig/Private/Units/Execution/RigUnit_BeginExecution.cpp
@@ -1,6 +1,6 @@
 // Copyright Epic Games, Inc. All Rights Reserved.
 
-#include "RigUnit_BeginExecution.h"
+#include "Units/Execution/RigUnit_BeginExecution.h"
 #include "Units/RigUnitContext.h"
 
 #include UE_INLINE_GENERATED_CPP_BY_NAME(RigUnit_BeginExecution)
diff --git a/Engine/Plugins/Animation/ControlRig/Source/ControlRig/Private/Units/Execution/RigUnit_Collection.cpp b/Engine/Plugins/Animation/ControlRig/Source/ControlRig/Private/Units/Execution/RigUnit_Collection.cpp
--- a/Engine/Plugins/Animation/ControlRig/Source/ControlRig/Private/Units/Execution/RigUnit_Collection.cpp
+++ b/Engine/Plugins/Animation/ControlRig/Source/ControlRig/Private/Units/Execution/RigUnit_Collection.cpp
@@ -1,6 +1,6 @@
 // Copyright Epic Games, Inc. All Rights Reserved.
 
-#include "RigUnit_Collection.h"
+#include "Units/Execution/RigUnit_Collection.h"
 #include "Units/Execution/RigUnit_Item.h"
 #include "Units/RigUnitContext.h"
 
diff --git a/Engine/Plugins/Animation/ControlRig/Source/ControlRig/Private/Units/Execution/RigUnit_DynamicHierarchy.cpp b/Engine/Plugins/Animation/ControlRig/Source/ControlRig/Private/Units/Execution/RigUnit_DynamicHierarchy.cpp
--- a/Engine/Plugins/Animation/ControlRig/Source/ControlRig/Private/Units/Execution/RigUnit_DynamicHierarchy.cpp
+++ b/Engine/Plugins/Animation/ControlRig/Source/ControlRig/Private/Units/Execution/RigUnit_DynamicHierarchy.cpp
@@ -1,6 +1,6 @@
 // Copyright Epic Games, Inc. All Rights Reserved.
 
-#include "RigUnit_DynamicHierarchy.h"
+#include "Units/Execution/RigUnit_DynamicHierarchy.h"
 #include "Engine/SkeletalMesh.h"
 #include "Rigs/RigHierarchyController.h"
 #include "Units/RigUnitContext.h"
diff --git a/Engine/Plugins/Animation/ControlRig/Source/ControlRig/Private/Units/Execution/RigUnit_Hierarchy.cpp b/Engine/Plugins/Animation/ControlRig/Source/ControlRig/Private/Units/Execution/RigUnit_Hierarchy.cpp
--- a/Engine/Plugins/Animation/ControlRig/Source/ControlRig/Private/Units/Execution/RigUnit_Hierarchy.cpp
+++ b/Engine/Plugins/Animation/ControlRig/Source/ControlRig/Private/Units/Execution/RigUnit_Hierarchy.cpp
@@ -1,9 +1,9 @@
 // Copyright Epic Games, Inc. All Rights Reserved.
 
-#include "RigUnit_Hierarchy.h"
+#include "Units/Execution/RigUnit_Hierarchy.h"
 #include "Units/Execution/RigUnit_Item.h"
 #include "Units/RigUnitContext.h"
-#include "RigUnit_Collection.h"
+#include "Units/Execution/RigUnit_Collection.h"
 
 #include UE_INLINE_GENERATED_CPP_BY_NAME(RigUnit_Hierarchy)
 
diff --git a/Engine/Plugins/Animation/ControlRig/Source/ControlRig/Private/Units/Execution/RigUnit_InteractionExecution.cpp b/Engine/Plugins/Animation/ControlRig/Source/ControlRig/Private/Units/Execution/RigUnit_InteractionExecution.cpp
--- a/Engine/Plugins/Animation/ControlRig/Source/ControlRig/Private/Units/Execution/RigUnit_InteractionExecution.cpp
+++ b/Engine/Plugins/Animation/ControlRig/Source/ControlRig/Private/Units/Execution/RigUnit_InteractionExecution.cpp
@@ -1,6 +1,6 @@
 // Copyright Epic Games, Inc. All Rights Reserved.
 
-#include "RigUnit_InteractionExecution.h"
+#include "Units/Execution/RigUnit_InteractionExecution.h"
 #include "Units/RigUnitContext.h"
 
 #include UE_INLINE_GENERATED_CPP_BY_NAME(RigUnit_InteractionExecution)
diff --git a/Engine/Plugins/Animation/ControlRig/Source/ControlRig/Private/Units/Execution/RigUnit_IsInteracting.cpp b/Engine/Plugins/Animation/ControlRig/Source/ControlRig/Private/Units/Execution/RigUnit_IsInteracting.cpp
--- a/Engine/Plugins/Animation/ControlRig/Source/ControlRig/Private/Units/Execution/RigUnit_IsInteracting.cpp
+++ b/Engine/Plugins/Animation/ControlRig/Source/ControlRig/Private/Units/Execution/RigUnit_IsInteracting.cpp
@@ -1,6 +1,6 @@
 // Copyright Epic Games, Inc. All Rights Reserved.
 
-#include "RigUnit_IsInteracting.h"
+#include "Units/Execution/RigUnit_IsInteracting.h"
 #include "Units/RigUnitContext.h"
 
 #include UE_INLINE_GENERATED_CPP_BY_NAME(RigUnit_IsInteracting)
diff --git a/Engine/Plugins/Animation/ControlRig/Source/ControlRig/Private/Units/Execution/RigUnit_Item.cpp b/Engine/Plugins/Animation/ControlRig/Source/ControlRig/Private/Units/Execution/RigUnit_Item.cpp
--- a/Engine/Plugins/Animation/ControlRig/Source/ControlRig/Private/Units/Execution/RigUnit_Item.cpp
+++ b/Engine/Plugins/Animation/ControlRig/Source/ControlRig/Private/Units/Execution/RigUnit_Item.cpp
@@ -1,6 +1,6 @@
 // Copyright Epic Games, Inc. All Rights Reserved.
 
-#include "RigUnit_Item.h"
+#include "Units/Execution/RigUnit_Item.h"
 #include "RigVMFunctions/RigVMFunction_Name.h"
 #include "Units/RigUnitContext.h"
 #include "RigVMFunctions/RigVMDispatch_Core.h"
diff --git a/Engine/Plugins/Animation/ControlRig/Source/ControlRig/Private/Units/Execution/RigUnit_PrepareForExecution.cpp b/Engine/Plugins/Animation/ControlRig/Source/ControlRig/Private/Units/Execution/RigUnit_PrepareForExecution.cpp
--- a/Engine/Plugins/Animation/ControlRig/Source/ControlRig/Private/Units/Execution/RigUnit_PrepareForExecution.cpp
+++ b/Engine/Plugins/Animation/ControlRig/Source/ControlRig/Private/Units/Execution/RigUnit_PrepareForExecution.cpp
@@ -1,6 +1,6 @@
 // Copyright Epic Games, Inc. All Rights Reserved.
 
-#include "RigUnit_PrepareForExecution.h"
+#include "Units/Execution/RigUnit_PrepareForExecution.h"
 #include "Units/RigUnitContext.h"
 
 #include UE_INLINE_GENERATED_CPP_BY_NAME(RigUnit_PrepareForExecution)
diff --git a/Engine/Plugins/Animation/ControlRig/Source/ControlRig/Private/Units/Execution/RigUnit_SequenceExecution.cpp b/Engine/Plugins/Animation/ControlRig/Source/ControlRig/Private/Units/Execution/RigUnit_SequenceExecution.cpp
--- a/Engine/Plugins/Animation/ControlRig/Source/ControlRig/Private/Units/Execution/RigUnit_SequenceExecution.cpp
+++ b/Engine/Plugins/Animation/ControlRig/Source/ControlRig/Private/Units/Execution/RigUnit_SequenceExecution.cpp
@@ -1,6 +1,6 @@
 // Copyright Epic Games, Inc. All Rights Reserved.
 
-#include "RigUnit_SequenceExecution.h"
+#include "Units/Execution/RigUnit_SequenceExecution.h"
 #include "Units/RigUnitContext.h"
 
 #include UE_INLINE_GENERATED_CPP_BY_NAME(RigUnit_SequenceExecution)
diff --git a/Engine/Plugins/Animation/ControlRig/Source/ControlRig/Private/Units/Hierarchy/RigUnit_BoneName.cpp b/Engine/Plugins/Animation/ControlRig/Source/ControlRig/Private/Units/Hierarchy/RigUnit_BoneName.cpp
--- a/Engine/Plugins/Animation/ControlRig/Source/ControlRig/Private/Units/Hierarchy/RigUnit_BoneName.cpp
+++ b/Engine/Plugins/Animation/ControlRig/Source/ControlRig/Private/Units/Hierarchy/RigUnit_BoneName.cpp
@@ -1,6 +1,6 @@
 // Copyright Epic Games, Inc. All Rights Reserved.
 
-#include "RigUnit_BoneName.h"
+#include "Units/Hierarchy/RigUnit_BoneName.h"
 #include "Units/RigUnitContext.h"
 
 #include UE_INLINE_GENERATED_CPP_BY_NAME(RigUnit_BoneName)
diff --git a/Engine/Plugins/Animation/ControlRig/Source/ControlRig/Private/Units/Hierarchy/RigUnit_ControlChannel.cpp b/Engine/Plugins/Animation/ControlRig/Source/ControlRig/Private/Units/Hierarchy/RigUnit_ControlChannel.cpp
--- a/Engine/Plugins/Animation/ControlRig/Source/ControlRig/Private/Units/Hierarchy/RigUnit_ControlChannel.cpp
+++ b/Engine/Plugins/Animation/ControlRig/Source/ControlRig/Private/Units/Hierarchy/RigUnit_ControlChannel.cpp
@@ -1,6 +1,6 @@
 // Copyright Epic Games, Inc. All Rights Reserved.
 
-#include "RigUnit_ControlChannel.h"
+#include "Units/Hierarchy/RigUnit_ControlChannel.h"
 #include "Units/RigUnitContext.h"
 
 #include UE_INLINE_GENERATED_CPP_BY_NAME(RigUnit_ControlChannel)
diff --git a/Engine/Plugins/Animation/ControlRig/Source/ControlRig/Private/Units/Hierarchy/RigUnit_GetBoneTransform.cpp b/Engine/Plugins/Animation/ControlRig/Source/ControlRig/Private/Units/Hierarchy/RigUnit_GetBoneTransform.cpp
--- a/Engine/Plugins/Animation/ControlRig/Source/ControlRig/Private/Units/Hierarchy/RigUnit_GetBoneTransform.cpp
+++ b/Engine/Plugins/Animation/ControlRig/Source/ControlRig/Private/Units/Hierarchy/RigUnit_GetBoneTransform.cpp
@@ -1,7 +1,7 @@
 // Copyright Epic Games, Inc. All Rights Reserved.
 
-#include "RigUnit_GetBoneTransform.h"
-#include "RigUnit_GetInitialBoneTransform.h"
+#include "Units/Hierarchy/RigUnit_GetBoneTransform.h"
+#include "Units/Hierarchy/RigUnit_GetInitialBoneTransform.h"
 #include "Units/RigUnitContext.h"
 #include "Units/Hierarchy/RigUnit_GetTransform.h"
 
diff --git a/Engine/Plugins/Animation/ControlRig/Source/ControlRig/Private/Units/Hierarchy/RigUnit_GetControlInitialTransform.cpp b/Engine/Plugins/Animation/ControlRig/Source/ControlRig/Private/Units/Hierarchy/RigUnit_GetControlInitialTransform.cpp
--- a/Engine/Plugins/Animation/ControlRig/Source/ControlRig/Private/Units/Hierarchy/RigUnit_GetControlInitialTransform.cpp
+++ b/Engine/Plugins/Animation/ControlRig/Source/ControlRig/Private/Units/Hierarchy/RigUnit_GetControlInitialTransform.cpp
@@ -1,6 +1,6 @@
 // Copyright Epic Games, Inc. All Rights Reserved.
 
-#include "RigUnit_GetControlInitialTransform.h"
+#include "Units/Hierarchy/RigUnit_GetControlInitialTransform.h"
 #include "Units/RigUnitContext.h"
 #include "Units/Hierarchy/RigUnit_GetTransform.h"
 
diff --git a/Engine/Plugins/Animation/ControlRig/Source/ControlRig/Private/Units/Hierarchy/RigUnit_GetControlTransform.cpp b/Engine/Plugins/Animation/ControlRig/Source/ControlRig/Private/Units/Hierarchy/RigUnit_GetControlTransform.cpp
--- a/Engine/Plugins/Animation/ControlRig/Source/ControlRig/Private/Units/Hierarchy/RigUnit_GetControlTransform.cpp
+++ b/Engine/Plugins/Animation/ControlRig/Source/ControlRig/Private/Units/Hierarchy/RigUnit_GetControlTransform.cpp
@@ -1,6 +1,6 @@
 // Copyright Epic Games, Inc. All Rights Reserved.
 
-#include "RigUnit_GetControlTransform.h"
+#include "Units/Hierarchy/RigUnit_GetControlTransform.h"
 #include "Units/RigUnitContext.h"
 
 #include UE_INLINE_GENERATED_CPP_BY_NAME(RigUnit_GetControlTransform)
diff --git a/Engine/Plugins/Animation/ControlRig/Source/ControlRig/Private/Units/Hierarchy/RigUnit_GetCurveValue.cpp b/Engine/Plugins/Animation/ControlRig/Source/ControlRig/Private/Units/Hierarchy/RigUnit_GetCurveValue.cpp
--- a/Engine/Plugins/Animation/ControlRig/Source/ControlRig/Private/Units/Hierarchy/RigUnit_GetCurveValue.cpp
+++ b/Engine/Plugins/Animation/ControlRig/Source/ControlRig/Private/Units/Hierarchy/RigUnit_GetCurveValue.cpp
@@ -1,6 +1,6 @@
 // Copyright Epic Games, Inc. All Rights Reserved.
 
-#include "RigUnit_GetCurveValue.h"
+#include "Units/Hierarchy/RigUnit_GetCurveValue.h"
 #include "Units/RigUnitContext.h"
 
 #include UE_INLINE_GENERATED_CPP_BY_NAME(RigUnit_GetCurveValue)
diff --git a/Engine/Plugins/Animation/ControlRig/Source/ControlRig/Private/Units/Hierarchy/RigUnit_GetInitialBoneTransform.cpp b/Engine/Plugins/Animation/ControlRig/Source/ControlRig/Private/Units/Hierarchy/RigUnit_GetInitialBoneTransform.cpp
--- a/Engine/Plugins/Animation/ControlRig/Source/ControlRig/Private/Units/Hierarchy/RigUnit_GetInitialBoneTransform.cpp
+++ b/Engine/Plugins/Animation/ControlRig/Source/ControlRig/Private/Units/Hierarchy/RigUnit_GetInitialBoneTransform.cpp
@@ -1,6 +1,6 @@
 // Copyright Epic Games, Inc. All Rights Reserved.
 
-#include "RigUnit_GetInitialBoneTransform.h"
+#include "Units/Hierarchy/RigUnit_GetInitialBoneTransform.h"
 #include "Units/RigUnitContext.h"
 #include "Units/Hierarchy/RigUnit_GetTransform.h"
 
diff --git a/Engine/Plugins/Animation/ControlRig/Source/ControlRig/Private/Units/Hierarchy/RigUnit_GetRelativeBoneTransform.cpp b/Engine/Plugins/Animation/ControlRig/Source/ControlRig/Private/Units/Hierarchy/RigUnit_GetRelativeBoneTransform.cpp
--- a/Engine/Plugins/Animation/ControlRig/Source/ControlRig/Private/Units/Hierarchy/RigUnit_GetRelativeBoneTransform.cpp
+++ b/Engine/Plugins/Animation/ControlRig/Source/ControlRig/Private/Units/Hierarchy/RigUnit_GetRelativeBoneTransform.cpp
@@ -1,6 +1,6 @@
 // Copyright Epic Games, Inc. All Rights Reserved.
 
-#include "RigUnit_GetRelativeBoneTransform.h"
+#include "Units/Hierarchy/RigUnit_GetRelativeBoneTransform.h"
 #include "Units/RigUnitContext.h"
 #include "Units/Hierarchy/RigUnit_GetRelativeTransform.h"
 
diff --git a/Engine/Plugins/Animation/ControlRig/Source/ControlRig/Private/Units/Hierarchy/RigUnit_GetRelativeTransform.cpp b/Engine/Plugins/Animation/ControlRig/Source/ControlRig/Private/Units/Hierarchy/RigUnit_GetRelativeTransform.cpp
--- a/Engine/Plugins/Animation/ControlRig/Source/ControlRig/Private/Units/Hierarchy/RigUnit_GetRelativeTransform.cpp
+++ b/Engine/Plugins/Animation/ControlRig/Source/ControlRig/Private/Units/Hierarchy/RigUnit_GetRelativeTransform.cpp
@@ -1,6 +1,6 @@
 // Copyright Epic Games, Inc. All Rights Reserved.
 
-#include "RigUnit_GetRelativeTransform.h"
+#include "Units/Hierarchy/RigUnit_GetRelativeTransform.h"
 #include "RigVMFunctions/Math/RigVMFunction_MathTransform.h"
 #include "Units/Hierarchy/RigUnit_GetTransform.h"
 #include "Units/RigUnitContext.h"
diff --git a/Engine/Plugins/Animation/ControlRig/Source/ControlRig/Private/Units/Hierarchy/RigUnit_GetSpaceTransform.cpp b/Engine/Plugins/Animation/ControlRig/Source/ControlRig/Private/Units/Hierarchy/RigUnit_GetSpaceTransform.cpp
--- a/Engine/Plugins/Animation/ControlRig/Source/ControlRig/Private/Units/Hierarchy/RigUnit_GetSpaceTransform.cpp
+++ b/Engine/Plugins/Animation/ControlRig/Source/ControlRig/Private/Units/Hierarchy/RigUnit_GetSpaceTransform.cpp
@@ -1,6 +1,6 @@
 // Copyright Epic Games, Inc. All Rights Reserved.
 
-#include "RigUnit_GetSpaceTransform.h"
+#include "Units/Hierarchy/RigUnit_GetSpaceTransform.h"
 #include "Units/RigUnitContext.h"
 #include "Units/Hierarchy/RigUnit_GetTransform.h"
 
diff --git a/Engine/Plugins/Animation/ControlRig/Source/ControlRig/Private/Units/Hierarchy/RigUnit_GetTransform.cpp b/Engine/Plugins/Animation/ControlRig/Source/ControlRig/Private/Units/Hierarchy/RigUnit_GetTransform.cpp
--- a/Engine/Plugins/Animation/ControlRig/Source/ControlRig/Private/Units/Hierarchy/RigUnit_GetTransform.cpp
+++ b/Engine/Plugins/Animation/ControlRig/Source/ControlRig/Private/Units/Hierarchy/RigUnit_GetTransform.cpp
@@ -1,6 +1,6 @@
 // Copyright Epic Games, Inc. All Rights Reserved.
 
-#include "RigUnit_GetTransform.h"
+#include "Units/Hierarchy/RigUnit_GetTransform.h"
 #include "Units/RigUnitContext.h"
 
 #include UE_INLINE_GENERATED_CPP_BY_NAME(RigUnit_GetTransform)
diff --git a/Engine/Plugins/Animation/ControlRig/Source/ControlRig/Private/Units/Hierarchy/RigUnit_Metadata.cpp b/Engine/Plugins/Animation/ControlRig/Source/ControlRig/Private/Units/Hierarchy/RigUnit_Metadata.cpp
--- a/Engine/Plugins/Animation/ControlRig/Source/ControlRig/Private/Units/Hierarchy/RigUnit_Metadata.cpp
+++ b/Engine/Plugins/Animation/ControlRig/Source/ControlRig/Private/Units/Hierarchy/RigUnit_Metadata.cpp
@@ -1,6 +1,6 @@
 // Copyright Epic Games, Inc. All Rights Reserved.
 
-#include "RigUnit_Metadata.h"
+#include "Units/Hierarchy/RigUnit_Metadata.h"
 #include "RigVMTypeUtils.h"
 #include "RigVMCore/RigVMStruct.h"
 #include "Units/RigUnitContext.h"
diff --git a/Engine/Plugins/Animation/ControlRig/Source/ControlRig/Private/Units/Hierarchy/RigUnit_OffsetTransform.cpp b/Engine/Plugins/Animation/ControlRig/Source/ControlRig/Private/Units/Hierarchy/RigUnit_OffsetTransform.cpp
--- a/Engine/Plugins/Animation/ControlRig/Source/ControlRig/Private/Units/Hierarchy/RigUnit_OffsetTransform.cpp
+++ b/Engine/Plugins/Animation/ControlRig/Source/ControlRig/Private/Units/Hierarchy/RigUnit_OffsetTransform.cpp
@@ -1,6 +1,6 @@
 // Copyright Epic Games, Inc. All Rights Reserved.
 
-#include "RigUnit_OffsetTransform.h"
+#include "Units/Hierarchy/RigUnit_OffsetTransform.h"
 #include "RigVMFunctions/Math/RigVMFunction_MathTransform.h"
 #include "Units/Hierarchy/RigUnit_GetTransform.h"
 #include "Units/Hierarchy/RigUnit_SetTransform.h"
diff --git a/Engine/Plugins/Animation/ControlRig/Source/ControlRig/Private/Units/Hierarchy/RigUnit_ParentSwitchConstraint.cpp b/Engine/Plugins/Animation/ControlRig/Source/ControlRig/Private/Units/Hierarchy/RigUnit_ParentSwitchConstraint.cpp
--- a/Engine/Plugins/Animation/ControlRig/Source/ControlRig/Private/Units/Hierarchy/RigUnit_ParentSwitchConstraint.cpp
+++ b/Engine/Plugins/Animation/ControlRig/Source/ControlRig/Private/Units/Hierarchy/RigUnit_ParentSwitchConstraint.cpp
@@ -1,6 +1,6 @@
 // Copyright Epic Games, Inc. All Rights Reserved.
 
-#include "RigUnit_ParentSwitchConstraint.h"
+#include "Units/Hierarchy/RigUnit_ParentSwitchConstraint.h"
 #include "Units/RigUnitContext.h"
 #include "Math/ControlRigMathLibrary.h"
 
diff --git a/Engine/Plugins/Animation/ControlRig/Source/ControlRig/Private/Units/Hierarchy/RigUnit_ProjectTransformToNewParent.cpp b/Engine/Plugins/Animation/ControlRig/Source/ControlRig/Private/Units/Hierarchy/RigUnit_ProjectTransformToNewParent.cpp
--- a/Engine/Plugins/Animation/ControlRig/Source/ControlRig/Private/Units/Hierarchy/RigUnit_ProjectTransformToNewParent.cpp
+++ b/Engine/Plugins/Animation/ControlRig/Source/ControlRig/Private/Units/Hierarchy/RigUnit_ProjectTransformToNewParent.cpp
@@ -1,6 +1,6 @@
 // Copyright Epic Games, Inc. All Rights Reserved.
 
-#include "RigUnit_ProjectTransformToNewParent.h"
+#include "Units/Hierarchy/RigUnit_ProjectTransformToNewParent.h"
 #include "RigVMFunctions/Math/RigVMFunction_MathTransform.h"
 #include "Units/Hierarchy/RigUnit_GetTransform.h"
 #include "Units/RigUnitContext.h"
diff --git a/Engine/Plugins/Animation/ControlRig/Source/ControlRig/Private/Units/Hierarchy/RigUnit_PropagateTransform.cpp b/Engine/Plugins/Animation/ControlRig/Source/ControlRig/Private/Units/Hierarchy/RigUnit_PropagateTransform.cpp
--- a/Engine/Plugins/Animation/ControlRig/Source/ControlRig/Private/Units/Hierarchy/RigUnit_PropagateTransform.cpp
+++ b/Engine/Plugins/Animation/ControlRig/Source/ControlRig/Private/Units/Hierarchy/RigUnit_PropagateTransform.cpp
@@ -1,6 +1,6 @@
 // Copyright Epic Games, Inc. All Rights Reserved.
 
-#include "RigUnit_PropagateTransform.h"
+#include "Units/Hierarchy/RigUnit_PropagateTransform.h"
 #include "RigVMFunctions/Math/RigVMFunction_MathTransform.h"
 #include "Units/Hierarchy/RigUnit_GetTransform.h"
 #include "Units/Hierarchy/RigUnit_SetTransform.h"
diff --git a/Engine/Plugins/Animation/ControlRig/Source/ControlRig/Private/Units/Hierarchy/RigUnit_SendEvent.cpp b/Engine/Plugins/Animation/ControlRig/Source/ControlRig/Private/Units/Hierarchy/RigUnit_SendEvent.cpp
--- a/Engine/Plugins/Animation/ControlRig/Source/ControlRig/Private/Units/Hierarchy/RigUnit_SendEvent.cpp
+++ b/Engine/Plugins/Animation/ControlRig/Source/ControlRig/Private/Units/Hierarchy/RigUnit_SendEvent.cpp
@@ -1,6 +1,6 @@
 // Copyright Epic Games, Inc. All Rights Reserved.
 
-#include "RigUnit_SendEvent.h"
+#include "Units/Hierarchy/RigUnit_SendEvent.h"
 #include "Units/RigUnitContext.h"
 
 #include UE_INLINE_GENERATED_CPP_BY_NAME(RigUnit_SendEvent)
diff --git a/Engine/Plugins/Animation/ControlRig/Source/ControlRig/Private/Units/Hierarchy/RigUnit_SetBoneInitialTransform.cpp b/Engine/Plugins/Animation/ControlRig/Source/ControlRig/Private/Units/Hierarchy/RigUnit_SetBoneInitialTransform.cpp
--- a/Engine/Plugins/Animation/ControlRig/Source/ControlRig/Private/Units/Hierarchy/RigUnit_SetBoneInitialTransform.cpp
+++ b/Engine/Plugins/Animation/ControlRig/Source/ControlRig/Private/Units/Hierarchy/RigUnit_SetBoneInitialTransform.cpp
@@ -1,6 +1,6 @@
 // Copyright Epic Games, Inc. All Rights Reserved.
 
-#include "RigUnit_SetBoneInitialTransform.h"
+#include "Units/Hierarchy/RigUnit_SetBoneInitialTransform.h"
 #include "Units/RigUnitContext.h"
 #include "Units/Execution/RigUnit_PrepareForExecution.h"
 #include "Units/Hierarchy/RigUnit_SetTransform.h"
diff --git a/Engine/Plugins/Animation/ControlRig/Source/ControlRig/Private/Units/Hierarchy/RigUnit_SetBoneRotation.cpp b/Engine/Plugins/Animation/ControlRig/Source/ControlRig/Private/Units/Hierarchy/RigUnit_SetBoneRotation.cpp
--- a/Engine/Plugins/Animation/ControlRig/Source/ControlRig/Private/Units/Hierarchy/RigUnit_SetBoneRotation.cpp
+++ b/Engine/Plugins/Animation/ControlRig/Source/ControlRig/Private/Units/Hierarchy/RigUnit_SetBoneRotation.cpp
@@ -1,6 +1,6 @@
 // Copyright Epic Games, Inc. All Rights Reserved.
 
-#include "RigUnit_SetBoneRotation.h"
+#include "Units/Hierarchy/RigUnit_SetBoneRotation.h"
 #include "Units/RigUnitContext.h"
 #include "Units/Hierarchy/RigUnit_SetTransform.h"
 
diff --git a/Engine/Plugins/Animation/ControlRig/Source/ControlRig/Private/Units/Hierarchy/RigUnit_SetBoneTransform.cpp b/Engine/Plugins/Animation/ControlRig/Source/ControlRig/Private/Units/Hierarchy/RigUnit_SetBoneTransform.cpp
--- a/Engine/Plugins/Animation/ControlRig/Source/ControlRig/Private/Units/Hierarchy/RigUnit_SetBoneTransform.cpp
+++ b/Engine/Plugins/Animation/ControlRig/Source/ControlRig/Private/Units/Hierarchy/RigUnit_SetBoneTransform.cpp
@@ -1,6 +1,6 @@
 // Copyright Epic Games, Inc. All Rights Reserved.
 
-#include "RigUnit_SetBoneTransform.h"
+#include "Units/Hierarchy/RigUnit_SetBoneTransform.h"
 #include "Units/RigUnitContext.h"
 #include "Math/ControlRigMathLibrary.h"
 #include "Units/Hierarchy/RigUnit_SetTransform.h"
diff --git a/Engine/Plugins/Animation/ControlRig/Source/ControlRig/Private/Units/Hierarchy/RigUnit_SetBoneTranslation.cpp b/Engine/Plugins/Animation/ControlRig/Source/ControlRig/Private/Units/Hierarchy/RigUnit_SetBoneTranslation.cpp
--- a/Engine/Plugins/Animation/ControlRig/Source/ControlRig/Private/Units/Hierarchy/RigUnit_SetBoneTranslation.cpp
+++ b/Engine/Plugins/Animation/ControlRig/Source/ControlRig/Private/Units/Hierarchy/RigUnit_SetBoneTranslation.cpp
@@ -1,6 +1,6 @@
 // Copyright Epic Games, Inc. All Rights Reserved.
 
-#include "RigUnit_SetBoneTranslation.h"
+#include "Units/Hierarchy/RigUnit_SetBoneTranslation.h"
 #include "Units/RigUnitContext.h"
 #include "Units/Hierarchy/RigUnit_SetTransform.h"
 
diff --git a/Engine/Plugins/Animation/ControlRig/Source/ControlRig/Private/Units/Hierarchy/RigUnit_SetControlColor.cpp b/Engine/Plugins/Animation/ControlRig/Source/ControlRig/Private/Units/Hierarchy/RigUnit_SetControlColor.cpp
--- a/Engine/Plugins/Animation/ControlRig/Source/ControlRig/Private/Units/Hierarchy/RigUnit_SetControlColor.cpp
+++ b/Engine/Plugins/Animation/ControlRig/Source/ControlRig/Private/Units/Hierarchy/RigUnit_SetControlColor.cpp
@@ -1,6 +1,6 @@
 // Copyright Epic Games, Inc. All Rights Reserved.
 
-#include "RigUnit_SetControlColor.h"
+#include "Units/Hierarchy/RigUnit_SetControlColor.h"
 #include "Units/RigUnitContext.h"
 
 #include UE_INLINE_GENERATED_CPP_BY_NAME(RigUnit_SetControlColor)
diff --git a/Engine/Plugins/Animation/ControlRig/Source/ControlRig/Private/Units/Hierarchy/RigUnit_SetControlDrivenList.cpp b/Engine/Plugins/Animation/ControlRig/Source/ControlRig/Private/Units/Hierarchy/RigUnit_SetControlDrivenList.cpp
--- a/Engine/Plugins/Animation/ControlRig/Source/ControlRig/Private/Units/Hierarchy/RigUnit_SetControlDrivenList.cpp
+++ b/Engine/Plugins/Animation/ControlRig/Source/ControlRig/Private/Units/Hierarchy/RigUnit_SetControlDrivenList.cpp
@@ -1,6 +1,6 @@
 // Copyright Epic Games, Inc. All Rights Reserved.
 
-#include "RigUnit_SetControlDrivenList.h"
+#include "Units/Hierarchy/RigUnit_SetControlDrivenList.h"
 
 #include "Rigs/RigHierarchyController.h"
 #include "Units/RigUnitContext.h"
diff --git a/Engine/Plugins/Animation/ControlRig/Source/ControlRig/Private/Units/Hierarchy/RigUnit_SetControlOffset.cpp b/Engine/Plugins/Animation/ControlRig/Source/ControlRig/Private/Units/Hierarchy/RigUnit_SetControlOffset.cpp
--- a/Engine/Plugins/Animation/ControlRig/Source/ControlRig/Private/Units/Hierarchy/RigUnit_SetControlOffset.cpp
+++ b/Engine/Plugins/Animation/ControlRig/Source/ControlRig/Private/Units/Hierarchy/RigUnit_SetControlOffset.cpp
@@ -1,6 +1,6 @@
 // Copyright Epic Games, Inc. All Rights Reserved.
 
-#include "RigUnit_SetControlOffset.h"
+#include "Units/Hierarchy/RigUnit_SetControlOffset.h"
 
 #include "GeometryCollection/GeometryCollectionAlgo.h"
 #include "Units/RigUnitContext.h"
diff --git a/Engine/Plugins/Animation/ControlRig/Source/ControlRig/Private/Units/Hierarchy/RigUnit_SetControlTransform.cpp b/Engine/Plugins/Animation/ControlRig/Source/ControlRig/Private/Units/Hierarchy/RigUnit_SetControlTransform.cpp
--- a/Engine/Plugins/Animation/ControlRig/Source/ControlRig/Private/Units/Hierarchy/RigUnit_SetControlTransform.cpp
+++ b/Engine/Plugins/Animation/ControlRig/Source/ControlRig/Private/Units/Hierarchy/RigUnit_SetControlTransform.cpp
@@ -1,6 +1,6 @@
 // Copyright Epic Games, Inc. All Rights Reserved.
 
-#include "RigUnit_SetControlTransform.h"
+#include "Units/Hierarchy/RigUnit_SetControlTransform.h"
 #include "Units/RigUnitContext.h"
 #include "Math/ControlRigMathLibrary.h"
 #include "Units/Hierarchy/RigUnit_SetTransform.h"
diff --git a/Engine/Plugins/Animation/ControlRig/Source/ControlRig/Private/Units/Hierarchy/RigUnit_SetControlVisibility.cpp b/Engine/Plugins/Animation/ControlRig/Source/ControlRig/Private/Units/Hierarchy/RigUnit_SetControlVisibility.cpp
--- a/Engine/Plugins/Animation/ControlRig/Source/ControlRig/Private/Units/Hierarchy/RigUnit_SetControlVisibility.cpp
+++ b/Engine/Plugins/Animation/ControlRig/Source/ControlRig/Private/Units/Hierarchy/RigUnit_SetControlVisibility.cpp
@@ -1,6 +1,6 @@
 // Copyright Epic Games, Inc. All Rights Reserved.
 
-#include "RigUnit_SetControlVisibility.h"
+#include "Units/Hierarchy/RigUnit_SetControlVisibility.h"
 #include "Units/RigUnitContext.h"
 
 #include UE_INLINE_GENERATED_CPP_BY_NAME(RigUnit_SetControlVisibility)
diff --git a/Engine/Plugins/Animation/ControlRig/Source/ControlRig/Private/Units/Hierarchy/RigUnit_SetRelativeBoneTransform.cpp b/Engine/Plugins/Animation/ControlRig/Source/ControlRig/Private/Units/Hierarchy/RigUnit_SetRelativeBoneTransform.cpp
--- a/Engine/Plugins/Animation/ControlRig/Source/ControlRig/Private/Units/Hierarchy/RigUnit_SetRelativeBoneTransform.cpp
+++ b/Engine/Plugins/Animation/ControlRig/Source/ControlRig/Private/Units/Hierarchy/RigUnit_SetRelativeBoneTransform.cpp
@@ -1,6 +1,6 @@
 // Copyright Epic Games, Inc. All Rights Reserved.
 
-#include "RigUnit_SetRelativeBoneTransform.h"
+#include "Units/Hierarchy/RigUnit_SetRelativeBoneTransform.h"
 #include "Units/RigUnitContext.h"
 #include "Units/Hierarchy/RigUnit_SetRelativeTransform.h"
 #include "Math/ControlRigMathLibrary.h"
diff --git a/Engine/Plugins/Animation/ControlRig/Source/ControlRig/Private/Units/Hierarchy/RigUnit_SetRelativeTransform.cpp b/Engine/Plugins/Animation/ControlRig/Source/ControlRig/Private/Units/Hierarchy/RigUnit_SetRelativeTransform.cpp
--- a/Engine/Plugins/Animation/ControlRig/Source/ControlRig/Private/Units/Hierarchy/RigUnit_SetRelativeTransform.cpp
+++ b/Engine/Plugins/Animation/ControlRig/Source/ControlRig/Private/Units/Hierarchy/RigUnit_SetRelativeTransform.cpp
@@ -1,6 +1,6 @@
 // Copyright Epic Games, Inc. All Rights Reserved.
 
-#include "RigUnit_SetRelativeTransform.h"
+#include "Units/Hierarchy/RigUnit_SetRelativeTransform.h"
 #include "RigVMFunctions/Math/RigVMFunction_MathTransform.h"
 #include "Units/Hierarchy/RigUnit_GetTransform.h"
 #include "Units/Hierarchy/RigUnit_SetTransform.h"
diff --git a/Engine/Plugins/Animation/ControlRig/Source/ControlRig/Private/Units/Hierarchy/RigUnit_SetSpaceInitialTransform.cpp b/Engine/Plugins/Animation/ControlRig/Source/ControlRig/Private/Units/Hierarchy/RigUnit_SetSpaceInitialTransform.cpp
--- a/Engine/Plugins/Animation/ControlRig/Source/ControlRig/Private/Units/Hierarchy/RigUnit_SetSpaceInitialTransform.cpp
+++ b/Engine/Plugins/Animation/ControlRig/Source/ControlRig/Private/Units/Hierarchy/RigUnit_SetSpaceInitialTransform.cpp
@@ -1,6 +1,6 @@
 // Copyright Epic Games, Inc. All Rights Reserved.
 
-#include "RigUnit_SetSpaceInitialTransform.h"
+#include "Units/Hierarchy/RigUnit_SetSpaceInitialTransform.h"
 #include "Units/RigUnitContext.h"
 #include "Units/Execution/RigUnit_PrepareForExecution.h"
 #include "Units/Hierarchy/RigUnit_SetTransform.h"
diff --git a/Engine/Plugins/Animation/ControlRig/Source/ControlRig/Private/Units/Hierarchy/RigUnit_SetSpaceTransform.cpp b/Engine/Plugins/Animation/ControlRig/Source/ControlRig/Private/Units/Hierarchy/RigUnit_SetSpaceTransform.cpp
--- a/Engine/Plugins/Animation/ControlRig/Source/ControlRig/Private/Units/Hierarchy/RigUnit_SetSpaceTransform.cpp
+++ b/Engine/Plugins/Animation/ControlRig/Source/ControlRig/Private/Units/Hierarchy/RigUnit_SetSpaceTransform.cpp
@@ -1,6 +1,6 @@
 // Copyright Epic Games, Inc. All Rights Reserved.
 
-#include "RigUnit_SetSpaceTransform.h"
+#include "Units/Hierarchy/RigUnit_SetSpaceTransform.h"
 #include "Units/RigUnitContext.h"
 #include "Math/ControlRigMathLibrary.h"
 #include "Units/Hierarchy/RigUnit_SetTransform.h"
diff --git a/Engine/Plugins/Animation/ControlRig/Source/ControlRig/Private/Units/Hierarchy/RigUnit_SetTransform.cpp b/Engine/Plugins/Animation/ControlRig/Source/ControlRig/Private/Units/Hierarchy/RigUnit_SetTransform.cpp
--- a/Engine/Plugins/Animation/ControlRig/Source/ControlRig/Private/Units/Hierarchy/RigUnit_SetTransform.cpp
+++ b/Engine/Plugins/Animation/ControlRig/Source/ControlRig/Private/Units/Hierarchy/RigUnit_SetTransform.cpp
@@ -1,7 +1,7 @@
 // Copyright Epic Games, Inc. All Rights Reserved.
 
 #include "Units/Hierarchy/RigUnit_SetTransform.h"
-#include "RigUnit_GetTransform.h"
+#include "Units/Hierarchy/RigUnit_GetTransform.h"
 #include "Units/RigUnitContext.h"
 #include "Units/Execution/RigUnit_PrepareForExecution.h"
 #include "Units/Execution/RigUnit_BeginExecution.h"
diff --git a/Engine/Plugins/Animation/ControlRig/Source/ControlRig/Private/Units/Hierarchy/RigUnit_UnsetCurveValue.cpp b/Engine/Plugins/Animation/ControlRig/Source/ControlRig/Private/Units/Hierarchy/RigUnit_UnsetCurveValue.cpp
--- a/Engine/Plugins/Animation/ControlRig/Source/ControlRig/Private/Units/Hierarchy/RigUnit_UnsetCurveValue.cpp
+++ b/Engine/Plugins/Animation/ControlRig/Source/ControlRig/Private/Units/Hierarchy/RigUnit_UnsetCurveValue.cpp
@@ -1,6 +1,6 @@
 ﻿// Copyright Epic Games, Inc. All Rights Reserved.
 
-#include "RigUnit_UnsetCurveValue.h"
+#include "Units/Hierarchy/RigUnit_UnsetCurveValue.h"
 #include "Units/RigUnitContext.h"
 
 #include UE_INLINE_GENERATED_CPP_BY_NAME(RigUnit_UnsetCurveValue)
diff --git a/Engine/Plugins/Animation/ControlRig/Source/ControlRig/Private/Units/Hierarchy/RigUnit_WorldSpace.cpp b/Engine/Plugins/Animation/ControlRig/Source/ControlRig/Private/Units/Hierarchy/RigUnit_WorldSpace.cpp
--- a/Engine/Plugins/Animation/ControlRig/Source/ControlRig/Private/Units/Hierarchy/RigUnit_WorldSpace.cpp
+++ b/Engine/Plugins/Animation/ControlRig/Source/ControlRig/Private/Units/Hierarchy/RigUnit_WorldSpace.cpp
@@ -1,6 +1,6 @@
 // Copyright Epic Games, Inc. All Rights Reserved.
 
-#include "RigUnit_WorldSpace.h"
+#include "Units/Hierarchy/RigUnit_WorldSpace.h"
 #include "Units/RigUnitContext.h"
 
 #include UE_INLINE_GENERATED_CPP_BY_NAME(RigUnit_WorldSpace)
diff --git a/Engine/Plugins/Animation/ControlRig/Source/ControlRig/Private/Units/Highlevel/Hierarchy/RigUnit_CCDIK.cpp b/Engine/Plugins/Animation/ControlRig/Source/ControlRig/Private/Units/Highlevel/Hierarchy/RigUnit_CCDIK.cpp
--- a/Engine/Plugins/Animation/ControlRig/Source/ControlRig/Private/Units/Highlevel/Hierarchy/RigUnit_CCDIK.cpp
+++ b/Engine/Plugins/Animation/ControlRig/Source/ControlRig/Private/Units/Highlevel/Hierarchy/RigUnit_CCDIK.cpp
@@ -1,6 +1,6 @@
 // Copyright Epic Games, Inc. All Rights Reserved.
 
-#include "RigUnit_CCDIK.h"
+#include "Units/Highlevel/Hierarchy/RigUnit_CCDIK.h"
 #include "Units/RigUnitContext.h"
 #include "Math/ControlRigMathLibrary.h"
 
diff --git a/Engine/Plugins/Animation/ControlRig/Source/ControlRig/Private/Units/Highlevel/Hierarchy/RigUnit_FABRIK.cpp b/Engine/Plugins/Animation/ControlRig/Source/ControlRig/Private/Units/Highlevel/Hierarchy/RigUnit_FABRIK.cpp
--- a/Engine/Plugins/Animation/ControlRig/Source/ControlRig/Private/Units/Highlevel/Hierarchy/RigUnit_FABRIK.cpp
+++ b/Engine/Plugins/Animation/ControlRig/Source/ControlRig/Private/Units/Highlevel/Hierarchy/RigUnit_FABRIK.cpp
@@ -1,6 +1,6 @@
 // Copyright Epic Games, Inc. All Rights Reserved.
 
-#include "RigUnit_FABRIK.h"
+#include "Units/Highlevel/Hierarchy/RigUnit_FABRIK.h"
 #include "Units/RigUnitContext.h"
 #include "Math/ControlRigMathLibrary.h"
 
diff --git a/Engine/Plugins/Animation/ControlRig/Source/ControlRig/Private/Units/Highlevel/Hierarchy/RigUnit_ModifyBoneTransforms.cpp b/Engine/Plugins/Animation/ControlRig/Source/ControlRig/Private/Units/Highlevel/Hierarchy/RigUnit_ModifyBoneTransforms.cpp
--- a/Engine/Plugins/Animation/ControlRig/Source/ControlRig/Private/Units/Highlevel/Hierarchy/RigUnit_ModifyBoneTransforms.cpp
+++ b/Engine/Plugins/Animation/ControlRig/Source/ControlRig/Private/Units/Highlevel/Hierarchy/RigUnit_ModifyBoneTransforms.cpp
@@ -1,6 +1,6 @@
 // Copyright Epic Games, Inc. All Rights Reserved.
 
-#include "RigUnit_ModifyBoneTransforms.h"
+#include "Units/Highlevel/Hierarchy/RigUnit_ModifyBoneTransforms.h"
 #include "Units/RigUnitContext.h"
 
 #include UE_INLINE_GENERATED_CPP_BY_NAME(RigUnit_ModifyBoneTransforms)
diff --git a/Engine/Plugins/Animation/ControlRig/Source/ControlRig/Private/Units/Highlevel/Hierarchy/RigUnit_ModifyTransforms.cpp b/Engine/Plugins/Animation/ControlRig/Source/ControlRig/Private/Units/Highlevel/Hierarchy/RigUnit_ModifyTransforms.cpp
--- a/Engine/Plugins/Animation/ControlRig/Source/ControlRig/Private/Units/Highlevel/Hierarchy/RigUnit_ModifyTransforms.cpp
+++ b/Engine/Plugins/Animation/ControlRig/Source/ControlRig/Private/Units/Highlevel/Hierarchy/RigUnit_ModifyTransforms.cpp
@@ -1,6 +1,6 @@
 // Copyright Epic Games, Inc. All Rights Reserved.
 
-#include "RigUnit_ModifyTransforms.h"
+#include "Units/Highlevel/Hierarchy/RigUnit_ModifyTransforms.h"
 #include "Units/RigUnitContext.h"
 #include "Math/ControlRigMathLibrary.h"
 
diff --git a/Engine/Plugins/Animation/ControlRig/Source/ControlRig/Private/Units/Highlevel/Hierarchy/RigUnit_MultiFABRIK.cpp b/Engine/Plugins/Animation/ControlRig/Source/ControlRig/Private/Units/Highlevel/Hierarchy/RigUnit_MultiFABRIK.cpp
--- a/Engine/Plugins/Animation/ControlRig/Source/ControlRig/Private/Units/Highlevel/Hierarchy/RigUnit_MultiFABRIK.cpp
+++ b/Engine/Plugins/Animation/ControlRig/Source/ControlRig/Private/Units/Highlevel/Hierarchy/RigUnit_MultiFABRIK.cpp
@@ -1,6 +1,6 @@
 // Copyright Epic Games, Inc. All Rights Reserved.
 
-#include "RigUnit_MultiFABRIK.h"
+#include "Units/Highlevel/Hierarchy/RigUnit_MultiFABRIK.h"
 #include "Units/RigUnitContext.h"
 
 #include UE_INLINE_GENERATED_CPP_BY_NAME(RigUnit_MultiFABRIK)
diff --git a/Engine/Plugins/Animation/ControlRig/Source/ControlRig/Private/Units/Highlevel/Hierarchy/RigUnit_SphericalPoseReader.cpp b/Engine/Plugins/Animation/ControlRig/Source/ControlRig/Private/Units/Highlevel/Hierarchy/RigUnit_SphericalPoseReader.cpp
--- a/Engine/Plugins/Animation/ControlRig/Source/ControlRig/Private/Units/Highlevel/Hierarchy/RigUnit_SphericalPoseReader.cpp
+++ b/Engine/Plugins/Animation/ControlRig/Source/ControlRig/Private/Units/Highlevel/Hierarchy/RigUnit_SphericalPoseReader.cpp
@@ -1,6 +1,6 @@
 // Copyright Epic Games, Inc. All Rights Reserved.
 
-#include "RigUnit_SphericalPoseReader.h"
+#include "Units/Highlevel/Hierarchy/RigUnit_SphericalPoseReader.h"
 #include "Kismet/KismetMathLibrary.h"
 #include "Units/RigUnitContext.h"
 
diff --git a/Engine/Plugins/Animation/ControlRig/Source/ControlRig/Private/Units/Highlevel/Hierarchy/RigUnit_TransformConstraint.cpp b/Engine/Plugins/Animation/ControlRig/Source/ControlRig/Private/Units/Highlevel/Hierarchy/RigUnit_TransformConstraint.cpp
--- a/Engine/Plugins/Animation/ControlRig/Source/ControlRig/Private/Units/Highlevel/Hierarchy/RigUnit_TransformConstraint.cpp
+++ b/Engine/Plugins/Animation/ControlRig/Source/ControlRig/Private/Units/Highlevel/Hierarchy/RigUnit_TransformConstraint.cpp
@@ -1,6 +1,6 @@
 // Copyright Epic Games, Inc. All Rights Reserved.
 
-#include "RigUnit_TransformConstraint.h"
+#include "Units/Highlevel/Hierarchy/RigUnit_TransformConstraint.h"
 #include "Units/RigUnitContext.h"
 #include "HelperUtil.h"
 #include "AnimationCoreLibrary.h"
diff --git a/Engine/Plugins/Animation/ControlRigSpline/Source/ControlRigSpline/ControlRigSpline.build.cs b/Engine/Plugins/Animation/ControlRigSpline/Source/ControlRigSpline/ControlRigSpline.build.cs
--- a/Engine/Plugins/Animation/ControlRigSpline/Source/ControlRigSpline/ControlRigSpline.build.cs
+++ b/Engine/Plugins/Animation/ControlRigSpline/Source/ControlRigSpline/ControlRigSpline.build.cs
@@ -19,7 +19,7 @@ namespace UnrealBuildTool.Rules
 			});
 
 			// TODO: Should not be including private headers in public code
-			PublicIncludePaths.Add(Path.Combine(GetModuleDirectory("ControlRig"), "Private/Units/Highlevel/Hierarchy")); // For RigUnit_FitChainToCurve.h
+			PublicIncludePaths.Add(Path.Combine(GetModuleDirectory("ControlRig"), "Private")); // For RigUnit_FitChainToCurve.h
 		}
 	}
 }
diff --git a/Engine/Plugins/Animation/ControlRigSpline/Source/ControlRigSpline/Private/ControlRigSplineUnits.cpp b/Engine/Plugins/Animation/ControlRigSpline/Source/ControlRigSpline/Private/ControlRigSplineUnits.cpp
--- a/Engine/Plugins/Animation/ControlRigSpline/Source/ControlRigSpline/Private/ControlRigSplineUnits.cpp
+++ b/Engine/Plugins/Animation/ControlRigSpline/Source/ControlRigSpline/Private/ControlRigSplineUnits.cpp
@@ -4,7 +4,7 @@
 #include "Units/RigUnitContext.h"
 #include "Features/IModularFeatures.h"
 #include "Algo/BinarySearch.h"
-#include "RigUnit_FitChainToCurve.h"
+#include "Units/Highlevel/Hierarchy/RigUnit_FitChainToCurve.h"
 
 #include UE_INLINE_GENERATED_CPP_BY_NAME(ControlRigSplineUnits)
 
diff --git a/Engine/Plugins/Animation/ControlRigSpline/Source/ControlRigSpline/Public/ControlRigSplineUnits.h b/Engine/Plugins/Animation/ControlRigSpline/Source/ControlRigSpline/Public/ControlRigSplineUnits.h
--- a/Engine/Plugins/Animation/ControlRigSpline/Source/ControlRigSpline/Public/ControlRigSplineUnits.h
+++ b/Engine/Plugins/Animation/ControlRigSpline/Source/ControlRigSpline/Public/ControlRigSplineUnits.h
@@ -5,7 +5,7 @@
 #include "Units/RigUnit.h"
 #include "ControlRig.h"
 #include "ControlRigSplineTypes.h"
-#include "RigUnit_FitChainToCurve.h"
+#include "Units/Highlevel/Hierarchy/RigUnit_FitChainToCurve.h"
 #include "ControlRigSplineUnits.generated.h"
 
 USTRUCT(meta = (Abstract, NodeColor = "0.3 0.1 0.1"))
diff --git a/Engine/Source/Runtime/RigVM/Public/RigVMFunctions/Simulation/RigVMFunction_Accumulate.h b/Engine/Source/Runtime/RigVM/Public/RigVMFunctions/Simulation/RigVMFunction_Accumulate.h
--- a/Engine/Source/Runtime/RigVM/Public/RigVMFunctions/Simulation/RigVMFunction_Accumulate.h
+++ b/Engine/Source/Runtime/RigVM/Public/RigVMFunctions/Simulation/RigVMFunction_Accumulate.h
@@ -2,7 +2,7 @@
 
 #pragma once
 
-#include "RigVMFunction_SimBase.h"
+#include "RigVMFunctions/Simulation/RigVMFunction_SimBase.h"
 #include "RigVMFunction_Accumulate.generated.h"
 
 USTRUCT(meta=(Category="Simulation|Accumulate"))
diff --git a/Engine/Source/Runtime/RigVM/Public/RigVMFunctions/Simulation/RigVMFunction_AlphaInterp.h b/Engine/Source/Runtime/RigVM/Public/RigVMFunctions/Simulation/RigVMFunction_AlphaInterp.h
--- a/Engine/Source/Runtime/RigVM/Public/RigVMFunctions/Simulation/RigVMFunction_AlphaInterp.h
+++ b/Engine/Source/Runtime/RigVM/Public/RigVMFunctions/Simulation/RigVMFunction_AlphaInterp.h
@@ -2,7 +2,7 @@
 
 #pragma once
 
-#include "RigVMFunction_SimBase.h"
+#include "RigVMFunctions/Simulation/RigVMFunction_SimBase.h"
 #include "Animation/InputScaleBias.h"
 #include "RigVMFunction_AlphaInterp.generated.h"
 
diff --git a/Engine/Source/Runtime/RigVM/Public/RigVMFunctions/Simulation/RigVMFunction_DeltaFromPrevious.h b/Engine/Source/Runtime/RigVM/Public/RigVMFunctions/Simulation/RigVMFunction_DeltaFromPrevious.h
--- a/Engine/Source/Runtime/RigVM/Public/RigVMFunctions/Simulation/RigVMFunction_DeltaFromPrevious.h
+++ b/Engine/Source/Runtime/RigVM/Public/RigVMFunctions/Simulation/RigVMFunction_DeltaFromPrevious.h
@@ -2,7 +2,7 @@
 
 #pragma once
 
-#include "RigVMFunction_SimBase.h"
+#include "RigVMFunctions/Simulation/RigVMFunction_SimBase.h"
 #include "RigVMFunction_DeltaFromPrevious.generated.h"
 
 /**
diff --git a/Engine/Source/Runtime/RigVM/Public/RigVMFunctions/Simulation/RigVMFunction_Kalman.h b/Engine/Source/Runtime/RigVM/Public/RigVMFunctions/Simulation/RigVMFunction_Kalman.h
--- a/Engine/Source/Runtime/RigVM/Public/RigVMFunctions/Simulation/RigVMFunction_Kalman.h
+++ b/Engine/Source/Runtime/RigVM/Public/RigVMFunctions/Simulation/RigVMFunction_Kalman.h
@@ -2,7 +2,7 @@
 
 #pragma once
 
-#include "RigVMFunction_SimBase.h"
+#include "RigVMFunctions/Simulation/RigVMFunction_SimBase.h"
 #include "RigVMFunction_Kalman.generated.h"
 
 /**
diff --git a/Engine/Source/Runtime/RigVM/Public/RigVMFunctions/Simulation/RigVMFunction_TimeOffset.h b/Engine/Source/Runtime/RigVM/Public/RigVMFunctions/Simulation/RigVMFunction_TimeOffset.h
--- a/Engine/Source/Runtime/RigVM/Public/RigVMFunctions/Simulation/RigVMFunction_TimeOffset.h
+++ b/Engine/Source/Runtime/RigVM/Public/RigVMFunctions/Simulation/RigVMFunction_TimeOffset.h
@@ -2,7 +2,7 @@
 
 #pragma once
 
-#include "RigVMFunction_SimBase.h"
+#include "RigVMFunctions/Simulation/RigVMFunction_SimBase.h"
 #include "RigVMFunction_TimeOffset.generated.h"
 
 /**
diff --git a/Engine/Source/Runtime/RigVM/Public/RigVMFunctions/Simulation/RigVMFunction_Timeline.h b/Engine/Source/Runtime/RigVM/Public/RigVMFunctions/Simulation/RigVMFunction_Timeline.h
--- a/Engine/Source/Runtime/RigVM/Public/RigVMFunctions/Simulation/RigVMFunction_Timeline.h
+++ b/Engine/Source/Runtime/RigVM/Public/RigVMFunctions/Simulation/RigVMFunction_Timeline.h
@@ -2,7 +2,7 @@
 
 #pragma once
 
-#include "RigVMFunction_SimBase.h"
+#include "RigVMFunctions/Simulation/RigVMFunction_SimBase.h"
 #include "RigVMFunction_Timeline.generated.h"
 
 /**
diff --git a/Engine/Source/Runtime/RigVM/Public/RigVMFunctions/Simulation/RigVMFunction_Verlet.h b/Engine/Source/Runtime/RigVM/Public/RigVMFunctions/Simulation/RigVMFunction_Verlet.h
--- a/Engine/Source/Runtime/RigVM/Public/RigVMFunctions/Simulation/RigVMFunction_Verlet.h
+++ b/Engine/Source/Runtime/RigVM/Public/RigVMFunctions/Simulation/RigVMFunction_Verlet.h
@@ -2,7 +2,7 @@
 
 #pragma once
 
-#include "RigVMFunction_SimBase.h"
+#include "RigVMFunctions/Simulation/RigVMFunction_SimBase.h"
 #include "RigVMFunctions/Math/RigVMMathLibrary.h"
 #include "RigVMFunction_Verlet.generated.h"
 
