Control Rig: Fix CR component input mapped elements having wrong indices

diff --git a/Engine/Plugins/Animation/ControlRig/Source/ControlRig/Private/ControlRigComponent.cpp b/Engine/Plugins/Animation/ControlRig/Source/ControlRig/Private/ControlRigComponent.cpp
--- a/Engine/Plugins/Animation/ControlRig/Source/ControlRig/Private/ControlRigComponent.cpp
+++ b/Engine/Plugins/Animation/ControlRig/Source/ControlRig/Private/ControlRigComponent.cpp
@@ -1351,7 +1351,7 @@ void UControlRigComponent::ValidateMappingData()
 						{
 							if (MappedElement.ElementType != ERigElementType::Curve)
 							{
-								MappedElement.SubIndex = Skeleton->GetReferenceSkeleton().FindBoneIndex(MappedElement.TransformName);
+								MappedElement.SubIndex = SkeletalMeshComponent->GetBoneIndex(MappedElement.TransformName);
 							}
 						}
 						else
