ControlRigComponent: use the mesh pose bone index across CRComponent code, converting to SkeletonIndex only right before modifying the pose context FControlRigAnimInstanceProxy::Evaluate(). Mismatching the index type caused control rig component to write transform to incorrect bones.

diff --git a/Engine/Plugins/Animation/ControlRig/Source/ControlRig/Private/ControlRigAnimInstance.cpp b/Engine/Plugins/Animation/ControlRig/Source/ControlRig/Private/ControlRigAnimInstance.cpp
--- a/Engine/Plugins/Animation/ControlRig/Source/ControlRig/Private/ControlRigAnimInstance.cpp
+++ b/Engine/Plugins/Animation/ControlRig/Source/ControlRig/Private/ControlRigAnimInstance.cpp
@@ -23,10 +23,15 @@ bool FControlRigAnimInstanceProxy::Evaluate(FPoseContext& Output)
 	CSContext.Pose.InitPose(Output.Pose);
 
 	TArray<FCompactPoseBoneIndex> ModifiedBones;
-	for (TPair<int32, FTransform>& StoredTransform : StoredTransforms)
+	
+	const FBoneContainer& BoneContainer = Output.Pose.GetBoneContainer();
+	
+	for (TPair<FMeshPoseBoneIndex, FTransform>& StoredTransform : StoredTransforms)
 	{
-		const int32 BoneIndexToModify = StoredTransform.Key;
-		FCompactPoseBoneIndex CompactIndex = Output.Pose.GetBoneContainer().GetCompactPoseIndexFromSkeletonIndex(BoneIndexToModify);
+		const FMeshPoseBoneIndex& MeshPoseBoneIndexToModify = StoredTransform.Key;
+		FSkeletonPoseBoneIndex SkeletonPoseBoneIndex = BoneContainer.GetSkeletonPoseIndexFromMeshPoseIndex(MeshPoseBoneIndexToModify);
+		FCompactPoseBoneIndex CompactIndex = BoneContainer.GetCompactPoseIndexFromSkeletonPoseIndex(SkeletonPoseBoneIndex);
+		
 		if (CompactIndex.GetInt() != INDEX_NONE)
 		{
 			CSContext.Pose.SetComponentSpaceTransform(CompactIndex, StoredTransform.Value);
diff --git a/Engine/Plugins/Animation/ControlRig/Source/ControlRig/Private/ControlRigComponent.cpp b/Engine/Plugins/Animation/ControlRig/Source/ControlRig/Private/ControlRigComponent.cpp
--- a/Engine/Plugins/Animation/ControlRig/Source/ControlRig/Private/ControlRigComponent.cpp
+++ b/Engine/Plugins/Animation/ControlRig/Source/ControlRig/Private/ControlRigComponent.cpp
@@ -648,9 +648,9 @@ void UControlRigComponent::AddMappedSkeletalMesh(USkeletalMeshComponent* Skeleta
 		{
 			if (const USkeleton* Skeleton = SkeletalMesh->GetSkeleton())
 			{
-				CR->GetHierarchy()->ForEach<FRigBoneElement>([Skeleton, &BonesToMap](FRigBoneElement* BoneElement) -> bool
+				CR->GetHierarchy()->ForEach<FRigBoneElement>([SkeletalMeshComponent, &BonesToMap](FRigBoneElement* BoneElement) -> bool
 				{
-					if (Skeleton->GetReferenceSkeleton().FindBoneIndex(BoneElement->GetName()) != INDEX_NONE)
+					if (SkeletalMeshComponent->GetBoneIndex(BoneElement->GetName()) != INDEX_NONE)
 					{
 						FControlRigComponentMappedBone BoneToMap;
 						BoneToMap.Source = BoneElement->GetName();
@@ -1351,7 +1351,8 @@ void UControlRigComponent::ValidateMappingData()
 						{
 							if (MappedElement.ElementType != ERigElementType::Curve)
 							{
-								MappedElement.SubIndex = Skeleton->GetReferenceSkeleton().FindBoneIndex(MappedElement.TransformName);
+								// this is really getting the FMeshPoseBoneIndex
+								MappedElement.SubIndex = SkeletalMeshComponent->GetBoneIndex(MappedElement.TransformName);
 							}
 						}
 						else
@@ -1557,7 +1558,7 @@ void UControlRigComponent::TransferOutputs()
 						{
 							Transform = Transform.GetRelativeTransform(MappedElement.SceneComponent->GetComponentToWorld());
 						}
-						Proxy->StoredTransforms.FindOrAdd(MappedElement.SubIndex) = Transform;
+						Proxy->StoredTransforms.FindOrAdd(FMeshPoseBoneIndex(MappedElement.SubIndex)) = Transform;
 					}
 					else if (UInstancedStaticMeshComponent* InstancingComponent = Cast<UInstancedStaticMeshComponent>(MappedElement.SceneComponent))
 					{
diff --git a/Engine/Plugins/Animation/ControlRig/Source/ControlRig/Public/ControlRigAnimInstance.h b/Engine/Plugins/Animation/ControlRig/Source/ControlRig/Public/ControlRigAnimInstance.h
--- a/Engine/Plugins/Animation/ControlRig/Source/ControlRig/Public/ControlRigAnimInstance.h
+++ b/Engine/Plugins/Animation/ControlRig/Source/ControlRig/Public/ControlRigAnimInstance.h
@@ -6,6 +6,8 @@
 #include "Animation/AnimInstanceProxy.h"
 #include "ControlRigAnimInstance.generated.h"
 
+struct FMeshPoseBoneIndex;
+
 /** Proxy override for this UAnimInstance-derived class */
 USTRUCT()
 struct CONTROLRIG_API FControlRigAnimInstanceProxy : public FAnimInstanceProxy
@@ -29,7 +31,7 @@ public:
 	virtual bool Evaluate(FPoseContext& Output) override;
 	virtual void UpdateAnimationNode(const FAnimationUpdateContext& InContext) override;
 
-	TMap<int32, FTransform> StoredTransforms;
+	TMap<FMeshPoseBoneIndex, FTransform> StoredTransforms;
 	TMap<FName, float> StoredCurves;
 };
 
