Control Rig: Avoid crash when right clicking get animation channel node

diff --git a/Engine/Plugins/Animation/ControlRig/Source/ControlRig/Private/Units/Hierarchy/RigUnit_ControlChannel.cpp b/Engine/Plugins/Animation/ControlRig/Source/ControlRig/Private/Units/Hierarchy/RigUnit_ControlChannel.cpp
--- a/Engine/Plugins/Animation/ControlRig/Source/ControlRig/Private/Units/Hierarchy/RigUnit_ControlChannel.cpp
+++ b/Engine/Plugins/Animation/ControlRig/Source/ControlRig/Private/Units/Hierarchy/RigUnit_ControlChannel.cpp
@@ -7,6 +7,11 @@
 
 bool FRigUnit_GetAnimationChannelBase::UpdateCache(const URigHierarchy* InHierarchy, const FName& Control, const FName& Channel, FRigElementKey& Key, int32& Hash)
 {
+	if (!IsValid(InHierarchy))
+	{
+		return false;
+	}
+	
 	if(!Key.IsValid())
 	{
 		Hash = INDEX_NONE;
