Avoid trashing the graph pins twice when recreating the widgets

#prefligth 64060817aa00423335be582a

diff --git a/Engine/Plugins/Animation/ControlRig/Source/ControlRigEditor/Private/Graph/SControlRigGraphNode.cpp b/Engine/Plugins/Animation/ControlRig/Source/ControlRigEditor/Private/Graph/SControlRigGraphNode.cpp
--- a/Engine/Plugins/Animation/ControlRig/Source/ControlRigEditor/Private/Graph/SControlRigGraphNode.cpp
+++ b/Engine/Plugins/Animation/ControlRig/Source/ControlRigEditor/Private/Graph/SControlRigGraphNode.cpp
@@ -38,6 +38,7 @@
 #include "Slate/SlateTextures.h"
 #include "RigVMFunctions/RigVMDispatch_If.h"
 #include "RigVMFunctions/RigVMDispatch_Select.h"
+#include "Algo/Copy.h"
 
 #if WITH_EDITOR
 #include "Editor.h"
@@ -1148,8 +1149,12 @@ void SControlRigGraphNode::HandleNodePinsChanged()
 			}
 		}
 	}
-	
-	PinsToDelete.Append(LocalPinsToDelete);
+
+	// only store those pins if not already marked as garbage
+	Algo::CopyIf(LocalPinsToDelete, PinsToDelete, [](const UEdGraphPin* PinToDelete)
+	{
+		return !PinToDelete->IsPendingKill();
+	});
 
 	// Reconstruct the pin widgets. This could be done more surgically but will do for now.
 	InputPins.Reset();
