Control Rig: Fix changing library template not updating properly in UI pins

diff --git a/Engine/Plugins/Animation/ControlRig/Source/ControlRigEditor/Private/Graph/SControlRigGraphNode.cpp b/Engine/Plugins/Animation/ControlRig/Source/ControlRigEditor/Private/Graph/SControlRigGraphNode.cpp
--- a/Engine/Plugins/Animation/ControlRig/Source/ControlRigEditor/Private/Graph/SControlRigGraphNode.cpp
+++ b/Engine/Plugins/Animation/ControlRig/Source/ControlRigEditor/Private/Graph/SControlRigGraphNode.cpp
@@ -1421,6 +1421,24 @@ void SControlRigGraphNode::HandleModifiedEvent(ERigVMGraphNotifType InNotifType,
 			}
 			break;
 		}
+		case ERigVMGraphNotifType::LibraryTemplateChanged:
+		{
+			if(!ModelNode.IsValid())
+			{
+				return;
+			}
+			
+			if (URigVMLibraryNode* LibraryNode = Cast<URigVMLibraryNode>(InSubject))
+			{
+				if(LibraryNode == ModelNode.Get() ||
+					LibraryNode->GetEntryNode() == ModelNode.Get() ||
+					LibraryNode->GetReturnNode() == ModelNode.Get())
+				{
+					HandleNodePinsChanged();
+				}
+			}
+			break;
+		}
 		default:
 		{
 			break;
@@ -1441,6 +1459,12 @@ void SControlRigGraphNode::UpdatePinTreeView()
 	LeftNodeBox->ClearChildren();
 
 	UControlRigGraphNode* RigGraphNode = Cast<UControlRigGraphNode>(GraphNode);
+	URigVMNode* Node = RigGraphNode->GetModelNode();
+	const FRigVMTemplate* Template = nullptr;
+	if (URigVMTemplateNode* TemplateNode = Cast<URigVMTemplateNode>(Node))
+	{
+		Template = TemplateNode->GetTemplate();
+	}
 
 	TMap<UEdGraphPin*, int32> EdGraphPinToInputPin;
 	for(int32 InputPinIndex = 0; InputPinIndex < InputPins.Num(); InputPinIndex++)
@@ -1603,6 +1627,22 @@ void SControlRigGraphNode::UpdatePinTreeView()
 				PinWidgetForExpander = PinInfo.OutputPinWidget;
 				bPinWidgetForExpanderLeft = false;
 				bPinInfoIsValid = true;
+				if (Template)
+				{
+					if (URigVMPin* RootPin = ModelPin->GetRootPin())
+					{
+						FLinearColor PinColorAndOpacity = PinInfo.OutputPinWidget->GetColorAndOpacity();
+						if (Template->FindArgument(RootPin->GetFName()) == nullptr)
+						{
+							PinColorAndOpacity.A = 0.2f;
+						}
+						else
+						{
+							PinColorAndOpacity.A = 1.0f;
+						}
+						PinInfo.OutputPinWidget->SetColorAndOpacity(PinColorAndOpacity);
+					}
+				}
 			}
 		}
 		
@@ -1615,6 +1655,22 @@ void SControlRigGraphNode::UpdatePinTreeView()
 				PinWidgetForExpander = PinInfo.InputPinWidget;
 				bPinWidgetForExpanderLeft = true;
 				bPinInfoIsValid = true;
+				if (Template)
+				{
+					if (URigVMPin* RootPin = ModelPin->GetRootPin())
+					{
+						FLinearColor PinColorAndOpacity = PinInfo.InputPinWidget->GetColorAndOpacity();
+						if (Template->FindArgument(RootPin->GetFName()) == nullptr)
+						{
+							PinColorAndOpacity.A = 0.2f;
+						}
+						else
+						{
+							PinColorAndOpacity.A = 1.0f;
+						}
+						PinInfo.InputPinWidget->SetColorAndOpacity(PinColorAndOpacity);
+					}
+				}
 			}
 		}
 
diff --git a/Engine/Source/Developer/RigVMDeveloper/Private/RigVMModel/RigVMController.cpp b/Engine/Source/Developer/RigVMDeveloper/Private/RigVMModel/RigVMController.cpp
--- a/Engine/Source/Developer/RigVMDeveloper/Private/RigVMModel/RigVMController.cpp
+++ b/Engine/Source/Developer/RigVMDeveloper/Private/RigVMModel/RigVMController.cpp
@@ -14942,30 +14942,24 @@ void URigVMController::RepopulatePinsOnNode(URigVMNode* InNode, bool bFollowCore
 						bRequireDetachLinks = true;
 						bRequirePinStates = true;
 					}
-					
-					if(!PreviousPinsToOrphan.Contains(RootIndex))
-					{
-						URigVMPin* RootPin = Pin->GetRootPin();
 
-						bool bPinShouldBeOrphaned = bSetupOrphanPinsForThisNode; 
-						if(!bPinShouldBeOrphaned)
+					if (bSetupOrphanPinsForThisNode)
+					{
+						if(!PreviousPinsToOrphan.Contains(RootIndex))
 						{
+							URigVMPin* RootPin = Pin->GetRootPin();
+
 							if(RootPin->GetSourceLinks(true).Num() > 0 ||
-								RootPin->GetTargetLinks(true).Num() > 0)
+							   RootPin->GetTargetLinks(true).Num() > 0)
 							{
-								bPinShouldBeOrphaned = true;
-							}
-						}
-		
-						if(bPinShouldBeOrphaned)
-						{
-							PreviousPinsToOrphan.Add(RootIndex);
+								PreviousPinsToOrphan.Add(RootIndex);
 							
-							bRequireDetachLinks = true;
-							bRequirePinStates = true;
+								bRequireDetachLinks = true;
+								bRequirePinStates = true;
 #if UE_RIGVMCONTROLLER_VERBOSE_REPOPULATE
-							UE_LOG(LogRigVMDeveloper, Display, TEXT("Previously existing pin '%s' needs to be orphaned."), *RootPin->GetPinPath());
+								UE_LOG(LogRigVMDeveloper, Display, TEXT("Previously existing pin '%s' needs to be orphaned."), *RootPin->GetPinPath());
 #endif
+							}
 						}
 					}
 				}
