Control Rig Component: Guard against unreachable skeletalcomponents

diff --git a/Engine/Plugins/Animation/ControlRig/Source/ControlRig/Private/ControlRigComponent.cpp b/Engine/Plugins/Animation/ControlRig/Source/ControlRig/Private/ControlRigComponent.cpp
--- a/Engine/Plugins/Animation/ControlRig/Source/ControlRig/Private/ControlRigComponent.cpp
+++ b/Engine/Plugins/Animation/ControlRig/Source/ControlRig/Private/ControlRigComponent.cpp
@@ -175,7 +175,8 @@ void UControlRigComponent::OnUnregister()
 			{
 				if (Pair.Key->IsValidLowLevel() &&
 					!Pair.Key->HasAnyFlags(RF_BeginDestroyed) &&
-					IsValid(Pair.Key))
+					IsValid(Pair.Key) &&
+					!Pair.Key->IsUnreachable())
 				{
 					Pair.Value.Apply(Pair.Key);
 				}
