Blueprint / Control Rig: Support for filtering during ForceLoadMembers

diff --git a/Engine/Plugins/Animation/ControlRig/Source/ControlRigDeveloper/Private/ControlRigBlueprint.cpp b/Engine/Plugins/Animation/ControlRig/Source/ControlRigDeveloper/Private/ControlRigBlueprint.cpp
--- a/Engine/Plugins/Animation/ControlRig/Source/ControlRigDeveloper/Private/ControlRigBlueprint.cpp
+++ b/Engine/Plugins/Animation/ControlRig/Source/ControlRigDeveloper/Private/ControlRigBlueprint.cpp
@@ -55,6 +55,11 @@
 
 #define LOCTEXT_NAMESPACE "ControlRigBlueprint"
 
+TAutoConsoleVariable<bool> CVarEnablePreLoadFiltering(
+	TEXT("ControlRig.EnablePreLoadFiltering"),
+	true,
+	TEXT("When true the RigVMGraphs will be skipped during preload to speed up load times."));
+
 TAutoConsoleVariable<bool> CVarEnablePostLoadHashing(
 	TEXT("ControlRig.EnablePostLoadHashing"),
 	true,
@@ -123,6 +128,7 @@ UControlRigBlueprint::UControlRigBlueprint(const FObjectInitializer& ObjectIniti
 	bSuspendPythonMessagesForRigVMClient = true;
 
 #if WITH_EDITORONLY_DATA
+	ReferencedObjectPathsStored = false;
 	GizmoLibrary_DEPRECATED = nullptr;
 	ShapeLibraries.Add(UControlRigSettings::Get()->DefaultShapeLibrary);
 #endif
@@ -273,6 +279,23 @@ bool UControlRigBlueprint::CanImportGraphFromText(const FString& InClipboardText
 	return GetTemplateController(true)->CanImportNodesFromText(InClipboardText);
 }
 
+bool UControlRigBlueprint::RequiresForceLoadMembers(UObject* InObject) const
+{
+	// we can stop traversing when hitting a URigVMNode
+	// (except for collapse nodes - since they contain a graphs again)
+	if(InObject->IsA<URigVMNode>())
+	{
+		if(!InObject->IsA<URigVMCollapseNode>())
+		{
+			if(CVarEnablePreLoadFiltering->GetBool())
+			{
+				return false;
+			}
+		}
+	}
+	return UBlueprint::RequiresForceLoadMembers(InObject);
+}
+
 void UControlRigBlueprint::PostEditChangeChainProperty(FPropertyChangedChainEvent& PropertyChangedEvent)
 {
 	Super::PostEditChangeChainProperty(PropertyChangedEvent);
@@ -704,22 +727,37 @@ void UControlRigBlueprint::Serialize(FArchive& Ar)
 
 	if(Ar.IsObjectReferenceCollector())
 	{
-		TArray<IRigVMGraphFunctionHost*> ReferencedFunctionHosts = GetReferencedFunctionHosts(false);
+		Ar.UsingCustomVersion(FControlRigObjectVersion::GUID);
+
 
-		for(IRigVMGraphFunctionHost* ReferencedFunctionHost : ReferencedFunctionHosts)
+#if WITH_EDITORONLY_DATA
+		if (Ar.IsCooking() && ReferencedObjectPathsStored)
 		{
-			if (URigVMBlueprintGeneratedClass* BPGeneratedClass = Cast<URigVMBlueprintGeneratedClass>(ReferencedFunctionHost))
+			for (FSoftObjectPath ObjectPath : ReferencedObjectPaths)
 			{
-				Ar << BPGeneratedClass;
+				ObjectPath.Serialize(Ar);
 			}
 		}
-
-		for(const TSoftObjectPtr<UControlRigShapeLibrary>& ShapeLibraryPtr : ShapeLibraries)
+		else
+#endif
 		{
-			if(ShapeLibraryPtr.IsValid())
+			TArray<IRigVMGraphFunctionHost*> ReferencedFunctionHosts = GetReferencedFunctionHosts(false);
+
+			for(IRigVMGraphFunctionHost* ReferencedFunctionHost : ReferencedFunctionHosts)
 			{
-				UControlRigShapeLibrary* ShapeLibrary = ShapeLibraryPtr.Get();
-				Ar << ShapeLibrary;
+				if (URigVMBlueprintGeneratedClass* BPGeneratedClass = Cast<URigVMBlueprintGeneratedClass>(ReferencedFunctionHost))
+				{
+					Ar << BPGeneratedClass;
+				}
+			}
+
+			for(const TSoftObjectPtr<UControlRigShapeLibrary>& ShapeLibraryPtr : ShapeLibraries)
+			{
+				if(ShapeLibraryPtr.IsValid())
+				{
+					UControlRigShapeLibrary* ShapeLibrary = ShapeLibraryPtr.Get();
+					Ar << ShapeLibrary;
+				}
 			}
 		}
 	}
@@ -763,6 +801,29 @@ void UControlRigBlueprint::PreSave(FObjectPreSaveContext ObjectSaveContext)
 		}
 	}
 
+#if WITH_EDITORONLY_DATA
+	ReferencedObjectPaths.Reset();
+
+	TArray<IRigVMGraphFunctionHost*> ReferencedFunctionHosts = GetReferencedFunctionHosts(false);
+	for(IRigVMGraphFunctionHost* ReferencedFunctionHost : ReferencedFunctionHosts)
+	{
+		if (URigVMBlueprintGeneratedClass* BPGeneratedClass = Cast<URigVMBlueprintGeneratedClass>(ReferencedFunctionHost))
+		{
+			ReferencedObjectPaths.AddUnique(BPGeneratedClass);
+		}
+	}
+
+	for(const TSoftObjectPtr<UControlRigShapeLibrary>& ShapeLibraryPtr : ShapeLibraries)
+	{
+		if(ShapeLibraryPtr.IsValid())
+		{
+			UControlRigShapeLibrary* ShapeLibrary = ShapeLibraryPtr.Get();
+			ReferencedObjectPaths.AddUnique(ShapeLibrary);
+		}
+	}
+	ReferencedObjectPathsStored = true;
+#endif
+
 	bExposesAnimatableControls = false;
 	Hierarchy->ForEach<FRigControlElement>([this](FRigControlElement* ControlElement) -> bool
     {
diff --git a/Engine/Plugins/Animation/ControlRig/Source/ControlRigDeveloper/Public/ControlRigBlueprint.h b/Engine/Plugins/Animation/ControlRig/Source/ControlRigDeveloper/Public/ControlRigBlueprint.h
--- a/Engine/Plugins/Animation/ControlRig/Source/ControlRigDeveloper/Public/ControlRigBlueprint.h
+++ b/Engine/Plugins/Animation/ControlRig/Source/ControlRigDeveloper/Public/ControlRigBlueprint.h
@@ -257,6 +257,7 @@ public:
 	virtual bool ExportGraphToText(UEdGraph* InEdGraph, FString& OutText) override;
 	virtual bool TryImportGraphFromText(const FString& InClipboardText, UEdGraph** OutGraphPtr = nullptr) override;
 	virtual bool CanImportGraphFromText(const FString& InClipboardText) override;
+	virtual bool RequiresForceLoadMembers(UObject* InObject) const override;
 
 	// UObject interface
 	virtual void PostEditChangeChainProperty(struct FPropertyChangedChainEvent& PropertyChangedEvent) override;
@@ -416,6 +417,16 @@ protected:
 	UPROPERTY()
 	FRigVMClient RigVMClient;
 
+#if WITH_EDITORONLY_DATA
+
+	UPROPERTY()
+	bool ReferencedObjectPathsStored;
+
+	UPROPERTY()
+	TArray<FSoftObjectPath> ReferencedObjectPaths;
+
+#endif
+
 	/** Asset searchable information about exposed public functions on this rig */
 	UPROPERTY(AssetRegistrySearchable)
 	TArray<FControlRigPublicFunctionData> PublicFunctions_DEPRECATED;
diff --git a/Engine/Source/Developer/RigVMDeveloper/Private/RigVMModel/Nodes/RigVMAggregateNode.cpp b/Engine/Source/Developer/RigVMDeveloper/Private/RigVMModel/Nodes/RigVMAggregateNode.cpp
--- a/Engine/Source/Developer/RigVMDeveloper/Private/RigVMModel/Nodes/RigVMAggregateNode.cpp
+++ b/Engine/Source/Developer/RigVMDeveloper/Private/RigVMModel/Nodes/RigVMAggregateNode.cpp
@@ -117,14 +117,17 @@ URigVMNode* URigVMAggregateNode::GetLastInnerNode() const
 
 	if (LastInnerNodeCache == nullptr)
 	{
-		const FString ArgOppositeName = GetOppositeAggregatePin()->GetName();
-		if (IsInputAggregate())
-		{
-			LastInnerNodeCache = GetAggregateInputs().Last()->GetLinkedTargetPins()[0]->GetNode();
-		}
-		else
+		if(const URigVMPin* OppositeAggregatePin = GetOppositeAggregatePin())
 		{
-			LastInnerNodeCache = GetAggregateOutputs().Last()->GetLinkedSourcePins()[0]->GetNode();
+			const FString ArgOppositeName = OppositeAggregatePin->GetName();
+			if (IsInputAggregate())
+			{
+				LastInnerNodeCache = GetAggregateInputs().Last()->GetLinkedTargetPins()[0]->GetNode();
+			}
+			else
+			{
+				LastInnerNodeCache = GetAggregateOutputs().Last()->GetLinkedSourcePins()[0]->GetNode();
+			}
 		}
 	}
 #endif
diff --git a/Engine/Source/Developer/RigVMDeveloper/Private/RigVMModel/RigVMGraph.cpp b/Engine/Source/Developer/RigVMDeveloper/Private/RigVMModel/RigVMGraph.cpp
--- a/Engine/Source/Developer/RigVMDeveloper/Private/RigVMModel/RigVMGraph.cpp
+++ b/Engine/Source/Developer/RigVMDeveloper/Private/RigVMModel/RigVMGraph.cpp
@@ -172,7 +172,15 @@ URigVMNode* URigVMGraph::FindNode(const FString& InNodePath) const
 
 	if (URigVMCollapseNode* CollapseNode = Cast< URigVMCollapseNode>(FindNodeByName(*Left)))
 	{
-		return CollapseNode->GetContainedGraph()->FindNode(Right);
+		if(const URigVMGraph* ContainedGraph = CollapseNode->GetContainedGraph())
+		{
+			return ContainedGraph->FindNode(Right);
+		}
+		else
+		{
+			static constexpr TCHAR Format[] = TEXT("Collapse Node '%s' does not contain a subgraph.");
+			UE_LOG(LogRigVMDeveloper, Warning, Format, *CollapseNode->GetPathName());
+		}
 	}
 
 	return nullptr;
diff --git a/Engine/Source/Editor/UMGEditor/Private/WidgetBlueprintCompiler.cpp b/Engine/Source/Editor/UMGEditor/Private/WidgetBlueprintCompiler.cpp
--- a/Engine/Source/Editor/UMGEditor/Private/WidgetBlueprintCompiler.cpp
+++ b/Engine/Source/Editor/UMGEditor/Private/WidgetBlueprintCompiler.cpp
@@ -767,7 +767,7 @@ void FWidgetBlueprintCompilerContext::FinishCompilingClass(UClass* Class)
 	{
 		if( !WidgetBP->bHasBeenRegenerated )
 		{
-			UBlueprint::ForceLoadMembers(WidgetBP->WidgetTree);
+			UBlueprint::ForceLoadMembers(WidgetBP->WidgetTree, WidgetBP);
 		}
 
 		FixAbandonedWidgetTree(WidgetBP);
diff --git a/Engine/Source/Editor/UnrealEd/Private/Kismet2/BlueprintEditorUtils.cpp b/Engine/Source/Editor/UnrealEd/Private/Kismet2/BlueprintEditorUtils.cpp
--- a/Engine/Source/Editor/UnrealEd/Private/Kismet2/BlueprintEditorUtils.cpp
+++ b/Engine/Source/Editor/UnrealEd/Private/Kismet2/BlueprintEditorUtils.cpp
@@ -1369,7 +1369,7 @@ UClass* FBlueprintEditorUtils::RegenerateBlueprintClass(UBlueprint* Blueprint, U
 		UBlueprint::ForceLoadMetaData(Blueprint);
 		if (ensure(PreviousCDO))
 		{
-			UBlueprint::ForceLoadMembers(PreviousCDO);
+			UBlueprint::ForceLoadMembers(PreviousCDO, Blueprint);
 		}
 		UBlueprint::ForceLoadMembers(Blueprint);
 	}
diff --git a/Engine/Source/Runtime/Engine/Classes/Engine/Blueprint.h b/Engine/Source/Runtime/Engine/Classes/Engine/Blueprint.h
--- a/Engine/Source/Runtime/Engine/Classes/Engine/Blueprint.h
+++ b/Engine/Source/Runtime/Engine/Classes/Engine/Blueprint.h
@@ -866,10 +866,14 @@ public:
 #endif // WITH_EDITORONLY_DATA
 
 #if WITH_EDITOR
+	virtual bool RequiresForceLoadMembers(UObject* InObject) const { return true; }
+
 	static bool ForceLoad(UObject* Obj);
 
 	static void ForceLoadMembers(UObject* InObject);
 
+	static void ForceLoadMembers(UObject* InObject, const UBlueprint* InBlueprint);
+
 	static void ForceLoadMetaData(UObject* InObject);
 
 	static bool ValidateGeneratedClass(const UClass* InClass);
diff --git a/Engine/Source/Runtime/Engine/Private/Blueprint.cpp b/Engine/Source/Runtime/Engine/Private/Blueprint.cpp
--- a/Engine/Source/Runtime/Engine/Private/Blueprint.cpp
+++ b/Engine/Source/Runtime/Engine/Private/Blueprint.cpp
@@ -1206,6 +1206,42 @@ bool UBlueprint::ForceLoad(UObject* Obj)
 
 void UBlueprint::ForceLoadMembers(UObject* InObject)
 {
+	if(const UBlueprint* Blueprint = Cast<UBlueprint>(InObject))
+	{
+		ForceLoadMembers(InObject, Blueprint);
+		return;
+	}
+
+	if(const UClass* Class = Cast<UClass>(InObject))
+	{
+		ForceLoadMembers(InObject, Cast<UBlueprint>(Class->ClassGeneratedBy));
+		return;
+	}
+
+	if(InObject->HasAnyFlags(RF_ClassDefaultObject))
+	{
+		if(const UClass* Class = InObject->GetClass())
+		{
+			ForceLoadMembers(InObject, Cast<UBlueprint>(Class->ClassGeneratedBy));
+			return;
+		}
+	}
+
+	ForceLoadMembers(InObject, nullptr);
+}
+
+void UBlueprint::ForceLoadMembers(UObject* InObject, const UBlueprint* InBlueprint)
+{
+	check(InObject);
+	
+	if(InObject && InBlueprint)
+	{
+		if(!InBlueprint->RequiresForceLoadMembers(InObject))
+		{
+			return;
+		}
+	}
+	
 	// Collect a list of all things this element owns
 	TArray<UObject*> MemberReferences;
 	FReferenceFinder ComponentCollector(MemberReferences, InObject, false, true, true, true);
@@ -1217,7 +1253,7 @@ void UBlueprint::ForceLoadMembers(UObject* InObject)
 		UObject* CurrentObject = *it;
 		if (ForceLoad(CurrentObject))
 		{
-			ForceLoadMembers(CurrentObject);
+			ForceLoadMembers(CurrentObject, InBlueprint);
 		}
 	}
 }
