Control Rig: Change check to log warning message when encountered with corrupted graphs

diff --git a/Engine/Plugins/Animation/ControlRig/Source/ControlRigDeveloper/Private/Graph/ControlRigGraphNode.cpp b/Engine/Plugins/Animation/ControlRig/Source/ControlRigDeveloper/Private/Graph/ControlRigGraphNode.cpp
--- a/Engine/Plugins/Animation/ControlRig/Source/ControlRigDeveloper/Private/Graph/ControlRigGraphNode.cpp
+++ b/Engine/Plugins/Animation/ControlRig/Source/ControlRigDeveloper/Private/Graph/ControlRigGraphNode.cpp
@@ -1320,7 +1320,8 @@ FSlateIcon UControlRigGraphNode::GetIconAndTint(FLinearColor& OutColor) const
 		}
 		else
 		{
-			checkf(false, TEXT("Could not find first inner node in aggregate node %s in package %s"), *ModelNode->GetPathName(), *GetPackage()->GetPathName());
+			const FString Message = FString::Printf(TEXT("Could not find first inner node in aggregate node %s in package %s"), *ModelNode->GetPathName(), *GetPackage()->GetPathName());
+			FScriptExceptionHandler::Get().HandleException(ELogVerbosity::Warning, *Message, *FString());
 		}
 
 		if (ModelNode->IsA<URigVMFunctionReferenceNode>())
