GraphEditor: Guard against trashed UEdGraphPins

diff --git a/Engine/Plugins/Animation/ControlRig/Source/ControlRigDeveloper/Private/Graph/ControlRigGraph.cpp b/Engine/Plugins/Animation/ControlRig/Source/ControlRigDeveloper/Private/Graph/ControlRigGraph.cpp
--- a/Engine/Plugins/Animation/ControlRig/Source/ControlRigDeveloper/Private/Graph/ControlRigGraph.cpp
+++ b/Engine/Plugins/Animation/ControlRig/Source/ControlRigDeveloper/Private/Graph/ControlRigGraph.cpp
@@ -1155,6 +1155,11 @@ void UControlRigGraph::RemoveNode(UEdGraphNode* InNode)
 	// Make sure EdGraph is not part of the transaction
 	TGuardValue<ITransaction*> TransactionGuard(GUndo, nullptr);
 
+	if(UControlRigGraphNode* RigNode = Cast<UControlRigGraphNode>(InNode))
+	{
+		RigNode->OnNodeBeginRemoval().Broadcast();
+	}
+
 	// clear out the pin relationships
 	for(UEdGraphPin* Pin : InNode->Pins)
 	{
@@ -1176,9 +1181,9 @@ void UControlRigGraph::RemoveNode(UEdGraphNode* InNode)
 		while (ExistingObject);
 	}
 	InNode->Rename(*DeletedName, GetTransientPackage(), REN_ForceNoResetLoaders | REN_DontCreateRedirectors);	
-	Super::RemoveNode(InNode);
 
-	NotifyGraphChanged(FEdGraphEditAction(EEdGraphActionType::GRAPHACTION_RemoveNode, this, InNode, false));
+	// this also subsequently calls NotifyGraphChanged
+	Super::RemoveNode(InNode);
 }
 
 URigVMController* UControlRigGraph::GetTemplateController()
diff --git a/Engine/Plugins/Animation/ControlRig/Source/ControlRigDeveloper/Public/Graph/ControlRigGraphNode.h b/Engine/Plugins/Animation/ControlRig/Source/ControlRigDeveloper/Public/Graph/ControlRigGraphNode.h
--- a/Engine/Plugins/Animation/ControlRig/Source/ControlRigDeveloper/Public/Graph/ControlRigGraphNode.h
+++ b/Engine/Plugins/Animation/ControlRig/Source/ControlRigDeveloper/Public/Graph/ControlRigGraphNode.h
@@ -81,6 +81,7 @@ public:
 	DECLARE_MULTICAST_DELEGATE(FNodeTitleDirtied);
 	DECLARE_MULTICAST_DELEGATE(FNodePinsChanged);
 	DECLARE_MULTICAST_DELEGATE(FNodePinExpansionChanged);
+	DECLARE_MULTICAST_DELEGATE(FNodeBeginRemoval);
 
 	struct FPinPair
 	{
@@ -190,6 +191,7 @@ public:
 	FNodeTitleDirtied& OnNodeTitleDirtied() { return NodeTitleDirtied; }
 	FNodePinsChanged& OnNodePinsChanged() { return NodePinsChanged; }
 	FNodePinExpansionChanged& OnNodePinExpansionChanged() { return NodePinExpansionChanged; }
+	FNodeBeginRemoval& OnNodeBeginRemoval() { return NodeBeginRemoval; }
 
 	/** Called when there's a drastic change in the pins */
 	bool ModelPinsChanged(bool bForce = false);
@@ -268,6 +270,7 @@ private:
 	FNodeTitleDirtied NodeTitleDirtied;
 	FNodePinsChanged NodePinsChanged;
 	FNodePinExpansionChanged NodePinExpansionChanged;
+	FNodeBeginRemoval NodeBeginRemoval;
 
 	mutable const FRigVMTemplate* CachedTemplate;
 	mutable TOptional<bool> DrawAsCompactNodeCache;
diff --git a/Engine/Plugins/Animation/ControlRig/Source/ControlRigEditor/Private/Graph/SControlRigGraphNode.cpp b/Engine/Plugins/Animation/ControlRig/Source/ControlRigEditor/Private/Graph/SControlRigGraphNode.cpp
--- a/Engine/Plugins/Animation/ControlRig/Source/ControlRigEditor/Private/Graph/SControlRigGraphNode.cpp
+++ b/Engine/Plugins/Animation/ControlRig/Source/ControlRigEditor/Private/Graph/SControlRigGraphNode.cpp
@@ -112,6 +112,7 @@ void SControlRigGraphNode::Construct( const FArguments& InArgs )
 
 	EdGraphNode->OnNodeTitleDirtied().AddSP(this, &SControlRigGraphNode::HandleNodeTitleDirtied);
 	EdGraphNode->OnNodePinsChanged().AddSP(this, &SControlRigGraphNode::HandleNodePinsChanged);
+	EdGraphNode->OnNodeBeginRemoval().AddSP(this, &SControlRigGraphNode::HandleNodeBeginRemoval);
 
 	LastHighDetailSize = FVector2D::ZeroVector;
 }
@@ -1169,6 +1170,37 @@ void SControlRigGraphNode::HandleNodePinsChanged()
 	UpdatePinTreeView();
 }
 
+void SControlRigGraphNode::HandleNodeBeginRemoval()
+{
+	if(UControlRigGraphNode* RigNode = Cast<UControlRigGraphNode>(GraphNode))
+	{
+		RigNode->OnNodeTitleDirtied().RemoveAll(this);
+		RigNode->OnNodePinsChanged().RemoveAll(this);
+		RigNode->OnNodeBeginRemoval().RemoveAll(this);
+		
+		if(URigVMController* Controller = RigNode->GetController())
+		{
+			Controller->OnModified().RemoveAll(this);
+		}
+	}
+	
+	for (const TSharedRef<SGraphPin>& GraphPin: InputPins)
+	{
+		GraphPin->SetPinObj(nullptr);
+	}
+	for (const TSharedRef<SGraphPin>& GraphPin: OutputPins)
+	{
+		GraphPin->SetPinObj(nullptr);
+	}
+
+	InputPins.Reset();
+	OutputPins.Reset();
+	PinInfos.Reset();
+	PinsToKeep.Reset();
+	
+	InvalidateGraphData();
+}
+
 FText SControlRigGraphNode::GetInstructionCountText() const
 {
 	if(Blueprint.IsValid())
diff --git a/Engine/Plugins/Animation/ControlRig/Source/ControlRigEditor/Private/Graph/SControlRigGraphNode.h b/Engine/Plugins/Animation/ControlRig/Source/ControlRigEditor/Private/Graph/SControlRigGraphNode.h
--- a/Engine/Plugins/Animation/ControlRig/Source/ControlRigEditor/Private/Graph/SControlRigGraphNode.h
+++ b/Engine/Plugins/Animation/ControlRig/Source/ControlRigEditor/Private/Graph/SControlRigGraphNode.h
@@ -81,6 +81,7 @@ private:
 
 	void HandleNodeTitleDirtied();
 	void HandleNodePinsChanged();
+	void HandleNodeBeginRemoval();
 
 	FText GetInstructionCountText() const;
 	FText GetInstructionDurationText() const;
diff --git a/Engine/Source/Editor/GraphEditor/Private/SGraphPin.cpp b/Engine/Source/Editor/GraphEditor/Private/SGraphPin.cpp
--- a/Engine/Source/Editor/GraphEditor/Private/SGraphPin.cpp
+++ b/Engine/Source/Editor/GraphEditor/Private/SGraphPin.cpp
@@ -906,6 +906,12 @@ void SGraphPin::Tick( const FGeometry& AllottedGeometry, const double InCurrentT
 UEdGraphPin* SGraphPin::GetPinObj() const
 {
 	ensureMsgf(!bGraphDataInvalid, TEXT("The Graph Pin Object has been invalidated. Someone is keeping a hard ref on the SGraphPin (%s). See InvalidateGraphData for more info"), *ToString());
+
+	if(GraphPinObj && GraphPinObj->bWasTrashed)
+	{
+		return nullptr;
+	}
+
 	return !bGraphDataInvalid ? GraphPinObj : nullptr;
 }
 
@@ -916,6 +922,11 @@ void SGraphPin::SetOwner( const TSharedRef<SGraphNode> OwnerNode )
 	OwnerNodePtr = OwnerNode;
 }
 
+void SGraphPin::SetPinObj(UEdGraphPin* PinObj)
+{
+	GraphPinObj = PinObj;
+}
+
 EVisibility SGraphPin::IsPinVisibleAsAdvanced() const
 {
 	bool bHideAdvancedPin = false;
diff --git a/Engine/Source/Editor/GraphEditor/Public/SGraphPin.h b/Engine/Source/Editor/GraphEditor/Public/SGraphPin.h
--- a/Engine/Source/Editor/GraphEditor/Public/SGraphPin.h
+++ b/Engine/Source/Editor/GraphEditor/Public/SGraphPin.h
@@ -157,6 +157,9 @@ public:
 	/** @param OwnerNode  The SGraphNode that this pin belongs to */
 	void SetOwner( const TSharedRef<SGraphNode> OwnerNode );
 
+	/** @param PinObj The UEdGraphPin object this pin widget represents */
+	void SetPinObj(UEdGraphPin* PinObj);
+
 	/** @return whether this pin is incoming or outgoing */
 	EEdGraphPinDirection GetDirection() const;
 
diff --git a/Engine/Source/Runtime/Engine/Private/EdGraph/EdGraphNode.cpp b/Engine/Source/Runtime/Engine/Private/EdGraph/EdGraphNode.cpp
--- a/Engine/Source/Runtime/Engine/Private/EdGraph/EdGraphNode.cpp
+++ b/Engine/Source/Runtime/Engine/Private/EdGraph/EdGraphNode.cpp
@@ -507,7 +507,10 @@ class UEdGraph* UEdGraphNode::GetGraph() const
 	UEdGraph* Graph = Cast<UEdGraph>(GetOuter());
 	if (Graph == nullptr && IsValid(this))
 	{
-		ensureMsgf(false, TEXT("EdGraphNode::GetGraph : '%s' does not have a UEdGraph as an Outer."), *GetPathName());
+		if(GetOuter() != GetTransientPackage())
+		{
+			ensureMsgf(false, TEXT("EdGraphNode::GetGraph : '%s' does not have a UEdGraph as an Outer."), *GetPathName());
+		}
 	}
 	return Graph;
 }
