do not call RebuildDebugDrawSkeleton when the construction event is running
    the HierarchyReset notification is sent when resetting/copying the hierarchy leading in an undefined Elements array so calling RebuildDebugDrawSkeleton before it has been rebuilt is unstable

diff --git a/Engine/Plugins/Animation/ControlRig/Source/ControlRigEditor/Private/Editor/ControlRigSkeletalMeshComponent.cpp b/Engine/Plugins/Animation/ControlRig/Source/ControlRigEditor/Private/Editor/ControlRigSkeletalMeshComponent.cpp
--- a/Engine/Plugins/Animation/ControlRig/Source/ControlRigEditor/Private/Editor/ControlRigSkeletalMeshComponent.cpp
+++ b/Engine/Plugins/Animation/ControlRig/Source/ControlRigEditor/Private/Editor/ControlRigSkeletalMeshComponent.cpp
@@ -13,6 +13,7 @@ UControlRigSkeletalMeshComponent::UControlRigSkeletalMeshComponent(const FObject
 	, DebugDrawSkeleton(false)
 	, HierarchyInteractionBracket(0)
 	, bRebuildDebugDrawSkeletonRequired(false)
+	, bIsConstructionEventRunning(false)
 {
 	SetDisablePostProcessBlueprint(true);
 }
@@ -48,6 +49,7 @@ void UControlRigSkeletalMeshComponent::InitAnim(bool bForceReinit)
 		}
 	}
 
+	bRebuildDebugDrawSkeletonRequired = true;
 	RebuildDebugDrawSkeleton();
 }
 
@@ -63,8 +65,11 @@ void UControlRigSkeletalMeshComponent::SetCustomDefaultPose()
 
 void UControlRigSkeletalMeshComponent::RebuildDebugDrawSkeleton()
 {
-	bRebuildDebugDrawSkeletonRequired = false;
- 
+	if ((HierarchyInteractionBracket != 0) || bIsConstructionEventRunning || !bRebuildDebugDrawSkeletonRequired)
+	{
+		return;
+	}
+	
 	UControlRigLayerInstance* ControlRigInstance = Cast<UControlRigLayerInstance>(GetAnimInstance());
 
 	if (ControlRigInstance)
@@ -136,6 +141,8 @@ void UControlRigSkeletalMeshComponent::RebuildDebugDrawSkeleton()
 			}
 		}
 	}
+	
+	bRebuildDebugDrawSkeletonRequired = false;
 }
 
 FTransform UControlRigSkeletalMeshComponent::GetDrawTransform(int32 BoneIndex) const
@@ -180,15 +187,31 @@ void UControlRigSkeletalMeshComponent::SetControlRigBeingDebugged(UControlRig* I
 			{
 				ControlRigBeingDebugged->GetHierarchy()->OnModified().RemoveAll(this);
 			}
+
+#if WITH_EDITOR
+			ControlRigBeingDebugged->OnPreConstructionForUI_AnyThread().RemoveAll(this);
+			ControlRigBeingDebugged->OnPostConstruction_AnyThread().RemoveAll(this);
+#endif
 		}
 	}
 
 	ControlRigBeingDebuggedPtr.Reset();
+	bRebuildDebugDrawSkeletonRequired = true;
 	
 	if(InControlRig)
 	{
 		ControlRigBeingDebuggedPtr = InControlRig;
+
+		InControlRig->GetHierarchy()->OnModified().RemoveAll(this);
  		InControlRig->GetHierarchy()->OnModified().AddUObject(this, &UControlRigSkeletalMeshComponent::OnHierarchyModified_AnyThread);
+
+#if WITH_EDITOR
+		InControlRig->OnPreConstructionForUI_AnyThread().RemoveAll(this);
+		InControlRig->OnPreConstructionForUI_AnyThread().AddUObject(this, &UControlRigSkeletalMeshComponent::OnPreConstruction_AnyThread);
+		
+		InControlRig->OnPostConstruction_AnyThread().RemoveAll(this);
+		InControlRig->OnPostConstruction_AnyThread().AddUObject(this, &UControlRigSkeletalMeshComponent::OnPostConstruction_AnyThread);
+#endif
 	}
 
 	RebuildDebugDrawSkeleton();
@@ -223,10 +246,7 @@ void UControlRigSkeletalMeshComponent::OnHierarchyModified(ERigHierarchyNotifica
 			HierarchyInteractionBracket = FMath::Max(HierarchyInteractionBracket - 1, 0);
 			if(HierarchyInteractionBracket == 0)
 			{
-				if(bRebuildDebugDrawSkeletonRequired)
-				{
-					RebuildDebugDrawSkeleton();
-				}
+				RebuildDebugDrawSkeleton();
 			}
 			break;
 		}
@@ -240,6 +260,11 @@ void UControlRigSkeletalMeshComponent::OnHierarchyModified(ERigHierarchyNotifica
 void UControlRigSkeletalMeshComponent::OnHierarchyModified_AnyThread(ERigHierarchyNotification InNotif,
     URigHierarchy* InHierarchy, const FRigBaseElement* InElement)
 {
+	if (bIsConstructionEventRunning)
+	{
+		return;
+	}
+	
 	FRigElementKey Key;
 	if(InElement)
 	{
@@ -248,17 +273,47 @@ void UControlRigSkeletalMeshComponent::OnHierarchyModified_AnyThread(ERigHierarc
 
 	TWeakObjectPtr<URigHierarchy> WeakHierarchy = InHierarchy;
 	
-	FFunctionGraphTask::CreateAndDispatchWhenReady([this, InNotif, WeakHierarchy, Key]()
-    {
-    	if(!WeakHierarchy.IsValid())
-    	{
-    		return;
-    	}
-        if (const FRigBaseElement* Element = WeakHierarchy.Get()->Find(Key))
-        {
-	        OnHierarchyModified(InNotif, WeakHierarchy.Get(), Element);
-        }
-		
-    }, TStatId(), NULL, ENamedThreads::GameThread);
+	auto Task = [this, InNotif, WeakHierarchy, Key]()
+	{
+		if(!WeakHierarchy.IsValid())
+		{
+			return;
+		}
+		const FRigBaseElement* Element = WeakHierarchy.Get()->Find(Key);
+		OnHierarchyModified(InNotif, WeakHierarchy.Get(), Element);
+	};
+
+	if (IsInGameThread())
+	{
+		Task();
+	}
+	else
+	{
+		FFunctionGraphTask::CreateAndDispatchWhenReady([Task]()
+		{
+			Task();		
+		}, TStatId(), NULL, ENamedThreads::GameThread);
+	}
 }
 
+void UControlRigSkeletalMeshComponent::OnPreConstruction_AnyThread(UControlRig* InControlRig, const EControlRigState InState, const FName& InEventName)
+{
+	bIsConstructionEventRunning = true;
+}
+
+void UControlRigSkeletalMeshComponent::OnPostConstruction_AnyThread(UControlRig* InControlRig, const EControlRigState InState, const FName& InEventName)
+{
+	bIsConstructionEventRunning = false;
+
+	if (IsInGameThread())
+	{
+		RebuildDebugDrawSkeleton();
+	}
+	else
+	{
+		FFunctionGraphTask::CreateAndDispatchWhenReady([this]()
+		{
+			RebuildDebugDrawSkeleton();
+		}, TStatId(), NULL, ENamedThreads::GameThread);
+	}
+}
\ No newline at end of file
diff --git a/Engine/Plugins/Animation/ControlRig/Source/ControlRigEditor/Private/Editor/ControlRigSkeletalMeshComponent.h b/Engine/Plugins/Animation/ControlRig/Source/ControlRigEditor/Private/Editor/ControlRigSkeletalMeshComponent.h
--- a/Engine/Plugins/Animation/ControlRig/Source/ControlRigEditor/Private/Editor/ControlRigSkeletalMeshComponent.h
+++ b/Engine/Plugins/Animation/ControlRig/Source/ControlRigEditor/Private/Editor/ControlRigSkeletalMeshComponent.h
@@ -54,6 +54,9 @@ public:
 	void OnHierarchyModified(ERigHierarchyNotification InNotif, URigHierarchy* InHierarchy, const FRigBaseElement* InElement);
 	void OnHierarchyModified_AnyThread(ERigHierarchyNotification InNotif, URigHierarchy* InHierarchy, const FRigBaseElement* InElement);
 
+	void OnPreConstruction_AnyThread(UControlRig* InControlRig, const EControlRigState InState, const FName& InEventName);
+	void OnPostConstruction_AnyThread(UControlRig* InControlRig, const EControlRigState InState, const FName& InEventName);
+
 private:
 	FReferenceSkeleton DebugDrawSkeleton;
 	TArray<FBoneIndexType> DebugDrawBones;
@@ -61,4 +64,5 @@ private:
 	TWeakObjectPtr<UControlRig> ControlRigBeingDebuggedPtr;
 	int32 HierarchyInteractionBracket;
 	bool bRebuildDebugDrawSkeletonRequired;
+	bool bIsConstructionEventRunning;
 };
