follow up of 24556971 (delegate changed signature)

diff --git a/Engine/Plugins/Animation/ControlRig/Source/ControlRigEditor/Private/Editor/ControlRigSkeletalMeshComponent.cpp b/Engine/Plugins/Animation/ControlRig/Source/ControlRigEditor/Private/Editor/ControlRigSkeletalMeshComponent.cpp
--- a/Engine/Plugins/Animation/ControlRig/Source/ControlRigEditor/Private/Editor/ControlRigSkeletalMeshComponent.cpp
+++ b/Engine/Plugins/Animation/ControlRig/Source/ControlRigEditor/Private/Editor/ControlRigSkeletalMeshComponent.cpp
@@ -279,8 +279,10 @@ void UControlRigSkeletalMeshComponent::OnHierarchyModified_AnyThread(ERigHierarc
 		{
 			return;
 		}
-		const FRigBaseElement* Element = WeakHierarchy.Get()->Find(Key);
-		OnHierarchyModified(InNotif, WeakHierarchy.Get(), Element);
+		if (const FRigBaseElement* Element = WeakHierarchy.Get()->Find(Key))
+		{
+			OnHierarchyModified(InNotif, WeakHierarchy.Get(), Element);
+		}
 	};
 
 	if (IsInGameThread())
@@ -296,12 +298,12 @@ void UControlRigSkeletalMeshComponent::OnHierarchyModified_AnyThread(ERigHierarc
 	}
 }
 
-void UControlRigSkeletalMeshComponent::OnPreConstruction_AnyThread(UControlRig* InControlRig, const EControlRigState InState, const FName& InEventName)
+void UControlRigSkeletalMeshComponent::OnPreConstruction_AnyThread(UControlRig* InControlRig, const FName& InEventName)
 {
 	bIsConstructionEventRunning = true;
 }
 
-void UControlRigSkeletalMeshComponent::OnPostConstruction_AnyThread(UControlRig* InControlRig, const EControlRigState InState, const FName& InEventName)
+void UControlRigSkeletalMeshComponent::OnPostConstruction_AnyThread(UControlRig* InControlRig, const FName& InEventName)
 {
 	bIsConstructionEventRunning = false;
 
diff --git a/Engine/Plugins/Animation/ControlRig/Source/ControlRigEditor/Private/Editor/ControlRigSkeletalMeshComponent.h b/Engine/Plugins/Animation/ControlRig/Source/ControlRigEditor/Private/Editor/ControlRigSkeletalMeshComponent.h
--- a/Engine/Plugins/Animation/ControlRig/Source/ControlRigEditor/Private/Editor/ControlRigSkeletalMeshComponent.h
+++ b/Engine/Plugins/Animation/ControlRig/Source/ControlRigEditor/Private/Editor/ControlRigSkeletalMeshComponent.h
@@ -54,8 +54,8 @@ public:
 	void OnHierarchyModified(ERigHierarchyNotification InNotif, URigHierarchy* InHierarchy, const FRigBaseElement* InElement);
 	void OnHierarchyModified_AnyThread(ERigHierarchyNotification InNotif, URigHierarchy* InHierarchy, const FRigBaseElement* InElement);
 
-	void OnPreConstruction_AnyThread(UControlRig* InControlRig, const EControlRigState InState, const FName& InEventName);
-	void OnPostConstruction_AnyThread(UControlRig* InControlRig, const EControlRigState InState, const FName& InEventName);
+	void OnPreConstruction_AnyThread(UControlRig* InControlRig, const FName& InEventName);
+	void OnPostConstruction_AnyThread(UControlRig* InControlRig, const FName& InEventName);
 
 private:
 	FReferenceSkeleton DebugDrawSkeleton;
