Control Rig: Add space to offset and initial for spawn control nodes

diff --git a/Engine/Plugins/Animation/ControlRig/Source/ControlRig/Private/Units/Execution/RigUnit_DynamicHierarchy.cpp b/Engine/Plugins/Animation/ControlRig/Source/ControlRig/Private/Units/Execution/RigUnit_DynamicHierarchy.cpp
--- a/Engine/Plugins/Animation/ControlRig/Source/ControlRig/Private/Units/Execution/RigUnit_DynamicHierarchy.cpp
+++ b/Engine/Plugins/Animation/ControlRig/Source/ControlRig/Private/Units/Execution/RigUnit_DynamicHierarchy.cpp
@@ -370,6 +370,19 @@ FRigUnit_HierarchyAddNull_Execute()
 	}
 }
 
+FTransform FRigUnit_HierarchyAddControlElement::ProjectOffsetTransform(const FTransform& InOffsetTransform,
+	ERigVMTransformSpace InOffsetSpace, const FRigElementKey& InParent, const URigHierarchy* InHierarchy)
+{
+	if(InOffsetSpace == ERigVMTransformSpace::GlobalSpace && InParent.IsValid())
+	{
+		const FTransform ParentTransform = InHierarchy->GetGlobalTransform(InParent);
+		FTransform OffsetTransform = InOffsetTransform.GetRelativeTransform(ParentTransform);
+		OffsetTransform.NormalizeRotation();
+		return OffsetTransform;
+	}
+	return InOffsetTransform;
+}
+
 void FRigUnit_HierarchyAddControl_Settings::Configure(FRigControlSettings& OutSettings) const
 {
 	OutSettings.DisplayName = DisplayName;
@@ -429,11 +442,13 @@ FRigUnit_HierarchyAddControlFloat_Execute()
 		FRigControlSettings ControlSettings;
 		Settings.Configure(ControlSettings);
 
+		const FTransform Offset = ProjectOffsetTransform(OffsetTransform, OffsetSpace, Parent, ExecuteContext.Hierarchy);
+
 		const FRigControlValue Value = FRigControlValue::Make<float>(InitialValue);
 		const FTransform ShapeTransform = Settings.Shape.Transform;
 		
 		FRigHierarchyControllerInstructionBracket InstructionBracket(Controller, ExecuteContext.GetInstructionIndex());
-		Item = Controller->AddControl(Name, Parent, ControlSettings, Value, OffsetTransform, ShapeTransform, false, false);
+		Item = Controller->AddControl(Name, Parent, ControlSettings, Value, Offset, ShapeTransform, false, false);
 	}
 }
 
@@ -477,11 +492,13 @@ FRigUnit_HierarchyAddControlInteger_Execute()
 		FRigControlSettings ControlSettings;
 		Settings.Configure(ControlSettings);
 
+		const FTransform Offset = ProjectOffsetTransform(OffsetTransform, OffsetSpace, Parent, ExecuteContext.Hierarchy);
+
 		const FRigControlValue Value = FRigControlValue::Make<int32>(InitialValue);
 		const FTransform ShapeTransform = Settings.Shape.Transform;
 		
 		FRigHierarchyControllerInstructionBracket InstructionBracket(Controller, ExecuteContext.GetInstructionIndex());
-		Item = Controller->AddControl(Name, Parent, ControlSettings, Value, OffsetTransform, ShapeTransform, false, false);
+		Item = Controller->AddControl(Name, Parent, ControlSettings, Value, Offset, ShapeTransform, false, false);
 	}
 }
 
@@ -527,11 +544,13 @@ FRigUnit_HierarchyAddControlVector2D_Execute()
 		FRigControlSettings ControlSettings;
 		Settings.Configure(ControlSettings);
 
+		const FTransform Offset = ProjectOffsetTransform(OffsetTransform, OffsetSpace, Parent, ExecuteContext.Hierarchy);
+
 		const FRigControlValue Value = FRigControlValue::Make<FVector3f>(FVector3f(InitialValue.X, InitialValue.Y, 0.f));
 		const FTransform ShapeTransform = Settings.Shape.Transform;
 		
 		FRigHierarchyControllerInstructionBracket InstructionBracket(Controller, ExecuteContext.GetInstructionIndex());
-		Item = Controller->AddControl(Name, Parent, ControlSettings, Value, OffsetTransform, ShapeTransform, false, false);
+		Item = Controller->AddControl(Name, Parent, ControlSettings, Value, Offset, ShapeTransform, false, false);
 	}
 }
 
@@ -577,11 +596,29 @@ FRigUnit_HierarchyAddControlVector_Execute()
 		FRigControlSettings ControlSettings;
 		Settings.Configure(ControlSettings);
 
-		const FRigControlValue Value = FRigControlValue::Make<FVector3f>(FVector3f(InitialValue));
+		const FTransform Offset = ProjectOffsetTransform(OffsetTransform, OffsetSpace, Parent, ExecuteContext.Hierarchy);
+
+		FVector Vector = InitialValue;
+		if(Settings.InitialSpace == ERigVMTransformSpace::GlobalSpace)
+		{
+			if(Settings.bIsPosition)
+			{
+				const FTransform ParentTransform = ExecuteContext.Hierarchy->GetGlobalTransform(Parent);
+				const FTransform GlobalOffsetTransform = Offset * ParentTransform;
+				const FTransform LocalTransform = FTransform(Vector).GetRelativeTransform(GlobalOffsetTransform);
+				Vector = LocalTransform.GetLocation();
+			}
+			else
+			{
+				Vector = Vector * ExecuteContext.Hierarchy->GetGlobalTransform(Parent).GetScale3D(); 
+			}
+		}
+
+		const FRigControlValue Value = FRigControlValue::Make<FVector3f>(FVector3f(Vector));
 		const FTransform ShapeTransform = Settings.Shape.Transform;
 		
 		FRigHierarchyControllerInstructionBracket InstructionBracket(Controller, ExecuteContext.GetInstructionIndex());
-		Item = Controller->AddControl(Name, Parent, ControlSettings, Value, OffsetTransform, ShapeTransform, false, false);
+		Item = Controller->AddControl(Name, Parent, ControlSettings, Value, Offset, ShapeTransform, false, false);
 	}
 }
 
@@ -627,11 +664,23 @@ FRigUnit_HierarchyAddControlRotator_Execute()
 		FRigControlSettings ControlSettings;
 		Settings.Configure(ControlSettings);
 
-		const FRigControlValue Value = FRigControlValue::Make<FVector3f>(FVector3f(InitialValue.Euler()));
+		const FTransform Offset = ProjectOffsetTransform(OffsetTransform, OffsetSpace, Parent, ExecuteContext.Hierarchy);
+
+		FRotator Rotator = InitialValue;
+		if(Settings.InitialSpace == ERigVMTransformSpace::GlobalSpace)
+		{
+			const FTransform ParentTransform = ExecuteContext.Hierarchy->GetGlobalTransform(Parent);
+			const FTransform GlobalOffsetTransform = Offset * ParentTransform;
+			FTransform LocalTransform = FTransform(Rotator).GetRelativeTransform(GlobalOffsetTransform);
+			LocalTransform.NormalizeRotation();
+			Rotator = LocalTransform.Rotator(); 
+		}
+
+		const FRigControlValue Value = FRigControlValue::Make<FVector3f>(FVector3f(Rotator.Euler()));
 		const FTransform ShapeTransform = Settings.Shape.Transform;
 		
 		FRigHierarchyControllerInstructionBracket InstructionBracket(Controller, ExecuteContext.GetInstructionIndex());
-		Item = Controller->AddControl(Name, Parent, ControlSettings, Value, OffsetTransform, ShapeTransform, false, false);
+		Item = Controller->AddControl(Name, Parent, ControlSettings, Value, Offset, ShapeTransform, false, false);
 	}
 }
 
@@ -665,12 +714,23 @@ FRigUnit_HierarchyAddControlTransform_Execute()
 		FRigControlSettings ControlSettings;
 		Settings.Configure(ControlSettings);
 
-		const FEulerTransform EulerTransform(InitialValue);
+		const FTransform Offset = ProjectOffsetTransform(OffsetTransform, OffsetSpace, Parent, ExecuteContext.Hierarchy);
+
+		FTransform Transform = InitialValue;
+		if(Settings.InitialSpace == ERigVMTransformSpace::GlobalSpace)
+		{
+			const FTransform ParentTransform = ExecuteContext.Hierarchy->GetGlobalTransform(Parent);
+			const FTransform GlobalOffsetTransform = Offset * ParentTransform;
+			Transform = Transform.GetRelativeTransform(GlobalOffsetTransform);
+			Transform.NormalizeRotation();
+		}
+		
+		const FEulerTransform EulerTransform(Transform);
 		const FRigControlValue Value = FRigControlValue::Make<FRigControlValue::FEulerTransform_Float>(EulerTransform);
 		const FTransform ShapeTransform = Settings.Shape.Transform;
 		
 		FRigHierarchyControllerInstructionBracket InstructionBracket(Controller, ExecuteContext.GetInstructionIndex());
-		Item = Controller->AddControl(Name, Parent, ControlSettings, Value, OffsetTransform, ShapeTransform, false, false);
+		Item = Controller->AddControl(Name, Parent, ControlSettings, Value, Offset, ShapeTransform, false, false);
 	}
 }
 
diff --git a/Engine/Plugins/Animation/ControlRig/Source/ControlRig/Private/Units/Execution/RigUnit_DynamicHierarchy.h b/Engine/Plugins/Animation/ControlRig/Source/ControlRig/Private/Units/Execution/RigUnit_DynamicHierarchy.h
--- a/Engine/Plugins/Animation/ControlRig/Source/ControlRig/Private/Units/Execution/RigUnit_DynamicHierarchy.h
+++ b/Engine/Plugins/Animation/ControlRig/Source/ControlRig/Private/Units/Execution/RigUnit_DynamicHierarchy.h
@@ -537,11 +537,14 @@ struct CONTROLRIG_API FRigUnit_HierarchyAddControlFloat_Settings : public FRigUn
 USTRUCT(meta=(TemplateName="SpawnControl", Keywords="Construction,Create,New,AddControl,NewControl,CreateControl", Varying))
 struct CONTROLRIG_API FRigUnit_HierarchyAddControlElement : public FRigUnit_HierarchyAddElement
 {
+public:
+	
 	GENERATED_BODY()
 
 	FRigUnit_HierarchyAddControlElement()
 	{
 		Name = TEXT("NewControl");
+		OffsetSpace = ERigVMTransformSpace::LocalSpace;
 	}
 
 	virtual ERigElementType GetElementTypeToSpawn() const override { return ERigElementType::Control; }
@@ -552,6 +555,16 @@ struct CONTROLRIG_API FRigUnit_HierarchyAddControlElement : public FRigUnit_Hier
 	 */
 	UPROPERTY(meta = (Input))
 	FTransform OffsetTransform;
+
+	/*
+	 * The space the offset is in
+	 */
+	UPROPERTY(meta = (Input))
+	ERigVMTransformSpace OffsetSpace;
+
+protected:
+	
+	static FTransform ProjectOffsetTransform(const FTransform& InOffsetTransform, ERigVMTransformSpace InOffsetSpace, const FRigElementKey& InParent, const URigHierarchy* InHierarchy);
 };
 
 /**
@@ -807,6 +820,7 @@ struct CONTROLRIG_API FRigUnit_HierarchyAddControlVector_Settings : public FRigU
 	
 	FRigUnit_HierarchyAddControlVector_Settings()
 		: FRigUnit_HierarchyAddControl_Settings()
+		, InitialSpace(ERigVMTransformSpace::LocalSpace)
 		, bIsPosition(true)
 	{
 		FilteredChannels.Reset();
@@ -815,6 +829,9 @@ struct CONTROLRIG_API FRigUnit_HierarchyAddControlVector_Settings : public FRigU
 
 	virtual void Configure(FRigControlSettings& OutSettings) const override;
 
+	UPROPERTY(EditAnywhere, BlueprintReadWrite, Category = "Settings")
+	ERigVMTransformSpace InitialSpace;
+
 	UPROPERTY(EditAnywhere, BlueprintReadWrite, Category = "Settings")
 	bool bIsPosition;
 
@@ -908,6 +925,7 @@ struct CONTROLRIG_API FRigUnit_HierarchyAddControlRotator_Settings : public FRig
 	
 	FRigUnit_HierarchyAddControlRotator_Settings()
 		: FRigUnit_HierarchyAddControl_Settings()
+		, InitialSpace(ERigVMTransformSpace::LocalSpace)
 	{
 		FilteredChannels.Reset();
 	}
@@ -915,6 +933,9 @@ struct CONTROLRIG_API FRigUnit_HierarchyAddControlRotator_Settings : public FRig
 
 	virtual void Configure(FRigControlSettings& OutSettings) const override;
 
+	UPROPERTY(EditAnywhere, BlueprintReadWrite, Category = "Settings")
+	ERigVMTransformSpace InitialSpace;
+
 	UPROPERTY(EditAnywhere, BlueprintReadWrite, Category = "Settings")
 	FRigUnit_HierarchyAddControlRotator_LimitSettings Limits;;
 
@@ -967,6 +988,7 @@ struct CONTROLRIG_API FRigUnit_HierarchyAddControlTransform_Settings : public FR
 	
 	FRigUnit_HierarchyAddControlTransform_Settings()
 		: FRigUnit_HierarchyAddControl_Settings()
+		, InitialSpace(ERigVMTransformSpace::LocalSpace)
 	{
 		FilteredChannels.Reset();
 	}
@@ -974,6 +996,9 @@ struct CONTROLRIG_API FRigUnit_HierarchyAddControlTransform_Settings : public FR
 
 	virtual void Configure(FRigControlSettings& OutSettings) const override;
 
+	UPROPERTY(EditAnywhere, BlueprintReadWrite, Category = "Settings")
+	ERigVMTransformSpace InitialSpace;
+
 	UPROPERTY(EditAnywhere, BlueprintReadWrite, Category = "Settings")
 	FRigUnit_HierarchyAddControl_ShapeSettings Shape;
 
