Control Rig: Fix double creation of variable node when undoing external variable removal

diff --git a/Engine/Plugins/Animation/ControlRig/Source/ControlRigDeveloper/Private/ControlRigBlueprint.cpp b/Engine/Plugins/Animation/ControlRig/Source/ControlRigDeveloper/Private/ControlRigBlueprint.cpp
--- a/Engine/Plugins/Animation/ControlRig/Source/ControlRigDeveloper/Private/ControlRigBlueprint.cpp
+++ b/Engine/Plugins/Animation/ControlRig/Source/ControlRigDeveloper/Private/ControlRigBlueprint.cpp
@@ -159,6 +159,8 @@ UControlRigBlueprint::UControlRigBlueprint(const FObjectInitializer& ObjectIniti
 	Validator = ObjectInitializer.CreateDefaultSubobject<UControlRigValidator>(this, TEXT("ControlRigValidator"));
 
 	DebugBoneRadius = 1.f;
+
+	bUpdatingExternalVariables = false;
 	
 	bDirtyDuringLoad = false;
 	bErrorsDuringCompilation = false;
@@ -4997,6 +4999,12 @@ void UControlRigBlueprint::OnPostVariableChange(UBlueprint* InBlueprint)
 		return;
 	}
 
+	if (bUpdatingExternalVariables)
+	{
+		return;
+	}
+
+	TGuardValue<bool> UpdatingVariablesGuard(bUpdatingExternalVariables, true);
 	TArray<FBPVariableDescription> LocalLastNewVariables = LastNewVariables;
 
 	TMap<FGuid, int32> NewVariablesByGuid;
diff --git a/Engine/Plugins/Animation/ControlRig/Source/ControlRigDeveloper/Public/ControlRigBlueprint.h b/Engine/Plugins/Animation/ControlRig/Source/ControlRigDeveloper/Public/ControlRigBlueprint.h
--- a/Engine/Plugins/Animation/ControlRig/Source/ControlRigDeveloper/Public/ControlRigBlueprint.h
+++ b/Engine/Plugins/Animation/ControlRig/Source/ControlRigDeveloper/Public/ControlRigBlueprint.h
@@ -732,6 +732,7 @@ public:
 private:
 
 	FOnExternalVariablesChanged ExternalVariablesChangedEvent;
+	bool bUpdatingExternalVariables;
 	void BroadcastExternalVariablesChangedEvent();
 	FCompilerResultsLog CompileLog;
 
