Control Rig: Prevent double movement for nodes inside comments in graph

diff --git a/Engine/Plugins/Animation/ControlRig/Source/ControlRigEditor/Private/Graph/SControlRigGraphNodeComment.cpp b/Engine/Plugins/Animation/ControlRig/Source/ControlRigEditor/Private/Graph/SControlRigGraphNodeComment.cpp
--- a/Engine/Plugins/Animation/ControlRig/Source/ControlRigEditor/Private/Graph/SControlRigGraphNodeComment.cpp
+++ b/Engine/Plugins/Animation/ControlRig/Source/ControlRigEditor/Private/Graph/SControlRigGraphNodeComment.cpp
@@ -157,6 +157,12 @@ void SControlRigGraphNodeComment::MoveTo(const FVector2D& NewPosition, FNodeSet&
 				TArray<UEdGraphNode*> NodesToMove = UControlRigGraphSchema::GetNodesToMoveForNode(GraphNode);
 				for(UEdGraphNode* NodeToMove : NodesToMove)
 				{
+					// If the internal node is also selected, it will also be moved, so we don't need to take care of it
+					if (NodeToMove->IsSelected())
+					{
+						continue;
+					}
+					
 					const FVector2D OldPosition = RigSchema->GetNodePositionAtStartOfInteraction(NodeToMove);
 
 					// Don't drag note content if either of the shift keys are down.
