Control Rig: Disable preload filtering on Control Rig blueprint since it is breaking assets

diff --git a/Engine/Plugins/Animation/ControlRig/Source/ControlRig/Private/AnimNode_ControlRig.cpp b/Engine/Plugins/Animation/ControlRig/Source/ControlRig/Private/AnimNode_ControlRig.cpp
--- a/Engine/Plugins/Animation/ControlRig/Source/ControlRig/Private/AnimNode_ControlRig.cpp
+++ b/Engine/Plugins/Animation/ControlRig/Source/ControlRig/Private/AnimNode_ControlRig.cpp
@@ -519,32 +519,46 @@ void FAnimNode_ControlRig::PropagateInputProperties(const UObject* InSourceInsta
 
 				const uint8* SrcPtr = CallerProperty->ContainerPtrToValuePtr<uint8>(InSourceInstance);
 
-				if (CastField<FBoolProperty>(CallerProperty) != nullptr && Variable.TypeName == TEXT("bool"))
+				if (CastField<FBoolProperty>(CallerProperty) != nullptr && Variable.TypeName == RigVMTypeUtils::BoolTypeName)
 				{
 					const bool Value = *(const bool*)SrcPtr;
 					Variable.SetValue<bool>(Value);
 				}
-				else if (CastField<FFloatProperty>(CallerProperty) != nullptr && Variable.TypeName == TEXT("float"))
+				else if (CastField<FFloatProperty>(CallerProperty) != nullptr && (Variable.TypeName == RigVMTypeUtils::FloatTypeName ||  Variable.TypeName == RigVMTypeUtils::DoubleTypeName))
 				{
 					const float Value = *(const float*)SrcPtr;
-					Variable.SetValue<float>(Value);
+					if(Variable.TypeName == RigVMTypeUtils::FloatTypeName)
+					{
+						Variable.SetValue<float>(Value);
+					}
+					else
+					{
+						Variable.SetValue<double>(Value);
+					}
 				}
-				else if (CastField<FDoubleProperty>(CallerProperty) != nullptr && Variable.TypeName == TEXT("double"))
+				else if (CastField<FDoubleProperty>(CallerProperty) != nullptr && (Variable.TypeName == RigVMTypeUtils::FloatTypeName ||  Variable.TypeName == RigVMTypeUtils::DoubleTypeName))
 				{
 					const double Value = *(const double*)SrcPtr;
-					Variable.SetValue<double>(Value);
+					if(Variable.TypeName == RigVMTypeUtils::FloatTypeName)
+					{
+						Variable.SetValue<float>((float)Value);
+					}
+					else
+					{
+						Variable.SetValue<double>(Value);
+					}
 				}
-				else if (CastField<FIntProperty>(CallerProperty) != nullptr && Variable.TypeName == TEXT("int32"))
+				else if (CastField<FIntProperty>(CallerProperty) != nullptr && Variable.TypeName == RigVMTypeUtils::Int32TypeName)
 				{
 					const int32 Value = *(const int32*)SrcPtr;
 					Variable.SetValue<int32>(Value);
 				}
-				else if (CastField<FNameProperty>(CallerProperty) != nullptr && Variable.TypeName == TEXT("FName"))
+				else if (CastField<FNameProperty>(CallerProperty) != nullptr && Variable.TypeName == RigVMTypeUtils::FNameTypeName)
 				{
 					const FName Value = *(const FName*)SrcPtr;
 					Variable.SetValue<FName>(Value);
 				}
-				else if (CastField<FNameProperty>(CallerProperty) != nullptr && Variable.TypeName == TEXT("FString"))
+				else if (CastField<FNameProperty>(CallerProperty) != nullptr && Variable.TypeName == RigVMTypeUtils::FStringTypeName)
 				{
 					const FString Value = *(const FString*)SrcPtr;
 					Variable.SetValue<FString>(Value);
diff --git a/Engine/Plugins/Animation/ControlRig/Source/ControlRigDeveloper/Private/ControlRigBlueprint.cpp b/Engine/Plugins/Animation/ControlRig/Source/ControlRigDeveloper/Private/ControlRigBlueprint.cpp
--- a/Engine/Plugins/Animation/ControlRig/Source/ControlRigDeveloper/Private/ControlRigBlueprint.cpp
+++ b/Engine/Plugins/Animation/ControlRig/Source/ControlRigDeveloper/Private/ControlRigBlueprint.cpp
@@ -57,7 +57,7 @@
 
 TAutoConsoleVariable<bool> CVarEnablePreLoadFiltering(
 	TEXT("ControlRig.EnablePreLoadFiltering"),
-	true,
+	false,
 	TEXT("When true the RigVMGraphs will be skipped during preload to speed up load times."));
 
 TAutoConsoleVariable<bool> CVarEnablePostLoadHashing(
