Control Rig: Fix transparent execute pins

diff --git a/Engine/Plugins/Animation/ControlRig/Source/ControlRigEditor/Private/Graph/SControlRigGraphNode.cpp b/Engine/Plugins/Animation/ControlRig/Source/ControlRigEditor/Private/Graph/SControlRigGraphNode.cpp
--- a/Engine/Plugins/Animation/ControlRig/Source/ControlRigEditor/Private/Graph/SControlRigGraphNode.cpp
+++ b/Engine/Plugins/Animation/ControlRig/Source/ControlRigEditor/Private/Graph/SControlRigGraphNode.cpp
@@ -1659,7 +1659,7 @@ void SControlRigGraphNode::UpdatePinTreeView()
 				PinWidgetForExpander = PinInfo.OutputPinWidget;
 				bPinWidgetForExpanderLeft = false;
 				bPinInfoIsValid = true;
-				if (Template)
+				if (Template && !ModelPin->IsExecuteContext())
 				{
 					if (URigVMPin* RootPin = ModelPin->GetRootPin())
 					{
@@ -1687,7 +1687,7 @@ void SControlRigGraphNode::UpdatePinTreeView()
 				PinWidgetForExpander = PinInfo.InputPinWidget;
 				bPinWidgetForExpanderLeft = true;
 				bPinInfoIsValid = true;
-				if (Template)
+				if (Template && !ModelPin->IsExecuteContext())
 				{
 					if (URigVMPin* RootPin = ModelPin->GetRootPin())
 					{
