Control Rig: Reenable preload filtering and guard against old assets

diff --git a/Engine/Plugins/Animation/ControlRig/Source/ControlRigDeveloper/Private/ControlRigBlueprint.cpp b/Engine/Plugins/Animation/ControlRig/Source/ControlRigDeveloper/Private/ControlRigBlueprint.cpp
--- a/Engine/Plugins/Animation/ControlRig/Source/ControlRigDeveloper/Private/ControlRigBlueprint.cpp
+++ b/Engine/Plugins/Animation/ControlRig/Source/ControlRigDeveloper/Private/ControlRigBlueprint.cpp
@@ -57,7 +57,7 @@
 
 TAutoConsoleVariable<bool> CVarEnablePreLoadFiltering(
 	TEXT("ControlRig.EnablePreLoadFiltering"),
-	false,
+	true,
 	TEXT("When true the RigVMGraphs will be skipped during preload to speed up load times."));
 
 TAutoConsoleVariable<bool> CVarEnablePostLoadHashing(
@@ -283,16 +283,27 @@ bool UControlRigBlueprint::CanImportGraphFromText(const FString& InClipboardText
 
 bool UControlRigBlueprint::RequiresForceLoadMembers(UObject* InObject) const
 {
+	// only filter if the console variable is enabled
+	if(!CVarEnablePreLoadFiltering->GetBool())
+	{
+		return UBlueprint::RequiresForceLoadMembers(InObject);
+	}
+
+	// old assets don't support preload filtering
+	if (GetLinkerCustomVersion(FControlRigObjectVersion::GUID) < FControlRigObjectVersion::RemoveParameters)
+	{
+		return UBlueprint::RequiresForceLoadMembers(InObject);
+	}
+	
 	// we can stop traversing when hitting a URigVMNode
-	// (except for collapse nodes - since they contain a graphs again)
+	// except for collapse nodes - since they contain a graphs again
+	// and variable  nodes - since they are needed during preload by the BP compiler
 	if(InObject->IsA<URigVMNode>())
 	{
-		if(!InObject->IsA<URigVMCollapseNode>())
+		if(!InObject->IsA<URigVMCollapseNode>() &&
+			!InObject->IsA<URigVMVariableNode>())
 		{
-			if(CVarEnablePreLoadFiltering->GetBool())
-			{
-				return false;
-			}
+			return false;
 		}
 	}
 	return UBlueprint::RequiresForceLoadMembers(InObject);
