Control Rix: Fix crash when adding pins to functions

diff --git a/Engine/Plugins/Animation/ControlRig/Source/ControlRigEditor/Private/Graph/SControlRigGraphNode.cpp b/Engine/Plugins/Animation/ControlRig/Source/ControlRigEditor/Private/Graph/SControlRigGraphNode.cpp
--- a/Engine/Plugins/Animation/ControlRig/Source/ControlRigEditor/Private/Graph/SControlRigGraphNode.cpp
+++ b/Engine/Plugins/Animation/ControlRig/Source/ControlRigEditor/Private/Graph/SControlRigGraphNode.cpp
@@ -1090,11 +1090,17 @@ void SControlRigGraphNode::HandleNodePinsChanged()
 	
 	for (const TSharedRef<SGraphPin>& GraphPin: InputPins)
 	{
-		LocalPinsToDelete.Add(GraphPin->GetPinObj());
+		if (UEdGraphPin* PinObj = GraphPin->GetPinObj())
+		{
+			LocalPinsToDelete.Add(PinObj);
+		}
 	}
 	for (const TSharedRef<SGraphPin>& GraphPin: OutputPins)
 	{
-		LocalPinsToDelete.Add(GraphPin->GetPinObj());
+		if (UEdGraphPin* PinObj = GraphPin->GetPinObj())
+		{
+			LocalPinsToDelete.Add(PinObj);
+		}
 	}
 
 	check(PinsToKeep.IsEmpty());
