Control Rig: Serialize previous name map to avoid loss of animation

diff --git a/Engine/Plugins/Animation/ControlRig/Source/ControlRig/Private/Rigs/RigHierarchy.cpp b/Engine/Plugins/Animation/ControlRig/Source/ControlRig/Private/Rigs/RigHierarchy.cpp
--- a/Engine/Plugins/Animation/ControlRig/Source/ControlRig/Private/Rigs/RigHierarchy.cpp
+++ b/Engine/Plugins/Animation/ControlRig/Source/ControlRig/Private/Rigs/RigHierarchy.cpp
@@ -8,6 +8,7 @@
 #include "Units/RigUnitContext.h"
 #include "Math/ControlRigMathLibrary.h"
 #include "UObject/AnimObjectVersion.h"
+#include "ControlRigObjectVersion.h"
 #include "Serialization/MemoryReader.h"
 #include "Serialization/MemoryWriter.h"
 #include "HAL/LowLevelMemTracker.h"
@@ -142,6 +143,7 @@ void URigHierarchy::Serialize(FArchive& Ar)
 	LLM_SCOPE_BYNAME(TEXT("Animation/ControlRig"));
 	
 	Ar.UsingCustomVersion(FAnimObjectVersion::GUID);
+	Ar.UsingCustomVersion(FControlRigObjectVersion::GUID);
 
 	if (Ar.IsSaving() || Ar.IsObjectReferenceCollector() || Ar.IsCountingMemory())
 	{
@@ -199,6 +201,9 @@ void URigHierarchy::Save(FArchive& Ar)
 		FRigBaseElement* Element = Elements[ElementIndex];
 		Element->Serialize(Ar, this, FRigBaseElement::InterElementData);
 	}
+
+	Ar << PreviousNameMap;
+	Ar << PreviousParentMap;
 }
 
 void URigHierarchy::Load(FArchive& Ar)
@@ -278,6 +283,17 @@ void URigHierarchy::Load(FArchive& Ar)
 		}
 	}
 
+	if (Ar.CustomVer(FControlRigObjectVersion::GUID) >= FControlRigObjectVersion::RigHierarchyStoringPreviousNames)
+	{
+		Ar << PreviousNameMap;
+		Ar << PreviousParentMap;
+	}
+	else
+	{
+		PreviousNameMap.Reset();
+		PreviousParentMap.Reset();
+	}
+
 	Notify(ERigHierarchyNotification::HierarchyReset, nullptr);
 }
 
diff --git a/Engine/Plugins/Animation/ControlRig/Source/ControlRig/Public/ControlRigObjectVersion.h b/Engine/Plugins/Animation/ControlRig/Source/ControlRig/Public/ControlRigObjectVersion.h
--- a/Engine/Plugins/Animation/ControlRig/Source/ControlRig/Public/ControlRigObjectVersion.h
+++ b/Engine/Plugins/Animation/ControlRig/Source/ControlRig/Public/ControlRigObjectVersion.h
@@ -84,6 +84,9 @@ struct CONTROLRIG_API FControlRigObjectVersion
 		// Store function information (and compilation data) in blueprint generated class
 		StoreFunctionsInGeneratedClass,
 
+		// Hierarchy storing previous names
+		RigHierarchyStoringPreviousNames,
+
 		// -----<new versions can be added above this line>-------------------------------------------------
 		VersionPlusOne,
 		LatestVersion = VersionPlusOne - 1
