Control Rig: Avoid temporary UObjects as part of transaction

diff --git a/Engine/Plugins/Animation/ControlRig/Source/ControlRigEditor/Private/Editor/SRigHierarchy.cpp b/Engine/Plugins/Animation/ControlRig/Source/ControlRigEditor/Private/Editor/SRigHierarchy.cpp
--- a/Engine/Plugins/Animation/ControlRig/Source/ControlRigEditor/Private/Editor/SRigHierarchy.cpp
+++ b/Engine/Plugins/Animation/ControlRig/Source/ControlRigEditor/Private/Editor/SRigHierarchy.cpp
@@ -1715,6 +1715,7 @@ void SRigHierarchy::HandleNewItem(ERigElementType InElementType, bool bIsAnimati
 		return;
 	}
 
+	FRigElementKey NewItemKey;
 	URigHierarchy* Hierarchy = GetDefaultHierarchy();
 	if (Hierarchy)
 	{
@@ -1726,7 +1727,6 @@ void SRigHierarchy::HandleNewItem(ERigElementType InElementType, bool bIsAnimati
 
 		FScopedTransaction Transaction(LOCTEXT("HierarchyTreeAdded", "Add new item to hierarchy"));
 
-		FRigElementKey NewItemKey;
 		FRigElementKey ParentKey;
 		FTransform ParentTransform = FTransform::Identity;
 
@@ -1803,12 +1803,18 @@ void SRigHierarchy::HandleNewItem(ERigElementType InElementType, bool bIsAnimati
 					return;
 				}
 			}
-
-			if (ControlRigBlueprint.IsValid())
-			{
-				ControlRigBlueprint->BroadcastRefreshEditor();
-			}
 		}
+	}
+
+	if (ControlRigBlueprint.IsValid())
+	{
+		ControlRigBlueprint->BroadcastRefreshEditor();
+	}
+	
+	if (Hierarchy && NewItemKey.IsValid())
+	{
+		URigHierarchyController* Controller = Hierarchy->GetController(true);
+		check(Controller);
 
 		Controller->ClearSelection();
 		Controller->SelectElement(NewItemKey);
