Control Rig: Catch nullptr for ObjectBinding in Additive / FK Control Rig

diff --git a/Engine/Plugins/Animation/ControlRig/Source/ControlRig/Private/Rigs/AdditiveControlRig.cpp b/Engine/Plugins/Animation/ControlRig/Source/ControlRig/Private/Rigs/AdditiveControlRig.cpp
--- a/Engine/Plugins/Animation/ControlRig/Source/ControlRig/Private/Rigs/AdditiveControlRig.cpp
+++ b/Engine/Plugins/Animation/ControlRig/Source/ControlRig/Private/Rigs/AdditiveControlRig.cpp
@@ -61,9 +61,9 @@ bool UAdditiveControlRig::Execute_Internal(const FName& InEventName)
 	}
 	else if(InEventName == FRigUnit_PrepareForExecution::EventName)
 	{
-		if (TSharedPtr<IControlRigObjectBinding> Binding = GetObjectBinding())
+		if(const TSharedPtr<IControlRigObjectBinding> Binding = GetObjectBinding())
 		{
-			if (USkeletalMeshComponent* SkeletalMeshComponent = Cast<USkeletalMeshComponent>(Binding->GetBoundObject()))
+			if (const USkeletalMeshComponent* SkeletalMeshComponent = Cast<USkeletalMeshComponent>(Binding->GetBoundObject()))
 			{
 				CreateRigElements(SkeletalMeshComponent->GetSkeletalMeshAsset());
 			}
diff --git a/Engine/Plugins/Animation/ControlRig/Source/ControlRig/Private/Rigs/FKControlRig.cpp b/Engine/Plugins/Animation/ControlRig/Source/ControlRig/Private/Rigs/FKControlRig.cpp
--- a/Engine/Plugins/Animation/ControlRig/Source/ControlRig/Private/Rigs/FKControlRig.cpp
+++ b/Engine/Plugins/Animation/ControlRig/Source/ControlRig/Private/Rigs/FKControlRig.cpp
@@ -284,21 +284,19 @@ void UFKControlRig::Initialize(bool bInitRigUnits /*= true*/)
 
 	Super::Initialize(bInitRigUnits);
 
-	if (GetObjectBinding() == nullptr)
+	if(const TSharedPtr<IControlRigObjectBinding> Binding = GetObjectBinding())
 	{
-		return;
-	}
-
-	// we do this after Initialize because Initialize will copy from CDO. 
-	// create hierarchy from the incoming skeleton
-	if (USkeletalMeshComponent* SkeletalMeshComponent = Cast<USkeletalMeshComponent>(GetObjectBinding()->GetBoundObject()))
-	{
-		CreateRigElements(SkeletalMeshComponent->GetSkeletalMeshAsset());
-	}
-	else if (USkeleton* Skeleton = Cast<USkeleton>(GetObjectBinding()->GetBoundObject()))
-	{
-		const FReferenceSkeleton& RefSkeleton = Skeleton->GetReferenceSkeleton();
-		CreateRigElements(RefSkeleton, Skeleton);
+		// we do this after Initialize because Initialize will copy from CDO. 
+		// create hierarchy from the incoming skeleton
+		if (USkeletalMeshComponent* SkeletalMeshComponent = Cast<USkeletalMeshComponent>(Binding->GetBoundObject()))
+		{
+			CreateRigElements(SkeletalMeshComponent->GetSkeletalMeshAsset());
+		}
+		else if (USkeleton* Skeleton = Cast<USkeleton>(Binding->GetBoundObject()))
+		{
+			const FReferenceSkeleton& RefSkeleton = Skeleton->GetReferenceSkeleton();
+			CreateRigElements(RefSkeleton, Skeleton);
+		}
 	}
 }
 
