Sequencer: Restore Selection when binding changes so we can handle dynamic controls getting recreated, need to do it on next tick so Sequencer get's reset.
Note preflight error is a backend problem
#trivial

diff --git a/Engine/Plugins/Animation/ControlRig/Source/ControlRigEditor/Private/Sequencer/ControlRigParameterTrackEditor.cpp b/Engine/Plugins/Animation/ControlRig/Source/ControlRigEditor/Private/Sequencer/ControlRigParameterTrackEditor.cpp
--- a/Engine/Plugins/Animation/ControlRig/Source/ControlRigEditor/Private/Sequencer/ControlRigParameterTrackEditor.cpp
+++ b/Engine/Plugins/Animation/ControlRig/Source/ControlRigEditor/Private/Sequencer/ControlRigParameterTrackEditor.cpp
@@ -3145,10 +3145,21 @@ void FControlRigParameterTrackEditor::HandleOnControlRigBound(UControlRig* InCon
 
 void FControlRigParameterTrackEditor::HandleOnObjectBoundToControlRig(UObject* InObject)
 {
+	//reselect these control rigs since selection may get
+	TMap<TWeakObjectPtr<UControlRig>, TArray<FName>> ReselectIfNeeded;
 	// look for sections to update
 	TArray<UMovieSceneControlRigParameterSection*> SectionsToUpdate;
-	for (const TWeakObjectPtr<UControlRig>& ControlRigPtr : BoundControlRigs)
+	for (TWeakObjectPtr<UControlRig>& ControlRigPtr : BoundControlRigs)
 	{
+		if (ControlRigPtr.IsValid() == false)
+		{
+			continue;
+		}
+		TArray<FName> Selection = ControlRigPtr->CurrentControlSelection();
+		if (Selection.Num() > 0)
+		{
+			ReselectIfNeeded.Add(ControlRigPtr, Selection);
+		}
 		const TSharedPtr<IControlRigObjectBinding> Binding =
 			ControlRigPtr.IsValid() ? ControlRigPtr->GetObjectBinding() : nullptr;
 		const UObject* CurrentObject = Binding ? Binding->GetBoundObject() : nullptr;
@@ -3176,6 +3187,27 @@ void FControlRigParameterTrackEditor::HandleOnObjectBoundToControlRig(UObject* I
 			Section->MarkAsChanged();
 		}
 	}
+	if (ReselectIfNeeded.Num() > 0)
+	{
+		GEditor->GetTimerManager()->SetTimerForNextTick([ReselectIfNeeded]()
+		{
+			GEditor->GetTimerManager()->SetTimerForNextTick([ReselectIfNeeded]()
+				{
+					for (const TPair <TWeakObjectPtr<UControlRig>, TArray<FName>>& Pair : ReselectIfNeeded)
+					{
+						if (Pair.Key.IsValid())
+						{
+							Pair.Key->ClearControlSelection();
+							for (const FName& ControlName : Pair.Value)
+							{
+								Pair.Key->SelectControl(ControlName, true);
+							}
+						}
+					}
+				});
+			
+		});
+	}
 }
 
 void FControlRigParameterTrackEditor::GetControlRigKeys(
