Control Rig: Fix crash when removing dot nodes

diff --git a/Engine/Plugins/Animation/ControlRig/Source/ControlRigEditor/Private/Graph/SControlRigGraphNodeKnot.cpp b/Engine/Plugins/Animation/ControlRig/Source/ControlRigEditor/Private/Graph/SControlRigGraphNodeKnot.cpp
--- a/Engine/Plugins/Animation/ControlRig/Source/ControlRigEditor/Private/Graph/SControlRigGraphNodeKnot.cpp
+++ b/Engine/Plugins/Animation/ControlRig/Source/ControlRigEditor/Private/Graph/SControlRigGraphNodeKnot.cpp
@@ -12,6 +12,11 @@
 void SControlRigGraphNodeKnot::Construct(const FArguments& InArgs, UEdGraphNode* InKnot)
 {
 	SGraphNodeKnot::Construct(SGraphNodeKnot::FArguments(), InKnot);
+
+	if (UControlRigGraphNode* CRNode = Cast<UControlRigGraphNode>(InKnot))
+	{
+		CRNode->OnNodeBeginRemoval().AddSP(this, &SControlRigGraphNodeKnot::HandleNodeBeginRemoval);
+	}
 }
 
 void SControlRigGraphNodeKnot::EndUserInteraction() const
@@ -48,4 +53,26 @@ void SControlRigGraphNodeKnot::MoveTo(const FVector2D& NewPosition, FNodeSet& No
 	}
 }
 
+void SControlRigGraphNodeKnot::HandleNodeBeginRemoval()
+{
+	if(UControlRigGraphNode* RigNode = Cast<UControlRigGraphNode>(GraphNode))
+	{
+		RigNode->OnNodeBeginRemoval().RemoveAll(this);
+	}
+	
+	for (const TSharedRef<SGraphPin>& GraphPin: InputPins)
+	{
+		GraphPin->SetPinObj(nullptr);
+	}
+	for (const TSharedRef<SGraphPin>& GraphPin: OutputPins)
+	{
+		GraphPin->SetPinObj(nullptr);
+	}
+
+	InputPins.Reset();
+	OutputPins.Reset();
+	
+	InvalidateGraphData();
+}
+
 #undef LOCTEXT_NAMESPACE
diff --git a/Engine/Plugins/Animation/ControlRig/Source/ControlRigEditor/Private/Graph/SControlRigGraphNodeKnot.h b/Engine/Plugins/Animation/ControlRig/Source/ControlRigEditor/Private/Graph/SControlRigGraphNodeKnot.h
--- a/Engine/Plugins/Animation/ControlRig/Source/ControlRigEditor/Private/Graph/SControlRigGraphNodeKnot.h
+++ b/Engine/Plugins/Animation/ControlRig/Source/ControlRigEditor/Private/Graph/SControlRigGraphNodeKnot.h
@@ -16,4 +16,6 @@ class SControlRigGraphNodeKnot : public SGraphNodeKnot
 	// SGraphNode interface
 	virtual void EndUserInteraction() const override;
 	virtual void MoveTo(const FVector2D& NewPosition, FNodeSet& NodeFilter, bool bMarkDirty = true) override;
+
+	void HandleNodeBeginRemoval();
 };
\ No newline at end of file
diff --git a/Engine/Source/Editor/GraphEditor/Private/SGraphNodeKnot.cpp b/Engine/Source/Editor/GraphEditor/Private/SGraphNodeKnot.cpp
--- a/Engine/Source/Editor/GraphEditor/Private/SGraphNodeKnot.cpp
+++ b/Engine/Source/Editor/GraphEditor/Private/SGraphNodeKnot.cpp
@@ -246,12 +246,12 @@ void SGraphPinKnot::OnDragEnter(const FGeometry& MyGeometry, const FDragDropEven
 FSlateColor SGraphPinKnot::GetPinColor() const
 {
 	// Make ourselves transparent if we're the input, since we are underneath the output pin and would double-blend looking ugly
-	if (UEdGraphPin* EdGraphPin = GetPinObj())
+	if (UEdGraphPin* PinObj = GetPinObj())
 	{
-		if (EdGraphPin->Direction == EEdGraphPinDirection::EGPD_Input)
+		if (PinObj->Direction == EEdGraphPinDirection::EGPD_Input)
 		{
 			return FLinearColor::Transparent;
-		}
+		} 
 	}
 	return SGraphPin::GetPinColor();
 }
