Fix constraints not transacting over multi user

diff --git a/Engine/Plugins/Animation/ControlRig/Source/ControlRig/Private/ControlRig.cpp b/Engine/Plugins/Animation/ControlRig/Source/ControlRig/Private/ControlRig.cpp
--- a/Engine/Plugins/Animation/ControlRig/Source/ControlRig/Private/ControlRig.cpp
+++ b/Engine/Plugins/Animation/ControlRig/Source/ControlRig/Private/ControlRig.cpp
@@ -2640,9 +2640,8 @@ UTransformableControlHandle* UControlRig::CreateTransformableControlHandle(
 		return nullptr;
 	}
 	
-	UTransformableControlHandle* CtrlHandle = NewObject<UTransformableControlHandle>(InOuter);
-	ensure(CtrlHandle);
-	CtrlHandle->SetFlags(RF_Transactional);
+	UTransformableControlHandle* CtrlHandle = NewObject<UTransformableControlHandle>(InOuter, NAME_None, RF_Transactional);
+	check(CtrlHandle);
 	CtrlHandle->ControlRig = this;
 	CtrlHandle->ControlName = InControlName;
 	CtrlHandle->RegisterDelegates();
