Sequencer: Procedural Control Rigs: On copy/paste need to make sure we have the initialize callback setup since controls aren't created unilt initialze happens with a bound skel mesh.

diff --git a/Engine/Plugins/Animation/ControlRig/Source/ControlRig/Private/Sequencer/MovieSceneControlRigParameterTrack.cpp b/Engine/Plugins/Animation/ControlRig/Source/ControlRig/Private/Sequencer/MovieSceneControlRigParameterTrack.cpp
--- a/Engine/Plugins/Animation/ControlRig/Source/ControlRig/Private/Sequencer/MovieSceneControlRigParameterTrack.cpp
+++ b/Engine/Plugins/Animation/ControlRig/Source/ControlRig/Private/Sequencer/MovieSceneControlRigParameterTrack.cpp
@@ -526,6 +526,10 @@ void UMovieSceneControlRigParameterTrack::PostEditImport()
 	Super::PostEditImport();
 	if (ControlRig)
 	{
+		if (ControlRig->OnInitialized_AnyThread().IsBoundToObject(this) == false)
+		{
+			ControlRig->OnInitialized_AnyThread().AddUObject(this, &UMovieSceneControlRigParameterTrack::HandleOnInitialized);
+		}
 		ControlRig->ClearFlags(RF_Transient); //when copied make sure it's no longer transient, sequencer does this for tracks/sections 
 											  //but not for all objects in them since the control rig itself has transient objects.
 	}
