RigVM: Make sure to call UEdGraphPin::ResolveAllPinReferences when loading a dependency

diff --git a/Engine/Plugins/Animation/ControlRig/Source/ControlRigDeveloper/Private/ControlRigBlueprint.cpp b/Engine/Plugins/Animation/ControlRig/Source/ControlRigDeveloper/Private/ControlRigBlueprint.cpp
--- a/Engine/Plugins/Animation/ControlRig/Source/ControlRigDeveloper/Private/ControlRigBlueprint.cpp
+++ b/Engine/Plugins/Animation/ControlRig/Source/ControlRigDeveloper/Private/ControlRigBlueprint.cpp
@@ -1107,6 +1107,8 @@ void UControlRigBlueprint::PostLoad()
 	// in which case we have to refresh the model
 	FRigVMRegistry::Get().OnRigVMRegistryChanged().RemoveAll(this);
 	FRigVMRegistry::Get().OnRigVMRegistryChanged().AddUObject(this, &UControlRigBlueprint::OnRigVMRegistryChanged);
+	
+	UEdGraphPin::ResolveAllPinReferences();
 }
 
 #if WITH_EDITORONLY_DATA
