Control Rig: Fix find references not working for local variables

diff --git a/Engine/Plugins/Animation/ControlRig/Source/ControlRigDeveloper/Private/Graph/ControlRigGraphSchema.cpp b/Engine/Plugins/Animation/ControlRig/Source/ControlRigDeveloper/Private/Graph/ControlRigGraphSchema.cpp
--- a/Engine/Plugins/Animation/ControlRig/Source/ControlRigDeveloper/Private/Graph/ControlRigGraphSchema.cpp
+++ b/Engine/Plugins/Animation/ControlRig/Source/ControlRigDeveloper/Private/Graph/ControlRigGraphSchema.cpp
@@ -1214,6 +1214,11 @@ FString UControlRigGraphSchema::GetFindReferenceSearchTerm(const FEdGraphSchemaA
 			const FEdGraphSchemaAction_K2Var* VarAction = (const FEdGraphSchemaAction_K2Var*)InGraphAction;
 			return VarAction->GetVariableName().ToString();
 		}
+		else if (InGraphAction->GetTypeId() == FControlRigGraphSchemaAction_LocalVar::StaticGetTypeId())
+		{
+			const FControlRigGraphSchemaAction_LocalVar* VarAction = (const FControlRigGraphSchemaAction_LocalVar*)InGraphAction;
+			return VarAction->GetVariableName().ToString();
+		}
 	}
 	return Super::GetFindReferenceSearchTerm(InGraphAction);
 }
