ControlRig: AddObjectReferenced SourceInstnace and TargetInstance of FAnimNode_ControlRig_ExternalSource in FControlRigLayerInstanceProxy

diff --git a/Engine/Plugins/Animation/ControlRig/Source/ControlRig/Private/Sequencer/ControlRigLayerInstanceProxy.cpp b/Engine/Plugins/Animation/ControlRig/Source/ControlRig/Private/Sequencer/ControlRigLayerInstanceProxy.cpp
--- a/Engine/Plugins/Animation/ControlRig/Source/ControlRig/Private/Sequencer/ControlRigLayerInstanceProxy.cpp
+++ b/Engine/Plugins/Animation/ControlRig/Source/ControlRig/Private/Sequencer/ControlRigLayerInstanceProxy.cpp
@@ -263,6 +263,14 @@ void FControlRigLayerInstanceProxy::AddReferencedObjects(UAnimInstance* InAnimIn
 	{
 		Collector.AddReferencedObject(CurrentSourceAnimInstance);
 	}
+
+	for (TSharedPtr<FAnimNode_ControlRig_ExternalSource>& Node : ControlRigNodes)
+	{
+#if WITH_EDITORONLY_DATA
+		Collector.AddReferencedObject(Node->SourceInstance);
+#endif
+		Collector.AddReferencedObject(Node->TargetInstance);
+	}
 }
 
 void FControlRigLayerInstanceProxy::InitializeCustomProxy(FAnimInstanceProxy* InputProxy, UAnimInstance* InAnimInstance)
