Control Rig: Modified details wrapper to have the Blueprint as outer, and hierarchy/node as subject

diff --git a/Engine/Plugins/Animation/ControlRig/Source/ControlRigEditor/Private/ControlRigElementDetails.cpp b/Engine/Plugins/Animation/ControlRig/Source/ControlRigEditor/Private/ControlRigElementDetails.cpp
--- a/Engine/Plugins/Animation/ControlRig/Source/ControlRigEditor/Private/ControlRigElementDetails.cpp
+++ b/Engine/Plugins/Animation/ControlRig/Source/ControlRigEditor/Private/ControlRigElementDetails.cpp
@@ -678,7 +678,7 @@ void FRigBaseElementDetails::CustomizeDetails(IDetailLayoutBuilder& DetailBuilde
 
 		FPerElementInfo Info;
 		Info.WrapperObject = WrapperObject;
-		Info.Element = Cast<URigHierarchy>(WrapperObject->GetOuter())->GetHandle(Key);
+		Info.Element = Cast<URigHierarchy>(WrapperObject->GetSubject())->GetHandle(Key);
 
 		if(!Info.Element.IsValid())
 		{
diff --git a/Engine/Plugins/Animation/ControlRig/Source/ControlRigEditor/Private/ControlRigGraphDetails.cpp b/Engine/Plugins/Animation/ControlRig/Source/ControlRigEditor/Private/ControlRigGraphDetails.cpp
--- a/Engine/Plugins/Animation/ControlRig/Source/ControlRigEditor/Private/ControlRigGraphDetails.cpp
+++ b/Engine/Plugins/Animation/ControlRig/Source/ControlRigEditor/Private/ControlRigGraphDetails.cpp
@@ -1304,7 +1304,7 @@ void FControlRigWrappedNodeDetails::CustomizeDetails(IDetailLayoutBuilder& Detai
 		}
 
 		ObjectsBeingCustomized.Add(WrapperObject);
-		NodesBeingCustomized.Add(CastChecked<URigVMNode>(WrapperObject->GetOuter()));
+		NodesBeingCustomized.Add(CastChecked<URigVMNode>(WrapperObject->GetSubject()));
 	}
 
 	if (BlueprintBeingCustomized == nullptr || ObjectsBeingCustomized.IsEmpty() || NodesBeingCustomized.IsEmpty())
diff --git a/Engine/Plugins/Animation/ControlRig/Source/ControlRigEditor/Private/Editor/ControlRigEditor.cpp b/Engine/Plugins/Animation/ControlRig/Source/ControlRigEditor/Private/Editor/ControlRigEditor.cpp
--- a/Engine/Plugins/Animation/ControlRig/Source/ControlRigEditor/Private/Editor/ControlRigEditor.cpp
+++ b/Engine/Plugins/Animation/ControlRig/Source/ControlRigEditor/Private/Editor/ControlRigEditor.cpp
@@ -1861,7 +1861,7 @@ void FControlRigEditor::SetDetailObjects(const TArray<UObject*>& InObjects, bool
 			bool bClassCreated = UDetailsViewWrapperObject::GetClassForNodes(ModelNodes, false) == nullptr;
 
 			// create the wrapper object
-			if(UDetailsViewWrapperObject* WrapperObject = UDetailsViewWrapperObject::MakeInstance(ModelNodes, ModelNode, ModelNode))
+			if(UDetailsViewWrapperObject* WrapperObject = UDetailsViewWrapperObject::MakeInstance(GetBlueprintObj(), ModelNodes, ModelNode))
 			{
 				WrapperObject->GetWrappedPropertyChangedChainEvent().AddSP(this, &FControlRigEditor::OnWrappedPropertyChangedChainEvent);
 				WrapperObject->AddToRoot();
@@ -1922,7 +1922,7 @@ void FControlRigEditor::SetDetailViewForRigElements()
 			continue;
 		}
 
-		UDetailsViewWrapperObject* WrapperObject = UDetailsViewWrapperObject::MakeInstance(Element->GetElementStruct(), (uint8*)Element, HierarchyBeingDebugged);
+		UDetailsViewWrapperObject* WrapperObject = UDetailsViewWrapperObject::MakeInstance(GetBlueprintObj(), Element->GetElementStruct(), (uint8*)Element, HierarchyBeingDebugged);
 		WrapperObject->GetWrappedPropertyChangedChainEvent().AddSP(this, &FControlRigEditor::OnWrappedPropertyChangedChainEvent);
 		WrapperObject->AddToRoot();
 
@@ -4852,7 +4852,7 @@ void FControlRigEditor::OnWrappedPropertyChangedChainEvent(UDetailsViewWrapperOb
 		{
 			check(WrappedStruct == WrapperObjects[0]->GetWrappedStruct());
 
-			URigHierarchy* Hierarchy = CastChecked<URigHierarchy>(InWrapperObject->GetOuter());
+			URigHierarchy* Hierarchy = CastChecked<URigHierarchy>(InWrapperObject->GetSubject());
 			const FRigBaseElement WrappedElement = InWrapperObject->GetContent<FRigBaseElement>();
 			const FRigBaseElement FirstWrappedElement = WrapperObjects[0]->GetContent<FRigBaseElement>();
 			const FRigElementKey& Key = WrappedElement.GetKey();
@@ -4994,7 +4994,7 @@ void FControlRigEditor::OnWrappedPropertyChangedChainEvent(UDetailsViewWrapperOb
 			check(WrappedStruct == WrapperObjects[0]->GetWrappedStruct());
 			
 			const FRigVMGraphVariableDescription VariableDescription = InWrapperObject->GetContent<FRigVMGraphVariableDescription>();
-			URigVMGraph* Graph = CastChecked<URigVMGraph>(InWrapperObject->GetOuter());
+			URigVMGraph* Graph = CastChecked<URigVMGraph>(InWrapperObject->GetSubject());
 			URigVMController* Controller = ControlRigBP->GetController(Graph);
 			if (PropertyPath == TEXT("Name") && MyBlueprintWidget.IsValid())
 			{
@@ -5051,7 +5051,7 @@ void FControlRigEditor::OnWrappedPropertyChangedChainEvent(UDetailsViewWrapperOb
 		FProperty* TargetProperty = WrapperObjects[0]->GetClass()->FindPropertyByName(RootPinName);
 		uint8* FirstPropertyStorage = TargetProperty->ContainerPtrToValuePtr<uint8>(WrapperObjects[0].Get());
 
-		URigVMNode* Node = CastChecked<URigVMNode>(InWrapperObject->GetOuter());
+		URigVMNode* Node = CastChecked<URigVMNode>(InWrapperObject->GetSubject());
 
 		FString DefaultValue = FRigVMStruct::ExportToFullyQualifiedText(TargetProperty, FirstPropertyStorage);
 
@@ -5284,7 +5284,7 @@ bool FControlRigEditor::SelectLocalVariable(const UEdGraph* Graph, const FName&
 				if (Variable.Name == VariableName)
 				{
 					UDetailsViewWrapperObject* WrapperObject = UDetailsViewWrapperObject::MakeInstance(
-						Variable.StaticStruct(), (uint8*)&Variable, RigVMGraph);
+						GetBlueprintObj(), Variable.StaticStruct(), (uint8*)&Variable, RigVMGraph);
 					WrapperObject->GetWrappedPropertyChangedChainEvent().AddSP(this, &FControlRigEditor::OnWrappedPropertyChangedChainEvent);
 					WrapperObject->AddToRoot();
 
diff --git a/Engine/Plugins/Animation/ControlRig/Source/ControlRigEditor/Private/Editor/DetailsViewWrapperObject.cpp b/Engine/Plugins/Animation/ControlRig/Source/ControlRigEditor/Private/Editor/DetailsViewWrapperObject.cpp
--- a/Engine/Plugins/Animation/ControlRig/Source/ControlRigEditor/Private/Editor/DetailsViewWrapperObject.cpp
+++ b/Engine/Plugins/Animation/ControlRig/Source/ControlRigEditor/Private/Editor/DetailsViewWrapperObject.cpp
@@ -197,7 +197,7 @@ UClass* UDetailsViewWrapperObject::GetClassForStruct(UScriptStruct* InStruct, bo
 	return WrapperClass;
 }
 
-UDetailsViewWrapperObject* UDetailsViewWrapperObject::MakeInstance(UScriptStruct* InStruct, uint8* InStructMemory, UObject* InOuter)
+UDetailsViewWrapperObject* UDetailsViewWrapperObject::MakeInstance(UObject* InOuter, UScriptStruct* InStruct, uint8* InStructMemory, UObject* InSubject)
 {
 	check(InStruct != nullptr);
 
@@ -211,6 +211,8 @@ UDetailsViewWrapperObject* UDetailsViewWrapperObject::MakeInstance(UScriptStruct
 
 	UDetailsViewWrapperObject* Instance = NewObject<UDetailsViewWrapperObject>(InOuter, WrapperClass, NAME_None, RF_Public | RF_Transient | RF_TextExportTransient | RF_DuplicateTransient);
 	Instance->SetContent(InStructMemory, InStruct);
+	Instance->SetSubject(InSubject);
+
 	return Instance;
 }
 
@@ -511,7 +513,7 @@ UClass* UDetailsViewWrapperObject::GetClassForNodes(TArray<URigVMNode*> InNodes,
 	return WrapperClass;	
 }
 
-UDetailsViewWrapperObject* UDetailsViewWrapperObject::MakeInstance(TArray<URigVMNode*> InNodes, URigVMNode* InSubject, UObject* InOuter)
+UDetailsViewWrapperObject* UDetailsViewWrapperObject::MakeInstance(UObject* InOuter, TArray<URigVMNode*> InNodes, URigVMNode* InSubject)
 {
 	InOuter = InOuter == nullptr ? GetTransientPackage() : InOuter;
 	
@@ -523,6 +525,7 @@ UDetailsViewWrapperObject* UDetailsViewWrapperObject::MakeInstance(TArray<URigVM
 
 	UDetailsViewWrapperObject* Instance = NewObject<UDetailsViewWrapperObject>(InOuter, WrapperClass, NAME_None, RF_Public | RF_Transient | RF_TextExportTransient | RF_DuplicateTransient);
 	Instance->SetContent(InSubject);
+	Instance->SetSubject(InSubject);
 
 	InSubject->GetGraph()->OnModified().AddUObject(Instance, &UDetailsViewWrapperObject::HandleModifiedEvent);
 	
@@ -581,6 +584,20 @@ void UDetailsViewWrapperObject::SetContent(URigVMNode* InNode)
 	}
 }
 
+UObject* UDetailsViewWrapperObject::GetSubject() const
+{
+	if(SubjectPtr.IsValid())
+	{
+		return SubjectPtr.Get();
+	}
+	return nullptr;
+}
+
+void UDetailsViewWrapperObject::SetSubject(UObject* InSubject)
+{
+	SubjectPtr = TWeakObjectPtr<UObject>(InSubject);
+}
+
 void UDetailsViewWrapperObject::PostEditChangeChainProperty(FPropertyChangedChainEvent& PropertyChangedEvent)
 {
 	Super::PostEditChangeChainProperty(PropertyChangedEvent);
@@ -640,7 +657,7 @@ void UDetailsViewWrapperObject::HandleModifiedEvent(ERigVMGraphNotifType InNotif
 	}
 	
 	URigVMPin* Pin = CastChecked<URigVMPin>(InSubject);
-	if(Pin->GetNode() != GetOuter())
+	if(Pin->GetNode() != GetSubject())
 	{
 		return;
 	}
diff --git a/Engine/Plugins/Animation/ControlRig/Source/ControlRigEditor/Private/Editor/DetailsViewWrapperObject.h b/Engine/Plugins/Animation/ControlRig/Source/ControlRigEditor/Private/Editor/DetailsViewWrapperObject.h
--- a/Engine/Plugins/Animation/ControlRig/Source/ControlRigEditor/Private/Editor/DetailsViewWrapperObject.h
+++ b/Engine/Plugins/Animation/ControlRig/Source/ControlRigEditor/Private/Editor/DetailsViewWrapperObject.h
@@ -21,12 +21,12 @@ public:
 
 	// Creating wrappers from a given struct
 	static UClass* GetClassForStruct(UScriptStruct* InStruct, bool bCreateIfNeeded = true);
-	static UDetailsViewWrapperObject* MakeInstance(UScriptStruct* InStruct, uint8* InStructMemory, UObject* InOuter = nullptr);
+	static UDetailsViewWrapperObject* MakeInstance(UObject* InOuter, UScriptStruct* InStruct, uint8* InStructMemory, UObject* InSubject = nullptr);
 	UScriptStruct* GetWrappedStruct() const;
 
 	// Creating wrappers from a RigVMNode
 	static UClass* GetClassForNodes(TArray<URigVMNode*> InNodes, bool bCreateIfNeeded = true);
-	static UDetailsViewWrapperObject* MakeInstance(TArray<URigVMNode*> InNodes, URigVMNode* InSubject, UObject* InOuter = nullptr);
+	static UDetailsViewWrapperObject* MakeInstance(UObject* InOuter, TArray<URigVMNode*> InNodes, URigVMNode* InSubject);
 	
 	static void MarkOutdatedClass(UClass* InClass);
 	static bool IsValidClass(UClass* InClass);
@@ -67,6 +67,9 @@ public:
 		SetContent((const uint8*)&InValue, T::StaticStruct());
 	}
 
+	UObject* GetSubject() const;
+	void SetSubject(UObject* InSubject);
+
 	FWrappedPropertyChangedChainEvent& GetWrappedPropertyChangedChainEvent() { return WrappedPropertyChangedChainEvent; }
 	virtual void PostEditChangeChainProperty(FPropertyChangedChainEvent& PropertyChangedEvent) override;
 
@@ -116,6 +119,9 @@ private:
 	static TMap<UClass*, FPerClassInfo> ClassToInfo;
 	static TSet<UClass*> OutdatedClassToRecreate;
 	bool bIsSettingValue;
+
+	UPROPERTY()
+	TWeakObjectPtr<UObject> SubjectPtr;
 	
 	FWrappedPropertyChangedChainEvent WrappedPropertyChangedChainEvent;
 };
