Sequencer: Animation Mode: Restore Gizmo Scale when leaving/entering Animation Mode
#trivial

diff --git a/Engine/Plugins/Animation/ControlRig/Source/ControlRigEditor/Private/EditMode/ControlRigEditMode.cpp b/Engine/Plugins/Animation/ControlRig/Source/ControlRigEditor/Private/EditMode/ControlRigEditMode.cpp
--- a/Engine/Plugins/Animation/ControlRig/Source/ControlRigEditor/Private/EditMode/ControlRigEditMode.cpp
+++ b/Engine/Plugins/Animation/ControlRig/Source/ControlRigEditor/Private/EditMode/ControlRigEditMode.cpp
@@ -360,7 +360,7 @@ void FControlRigEditMode::Enter()
 	// Call parent implementation
 	FEdMode::Enter();
 	LastMovieSceneSig = FGuid();
-	if(UsesToolkits())
+	if (UsesToolkits())
 	{
 		if (!Toolkit.IsValid())
 		{
@@ -381,13 +381,20 @@ void FControlRigEditMode::Enter()
 				CoordSystemPerWidgetMode[i] = CoordSystem;
 			}
 		}
-	
+
 		ModeManager->OnWidgetModeChanged().AddSP(this, &FControlRigEditMode::OnWidgetModeChanged);
 		ModeManager->OnCoordSystemChanged().AddSP(this, &FControlRigEditMode::OnCoordSystemChanged);
 	}
 	WorldPtr = GetWorld();
 	OnWorldCleanupHandle = FWorldDelegates::OnWorldCleanup.AddSP(this, &FControlRigEditMode::OnWorldCleanup);
 	SetObjects_Internal();
+
+	//set up gizmo scale to what we had last and save what it was.
+	PreviousGizmoScale = GetModeManager()->GetWidgetScale();
+	if (const UControlRigEditModeSettings* Settings = GetDefault<UControlRigEditModeSettings>())
+	{
+		GetModeManager()->SetWidgetScale(Settings->GizmoScale);
+	}
 }
 
 //todo get working with Persona
@@ -2491,8 +2498,7 @@ void FControlRigEditMode::DecreaseShapeSize()
 void FControlRigEditMode::ResetControlShapeSize()
 {
 	UControlRigEditModeSettings* Settings = GetMutableDefault<UControlRigEditModeSettings>();
-	Settings->GizmoScale = 1.0f;
-	GetModeManager()->SetWidgetScale(Settings->GizmoScale);
+	GetModeManager()->SetWidgetScale(PreviousGizmoScale);
 }
 
 uint8 FControlRigEditMode::GetInteractionType(FEditorViewportClient* InViewportClient)
diff --git a/Engine/Plugins/Animation/ControlRig/Source/ControlRigEditor/Public/EditMode/ControlRigEditMode.h b/Engine/Plugins/Animation/ControlRig/Source/ControlRigEditor/Public/EditMode/ControlRigEditMode.h
--- a/Engine/Plugins/Animation/ControlRig/Source/ControlRigEditor/Public/EditMode/ControlRigEditMode.h
+++ b/Engine/Plugins/Animation/ControlRig/Source/ControlRigEditor/Public/EditMode/ControlRigEditMode.h
@@ -345,6 +345,8 @@ protected:
 	/** Starting X Value*/
 	TOptional<int32> StartXValue;
 
+	/** previous Gizmo(Widget) scale before we enter this mode, used to set it back*/
+	float PreviousGizmoScale = 1.0f;
 	
 public: 
 	/** Clear all selected RigElements */
