ControlRig: Fix ensure when undoing the creation of an UAnimSequence asset.

diff --git a/Engine/Plugins/Animation/ControlRig/Source/ControlRig/Private/Rigs/RigHierarchyController.cpp b/Engine/Plugins/Animation/ControlRig/Source/ControlRig/Private/Rigs/RigHierarchyController.cpp
--- a/Engine/Plugins/Animation/ControlRig/Source/ControlRig/Private/Rigs/RigHierarchyController.cpp
+++ b/Engine/Plugins/Animation/ControlRig/Source/ControlRig/Private/Rigs/RigHierarchyController.cpp
@@ -1636,6 +1636,13 @@ void URigHierarchyController::HandleHierarchyModified(ERigHierarchyNotification
 	ModifiedEvent.Broadcast(InNotifType, InHierarchy, InElement);
 }
 
+bool URigHierarchyController::IsValid() const
+{
+	// If we're pending kill, it's fine for our Hierarchy to also be pending kill
+	const bool bPendingKillAcceptable = ::IsValid(this);
+	return Hierarchy.IsValid(bPendingKillAcceptable);
+}
+
 int32 URigHierarchyController::AddElement(FRigBaseElement* InElementToAdd, FRigBaseElement* InFirstParent, bool bMaintainGlobalTransform)
 {
 	ensure(IsValid());
diff --git a/Engine/Plugins/Animation/ControlRig/Source/ControlRig/Public/Rigs/RigHierarchyController.h b/Engine/Plugins/Animation/ControlRig/Source/ControlRig/Public/Rigs/RigHierarchyController.h
--- a/Engine/Plugins/Animation/ControlRig/Source/ControlRig/Public/Rigs/RigHierarchyController.h
+++ b/Engine/Plugins/Animation/ControlRig/Source/ControlRig/Public/Rigs/RigHierarchyController.h
@@ -598,7 +598,7 @@ private:
 	 * Returns true if this controller is valid / linked to a valid hierarchy.
 	 * @return Returns true if this controller is valid / linked to a valid hierarchy.
 	 */
-	bool IsValid() const { return Hierarchy.IsValid(); }
+	bool IsValid() const;
 
 	/**
 	 * Adds a new element to the hierarchy
