Control Rig: Autowire new node execute and creating a second link to the old connected pin

diff --git a/Engine/Plugins/Animation/ControlRig/Source/ControlRigDeveloper/Private/Graph/ControlRigGraphNode.cpp b/Engine/Plugins/Animation/ControlRig/Source/ControlRigDeveloper/Private/Graph/ControlRigGraphNode.cpp
--- a/Engine/Plugins/Animation/ControlRig/Source/ControlRigDeveloper/Private/Graph/ControlRigGraphNode.cpp
+++ b/Engine/Plugins/Animation/ControlRig/Source/ControlRigDeveloper/Private/Graph/ControlRigGraphNode.cpp
@@ -1674,37 +1674,7 @@ void UControlRigGraphNode::AutowireNewNode(UEdGraphPin* FromPin)
 	Super::AutowireNewNode(FromPin);
 
 	const UControlRigGraphSchema* Schema = GetDefault<UControlRigGraphSchema>();
-
-	// copying high level information into a local array since the try create connection below
-	// may cause the pin array to be destroyed / changed
-	TArray<TPair<FName, EEdGraphPinDirection>> PinsToVisit;
-	for(UEdGraphPin* Pin : Pins)
-	{
-		PinsToVisit.Emplace(Pin->GetFName(), Pin->Direction);
-	}
-
-	for(const TPair<FName, EEdGraphPinDirection>& PinToVisit : PinsToVisit)
-	{
-		UEdGraphPin* Pin = FindPin(PinToVisit.Key, PinToVisit.Value);
-		if(Pin == nullptr)
-		{
-			continue;
-		}
-		
-		if (Pin->ParentPin != nullptr)
-		{
-			continue;
-		}
-
-		FPinConnectionResponse ConnectResponse = Schema->CanCreateConnection(FromPin, Pin);
-		if(ConnectResponse.Response != ECanCreateConnectionResponse::CONNECT_RESPONSE_DISALLOW)
-		{
-			if (Schema->TryCreateConnection(FromPin, Pin))
-			{
-				break;
-			}
-		}
-	}
+	Schema->AutowireNewNode(this, FromPin);
 }
 
 bool UControlRigGraphNode::IsSelectedInEditor() const
diff --git a/Engine/Plugins/Animation/ControlRig/Source/ControlRigDeveloper/Private/Graph/ControlRigGraphSchema.cpp b/Engine/Plugins/Animation/ControlRig/Source/ControlRigDeveloper/Private/Graph/ControlRigGraphSchema.cpp
--- a/Engine/Plugins/Animation/ControlRig/Source/ControlRigDeveloper/Private/Graph/ControlRigGraphSchema.cpp
+++ b/Engine/Plugins/Animation/ControlRig/Source/ControlRigDeveloper/Private/Graph/ControlRigGraphSchema.cpp
@@ -2275,8 +2275,83 @@ FVector2D UControlRigGraphSchema::GetNodePositionAtStartOfInteraction(UEdGraphNo
 	return FVector2D::ZeroVector;
 }
 
+bool UControlRigGraphSchema::AutowireNewNode(UControlRigGraphNode* NewNode, UEdGraphPin* FromPin) const
+{
+	// copying high level information into a local array since the try create connection below
+	// may cause the pin array to be destroyed / changed
+	TArray<TPair<FName, EEdGraphPinDirection>> PinsToVisit;
+	for(UEdGraphPin* Pin : NewNode->Pins)
+	{
+		PinsToVisit.Emplace(Pin->GetFName(), Pin->Direction);
+	}
+
+	TArray<UEdGraphPin*> OldLinkedTo = FromPin->LinkedTo;
+	for(const TPair<FName, EEdGraphPinDirection>& PinToVisit : PinsToVisit)
+	{
+		UEdGraphPin* Pin = NewNode->FindPin(PinToVisit.Key, PinToVisit.Value);
+		if(Pin == nullptr)
+		{
+			continue;
+		}
+		
+		if (Pin->ParentPin != nullptr)
+		{
+			continue;
+		}
+
+		FPinConnectionResponse ConnectResponse = CanCreateConnection(FromPin, Pin);
+		if(ConnectResponse.Response != ECanCreateConnectionResponse::CONNECT_RESPONSE_DISALLOW)
+		{
+			if (TryCreateConnection(FromPin, Pin))
+			{
+				// It might have been linked in a different direction. Find pin with the correct direction
+				if (Pin->LinkedTo.IsEmpty())
+				{
+					for (UEdGraphPin* Linked : FromPin->LinkedTo)
+					{
+						if (Linked->GetOwningNode() == Pin->GetOwningNode())
+						{
+							Pin = Linked;
+							break;
+						}
+					}
+				}
+
+				// If the pin is an execute context, try to create a link to the pin which was previously connected to FromPin
+				if (URigVMPin* ModelPin = NewNode->FindModelPinFromGraphPin(Pin))
+				{
+					if (!OldLinkedTo.IsEmpty() && ModelPin->IsExecuteContext())
+					{
+						bool bIsInput = Pin->Direction == EEdGraphPinDirection::EGPD_Input;
+
+						// If the pin is an input to a control flow node, the output used should be the completed pin
+						if (ModelPin->GetNode()->IsControlFlowNode())
+						{
+							if (bIsInput)
+							{
+								ModelPin = ModelPin->GetNode()->FindPin(FRigVMStruct::ControlFlowCompletedName.ToString());
+							}
+						}
+
+						// Try to create the second connection
+						if (UEdGraphPin* OppositePin = NewNode->FindGraphPinFromModelPin(ModelPin, !bIsInput))
+						{
+							if (CanCreateConnection(OppositePin, OldLinkedTo[0]).Response != ECanCreateConnectionResponse::CONNECT_RESPONSE_DISALLOW)
+							{
+								TryCreateConnection(OppositePin, OldLinkedTo[0]);
+							}
+						}
+					}
+				}
+				return true;
+			}
+		}
+	}
+	return false;
+}
+
 void UControlRigGraphSchema::HandleModifiedEvent(ERigVMGraphNotifType InNotifType, URigVMGraph* InGraph,
-	UObject* InSubject)
+                                                 UObject* InSubject)
 {
 	switch(InNotifType)
 	{
diff --git a/Engine/Plugins/Animation/ControlRig/Source/ControlRigDeveloper/Public/Graph/ControlRigGraphSchema.h b/Engine/Plugins/Animation/ControlRig/Source/ControlRigDeveloper/Public/Graph/ControlRigGraphSchema.h
--- a/Engine/Plugins/Animation/ControlRig/Source/ControlRigDeveloper/Public/Graph/ControlRigGraphSchema.h
+++ b/Engine/Plugins/Animation/ControlRig/Source/ControlRigDeveloper/Public/Graph/ControlRigGraphSchema.h
@@ -322,6 +322,8 @@ public:
 	static TArray<UEdGraphNode*> GetNodesToMoveForNode(UEdGraphNode* InNode);
 	FVector2D GetNodePositionAtStartOfInteraction(UEdGraphNode* InNode) const;
 
+	bool AutowireNewNode(UControlRigGraphNode* NewNode, UEdGraphPin* FromPin) const;
+
 	void HandleModifiedEvent(ERigVMGraphNotifType InNotifType, URigVMGraph* InGraph, UObject* InSubject);
 
 	// Allow derived classes to spawn derived node classes
