Fix dramatic performance issue in Control Rig

diff --git a/Engine/Plugins/Animation/ControlRig/Source/ControlRig/Private/AnimNode_ControlRig.cpp b/Engine/Plugins/Animation/ControlRig/Source/ControlRig/Private/AnimNode_ControlRig.cpp
--- a/Engine/Plugins/Animation/ControlRig/Source/ControlRig/Private/AnimNode_ControlRig.cpp
+++ b/Engine/Plugins/Animation/ControlRig/Source/ControlRig/Private/AnimNode_ControlRig.cpp
@@ -369,7 +369,16 @@ void FAnimNode_ControlRig::PropagateInputProperties(const UObject* InSourceInsta
 	if (TargetInstance)
 	{
 		UControlRig* TargetControlRig = Cast<UControlRig>((UObject*)TargetInstance);
-
+		if(TargetControlRig == nullptr)
+		{
+			return;
+		}
+		URigHierarchy* TargetHierarchy = TargetControlRig->GetHierarchy();
+		if(TargetHierarchy == nullptr)
+		{
+			return;
+		}
+		
 		// First copy properties
 		check(SourceProperties.Num() == DestProperties.Num());
 		for (int32 PropIdx = 0; PropIdx < SourceProperties.Num(); ++PropIdx)
@@ -504,7 +513,7 @@ void FAnimNode_ControlRig::PropagateInputProperties(const UObject* InSourceInsta
 
 				if(bIsValid)
 				{
-					TargetControlRig->GetHierarchy()->SetControlValue(ControlElement, Value, ERigControlValueType::Current);
+					TargetHierarchy->SetControlValue(ControlElement, Value, ERigControlValueType::Current);
 				}
 				continue;
 			}
diff --git a/Engine/Plugins/Animation/ControlRig/Source/ControlRig/Private/AnimNode_ControlRigBase.cpp b/Engine/Plugins/Animation/ControlRig/Source/ControlRig/Private/AnimNode_ControlRigBase.cpp
--- a/Engine/Plugins/Animation/ControlRig/Source/ControlRig/Private/AnimNode_ControlRigBase.cpp
+++ b/Engine/Plugins/Animation/ControlRig/Source/ControlRig/Private/AnimNode_ControlRigBase.cpp
@@ -140,14 +140,17 @@ void FAnimNode_ControlRigBase::UpdateInput(UControlRig* ControlRig, const FPoseC
 		return;
 	}
 
+	URigHierarchy* Hierarchy = ControlRig->GetHierarchy();
+	if(Hierarchy == nullptr)
+	{
+		return;
+	}
+
 #if WITH_EDITOR
 	// if we are recording any change - let's clear the undo stack
-	if(URigHierarchy* Hierarchy = ControlRig->GetHierarchy())
+	if(Hierarchy->IsTracingChanges())
 	{
-		if(Hierarchy->IsTracingChanges())
-		{
-			Hierarchy->ResetTransformStack();
-		}
+		Hierarchy->ResetTransformStack();
 	}
 #endif
 
@@ -163,16 +166,16 @@ void FAnimNode_ControlRigBase::UpdateInput(UControlRig* ControlRig, const FPoseC
 		if(!ControlRigBoneInputMappingByName.IsEmpty() || bResetInputPoseToInitial)
 		{
 			{
-				FRigHierarchyValidityBracket ValidityBracket(ControlRig->GetHierarchy());
+				FRigHierarchyValidityBracket ValidityBracket(Hierarchy);
 			}
-			FRigHierarchyValidityBracket ValidityBracket(ControlRig->GetHierarchy());
+			FRigHierarchyValidityBracket ValidityBracket(Hierarchy);
 
 			{
 #if WITH_EDITOR
 				// make sure transient controls don't get reset
 				UControlRig::FTransientControlPoseScope PoseScope(ControlRig);
 #endif 
-				ControlRig->GetHierarchy()->ResetPoseToInitial(ERigElementType::Bone);
+				Hierarchy->ResetPoseToInitial(ERigElementType::Bone);
 			}
 		}
 
@@ -192,7 +195,7 @@ void FAnimNode_ControlRigBase::UpdateInput(UControlRig* ControlRig, const FPoseC
 					
 					FCompactPoseBoneIndex CompactPoseIndex(SkeletonIndex);
 					FTransform ComponentTransform = MeshPoses.GetComponentSpaceTransform(CompactPoseIndex);
-					ControlRig->GetHierarchy()->SetGlobalTransformByIndex(ControlRigIndex, ComponentTransform, false);
+					Hierarchy->SetGlobalTransformByIndex(ControlRigIndex, ComponentTransform, false);
 				}
 			}
 			else
@@ -209,7 +212,7 @@ void FAnimNode_ControlRigBase::UpdateInput(UControlRig* ControlRig, const FPoseC
 					{
 						ComponentTransform = NodeMappingContainer->GetSourceToTargetTransform(Name).GetRelativeTransformReverse(ComponentTransform);
 					}
-					ControlRig->GetHierarchy()->SetGlobalTransform(Key, ComponentTransform, false);
+					Hierarchy->SetGlobalTransform(Key, ComponentTransform, false);
 				}
 			}
 		}
@@ -224,7 +227,7 @@ void FAnimNode_ControlRigBase::UpdateInput(UControlRig* ControlRig, const FPoseC
 					
 					FCompactPoseBoneIndex CompactPoseIndex(SkeletonIndex);
 					FTransform LocalTransform = InOutput.Pose[CompactPoseIndex];
-					ControlRig->GetHierarchy()->SetLocalTransformByIndex(ControlRigIndex, LocalTransform, false);
+					Hierarchy->SetLocalTransformByIndex(ControlRigIndex, LocalTransform, false);
 				}
 			}
 			else
@@ -237,7 +240,7 @@ void FAnimNode_ControlRigBase::UpdateInput(UControlRig* ControlRig, const FPoseC
 
 					FCompactPoseBoneIndex CompactPoseIndex(SkeletonIndex);
 					FTransform LocalTransform = InOutput.Pose[CompactPoseIndex];
-					ControlRig->GetHierarchy()->SetLocalTransform(Key, LocalTransform, false);
+					Hierarchy->SetLocalTransform(Key, LocalTransform, false);
 				}
 			}
 		}
@@ -248,7 +251,6 @@ void FAnimNode_ControlRigBase::UpdateInput(UControlRig* ControlRig, const FPoseC
 	{
 		// TODO: Do we need to 'unset' old values manually here? The previous code could do this because it knows the
 		// global set of curves, but we dont know that any more
-		URigHierarchy* Hierarchy = ControlRig->GetHierarchy();
 		UE::Anim::FCurveUtils::BulkGet(InOutput.Curve, CachedBulkCurves, [Hierarchy](const UE::Anim::FNamedIndexElement& InBulkElement, float InValue)
 		{
 			const FRigElementKey Key(InBulkElement.Name, ERigElementType::Curve);
@@ -257,12 +259,9 @@ void FAnimNode_ControlRigBase::UpdateInput(UControlRig* ControlRig, const FPoseC
 	}
 
 #if WITH_EDITOR
-	if(URigHierarchy* Hierarchy = ControlRig->GetHierarchy())
+	if(Hierarchy->IsTracingChanges())
 	{
-		if(Hierarchy->IsTracingChanges())
-		{
-			Hierarchy->StorePoseForTrace(TEXT("FAnimNode_ControlRigBase::UpdateInput"));
-		}
+		Hierarchy->StorePoseForTrace(TEXT("FAnimNode_ControlRigBase::UpdateInput"));
 	}
 #endif
 }
@@ -276,6 +275,12 @@ void FAnimNode_ControlRigBase::UpdateOutput(UControlRig* ControlRig, FPoseContex
 		return;
 	}
 
+	URigHierarchy* Hierarchy = ControlRig->GetHierarchy();
+	if(Hierarchy == nullptr)
+	{
+		return;
+	}
+
 	DECLARE_SCOPE_HIERARCHICAL_COUNTER_FUNC()
 
 	if (OutputSettings.bUpdatePose)
@@ -307,7 +312,7 @@ void FAnimNode_ControlRigBase::UpdateOutput(UControlRig* ControlRig, FPoseContex
 					const uint16 SkeletonIndex = Pair.Value;
 
 					FCompactPoseBoneIndex CompactPoseIndex(SkeletonIndex);
-					FTransform ComponentTransform = ControlRig->GetHierarchy()->GetGlobalTransformByIndex(ControlRigIndex);
+					FTransform ComponentTransform = Hierarchy->GetGlobalTransformByIndex(ControlRigIndex);
 					MeshPoses.SetComponentSpaceTransform(CompactPoseIndex, ComponentTransform);
 				}
 			}
@@ -320,7 +325,7 @@ void FAnimNode_ControlRigBase::UpdateOutput(UControlRig* ControlRig, FPoseContex
 					const FRigElementKey Key(Name, ERigElementType::Bone);
 
 					FCompactPoseBoneIndex CompactPoseIndex(SkeletonIndex);
-					FTransform ComponentTransform = ControlRig->GetHierarchy()->GetGlobalTransform(Key);
+					FTransform ComponentTransform = Hierarchy->GetGlobalTransform(Key);
 					if (NodeMappingContainer.IsValid())
 					{
 						ComponentTransform = NodeMappingContainer->GetSourceToTargetTransform(Name) * ComponentTransform;
@@ -343,7 +348,7 @@ void FAnimNode_ControlRigBase::UpdateOutput(UControlRig* ControlRig, FPoseContex
 					const uint16 SkeletonIndex = Pair.Value;
 
 					FCompactPoseBoneIndex CompactPoseIndex(SkeletonIndex);
-					FTransform LocalTransform = ControlRig->GetHierarchy()->GetLocalTransformByIndex(ControlRigIndex);
+					FTransform LocalTransform = Hierarchy->GetLocalTransformByIndex(ControlRigIndex);
 					InOutput.Pose[CompactPoseIndex] = LocalTransform;
 				}
 			}
@@ -356,7 +361,7 @@ void FAnimNode_ControlRigBase::UpdateOutput(UControlRig* ControlRig, FPoseContex
 					const FRigElementKey Key(Name, ERigElementType::Bone);
 
 					FCompactPoseBoneIndex CompactPoseIndex(Index);
-					FTransform LocalTransform = ControlRig->GetHierarchy()->GetLocalTransform(Key);
+					FTransform LocalTransform = Hierarchy->GetLocalTransform(Key);
 					InOutput.Pose[CompactPoseIndex] = LocalTransform;
 				}
 			}
@@ -365,7 +370,6 @@ void FAnimNode_ControlRigBase::UpdateOutput(UControlRig* ControlRig, FPoseContex
 
 	if (OutputSettings.bUpdateCurves)
 	{
-		URigHierarchy* Hierarchy = ControlRig->GetHierarchy();
 		TArray<FRigCurveElement*> Curves = Hierarchy->GetCurves();
 		UE::Anim::FCurveUtils::BulkSet(InOutput.Curve, CachedBulkCurves, [Hierarchy, &Curves](const UE::Anim::FNamedIndexElement& InBulkElement)
 		{
@@ -375,13 +379,10 @@ void FAnimNode_ControlRigBase::UpdateOutput(UControlRig* ControlRig, FPoseContex
 	}
 
 #if WITH_EDITOR
-	if(URigHierarchy* Hierarchy = ControlRig->GetHierarchy())
+	if(Hierarchy->IsTracingChanges())
 	{
-		if(Hierarchy->IsTracingChanges())
-		{
-			Hierarchy->StorePoseForTrace(TEXT("FAnimNode_ControlRigBase::UpdateOutput"));
-			Hierarchy->DumpTransformStackToFile();
-		}
+		Hierarchy->StorePoseForTrace(TEXT("FAnimNode_ControlRigBase::UpdateOutput"));
+		Hierarchy->DumpTransformStackToFile();
 	}
 #endif
 }
@@ -447,15 +448,19 @@ void FAnimNode_ControlRigBase::ExecuteControlRig(FPoseContext& InOutput)
 		// first update input to the system
 		UpdateInput(ControlRig, InOutput);
 
+		
 		if (bExecute)
 		{
 #if WITH_EDITOR
-			if(URigHierarchy* Hierarchy = ControlRig->GetHierarchy())
+			URigHierarchy* Hierarchy = ControlRig->GetHierarchy();
+			if(Hierarchy == nullptr)
 			{
-				if(Hierarchy->IsTracingChanges())
-				{
-					Hierarchy->StorePoseForTrace(TEXT("FAnimNode_ControlRigBase::BeforeEvaluate"));
-				}
+				return;
+			}
+
+			if(Hierarchy->IsTracingChanges())
+			{
+				Hierarchy->StorePoseForTrace(TEXT("FAnimNode_ControlRigBase::BeforeEvaluate"));
 			}
 #endif
 
@@ -493,12 +498,9 @@ void FAnimNode_ControlRigBase::ExecuteControlRig(FPoseContext& InOutput)
 #endif
 
 #if WITH_EDITOR
-			if(URigHierarchy* Hierarchy = ControlRig->GetHierarchy())
+			if(Hierarchy->IsTracingChanges())
 			{
-				if(Hierarchy->IsTracingChanges())
-				{
-					Hierarchy->StorePoseForTrace(TEXT("FAnimNode_ControlRigBase::AfterEvaluate"));
-				}
+				Hierarchy->StorePoseForTrace(TEXT("FAnimNode_ControlRigBase::AfterEvaluate"));
 			}
 #endif
 		}
@@ -549,6 +551,12 @@ void FAnimNode_ControlRigBase::CacheBones_AnyThread(const FAnimationCacheBonesCo
 
 	if (UControlRig* ControlRig = GetControlRig())
 	{
+		URigHierarchy* Hierarchy = ControlRig->GetHierarchy();
+		if(Hierarchy == nullptr)
+		{
+			return;
+		}
+
 		// fill up node names
 		FBoneContainer& RequiredBones = Context.AnimInstanceProxy->GetRequiredBones();
 
@@ -677,7 +685,7 @@ void FAnimNode_ControlRigBase::CacheBones_AnyThread(const FAnimationCacheBonesCo
 					for (auto Iter = NameBasedMapping.CreateConstIterator(); Iter; ++Iter)
 					{
 						const uint16 SkeletonIndex = Iter.Value();
-						const int32 ControlRigIndex = ControlRig->GetHierarchy()->GetIndex(FRigElementKey(Iter.Key(), ERigElementType::Bone));
+						const int32 ControlRigIndex = Hierarchy->GetIndex(FRigElementKey(Iter.Key(), ERigElementType::Bone));
 						if(ControlRigIndex != INDEX_NONE)
 						{
 							IndexBasedMapping.Add(TPair<uint16, uint16>((uint16)ControlRigIndex, SkeletonIndex));
