Control Rig: Fixed dispatch node not being resolved after load

diff --git a/Engine/Plugins/Animation/ControlRig/Source/ControlRigDeveloper/Private/ControlRigBlueprint.cpp b/Engine/Plugins/Animation/ControlRig/Source/ControlRigDeveloper/Private/ControlRigBlueprint.cpp
--- a/Engine/Plugins/Animation/ControlRig/Source/ControlRigDeveloper/Private/ControlRigBlueprint.cpp
+++ b/Engine/Plugins/Animation/ControlRig/Source/ControlRigDeveloper/Private/ControlRigBlueprint.cpp
@@ -1522,6 +1522,18 @@ void UControlRigBlueprint::RefreshAllModels(EControlRigBlueprintLoadType InLoadT
 						static const TCHAR UnresolvedDispatchNodeMessage[] = TEXT("Dispatch node %s has no factory..");
 						Controller->ReportErrorf(UnresolvedDispatchNodeMessage, *ModelNode->GetNodePath(true));
 					}
+					else if (!DispatchNode->HasWildCardPin())
+					{
+						if (DispatchNode->GetResolvedFunction() == nullptr)
+						{
+							Controller->FullyResolveTemplateNode(DispatchNode, INDEX_NONE, false);
+						}
+						if (DispatchNode->GetResolvedFunction() == nullptr)
+						{
+							static const TCHAR UnresolvedDispatchNodeMessage[] = TEXT("Node %s could not be resolved.");
+							Controller->ReportErrorf(UnresolvedDispatchNodeMessage, *ModelNode->GetNodePath(true));
+						}
+					}
 				}
 			}
 		}
diff --git a/Engine/Source/Developer/RigVMDeveloper/Private/RigVMCompiler/RigVMCompiler.cpp b/Engine/Source/Developer/RigVMDeveloper/Private/RigVMCompiler/RigVMCompiler.cpp
--- a/Engine/Source/Developer/RigVMDeveloper/Private/RigVMCompiler/RigVMCompiler.cpp
+++ b/Engine/Source/Developer/RigVMDeveloper/Private/RigVMCompiler/RigVMCompiler.cpp
@@ -1584,8 +1584,8 @@ int32 URigVMCompiler::TraverseCallExtern(const FRigVMCallExternExprAST* InExpr,
 		{
 			Function = DispatchNode->GetResolvedFunction();
 		}
-			
-		check(Function);
+
+		checkf(Function, TEXT("Could not find function for node %s in package %s"), *Node->GetPathName(), *GetPackage()->GetPathName());
 
 		FRigVMOperand CountOperand;
 		FRigVMOperand IndexOperand;
