Control Rig: Fix sequence collapse graphs not loading links

diff --git a/Engine/Plugins/Animation/ControlRig/Source/ControlRigDeveloper/Private/ControlRigBlueprint.cpp b/Engine/Plugins/Animation/ControlRig/Source/ControlRigDeveloper/Private/ControlRigBlueprint.cpp
--- a/Engine/Plugins/Animation/ControlRig/Source/ControlRigDeveloper/Private/ControlRigBlueprint.cpp
+++ b/Engine/Plugins/Animation/ControlRig/Source/ControlRigDeveloper/Private/ControlRigBlueprint.cpp
@@ -1437,6 +1437,7 @@ void UControlRigBlueprint::RefreshAllModels(EControlRigBlueprintLoadType InLoadT
 			URigVMController* Controller = GetOrCreateController(Graph);
 			// temporarily disable default value validation during load time, serialized values should always be accepted
 			TGuardValue<bool> PerGraphDisablePinDefaultValueValidation(Controller->bValidatePinDefaults, false);
+			TGuardValue<bool> GuardEditGraph(Graph->bEditable, true);
 			FRigVMControllerNotifGuard NotifGuard(Controller, true);
 			LinkedPaths.Add(Graph, Controller->GetLinkedPaths());
 			Controller->FastBreakLinkedPaths(LinkedPaths.FindChecked(Graph));
@@ -1452,6 +1453,7 @@ void UControlRigBlueprint::RefreshAllModels(EControlRigBlueprintLoadType InLoadT
 	for (URigVMGraph* Graph : AllModelsLeavesFirst)
 	{
 		URigVMController* Controller = GetOrCreateController(Graph);
+		TGuardValue<bool> GuardEditGraph(Graph->bEditable, true);
 		FRigVMControllerNotifGuard NotifGuard(Controller, true);
 		{
 			URigVMController::FRestoreLinkedPathSettings Settings;
@@ -1459,8 +1461,6 @@ void UControlRigBlueprint::RefreshAllModels(EControlRigBlueprintLoadType InLoadT
 			Controller->RestoreLinkedPaths(LinkedPaths.FindChecked(Graph));
 		}
 
-		TGuardValue<bool> GuardEditGraph(Graph->bEditable, true);
-		
 		for(URigVMNode* ModelNode : Graph->GetNodes())
 		{
 			Controller->RemoveUnusedOrphanedPins(ModelNode);
diff --git a/Engine/Source/Developer/RigVMDeveloper/Private/RigVMModel/RigVMController.cpp b/Engine/Source/Developer/RigVMDeveloper/Private/RigVMModel/RigVMController.cpp
--- a/Engine/Source/Developer/RigVMDeveloper/Private/RigVMModel/RigVMController.cpp
+++ b/Engine/Source/Developer/RigVMDeveloper/Private/RigVMModel/RigVMController.cpp
@@ -9529,23 +9529,45 @@ bool URigVMController::AddLink(URigVMPin* OutputPin, URigVMPin* InputPin, bool b
 {
 	if (!bIsTransacting && !IsGraphEditable())
 	{
+		static const FString FailureReason = TEXT("Cannot add links in non-editable graph."); 
+		if(OutFailureReason)
+		{
+			*OutFailureReason = FailureReason;
+		}
+		ReportError(FailureReason);
 		return false;
 	}
 	
 	if(OutputPin == nullptr)
 	{
-		ReportError(TEXT("OutputPin is nullptr."));
+		static const FString FailureReason = TEXT("OutputPin is nullptr."); 
+		if(OutFailureReason)
+		{
+			*OutFailureReason = FailureReason;
+		}
+		ReportError(FailureReason);
 		return false;
 	}
 
 	if(InputPin == nullptr)
 	{
-		ReportError(TEXT("InputPin is nullptr."));
+		static const FString FailureReason = TEXT("InputPin is nullptr."); 
+		if(OutFailureReason)
+		{
+			*OutFailureReason = FailureReason;
+		}
+		ReportError(FailureReason);
 		return false;
 	}
 
 	if(!IsValidPinForGraph(OutputPin) || !IsValidPinForGraph(InputPin))
 	{
+		static const FString FailureReason = TEXT("Pin is not valid for graph."); 
+		if(OutFailureReason)
+		{
+			*OutFailureReason = FailureReason;
+		}
+		ReportError(FailureReason);
 		return false;
 	}
 
