Control Rig: Fix aggregate nodes removing subpins

diff --git a/Engine/Plugins/Animation/ControlRig/Source/ControlRigDeveloper/Private/Graph/ControlRigGraphNode.cpp b/Engine/Plugins/Animation/ControlRig/Source/ControlRigDeveloper/Private/Graph/ControlRigGraphNode.cpp
--- a/Engine/Plugins/Animation/ControlRig/Source/ControlRigDeveloper/Private/Graph/ControlRigGraphNode.cpp
+++ b/Engine/Plugins/Animation/ControlRig/Source/ControlRigDeveloper/Private/Graph/ControlRigGraphNode.cpp
@@ -848,8 +848,17 @@ bool UControlRigGraphNode::ModelPinsChanged(bool bForce)
 		for(const auto& Pair : CachedPins)
 		{
 			URigVMPin* ModelPin = Pair.Key;
-			const URigVMPin* RootPin = ModelPin->GetRootPin();
-			if(!PinListForPin(RootPin).Contains(RootPin))
+			TArray<URigVMPin*>& List = PinListForPin(ModelPin);
+			bool bRemove = true;
+			if (URigVMPin** FoundPin = List.FindByKey(ModelPin->GetRootPin()))
+			{
+				if (ModelPin->GetParentPin() == nullptr ||
+					(*FoundPin)->FindSubPin(ModelPin->GetSegmentPath()))
+				{
+					bRemove = false;
+				}
+			}
+			if(bRemove)
 			{
 				PinsToRemove.Add(ModelPin);
 			}
