Sequencer: Control Rig: Don't set previous key on first frame always to avoid evaluating out of range. #jira UE-181133

diff --git a/Engine/Plugins/Animation/ControlRig/Source/ControlRigEditor/Private/Sequencer/ControlRigSpaceChannelEditors.cpp b/Engine/Plugins/Animation/ControlRig/Source/ControlRigEditor/Private/Sequencer/ControlRigSpaceChannelEditors.cpp
--- a/Engine/Plugins/Animation/ControlRig/Source/ControlRigEditor/Private/Sequencer/ControlRigSpaceChannelEditors.cpp
+++ b/Engine/Plugins/Animation/ControlRig/Source/ControlRigEditor/Private/Sequencer/ControlRigSpaceChannelEditors.cpp
@@ -351,19 +351,19 @@ FKeyHandle FControlRigSpaceChannelHelpers::SequencerKeyControlRigSpaceChannel(UC
 			}
 		}
 		SectionToKey->Modify();
+		const FFrameNumber StartFrame = Sequencer->GetFocusedMovieSceneSequence()->GetMovieScene()->GetPlaybackRange().GetLowerBoundValue();
+		if (StartFrame == Time)
+		{
+			bSetPreviousKey = false;
+		}
 		//if we have no keys need to set key for current space at start frame, unless setting key at start time, where then don't do previous compensation
-		if (Channel->GetNumKeys() == 0 && Sequencer->GetFocusedMovieSceneSequence() && Sequencer->GetFocusedMovieSceneSequence()->GetMovieScene())
+		if (Channel->GetNumKeys() == 0)
 		{
-			FFrameNumber StartFrame = Sequencer->GetFocusedMovieSceneSequence()->GetMovieScene()->GetPlaybackRange().GetLowerBoundValue();
 			if (StartFrame != Time)
 			{
 				FMovieSceneControlRigSpaceBaseKey Original = ExistingValue;
 				ChannelInterface.AddKey(StartFrame, Forward<FMovieSceneControlRigSpaceBaseKey>(Original));
 			}
-			else
-			{
-				bSetPreviousKey = false;
-			}
 		}
 
 		TArray<FTransform> ControlRigParentWorldTransforms;
@@ -445,9 +445,9 @@ FKeyHandle FControlRigSpaceChannelHelpers::SequencerKeyControlRigSpaceChannel(UC
 		// set previous key if we're going to switch space
 		if (bSetPreviousKey)
 		{
-			FFrameTime GlobalTime(Time -1);
+			FFrameTime GlobalTime(Time - 1);
 			GlobalTime = GlobalTime * RootToLocalTransform.InverseLinearOnly();
-			
+
 			FMovieSceneContext SceneContext = FMovieSceneContext(FMovieSceneEvaluationRange(GlobalTime, TickResolution), Sequencer->GetPlaybackStatus()).SetHasJumped(true);
 			Sequencer->GetEvaluationTemplate().EvaluateSynchronousBlocking(SceneContext, *Sequencer);
 
