Animation/ Control Rig Mode:  Disable selection when clicking on an axis.
#trivial

diff --git a/Engine/Plugins/Animation/ControlRig/Source/ControlRigEditor/Private/EditMode/ControlRigEditMode.cpp b/Engine/Plugins/Animation/ControlRig/Source/ControlRigEditor/Private/EditMode/ControlRigEditMode.cpp
--- a/Engine/Plugins/Animation/ControlRig/Source/ControlRigEditor/Private/EditMode/ControlRigEditMode.cpp
+++ b/Engine/Plugins/Animation/ControlRig/Source/ControlRigEditor/Private/EditMode/ControlRigEditMode.cpp
@@ -1250,8 +1250,15 @@ bool FControlRigEditMode::GetCustomInputCoordinateSystem(FMatrix& OutMatrix, voi
 
 bool FControlRigEditMode::HandleClick(FEditorViewportClient* InViewportClient, HHitProxy *HitProxy, const FViewportClick &Click)
 {
+	const EAxisList::Type CurrentAxis = InViewportClient->GetCurrentWidgetAxis();
+	//if we are hitting a widget, besides arcball then bail saying we are handling it
+	if (CurrentAxis != EAxisList::None)
+	{
+		return true;
+	}
+
 	InteractionType = GetInteractionType(InViewportClient);
-	
+
 	if(HActor* ActorHitProxy = HitProxyCast<HActor>(HitProxy))
 	{
 		if(ActorHitProxy->Actor)
