Control Rig: Fix crash when resolving cast

diff --git a/Engine/Plugins/Animation/ControlRig/Source/ControlRigEditor/Private/ControlRigEditorModule.cpp b/Engine/Plugins/Animation/ControlRig/Source/ControlRigEditor/Private/ControlRigEditorModule.cpp
--- a/Engine/Plugins/Animation/ControlRig/Source/ControlRigEditor/Private/ControlRigEditorModule.cpp
+++ b/Engine/Plugins/Animation/ControlRig/Source/ControlRigEditor/Private/ControlRigEditorModule.cpp
@@ -1465,10 +1465,11 @@ void FControlRigEditorModule::GetContextMenuActions(const UControlRigGraphSchema
 									{
 										if(!Argument->IsSingleton())
 										{
+											TArray<TRigVMTypeIndex> ResolvedTypeIndices = Argument->GetSupportedTypeIndices(TemplateNode->GetResolvedPermutationIndices(true));
 											TSharedRef<SControlRigChangePinType> ChangePinTypeWidget =
 											SNew(SControlRigChangePinType)
 											.Blueprint(RigBlueprint)
-											.Types(Argument->GetTypeIndices())
+											.Types(ResolvedTypeIndices)
 											.OnTypeSelected_Lambda([RigBlueprint, ModelPin](const TRigVMTypeIndex& TypeSelected)
 											{
 												if (URigVMController* Controller = RigBlueprint->GetController(ModelPin->GetGraph()))
diff --git a/Engine/Plugins/Runtime/RigVM/Source/RigVMDeveloper/Private/RigVMModel/Nodes/RigVMTemplateNode.cpp b/Engine/Plugins/Runtime/RigVM/Source/RigVMDeveloper/Private/RigVMModel/Nodes/RigVMTemplateNode.cpp
--- a/Engine/Plugins/Runtime/RigVM/Source/RigVMDeveloper/Private/RigVMModel/Nodes/RigVMTemplateNode.cpp
+++ b/Engine/Plugins/Runtime/RigVM/Source/RigVMDeveloper/Private/RigVMModel/Nodes/RigVMTemplateNode.cpp
@@ -801,6 +801,7 @@ void URigVMTemplateNode::InvalidateCache()
 	
 	CachedFunction = nullptr;
 	CachedTemplate = nullptr;
+	ResolvedPermutation = INDEX_NONE;
 
 	if (HasWildCardPin())
 	{
