Control Rig: Support Asset User Data as input data

diff --git a/Engine/Plugins/Animation/ControlRig/Source/ControlRig/Private/ControlRig.cpp b/Engine/Plugins/Animation/ControlRig/Source/ControlRig/Private/ControlRig.cpp
--- a/Engine/Plugins/Animation/ControlRig/Source/ControlRig/Private/ControlRig.cpp
+++ b/Engine/Plugins/Animation/ControlRig/Source/ControlRig/Private/ControlRig.cpp
@@ -162,6 +162,115 @@ void UControlRig::Initialize(bool bRequestInit)
 	GetHierarchy()->UpdateVisibilityOnProxyControls();
 }
 
+void UControlRig::OnAddShapeLibrary(const FControlRigExecuteContext* InContext, const FString& InLibraryName, UControlRigShapeLibrary* InShapeLibrary, bool bReplaceExisting, bool bLogResults)
+{
+	// don't ever change the CDO
+	if (HasAnyFlags(RF_ClassDefaultObject))
+	{
+		return;
+	}
+
+	if(InShapeLibrary == nullptr)
+	{
+		if(InContext)
+		{
+			static constexpr TCHAR Message[] = TEXT("No shape library provided.");
+			InContext->Log(EMessageSeverity::Error, Message);
+		}
+		return;
+	}
+
+	// remove the shape library in case it was there before
+	ShapeLibraryNameMap.Remove(InShapeLibrary->GetName());
+	ShapeLibraries.Remove(InShapeLibrary);
+
+	// if we've removed all shape libraries - let's add the ones from the CDO back
+	if (ShapeLibraries.IsEmpty())
+	{
+		UControlRig* CDO = GetClass()->GetDefaultObject<UControlRig>();
+		ShapeLibraries = CDO->ShapeLibraries;
+	}
+
+	// if we are supposed to replace the library and the library name is empty
+	FString LibraryName = InLibraryName;
+	if(LibraryName.IsEmpty() && bReplaceExisting)
+	{
+		if(ShapeLibraries.Num() == 1)
+		{
+			LibraryName = ShapeLibraries[0]->GetName();
+		}
+	}
+
+	if(LibraryName.IsEmpty())
+	{
+		LibraryName = InShapeLibrary->GetName();
+	}
+
+	if(LibraryName != InShapeLibrary->GetName())
+	{
+		ShapeLibraryNameMap.FindOrAdd(InShapeLibrary->GetName()) = LibraryName;
+	}
+
+	if(bReplaceExisting)
+	{
+		for(int32 Index = 0; Index < ShapeLibraries.Num(); Index++)
+		{
+			const TSoftObjectPtr<UControlRigShapeLibrary>& ExistingShapeLibrary = ShapeLibraries[Index];
+			if(ExistingShapeLibrary.IsNull())
+			{
+				continue;
+			}
+			if(ExistingShapeLibrary->GetName().Equals(LibraryName, ESearchCase::IgnoreCase))
+			{
+				ShapeLibraries[Index] = InShapeLibrary;
+				break;
+			}
+		}
+	}
+
+	ShapeLibraries.AddUnique(InShapeLibrary);
+
+#if WITH_EDITOR
+	if(bLogResults)
+	{
+		static constexpr TCHAR MapFormat[] = TEXT("Control Rig '%s': Shape Library Name Map: '%s' -> '%s'");
+		static constexpr TCHAR LibraryFormat[] = TEXT("Control Rig '%s': Shape Library '%s' uses asset '%s'");
+		static constexpr TCHAR DefaultShapeFormat[] = TEXT("Control Rig '%s': Shape Library '%s' has default shape '%s'");
+		static constexpr TCHAR ShapeFormat[] = TEXT("Control Rig '%s': Shape Library '%s' contains shape %03d: '%s'");
+
+		const FString PathName = GetPathName();
+
+		for(const TPair<FString, FString>& Pair: ShapeLibraryNameMap)
+		{
+			UE_LOG(LogControlRig, Display, MapFormat, *PathName, *Pair.Key, *Pair.Value);
+		}
+
+		const int32 NumShapeLibraries = ShapeLibraries.Num();
+		
+		for(int32 Index = 0; Index <NumShapeLibraries; Index++)
+		{
+			const TSoftObjectPtr<UControlRigShapeLibrary>& ShapeLibrary = ShapeLibraries[Index];
+			if(ShapeLibrary.IsNull())
+			{
+				continue;
+			}
+
+			UE_LOG(LogControlRig, Display, LibraryFormat, *PathName, *ShapeLibrary->GetName(), *ShapeLibrary->GetPathName());
+
+			const bool bUseNameSpace = ShapeLibraries.Num() > 1;
+			const FString DefaultShapeName = UControlRigShapeLibrary::GetShapeName(ShapeLibrary.Get(), bUseNameSpace, ShapeLibraryNameMap, ShapeLibrary->DefaultShape);
+			UE_LOG(LogControlRig, Display, DefaultShapeFormat, *PathName, *ShapeLibrary->GetName(), *DefaultShapeName);
+
+			for(int32 ShapeIndex = 0; ShapeIndex < ShapeLibrary->Shapes.Num(); ShapeIndex++)
+			{
+				const FString ShapeName = UControlRigShapeLibrary::GetShapeName(ShapeLibrary.Get(), bUseNameSpace, ShapeLibraryNameMap, ShapeLibrary->Shapes[ShapeIndex]);
+				UE_LOG(LogControlRig, Display, ShapeFormat, *PathName, *ShapeLibrary->GetName(), ShapeIndex, *ShapeName);
+			}
+		}
+	}
+#endif
+}
+
 bool UControlRig::InitializeVM(const FName& InEventName)
 {
 	if(!Super::InitializeVM(InEventName))
@@ -306,6 +415,11 @@ bool UControlRig::Execute(const FName& InEventName)
 	
 	FRigVMExtendedExecuteContext& ExtendedExecuteContext = GetVM()->GetContext();
 	FControlRigExecuteContext& PublicContext = ExtendedExecuteContext.GetPublicDataSafe<FControlRigExecuteContext>();
+
+#if WITH_EDITOR
+	PublicContext.SetLog(RigVMLog); // may be nullptr
+#endif
+
 	PublicContext.SetDeltaTime(DeltaTime);
 	PublicContext.SetAbsoluteTime(AbsoluteTime);
 	PublicContext.SetFramesPerSecond(GetCurrentFramesPerSecond());
@@ -403,6 +517,9 @@ bool UControlRig::Execute(const FName& InEventName)
 	Context.HierarchySettings = HierarchySettings;
 	check(PublicContext.Hierarchy);
 
+	// allow access to the shape libraries
+	PublicContext.OnAddShapeLibraryDelegate.BindUObject(this, &UControlRig::OnAddShapeLibrary);
+
 	// allow access to the default hierarchy to allow to reset
 	if(!HasAnyFlags(RF_ClassDefaultObject))
 	{
@@ -449,6 +566,21 @@ bool UControlRig::Execute(const FName& InEventName)
 			}
 		}
 	}
+	
+	// setup the user data
+	// first using the data from the rig,
+	// followed by the data coming from the skeleton
+	// then the data coming from the skeletal mesh,
+	// and in the future we'll also add the data from the control rig component
+	PublicContext.AssetUserData.Reset();
+	if(const TArray<UAssetUserData*>* ControlRigUserDataArray = GetAssetUserDataArray())
+	{
+		for(const UAssetUserData* ControlRigUserData : *ControlRigUserDataArray)
+		{
+			PublicContext.AssetUserData.Add(ControlRigUserData);
+		}
+	}
+	PublicContext.AssetUserData.Remove(nullptr);
 
 	// if we have any referenced elements dirty them
 	if(GetHierarchy())
@@ -549,6 +681,13 @@ bool UControlRig::Execute(const FName& InEventName)
 						GetHierarchy()->ResetPoseToInitial(ERigElementType::All);
 					}
 
+					// clone the shape libraries again from the CDO 
+					if (!HasAnyFlags(RF_ClassDefaultObject) && ShapeLibraries.IsEmpty())
+					{
+						UControlRig* CDO = GetClass()->GetDefaultObject<UControlRig>();
+						ShapeLibraries = CDO->ShapeLibraries;
+					}
+
 					if (PreConstructionEvent.IsBound())
 					{
 						FControlRigBracketScope BracketScope(PreConstructionBracket);
@@ -1428,7 +1567,7 @@ const TArray<TSoftObjectPtr<UControlRigShapeLibrary>>& UControlRig::GetShapeLibr
 {
 	const TArray<TSoftObjectPtr<UControlRigShapeLibrary>>* LibrariesPtr = &ShapeLibraries;
 
-	if(!GetClass()->IsNative())
+	if(!GetClass()->IsNative() && ShapeLibraries.IsEmpty())
 	{
 		if (UControlRig* CDO = Cast<UControlRig>(GetClass()->GetDefaultObject()))
 		{
@@ -2169,6 +2308,46 @@ void UControlRig::PostEditChangeProperty(FPropertyChangedEvent& PropertyChangedE
 }
 #endif
 
+const TArray<UAssetUserData*>* UControlRig::GetAssetUserDataArray() const
+{
+	if(HasAnyFlags(RF_ClassDefaultObject))
+	{
+		return &ToRawPtrTArrayUnsafe(AssetUserData);
+	}
+
+	CombinedAssetUserData.Reset();
+	CombinedAssetUserData.Append(AssetUserData);
+
+	if(OuterSceneComponent.IsValid())
+	{
+		if(const USkeletalMeshComponent* SkeletalMeshComponent = Cast<USkeletalMeshComponent>(OuterSceneComponent))
+		{
+			if(USkeletalMesh* SkeletalMesh = SkeletalMeshComponent->GetSkeletalMeshAsset())
+			{
+				if(const USkeleton* Skeleton = SkeletalMesh->GetSkeleton())
+				{
+					if(const TArray<UAssetUserData*>* SkeletonUserDataArray = Skeleton->GetAssetUserDataArray())
+					{
+						for(UAssetUserData* SkeletonUserData : *SkeletonUserDataArray)
+						{
+							CombinedAssetUserData.Add(SkeletonUserData);
+						}
+					}
+				}
+				if(const TArray<UAssetUserData*>* SkeletalMeshUserDataArray = SkeletalMesh->GetAssetUserDataArray())
+				{
+					for(UAssetUserData* SkeletalMeshUserData : *SkeletalMeshUserDataArray)
+					{
+						CombinedAssetUserData.Add(SkeletalMeshUserData);
+					}
+				}
+			}
+		}
+	}
+
+	return &ToRawPtrTArrayUnsafe(CombinedAssetUserData);
+}
+
 void UControlRig::CopyPoseFromOtherRig(UControlRig* Subject)
 {
 	check(DynamicHierarchy);
diff --git a/Engine/Plugins/Animation/ControlRig/Source/ControlRig/Private/ControlRigAssetUserData.cpp b/Engine/Plugins/Animation/ControlRig/Source/ControlRig/Private/ControlRigAssetUserData.cpp
new file mode 100644
--- /dev/null
+++ b/Engine/Plugins/Animation/ControlRig/Source/ControlRig/Private/ControlRigAssetUserData.cpp
@@ -0,0 +1,572 @@
+// Copyright Epic Games, Inc. All Rights Reserved.
+
+#include "ControlRigAssetUserData.h"
+#include "Engine/UserDefinedStruct.h"
+
+#include UE_INLINE_GENERATED_CPP_BY_NAME(ControlRigAssetUserData)
+
+const TArray<const UNameSpacedUserData::FUserData*> UNameSpacedUserData::EmptyUserDatas;
+
+FString UNameSpacedUserData::FUserData::GetName() const
+{
+	FString Left, Right;
+	if(SplitAtEnd(Path, Left, Right))
+	{
+		return Right;
+	}
+	return Path;
+}
+
+#if WITH_EDITOR
+
+FString UNameSpacedUserData::FUserData::GetDisplayName() const
+{
+	if(Property)
+	{
+		if(Property->GetTypedOwner<UUserDefinedStruct>())
+		{
+			return Property->GetDisplayNameText().ToString();
+		}
+	}
+	return GetName();
+}
+
+bool UNameSpacedUserData::FUserData::IsArray() const
+{
+	return CastField<FArrayProperty>(Property) != nullptr;
+}
+
+bool UNameSpacedUserData::FUserData::IsArrayElement() const
+{
+	if(Property)
+	{
+		return CastField<FArrayProperty>(Property->GetOwnerProperty()) != nullptr;
+	}
+	return false;
+}
+
+bool UNameSpacedUserData::FUserData::IsUObject() const
+{
+	return CastField<FObjectProperty>(Property) != nullptr;
+}
+
+#endif
+
+TArray<UStruct*> UNameSpacedUserData::FUserData::GetSuperStructs(UStruct* InStruct)
+{
+	// Create an array of structs, ordered super -> child struct
+	TArray<UStruct*> SuperStructs = {InStruct};
+	while(true)
+	{
+		if(UStruct* SuperStruct = SuperStructs[0]->GetSuperStruct())
+		{
+			SuperStructs.Insert(SuperStruct, 0);
+		}
+		else
+		{
+			break;
+		}
+	}
+	return SuperStructs;
+}
+
+void UNameSpacedUserData::Serialize(FArchive& Ar)
+{
+	Super::Serialize(Ar);
+	InvalidateCache();
+}
+
+const UNameSpacedUserData::FUserData* UNameSpacedUserData::GetUserData(const FString& InPath, FString* OutErrorMessage) const
+{
+	FScopeLock Lock(&CacheLock);
+	if(FUserData** UserDataInCache = CachedUserData.Find(InPath))
+	{
+		if(*UserDataInCache)
+		{
+			return *UserDataInCache;
+		}
+	}
+	return nullptr;
+}
+
+const TArray<const UNameSpacedUserData::FUserData*>& UNameSpacedUserData::GetUserDataArray(const FString& InParentPath, FString* OutErrorMessage) const
+{
+	FScopeLock Lock(&CacheLock);
+	if(const TArray<const FUserData*>* UserDataArrayInCache = CachedUserDataArray.Find(InParentPath))
+	{
+		return *UserDataArrayInCache;
+	}
+	return EmptyUserDatas;
+}
+
+const UNameSpacedUserData::FUserData* UNameSpacedUserData::GetUserDataWithinStruct(const UStruct* InStruct, const uint8* InMemory, const FString& InPath, const FString& InPropertyName, FString* OutErrorMessage) const
+{
+	check(InStruct);
+	check(InMemory);
+
+	FString Left, Right;
+	if(!InPropertyName.IsEmpty() || !FUserData::SplitAtEnd(InPath, Left, Right))
+	{
+		const FString& PropertyName = InPropertyName.IsEmpty() ? InPath : InPropertyName; 
+		if(const FProperty* Property = FindPropertyByName(InStruct, *PropertyName))
+		{
+			if(!IsPropertySupported(Property, InPath, true, OutErrorMessage))
+			{
+				return nullptr;
+			}
+
+			const uint8* PropertyMemory = Property->ContainerPtrToValuePtr<uint8>(InMemory);
+			const FUserData UserData(InPath, Property, PropertyMemory);
+			return StoreCacheForUserData(UserData);
+		}
+	}
+	else
+	{
+		if(const FUserData* ParentUserData = GetUserData(Left, OutErrorMessage))
+		{
+			check(ParentUserData->GetProperty());
+
+			const uint8* Memory = ParentUserData->GetMemory();
+			if(Memory == nullptr)
+			{
+				if(OutErrorMessage && OutErrorMessage->IsEmpty())
+				{
+					(*OutErrorMessage) = FString::Printf(InvalidMemoryFormat, *InPath);
+				}
+				return nullptr;
+			}
+			
+			if(const FStructProperty* StructProperty = CastField<FStructProperty>(ParentUserData->GetProperty()))
+			{
+				return GetUserDataWithinStruct(StructProperty->Struct, Memory, InPath, Right, OutErrorMessage);
+			}
+
+			if(const FArrayProperty* ArrayProperty = CastField<FArrayProperty>(ParentUserData->GetProperty()))
+			{
+				const FProperty* Property = ArrayProperty->Inner;
+				if(!IsPropertySupported(Property, InPath, false, OutErrorMessage))
+				{
+					return nullptr;
+				}
+
+				// make sure the next bit is provided as a number and is within bounds
+				if(!Right.IsNumeric())
+				{
+					if(OutErrorMessage && OutErrorMessage->IsEmpty())
+					{
+						(*OutErrorMessage) = FString::Printf(InvalidArrayIndexFormat, *Right, *InPath);
+					}
+					return nullptr;
+				}
+
+				FScriptArrayHelper ArrayHelper(ArrayProperty, Memory);
+
+				const int32 Index = FCString::Atoi(*Right);
+				if(!ArrayHelper.IsValidIndex(Index))
+				{
+					if(OutErrorMessage && OutErrorMessage->IsEmpty())
+					{
+						(*OutErrorMessage) = FString::Printf(OutOfBoundArrayIndexFormat, *InPath, ArrayHelper.Num());
+					}
+					return nullptr;
+				}
+
+				Memory = ArrayHelper.GetRawPtr(Index);
+				const FUserData UserData(InPath, Property, Memory);
+				return StoreCacheForUserData(UserData);
+			}
+
+			if(OutErrorMessage && OutErrorMessage->IsEmpty())
+			{
+				(*OutErrorMessage) = FString::Printf(UnSupportedSubPathsFormat, *ParentUserData->GetCPPType().ToString(), *Left);
+			}
+		}
+	}
+
+	if(OutErrorMessage && OutErrorMessage->IsEmpty())
+	{
+		(*OutErrorMessage) = FString::Printf(PathNotFoundFormat, *InPath);
+	}
+	return nullptr;
+}
+
+const TArray<const UNameSpacedUserData::FUserData*>& UNameSpacedUserData::GetUserDataArrayWithinStruct(UStruct* InStruct, const uint8* InMemory, const FString& InPath, FString* OutErrorMessage) const
+{
+	FScopeLock Lock(&CacheLock);
+
+	TArray<const FUserData*> UserDataArray;
+
+	const TArray<UStruct*> Structs = FUserData::GetSuperStructs(InStruct);
+	for(const UStruct* Struct : Structs)
+	{
+		//const bool bIsUserDefinedStruct = Struct->IsA<UUserDefinedStruct>();
+		for (TFieldIterator<FProperty> It(Struct, EFieldIterationFlags::None); It; ++It)
+		{
+			const FProperty* Property = *It;
+
+			const FString& PropertyName = Property->GetNameCPP();
+			const FString PropertyPath = InPath.IsEmpty() ? PropertyName : FUserData::Join(InPath, PropertyName);
+					
+			if(!IsPropertySupported(Property, PropertyPath, true, OutErrorMessage))
+			{
+				continue;
+			}
+					
+			FUserData*& UserDataInCache = CachedUserData.FindOrAdd(PropertyPath);
+			if(UserDataInCache == nullptr)
+			{
+				const uint8* PropertyMemory = Property->ContainerPtrToValuePtr<uint8>(InMemory);
+				UserDataInCache = new FUserData(PropertyPath, Property, PropertyMemory);
+			}
+
+			UserDataArray.Add(UserDataInCache);
+		}
+	}
+
+	if(!UserDataArray.IsEmpty())
+	{
+		TArray<const FUserData*>& UserDataArrayInCache = CachedUserDataArray.FindOrAdd(InPath);
+		if(UserDataArrayInCache.IsEmpty())
+		{
+			UserDataArrayInCache = UserDataArray;
+		}
+		return UserDataArrayInCache;
+	}
+
+	return EmptyUserDatas;
+}
+
+const TArray<const UNameSpacedUserData::FUserData*>& UNameSpacedUserData::GetUserDataArrayWithinArray(const FArrayProperty* InArrayProperty, const uint8* InMemory, const FString& InPath, FString* OutErrorMessage) const
+{
+	check(!InPath.IsEmpty());
+	
+	FScopeLock Lock(&CacheLock);
+
+	TArray<const FUserData*> UserDataArray;
+
+	const FProperty* Property = InArrayProperty->Inner;
+	if(!IsPropertySupported(Property, InPath, false, OutErrorMessage))
+	{
+		return EmptyUserDatas;
+	}
+
+	FScriptArrayHelper ArrayHelper(InArrayProperty, InMemory);
+	for(int32 Index = 0; Index < ArrayHelper.Num(); Index++)
+	{
+		const FString PropertyPath = FUserData::Join(InPath, FString::FromInt(Index));
+
+		FUserData*& UserDataInCache = CachedUserData.FindOrAdd(PropertyPath);
+		if(UserDataInCache == nullptr)
+		{
+			const uint8* PropertyMemory = ArrayHelper.GetRawPtr(Index);
+			UserDataInCache = new FUserData(PropertyPath, Property, PropertyMemory);
+		}
+
+		UserDataArray.Add(UserDataInCache);
+	}
+
+	if(!UserDataArray.IsEmpty())
+	{
+		TArray<const FUserData*>& UserDataArrayInCache = CachedUserDataArray.FindOrAdd(InPath);
+		if(UserDataArrayInCache.IsEmpty())
+		{
+			UserDataArrayInCache = UserDataArray;
+		}
+		return UserDataArrayInCache;
+	}
+
+	return EmptyUserDatas;
+}
+
+const FProperty* UNameSpacedUserData::FindPropertyByName(const UStruct* InStruct, const FName& InName)
+{
+	check(InStruct);
+
+	if(const FProperty* Property = InStruct->FindPropertyByName(InName))
+	{
+		return Property;
+	}
+
+	// user defined properties may actually be called quite differently
+	return nullptr;
+}
+
+void UNameSpacedUserData::InvalidateCache() const
+{
+	FScopeLock Lock(&CacheLock);
+	for(const TPair<FString, FUserData*>& Pair : CachedUserData)
+	{
+		delete(Pair.Value);
+	}
+	CachedUserData.Reset();
+	CachedUserDataArray.Reset();
+}
+
+const UNameSpacedUserData::FUserData* UNameSpacedUserData::StoreCacheForUserData(const FUserData& InUserData) const
+{
+	check(InUserData.IsValid());
+	FScopeLock Lock(&CacheLock);
+
+	if(const FUserData* const* UserDataInCache = CachedUserData.Find(InUserData.GetPath()))
+	{
+		if(*UserDataInCache)
+		{
+			check((*UserDataInCache)->GetCPPType() == InUserData.GetCPPType());
+			return *UserDataInCache;
+		}
+	}
+
+	FUserData* UserDataPtr = new FUserData(InUserData);
+	CachedUserData.Add(InUserData.GetPath(), UserDataPtr);
+	return UserDataPtr;
+}
+
+const TArray<const UNameSpacedUserData::FUserData*>& UNameSpacedUserData::StoreCacheForUserDataArray(const FString& InPath, const TArray<const FUserData*>& InUserDataArray) const
+{
+	if(InUserDataArray.IsEmpty())
+	{
+		return EmptyUserDatas;
+	}
+	
+	FScopeLock Lock(&CacheLock);
+
+	TArray<const FUserData*>& UserDataArrayInCache = CachedUserDataArray.FindOrAdd(InPath);
+	if(UserDataArrayInCache.IsEmpty())
+	{
+		UserDataArrayInCache = InUserDataArray;
+	}
+	return UserDataArrayInCache;
+}
+
+bool UNameSpacedUserData::IsPropertySupported(const FProperty* InProperty, const FString& InPath, bool bCheckPropertyFlags, FString* OutErrorMessage) const
+{
+	check(InProperty);
+
+	if(bCheckPropertyFlags)
+	{
+		if(!InProperty->HasAnyPropertyFlags(CPF_Edit | CPF_EditConst | CPF_BlueprintVisible))
+		{
+			return false;
+		}
+		if(InProperty->HasAnyPropertyFlags(CPF_EditorOnly))
+		{
+			return false;
+		}
+	}
+
+	if(InProperty->IsA<FSetProperty>() ||
+		InProperty->IsA<FMapProperty>())
+	{
+		if(OutErrorMessage && OutErrorMessage->IsEmpty())
+		{
+			FString ExtendedCppType;
+			FString CPPTypeString = InProperty->GetCPPType(&ExtendedCppType);
+			CPPTypeString += ExtendedCppType;
+
+			static constexpr TCHAR UnsupportedCPPTypeFormat[] = TEXT("Unsupported user data type '%s' at path '%s'.");
+			(*OutErrorMessage) = FString::Printf(UnsupportedCPPTypeFormat, *CPPTypeString, *InPath);
+		}
+		return false;
+	}
+
+	return true;
+}
+
+void UDataAssetLink::SetDataAsset(UDataAsset* InDataAsset)
+{
+	InvalidateCache();
+	DataAsset = InDataAsset;
+
+	if(NameSpace.IsEmpty() && DataAsset)
+	{
+		NameSpace = DataAsset->GetName();
+	}
+}
+
+const UNameSpacedUserData::FUserData* UDataAssetLink::GetUserData(const FString& InPath, FString* OutErrorMessage) const
+{
+	// rely on super to look it up from the cache
+	if(const FUserData* ResultFromSuper = Super::GetUserData(InPath))
+	{
+		return ResultFromSuper;
+	}
+
+	if(DataAsset)
+	{
+		// this method caches as well - so the next time around Super::GetUserData should return the cache 
+		return GetUserDataWithinStruct(DataAsset->GetClass(), (const uint8*)DataAsset, InPath, FString(), OutErrorMessage);
+	}
+
+	if(OutErrorMessage && OutErrorMessage->IsEmpty())
+	{
+		(*OutErrorMessage) = FString::Printf(DataAssetNullFormat, *InPath);
+	}
+	return nullptr;
+}
+
+const TArray<const UNameSpacedUserData::FUserData*>& UDataAssetLink::GetUserDataArray(const FString& InParentPath, FString* OutErrorMessage) const
+{
+	// rely on super to look it up from the cache
+	const TArray<const FUserData*>& ResultFromSuper = Super::GetUserDataArray(InParentPath);
+	if(!ResultFromSuper.IsEmpty())
+	{
+		return ResultFromSuper;
+	}
+
+	if(DataAsset)
+	{
+		// we should only get here if we haven't cached this user data array before.
+		if(InParentPath.IsEmpty())
+		{
+			// this method caches as well - so the next time around Super::GetUserDataArray should return the cache 
+			return GetUserDataArrayWithinStruct(DataAsset->GetClass(), (const uint8*)DataAsset, InParentPath, OutErrorMessage);
+		}
+
+		if(const FUserData* ParentUserData = GetUserData(InParentPath, OutErrorMessage))
+		{
+			check(ParentUserData->GetProperty());
+
+			const uint8* Memory = ParentUserData->GetMemory();
+			if(Memory == nullptr)
+			{
+				return EmptyUserDatas;
+			}
+
+			if(const FStructProperty* StructProperty = CastField<FStructProperty>(ParentUserData->GetProperty()))
+			{
+				// this method caches as well - so the next time around Super::GetUserDataArray should return the cache 
+				return GetUserDataArrayWithinStruct(StructProperty->Struct, Memory, InParentPath, OutErrorMessage);
+			}
+			
+			if(const FArrayProperty* ArrayProperty = CastField<FArrayProperty>(ParentUserData->GetProperty()))
+			{
+				// this method caches as well - so the next time around Super::GetUserDataArray should return the cache 
+				return GetUserDataArrayWithinArray(ArrayProperty, Memory, InParentPath, OutErrorMessage);
+			}
+
+			// other types of containers are not expected here due to UNameSpacedUserData::IsPropertySupported
+		}
+	}
+
+	if(OutErrorMessage && OutErrorMessage->IsEmpty())
+	{
+		(*OutErrorMessage) = FString::Printf(DataAssetNullFormat, *InParentPath);
+	}
+	return EmptyUserDatas;
+}
+
+#if WITH_EDITOR
+
+void UDataAssetLink::PostEditChangeProperty(FPropertyChangedEvent& PropertyChangedEvent)
+{
+	Super::PostEditChangeProperty(PropertyChangedEvent);
+
+	static const FProperty* DataAssetProperty= FindPropertyByName(GetClass(), GET_MEMBER_NAME_CHECKED(UDataAssetLink, DataAsset));
+	if(PropertyChangedEvent.Property == DataAssetProperty ||
+		PropertyChangedEvent.MemberProperty == DataAssetProperty)
+	{
+		SetDataAsset(DataAsset);
+	}
+}
+
+#endif
+
+void UControlRigShapeLibraryLink::SetShapeLibrary(UControlRigShapeLibrary* InShapeLibrary)
+{
+	InvalidateCache();
+	ShapeLibrary = InShapeLibrary;
+	ShapeNames.Reset();
+	if(ShapeLibrary)
+	{
+		for(const FControlRigShapeDefinition& Shape : ShapeLibrary->Shapes)
+		{
+			ShapeNames.Add(Shape.ShapeName);
+		}
+	}
+}
+
+const UNameSpacedUserData::FUserData* UControlRigShapeLibraryLink::GetUserData(const FString& InPath, FString* OutErrorMessage) const
+{
+	// rely on super to look it up from the cache
+	if(const FUserData* ResultFromSuper = Super::GetUserData(InPath))
+	{
+		return ResultFromSuper;
+	}
+
+	if(ShapeLibrary)
+	{
+		if(InPath.Equals(GET_MEMBER_NAME_STRING_CHECKED(UControlRigShapeLibraryLink, ShapeLibrary), ESearchCase::CaseSensitive))
+		{
+			static const FProperty* ShapeLibraryProperty = FindPropertyByName(StaticClass(), GET_MEMBER_NAME_CHECKED(UControlRigShapeLibraryLink, ShapeLibrary));
+			return StoreCacheForUserData({InPath, ShapeLibraryProperty, (const uint8*)ShapeLibrary});
+		}
+		if(InPath.Equals(DefaultShapePath, ESearchCase::CaseSensitive))
+		{
+			static const FProperty* DefaultShapeProperty = FindPropertyByName(FControlRigShapeDefinition::StaticStruct(), GET_MEMBER_NAME_CHECKED(FControlRigShapeDefinition, ShapeName));
+			return StoreCacheForUserData({InPath, DefaultShapeProperty, (const uint8*)&ShapeLibrary->DefaultShape.ShapeName});
+		}
+		if(InPath.Equals(ShapeNamesPath, ESearchCase::CaseSensitive))
+		{
+			static const FProperty* ShapeNamesProperty = FindPropertyByName(StaticClass(), GET_MEMBER_NAME_CHECKED(UControlRigShapeLibraryLink, ShapeNames));
+			return StoreCacheForUserData({InPath, ShapeNamesProperty, (const uint8*)&ShapeNames});
+		}
+		if(OutErrorMessage && OutErrorMessage->IsEmpty())
+		{
+			(*OutErrorMessage) = FString::Printf(PathNotFoundFormat, *InPath);
+		}
+		return nullptr;
+	}
+
+	if(OutErrorMessage && OutErrorMessage->IsEmpty())
+	{
+		(*OutErrorMessage) = FString::Printf(ShapeLibraryNullFormat, *InPath);
+	}
+	return nullptr;
+}
+
+const TArray<const UNameSpacedUserData::FUserData*>& UControlRigShapeLibraryLink::GetUserDataArray(const FString& InParentPath, FString* OutErrorMessage) const
+{
+	const TArray<const FUserData*>& ResultFromSuper = Super::GetUserDataArray(InParentPath);
+	if(!ResultFromSuper.IsEmpty())
+	{
+		return ResultFromSuper;
+	}
+
+	if(ShapeLibrary)
+	{
+		// UControlRigShapeLibraryLink doesn't offer any arrays other than the top level 
+		if(InParentPath.IsEmpty())
+		{
+			const FUserData* ShapeLibraryUserData = GetUserData(GET_MEMBER_NAME_STRING_CHECKED(UControlRigShapeLibraryLink, ShapeLibrary), OutErrorMessage);
+			const FUserData* DefaultShapeUserData = GetUserData(DefaultShapePath, OutErrorMessage);
+			const FUserData* ShapeNamesUserData = GetUserData(ShapeNamesPath, OutErrorMessage);
+			check(ShapeLibraryUserData);
+			check(DefaultShapeUserData);
+			check(ShapeNamesUserData);
+			return StoreCacheForUserDataArray(InParentPath, {ShapeLibraryUserData, DefaultShapeUserData, ShapeNamesUserData});
+		}
+		return EmptyUserDatas;
+	}
+
+	if(OutErrorMessage && OutErrorMessage->IsEmpty())
+	{
+		(*OutErrorMessage) = FString::Printf(ShapeLibraryNullFormat, *InParentPath);
+	}
+	return EmptyUserDatas;
+}
+
+#if WITH_EDITOR
+
+void UControlRigShapeLibraryLink::PostEditChangeProperty(FPropertyChangedEvent& PropertyChangedEvent)
+{
+	Super::PostEditChangeProperty(PropertyChangedEvent);
+
+	static const FProperty* ShapeLibraryProperty= FindPropertyByName(GetClass(), GET_MEMBER_NAME_CHECKED(UControlRigShapeLibraryLink, ShapeLibrary));
+	if(PropertyChangedEvent.Property == ShapeLibraryProperty ||
+		PropertyChangedEvent.MemberProperty == ShapeLibraryProperty)
+	{
+		SetShapeLibrary(ShapeLibrary);
+	}
+}
+
+#endif
diff --git a/Engine/Plugins/Animation/ControlRig/Source/ControlRig/Private/ControlRigControlActor.cpp b/Engine/Plugins/Animation/ControlRig/Source/ControlRig/Private/ControlRigControlActor.cpp
--- a/Engine/Plugins/Animation/ControlRig/Source/ControlRig/Private/ControlRigControlActor.cpp
+++ b/Engine/Plugins/Animation/ControlRig/Source/ControlRig/Private/ControlRigControlActor.cpp
@@ -206,7 +206,7 @@ void AControlRigControlActor::Refresh()
 				case ERigControlType::TransformNoScale:
 				case ERigControlType::EulerTransform:
 				{
-					if (const FControlRigShapeDefinition* ShapeDef = UControlRigShapeLibrary::GetShapeByName(ControlElement->Settings.ShapeName, ShapeLibraries))
+					if (const FControlRigShapeDefinition* ShapeDef = UControlRigShapeLibrary::GetShapeByName(ControlElement->Settings.ShapeName, ShapeLibraries, ControlRig->ShapeLibraryNameMap))
 					{
 						UMaterialInterface* BaseMaterial = nullptr;
 						if (MaterialOverride && !ColorParameter.IsEmpty())
diff --git a/Engine/Plugins/Animation/ControlRig/Source/ControlRig/Private/ControlRigGizmoActor.cpp b/Engine/Plugins/Animation/ControlRig/Source/ControlRig/Private/ControlRigGizmoActor.cpp
--- a/Engine/Plugins/Animation/ControlRig/Source/ControlRig/Private/ControlRigGizmoActor.cpp
+++ b/Engine/Plugins/Animation/ControlRig/Source/ControlRig/Private/ControlRigGizmoActor.cpp
@@ -137,7 +137,7 @@ bool AControlRigShapeActor::UpdateControlSettings(
 	if(bLookupShape)
 	{
 		const TArray<TSoftObjectPtr<UControlRigShapeLibrary>> ShapeLibraries = InControlRig->GetShapeLibraries();
-		if (const FControlRigShapeDefinition* ShapeDef = UControlRigShapeLibrary::GetShapeByName(ControlSettings.ShapeName, ShapeLibraries))
+		if (const FControlRigShapeDefinition* ShapeDef = UControlRigShapeLibrary::GetShapeByName(ControlSettings.ShapeName, ShapeLibraries, InControlRig->ShapeLibraryNameMap))
 		{
 			MeshTransform = ShapeDef->Transform;
 
diff --git a/Engine/Plugins/Animation/ControlRig/Source/ControlRig/Private/ControlRigGizmoLibrary.cpp b/Engine/Plugins/Animation/ControlRig/Source/ControlRig/Private/ControlRigGizmoLibrary.cpp
--- a/Engine/Plugins/Animation/ControlRig/Source/ControlRig/Private/ControlRigGizmoLibrary.cpp
+++ b/Engine/Plugins/Animation/ControlRig/Source/ControlRig/Private/ControlRigGizmoLibrary.cpp
@@ -110,7 +110,7 @@ const FControlRigShapeDefinition* UControlRigShapeLibrary::GetShapeByName(const
 	return nullptr;
 }
 
-const FControlRigShapeDefinition* UControlRigShapeLibrary::GetShapeByName(const FName& InName, const TArray<TSoftObjectPtr<UControlRigShapeLibrary>>& InShapeLibraries)
+const FControlRigShapeDefinition* UControlRigShapeLibrary::GetShapeByName(const FName& InName, const TArray<TSoftObjectPtr<UControlRigShapeLibrary>>& InShapeLibraries, const TMap<FString, FString>& InLibraryNameMap)
 {
 	const FString InString = InName.ToString();
 	FString Left, Right;
@@ -133,7 +133,13 @@ const FControlRigShapeDefinition* UControlRigShapeLibrary::GetShapeByName(const
 			continue;
 		}
 
-		if(ShapeLibrary->GetName().Equals(Left) || Left.IsEmpty())
+		FString ShapeLibraryName = ShapeLibrary->GetName();
+		if(const FString* RemappedName = InLibraryNameMap.Find(ShapeLibraryName))
+		{
+			ShapeLibraryName = *RemappedName;
+		}
+
+		if(ShapeLibraryName.Equals(Left) || Left.IsEmpty())
 		{
 			if(const FControlRigShapeDefinition* Shape = ShapeLibrary->GetShapeByName(RightName))
 			{
@@ -145,6 +151,18 @@ const FControlRigShapeDefinition* UControlRigShapeLibrary::GetShapeByName(const
 	return nullptr;
 }
 
+const FString UControlRigShapeLibrary::GetShapeName(const UControlRigShapeLibrary* InShapeLibrary, bool bUseNameSpace, const TMap<FString, FString>& InLibraryNameMap, const FControlRigShapeDefinition& InShape)
+{
+	FString LibraryName = InShapeLibrary->GetName();
+	if(const FString* RemappedName = InLibraryNameMap.Find(LibraryName))
+	{
+		LibraryName = *RemappedName;
+	}
+		
+	const FString NameSpace = bUseNameSpace ? LibraryName + TEXT(".") : FString();
+	return NameSpace + InShape.ShapeName.ToString();
+}
+
 const TArray<FName> UControlRigShapeLibrary::GetUpdatedNameList(bool bReset)
 {
 	if (bReset)
diff --git a/Engine/Plugins/Animation/ControlRig/Source/ControlRig/Private/Units/Core/RigUnit_UserData.cpp b/Engine/Plugins/Animation/ControlRig/Source/ControlRig/Private/Units/Core/RigUnit_UserData.cpp
new file mode 100644
--- /dev/null
+++ b/Engine/Plugins/Animation/ControlRig/Source/ControlRig/Private/Units/Core/RigUnit_UserData.cpp
@@ -0,0 +1,230 @@
+// Copyright Epic Games, Inc. All Rights Reserved.
+
+#include "Units/Core/RigUnit_UserData.h"
+#include "RigVMCore/RigVMRegistry.h"
+#include "RigVMCore/RigVMMemoryStorage.h"
+#include "Units/RigUnitContext.h"
+#include "Units/Execution/RigUnit_PrepareForExecution.h"
+
+#include UE_INLINE_GENERATED_CPP_BY_NAME(RigUnit_UserData)
+
+#if WITH_EDITOR
+
+FString FRigDispatch_GetUserData::GetArgumentMetaData(const FName& InArgumentName, const FName& InMetaDataKey) const
+{
+	if(InMetaDataKey == FRigVMStruct::CustomWidgetMetaName)
+	{
+		if(InArgumentName == ArgNameSpaceName)
+		{
+			return TEXT("UserDataNameSpace");
+		}
+		if(InArgumentName == ArgPathName)
+		{
+			return TEXT("UserDataPath");
+		}
+	}
+	return FRigDispatchFactory::GetArgumentMetaData(InArgumentName, InMetaDataKey);
+}
+
+#endif
+
+TArray<FRigVMTemplateArgument> FRigDispatch_GetUserData::GetArguments() const
+{
+	const TArray<FRigVMTemplateArgument::ETypeCategory> ValueCategories = {
+		FRigVMTemplateArgument::ETypeCategory_SingleAnyValue,
+		FRigVMTemplateArgument::ETypeCategory_ArrayAnyValue
+	};
+	return {
+		FRigVMTemplateArgument(ArgNameSpaceName, ERigVMPinDirection::Input, RigVMTypeUtils::TypeIndex::FString),
+		FRigVMTemplateArgument(ArgPathName, ERigVMPinDirection::Input, RigVMTypeUtils::TypeIndex::FString),
+		FRigVMTemplateArgument(ArgDefaultName, ERigVMPinDirection::Input, ValueCategories),
+		FRigVMTemplateArgument(ArgResultName, ERigVMPinDirection::Output, ValueCategories),
+		FRigVMTemplateArgument(ArgFoundName, ERigVMPinDirection::Output, RigVMTypeUtils::TypeIndex::Bool)
+	};
+}
+
+FRigVMTemplateTypeMap FRigDispatch_GetUserData::OnNewArgumentType(const FName& InArgumentName,
+	TRigVMTypeIndex InTypeIndex) const
+{
+	FRigVMTemplateTypeMap Types;
+	Types.Add(ArgNameSpaceName, RigVMTypeUtils::TypeIndex::FString);
+	Types.Add(ArgPathName, RigVMTypeUtils::TypeIndex::FString);
+	Types.Add(ArgDefaultName, InTypeIndex);
+	Types.Add(ArgResultName, InTypeIndex);
+	Types.Add(ArgFoundName, RigVMTypeUtils::TypeIndex::Bool);
+	return Types;
+}
+
+FRigVMFunctionPtr FRigDispatch_GetUserData::GetDispatchFunctionImpl(const FRigVMTemplateTypeMap& InTypes) const
+{
+	check(InTypes.FindChecked(ArgNameSpaceName) == RigVMTypeUtils::TypeIndex::FString);
+	check(InTypes.FindChecked(ArgPathName) == RigVMTypeUtils::TypeIndex::FString);
+	const TRigVMTypeIndex TypeIndex = InTypes.FindChecked(ArgDefaultName);
+	check(FRigVMRegistry::Get().CanMatchTypes(TypeIndex, InTypes.FindChecked(ArgResultName), true));
+	check(InTypes.FindChecked(ArgFoundName) == RigVMTypeUtils::TypeIndex::Bool);
+	return &FRigDispatch_GetUserData::Execute;
+}
+
+void FRigDispatch_GetUserData::Execute(FRigVMExtendedExecuteContext& InContext, FRigVMMemoryHandleArray Handles)
+{
+	const FControlRigExecuteContext& ControlRigContext = InContext.GetPublicData<FControlRigExecuteContext>();
+
+	if(ControlRigContext.GetEventName().IsNone())
+	{
+		// may happen during init
+		return;
+	}
+	
+	// allow to run only in construction event
+	if(ControlRigContext.GetEventName() != FRigUnit_PrepareForExecution::EventName)
+	{
+		static constexpr TCHAR ErrorMessageFormat[] = TEXT("Node can only run in %s Event");
+		ControlRigContext.Logf(EMessageSeverity::Error, ErrorMessageFormat, *FRigUnit_PrepareForExecution::EventName.ToString());
+		return;
+	}
+
+	const FProperty* PropertyDefault = Handles[2].GetResolvedProperty(); 
+	const FProperty* PropertyResult = Handles[3].GetResolvedProperty(); 
+#if WITH_EDITOR
+	check(PropertyDefault);
+	check(PropertyResult);
+	check(Handles[0].IsString());
+	check(Handles[1].IsString());
+	check(Handles[4].IsBool());
+#endif
+
+	const FString& NameSpace = *(FString*)Handles[0].GetData();
+	const FString& Path = *(FString*)Handles[1].GetData();
+	const uint8* Default = Handles[2].GetData();
+	uint8* Result = Handles[3].GetData();
+	bool& bFound = *(bool*)Handles[4].GetData();
+	bFound = false;
+
+	auto Log = [&ControlRigContext](const FString& Message)
+	{
+#if WITH_EDITOR
+		if(ControlRigContext.GetLog())
+		{
+			ControlRigContext.GetLog()->Report(EMessageSeverity::Info, ControlRigContext.GetFunctionName(), ControlRigContext.GetInstructionIndex(), Message);
+		}
+		else
+#endif
+		{
+			ControlRigContext.Log(EMessageSeverity::Info, Message);
+		}
+	};
+
+	
+	if(const UNameSpacedUserData* UserDataObject = ControlRigContext.FindUserData(NameSpace))
+	{
+		FString ErrorMessage;
+		if(const UNameSpacedUserData::FUserData* UserData = UserDataObject->GetUserData(Path, &ErrorMessage))
+		{
+			if(UserData->GetProperty())
+			{
+#if WITH_EDITOR
+				FString ExtendedCPPType;
+				const FString CPPType = PropertyResult->GetCPPType(&ExtendedCPPType);
+				const FRigVMRegistry& Registry = FRigVMRegistry::Get();
+				const TRigVMTypeIndex ResultTypeIndex = Registry.GetTypeIndexFromCPPType(CPPType + ExtendedCPPType);
+				const TRigVMTypeIndex UserDataTypeIndex = Registry.GetTypeIndexFromCPPType(UserData->GetCPPType().ToString());
+
+				if(Registry.CanMatchTypes(ResultTypeIndex, UserDataTypeIndex, true))
+#endif
+				{
+					URigVMMemoryStorage::CopyProperty(PropertyResult, Result, UserData->GetProperty(), UserData->GetMemory());
+					bFound = true;
+					return;
+				}
+
+#if WITH_EDITOR
+				static constexpr TCHAR Format[] = TEXT("User data of type '%s' not compatible with pin of type '%s'.");
+				ControlRigContext.Logf(EMessageSeverity::Info, Format, *UserData->GetCPPType().ToString(), *(CPPType + ExtendedCPPType));
+#endif
+			}
+		}
+#if WITH_EDITOR
+		else
+		{
+			ControlRigContext.Log(EMessageSeverity::Info, ErrorMessage);
+		}
+#endif
+	}
+#if WITH_EDITOR
+	else
+	{
+		static constexpr TCHAR Format[] = TEXT("User data namespace '%s' cannot be found.");
+		ControlRigContext.Logf(EMessageSeverity::Info, Format, *NameSpace);
+	}
+#endif
+
+	// fall back behavior - copy default to result
+	URigVMMemoryStorage::CopyProperty(PropertyResult, Result, PropertyDefault, Default);
+}
+
+FRigUnit_SetupShapeLibraryFromUserData_Execute()
+{
+	// this may happen during init
+	if(ExecuteContext.GetEventName().IsNone())
+	{
+		return;
+	}
+
+	auto Log = [&ExecuteContext](const FString& Message)
+	{
+#if WITH_EDITOR
+		if(ExecuteContext.GetLog())
+		{
+			ExecuteContext.GetLog()->Report(EMessageSeverity::Info, ExecuteContext.GetFunctionName(), ExecuteContext.GetInstructionIndex(), Message);
+		}
+		else
+#endif
+		{
+			ExecuteContext.Log(EMessageSeverity::Info, Message);
+		}
+	};
+	
+	// allow to run only in construction event
+	if(ExecuteContext.GetEventName() != FRigUnit_PrepareForExecution::EventName)
+	{
+		static constexpr TCHAR ErrorMessageFormat[] = TEXT("Node can only run in %s Event");
+		Log(FString::Printf(ErrorMessageFormat, *FRigUnit_PrepareForExecution::EventName.ToString()));
+		return;
+	}
+
+	if(const UNameSpacedUserData* UserDataObject = ExecuteContext.FindUserData(NameSpace))
+	{
+		FString ErrorMessage;
+		if(const UNameSpacedUserData::FUserData* UserData = UserDataObject->GetUserData(Path, &ErrorMessage))
+		{
+			if(const FObjectProperty* ObjectProperty = CastField<FObjectProperty>(UserData->GetProperty()))
+			{
+				if(ObjectProperty->PropertyClass == UControlRigShapeLibrary::StaticClass())
+				{
+					if(ExecuteContext.OnAddShapeLibraryDelegate.IsBound())
+					{
+						ExecuteContext.OnAddShapeLibraryDelegate.Execute(&ExecuteContext, LibraryName, (UControlRigShapeLibrary*)UserData->GetMemory(), ReplaceExisting, LogShapeLibraries);
+					}
+					return;
+				}
+			}
+#if WITH_EDITOR
+			static constexpr TCHAR Format[] = TEXT("User data for path '%s' is of type '%s' - not a shape library.");
+			Log(FString::Printf(Format, *Path, *UserData->GetCPPType().ToString()));
+#endif
+		}
+#if WITH_EDITOR
+		else
+		{
+			Log(ErrorMessage);
+		}
+#endif
+	}
+#if WITH_EDITOR
+	else
+	{
+		static constexpr TCHAR Format[] = TEXT("User data namespace '%s' cannot be found.");
+		Log(FString::Printf(Format, *NameSpace));
+	}
+#endif
+}
\ No newline at end of file
diff --git a/Engine/Plugins/Animation/ControlRig/Source/ControlRig/Private/Units/Core/RigUnit_UserData.h b/Engine/Plugins/Animation/ControlRig/Source/ControlRig/Private/Units/Core/RigUnit_UserData.h
new file mode 100644
--- /dev/null
+++ b/Engine/Plugins/Animation/ControlRig/Source/ControlRig/Private/Units/Core/RigUnit_UserData.h
@@ -0,0 +1,86 @@
+// Copyright Epic Games, Inc. All Rights Reserved.
+
+#pragma once
+
+#include "Units/RigDispatchFactory.h"
+#include "RigUnit_UserData.generated.h"
+
+/*
+ * Retrieves a value from a namespaces user data
+ */
+USTRUCT(meta=(DisplayName = "Get User Data", Keywords = "AssetUserData,Metadata", ExecuteContext="FControlRigExecuteContext", Varying))
+struct FRigDispatch_GetUserData : public FRigDispatchFactory
+{
+	GENERATED_BODY()
+
+public:
+
+#if WITH_EDITOR
+	virtual FString GetArgumentMetaData(const FName& InArgumentName, const FName& InMetaDataKey) const override;
+#endif
+	virtual TArray<FRigVMTemplateArgument> GetArguments() const override;
+	virtual FRigVMTemplateTypeMap OnNewArgumentType(const FName& InArgumentName, TRigVMTypeIndex InTypeIndex) const override;
+
+protected:
+
+	virtual FRigVMFunctionPtr GetDispatchFunctionImpl(const FRigVMTemplateTypeMap& InTypes) const override;
+	static void Execute(FRigVMExtendedExecuteContext& InContext, FRigVMMemoryHandleArray Handles);
+
+	static inline constexpr TCHAR ArgNameSpaceName[] = TEXT("NameSpace");
+	static inline constexpr TCHAR ArgPathName[] = TEXT("Path");
+	static inline constexpr TCHAR ArgDefaultName[] = TEXT("Default");
+	static inline constexpr TCHAR ArgResultName[] = TEXT("Result");
+	static inline constexpr TCHAR ArgFoundName[] = TEXT("Found");
+};
+
+/**
+ * Allows to set / add a shape library on the running control rig from user data
+ */
+USTRUCT(meta=(DisplayName="Set Shape Library from User Data", Keywords="Shape,Gizmo,Library"))
+struct CONTROLRIG_API FRigUnit_SetupShapeLibraryFromUserData : public FRigUnitMutable
+{
+	GENERATED_BODY()
+
+	FRigUnit_SetupShapeLibraryFromUserData()
+	{
+		ReplaceExisting = true;
+		LogShapeLibraries = false;
+	}
+
+	/*
+	 * The name space of the user data to look the shape library up within
+	 */
+	UPROPERTY(meta = (Input, CustomWidget=UserDataNameSpace))
+	FString NameSpace;
+
+	/*
+	 * The path within the user data for the shape library
+	 */
+	UPROPERTY(meta = (Input, CustomWidget=UserDataPath, AllowUObjects))
+	FString Path;
+
+	/*
+	 * Optionally provide the namespace of the shape library to use.
+	 * This is only useful if you have multiple shape libraries and you
+	 * want to override a specific one.
+	 * If this is empty and ReplaceExisting is checked we'll replace the first
+	 * shape library on this rig. 
+	 */
+	UPROPERTY(meta = (Input))
+	FString LibraryName;
+
+	/*
+	 * If true we'll try to replace an existing library.
+	 */
+	UPROPERTY(meta = (Input))
+	bool ReplaceExisting;
+
+	/*
+	 * If this is checked we'll output the resulting shape libraries to the log for debugging.
+	 */
+	UPROPERTY(meta = (Input))
+	bool LogShapeLibraries;
+
+	RIGVM_METHOD()
+	virtual void Execute() override;
+};
\ No newline at end of file
diff --git a/Engine/Plugins/Animation/ControlRig/Source/ControlRig/Public/ControlRig.h b/Engine/Plugins/Animation/ControlRig/Source/ControlRig/Public/ControlRig.h
--- a/Engine/Plugins/Animation/ControlRig/Source/ControlRig/Public/ControlRig.h
+++ b/Engine/Plugins/Animation/ControlRig/Source/ControlRig/Public/ControlRig.h
@@ -350,6 +350,9 @@ protected:
 	UPROPERTY()
 	TArray<TSoftObjectPtr<UControlRigShapeLibrary>> ShapeLibraries;
 
+	UPROPERTY(transient)
+	TMap<FString, FString> ShapeLibraryNameMap;
+
 	/** Runtime object binding */
 	TSharedPtr<IControlRigObjectBinding> ObjectBinding;
 
@@ -457,8 +460,10 @@ private:
 	void HandleHierarchyEvent(URigHierarchy* InHierarchy, const FRigEventContext& InEvent);
 	FRigEventDelegate RigEventDelegate;
 
+	void OnAddShapeLibrary(const FControlRigExecuteContext* InContext, const FString& InLibraryName, UControlRigShapeLibrary* InShapeLibrary, bool bReplaceExisting, bool bLogResults);
 	virtual void InitializeFromCDO() override;
 
+
 	UPROPERTY()
 	FRigInfluenceMapPerEvent Influences;
 
@@ -492,6 +497,13 @@ public:
 
 	float GetDebugBoneRadiusMultiplier() const { return DebugBoneRadiusMultiplier; }
 
+public:
+	//~ Begin IInterface_AssetUserData Interface
+	virtual const TArray<UAssetUserData*>* GetAssetUserDataArray() const override;
+	//~ End IInterface_AssetUserData Interface
+protected:
+	mutable TArray<TObjectPtr<UAssetUserData>> CombinedAssetUserData;
+
 private:
 
 	void CopyPoseFromOtherRig(UControlRig* Subject);
@@ -697,6 +709,9 @@ private:
 	friend struct FAnimNode_ControlRig;
 	friend class URigHierarchy;
 	friend class UFKControlRig;
+	friend class UControlRigGraph;
+	friend class AControlRigControlActor;
+	friend class AControlRigShapeActor;
 };
 
 class CONTROLRIG_API FControlRigBracketScope
diff --git a/Engine/Plugins/Animation/ControlRig/Source/ControlRig/Public/ControlRigAssetUserData.h b/Engine/Plugins/Animation/ControlRig/Source/ControlRig/Public/ControlRigAssetUserData.h
new file mode 100644
--- /dev/null
+++ b/Engine/Plugins/Animation/ControlRig/Source/ControlRig/Public/ControlRigAssetUserData.h
@@ -0,0 +1,250 @@
+// Copyright Epic Games, Inc. All Rights Reserved.
+
+#pragma once
+
+#include "ControlRigGizmoLibrary.h"
+#include "Engine/AssetUserData.h"
+#include "Engine/DataAsset.h"
+#include "ControlRigAssetUserData.generated.h"
+
+/**
+* User data that can be attached to assets to provide namespaced data access.
+*/
+UCLASS(Abstract)
+class CONTROLRIG_API UNameSpacedUserData : public UAssetUserData
+{
+	GENERATED_BODY()
+
+public:
+
+	/** A helper struct to represent a single user data */
+	struct CONTROLRIG_API FUserData
+	{
+	public:
+
+		FUserData()
+		: Path()
+		, CPPType(NAME_None)
+		, Property(nullptr)
+		, Memory(nullptr)
+		{
+		}
+
+		FUserData(const FString& InPath, const FProperty* InProperty, const uint8* InMemory)
+		: Path(InPath)
+		, CPPType(NAME_None)
+		, Property(InProperty)
+		, Memory(InMemory)
+		{
+		}
+
+		FUserData(const FString& InPath, const FName& InCPPType, const FProperty* InProperty, const uint8* InMemory)
+		: Path(InPath)
+		, CPPType(InCPPType)
+		, Property(InProperty)
+		, Memory(InMemory)
+		{
+		}
+
+		// Returns true if this user data is mapped correctly
+		bool IsValid() const
+		{
+			return (!Path.IsEmpty()) && (Memory != nullptr);
+		}
+
+		operator bool() const { return IsValid(); }
+
+		operator FString() const { return Path; }
+
+		// returns the string path inside of the asset user data
+		const FString& GetPath() const { return Path; }
+
+		// returns the name of the last segment of the path
+		FString GetName() const;
+
+#if WITH_EDITOR
+
+		// returns the display name of the last segment of the path
+		FString GetDisplayName() const;
+
+#endif
+
+		// returns the CPP type of this user data (computes it from the property as needed)
+		const FName& GetCPPType() const
+		{
+			if(CPPType.IsNone() && Property != nullptr)
+			{
+				FString ExtendedCppType;
+  				FString CPPTypeString = Property->GetCPPType(&ExtendedCppType);
+  				CPPTypeString += ExtendedCppType;
+				CPPType = *CPPTypeString;
+			}
+			return CPPType;
+		}
+
+		bool IsArray() const;
+
+		bool IsArrayElement() const;
+
+		bool IsUObject() const;
+
+		// returns the property of this user data or nullptr
+		const FProperty* GetProperty() const { return Property; }
+
+		// returns the mapped memory of this user data
+		const uint8* GetMemory() const { return Memory; }
+
+		// joints two user data path strings 
+		static FString Join(const FString& InLeft, const FString& InRight)
+		{
+			check(!InLeft.IsEmpty());
+			check(!InRight.IsEmpty());
+			return InLeft + TEXT(".") + InRight;
+		}
+
+		// splits a user data string path at the first segment
+		static bool SplitAtStart(const FString& InPath, FString& OutLeft, FString& OutRight)
+		{
+			return InPath.Split(TEXT("."), &OutLeft, &OutRight, ESearchCase::IgnoreCase, ESearchDir::FromStart);
+		}
+
+		// splits a user data string path at the last segment
+		static bool SplitAtEnd(const FString& InPath, FString& OutLeft, FString& OutRight)
+		{
+			return InPath.Split(TEXT("."), &OutLeft, &OutRight, ESearchCase::IgnoreCase, ESearchDir::FromEnd);
+		}
+
+		static TArray<UStruct*> GetSuperStructs(UStruct* InStruct);
+
+	private:
+
+		FString Path;
+		mutable FName CPPType;
+		const FProperty* Property;
+		const uint8* Memory;
+	};
+
+	virtual ~UNameSpacedUserData() override
+	{
+		InvalidateCache();
+	}
+
+	virtual void Serialize(FArchive& Ar) override;
+
+	/** The namespace to use when looking up values inside of the user data. */
+	UPROPERTY(EditAnywhere, Category = General)
+	FString NameSpace;
+
+	/** Returns true if userdata exists for a given path */
+	bool ContainsUserData(const FString& InPath) const { return GetUserData(InPath) != nullptr; }
+
+	/** Returns a user data path given its string path */
+	virtual const FUserData* GetUserData(const FString& InPath, FString* OutErrorMessage = nullptr) const;
+
+	/** Retrieves the user data paths given a (optional) parent path */
+	virtual const TArray<const FUserData*>& GetUserDataArray(const FString& InParentPath = FString(), FString* OutErrorMessage = nullptr) const;
+
+	/** Retrieves the user data path given the path string */
+	const FUserData* GetParentUserData(const FString& InPath, FString* OutErrorMessage = nullptr) const
+	{
+		FString Left, Right;
+		if(FUserData::SplitAtEnd(InPath, Left, Right))
+		{
+			return GetUserData(Left, OutErrorMessage);
+		}
+		return nullptr;
+	}
+
+protected:
+
+	const FUserData* GetUserDataWithinStruct(const UStruct* InStruct, const uint8* InMemory, const FString& InPath, const FString& InPropertyName, FString* OutErrorMessage = nullptr) const;
+	const TArray<const FUserData*>& GetUserDataArrayWithinStruct(UStruct* InStruct, const uint8* InMemory, const FString& InPath, FString* OutErrorMessage = nullptr) const;
+	const TArray<const FUserData*>& GetUserDataArrayWithinArray(const FArrayProperty* InArrayProperty, const uint8* InMemory, const FString& InPath, FString* OutErrorMessage = nullptr) const;
+	static const FProperty* FindPropertyByName(const UStruct* InStruct, const FName& InName);
+
+	void InvalidateCache() const;
+	const FUserData* StoreCacheForUserData(const FUserData& InUserData) const; 
+	const TArray<const FUserData*>& StoreCacheForUserDataArray(const FString& InPath, const TArray<const FUserData*>& InUserDataArray) const;
+	virtual bool IsPropertySupported(const FProperty* InProperty, const FString& InPath, bool bCheckPropertyFlags, FString* OutErrorMessage = nullptr) const;
+	
+	static const TArray<const FUserData*> EmptyUserDatas;
+
+	mutable TMap<FString, FUserData*> CachedUserData;
+	mutable TMap<FString, TArray<const FUserData*>> CachedUserDataArray;
+	mutable FCriticalSection CacheLock;
+
+	static inline constexpr TCHAR PathNotFoundFormat[] = TEXT("UserData path '%s' could not be found.");
+	static inline constexpr TCHAR InvalidMemoryFormat[] = TEXT("UserData at path '%s' does not have any data backing it up.");
+	static inline constexpr TCHAR InvalidArrayIndexFormat[] = TEXT("Provided array index '%s' for path '%s' is invalid.");
+	static inline constexpr TCHAR OutOfBoundArrayIndexFormat[] = TEXT("Provided array index for path '%s' is out of bounds (num == %d).");
+	static inline constexpr TCHAR UnSupportedSubPathsFormat[] = TEXT("Cannot access sub paths of user data type '%s' for path '%s'.");
+};
+
+/**
+* Namespaced user data which provides access to a linked data asset
+*/
+UCLASS(BlueprintType)
+class CONTROLRIG_API UDataAssetLink : public UNameSpacedUserData
+{
+	GENERATED_BODY()
+
+public:
+
+	/** If assigned, the data asset link will provide access to the data asset's content. */
+	UPROPERTY(EditAnywhere, BlueprintReadWrite, Category = General, BlueprintGetter = GetDataAsset, BlueprintSetter = SetDataAsset, Meta = (DisplayAfter="NameSpace"))
+	TObjectPtr<UDataAsset> DataAsset;
+
+	UFUNCTION(BlueprintGetter)
+	UDataAsset* GetDataAsset() const { return DataAsset; }
+
+	UFUNCTION(BlueprintSetter)
+	void SetDataAsset(UDataAsset* InDataAsset);
+
+	virtual const FUserData* GetUserData(const FString& InPath, FString* OutErrorMessage = nullptr) const override;
+	virtual const TArray<const FUserData*>& GetUserDataArray(const FString& InParentPath = FString(), FString* OutErrorMessage = nullptr) const override;
+
+#if WITH_EDITOR
+	virtual void PostEditChangeProperty(FPropertyChangedEvent& PropertyChangedEvent) override;
+#endif
+
+protected:
+
+	static inline constexpr TCHAR DataAssetNullFormat[] = TEXT("User data path '%s' could not be found (DataAsset not provided)");
+};
+
+/**
+* Namespaced user data which provides access to a linked shape library
+*/
+UCLASS(BlueprintType)
+class CONTROLRIG_API UControlRigShapeLibraryLink : public UNameSpacedUserData
+{
+	GENERATED_BODY()
+
+public:
+
+	/** If assigned, the data asset link will provide access to the data asset's content. */
+	UPROPERTY(EditAnywhere, BlueprintReadWrite, Category = General, BlueprintGetter = GetShapeLibrary, BlueprintSetter = SetShapeLibrary, Meta = (DisplayAfter="NameSpace"))
+	TObjectPtr<UControlRigShapeLibrary> ShapeLibrary;
+
+	UFUNCTION(BlueprintGetter)
+	UControlRigShapeLibrary* GetShapeLibrary() const { return ShapeLibrary; }
+
+	UFUNCTION(BlueprintSetter)
+	void SetShapeLibrary(UControlRigShapeLibrary* InShapeLibrary);
+
+	virtual const FUserData* GetUserData(const FString& InPath, FString* OutErrorMessage = nullptr) const override;
+	virtual const TArray<const FUserData*>& GetUserDataArray(const FString& InParentPath = FString(), FString* OutErrorMessage = nullptr) const override;
+
+#if WITH_EDITOR
+	virtual void PostEditChangeProperty(FPropertyChangedEvent& PropertyChangedEvent) override;
+#endif
+
+protected:
+
+	UPROPERTY(transient)
+	TArray<FName> ShapeNames;
+
+	static inline constexpr TCHAR DefaultShapePath[] = TEXT("DefaultShape");
+	static inline constexpr TCHAR ShapeNamesPath[] = TEXT("ShapeNames");
+	static inline constexpr TCHAR ShapeLibraryNullFormat[] = TEXT("User data path '%s' could not be found (ShapeLibrary not provided)");
+};
\ No newline at end of file
diff --git a/Engine/Plugins/Animation/ControlRig/Source/ControlRig/Public/ControlRigGizmoLibrary.h b/Engine/Plugins/Animation/ControlRig/Source/ControlRig/Public/ControlRigGizmoLibrary.h
--- a/Engine/Plugins/Animation/ControlRig/Source/ControlRig/Public/ControlRigGizmoLibrary.h
+++ b/Engine/Plugins/Animation/ControlRig/Source/ControlRig/Public/ControlRigGizmoLibrary.h
@@ -61,7 +61,8 @@ public:
 	TArray<FControlRigShapeDefinition> Shapes;
 
 	const FControlRigShapeDefinition* GetShapeByName(const FName& InName, bool bUseDefaultIfNotFound = false) const;
-	static const FControlRigShapeDefinition* GetShapeByName(const FName& InName, const TArray<TSoftObjectPtr<UControlRigShapeLibrary>>& InShapeLibraries);
+	static const FControlRigShapeDefinition* GetShapeByName(const FName& InName, const TArray<TSoftObjectPtr<UControlRigShapeLibrary>>& InShapeLibraries, const TMap<FString, FString>& InLibraryNameMap);
+	static const FString GetShapeName(const UControlRigShapeLibrary* InShapeLibrary, bool bUseNameSpace, const TMap<FString, FString>& InLibraryNameMap, const FControlRigShapeDefinition& InShape);
 
 #if WITH_EDITOR
 
diff --git a/Engine/Plugins/Animation/ControlRig/Source/ControlRig/Public/Units/RigUnitContext.h b/Engine/Plugins/Animation/ControlRig/Source/ControlRig/Public/Units/RigUnitContext.h
--- a/Engine/Plugins/Animation/ControlRig/Source/ControlRig/Public/Units/RigUnitContext.h
+++ b/Engine/Plugins/Animation/ControlRig/Source/ControlRig/Public/Units/RigUnitContext.h
@@ -7,9 +7,11 @@
 #include "Rigs/RigCurveContainer.h"
 #include "AnimationDataSource.h"
 #include "Animation/AttributesRuntime.h"
+#include "ControlRigAssetUserData.h"
 #include "RigVMCore/RigVMExecuteContext.h"
 #include "RigUnitContext.generated.h"
 
+class UControlRigShapeLibrary;
 /**
  * The type of interaction happening on a rig
  */
@@ -109,6 +111,35 @@ struct FControlRigExecuteContext : public FRigVMExecuteContext
 		Hierarchy = OtherContext->Hierarchy;
 	}
 
+	/**
+     * Finds a name spaced user data object
+     */
+	const UNameSpacedUserData* FindUserData(const FString& InNameSpace) const
+	{
+		// walk in reverse since asset user data at the end with the same
+		// namespace overrides previously listed user data
+		for(int32 Index = AssetUserData.Num() - 1; AssetUserData.IsValidIndex(Index); Index--)
+		{
+			if(!IsValid(AssetUserData[Index]))
+			{
+				continue;
+			}
+			if(const UNameSpacedUserData* NameSpacedUserData = Cast<UNameSpacedUserData>(AssetUserData[Index]))
+			{
+				if(NameSpacedUserData->NameSpace.Equals(InNameSpace, ESearchCase::CaseSensitive))
+				{
+					return NameSpacedUserData;
+				}
+			}
+		}
+		return nullptr;
+	}
+
+	/** The list of available asset user data object */
+	TArray<const UAssetUserData*> AssetUserData;
+
+	DECLARE_DELEGATE_FiveParams(FOnAddShapeLibrary, const FControlRigExecuteContext* InContext, const FString&, UControlRigShapeLibrary*, bool /* replace? */, bool /* log results */);
+	FOnAddShapeLibrary OnAddShapeLibraryDelegate; 
 	FRigUnitContext UnitContext;
 	URigHierarchy* Hierarchy;
 };
diff --git a/Engine/Plugins/Animation/ControlRig/Source/ControlRigDeveloper/Private/ControlRigBlueprint.cpp b/Engine/Plugins/Animation/ControlRig/Source/ControlRigDeveloper/Private/ControlRigBlueprint.cpp
--- a/Engine/Plugins/Animation/ControlRig/Source/ControlRigDeveloper/Private/ControlRigBlueprint.cpp
+++ b/Engine/Plugins/Animation/ControlRig/Source/ControlRigDeveloper/Private/ControlRigBlueprint.cpp
@@ -2895,7 +2895,12 @@ bool UControlRigBlueprint::ChangeMemberVariableType(const FName& InName, const F
 
 const FControlRigShapeDefinition* UControlRigBlueprint::GetControlShapeByName(const FName& InName) const
 {
-	return UControlRigShapeLibrary::GetShapeByName(InName, ShapeLibraries);
+	TMap<FString, FString> LibraryNameMap;
+	if(UControlRig* ControlRig = Cast<UControlRig>(GetObjectBeingDebugged()))
+	{
+		LibraryNameMap = ControlRig->ShapeLibraryNameMap;
+	}
+	return UControlRigShapeLibrary::GetShapeByName(InName, ShapeLibraries, LibraryNameMap);
 }
 
 FName UControlRigBlueprint::AddTransientControl(URigVMPin* InPin)
diff --git a/Engine/Plugins/Animation/ControlRig/Source/ControlRigDeveloper/Private/Graph/ControlRigGraph.cpp b/Engine/Plugins/Animation/ControlRig/Source/ControlRigDeveloper/Private/Graph/ControlRigGraph.cpp
--- a/Engine/Plugins/Animation/ControlRig/Source/ControlRigDeveloper/Private/Graph/ControlRigGraph.cpp
+++ b/Engine/Plugins/Animation/ControlRig/Source/ControlRigDeveloper/Private/Graph/ControlRigGraph.cpp
@@ -206,6 +206,12 @@ void UControlRigGraph::CacheNameLists(URigHierarchy* InHierarchy, const FRigVMDr
 	ShapeNameList.Reset();
 	ShapeNameList.Add(MakeShared<FString>(FName(NAME_None).ToString()));
 
+	TMap<FString, FString> LibraryNameMap;
+	if(UControlRig* ControlRig = InHierarchy->GetTypedOuter<UControlRig>())
+	{
+		LibraryNameMap = ControlRig->ShapeLibraryNameMap;
+	}
+
 	for(const TSoftObjectPtr<UControlRigShapeLibrary>& ShapeLibrary : ShapeLibraries)
 	{
 		if(ShapeLibrary.IsNull() || !ShapeLibrary.IsValid())
@@ -219,11 +225,17 @@ void UControlRigGraph::CacheNameLists(URigHierarchy* InHierarchy, const FRigVMDr
 		}
 
 		const bool bUseNameSpace = ShapeLibraries.Num() > 1;
-		const FString NameSpace = bUseNameSpace ? ShapeLibrary->GetName() + TEXT(".") : FString();
-		ShapeNameList.Add(MakeShared<FString>(NameSpace + ShapeLibrary->DefaultShape.ShapeName.ToString()));
+		FString LibraryName = ShapeLibrary->GetName();
+		if(const FString* RemappedName = LibraryNameMap.Find(LibraryName))
+		{
+			LibraryName = *RemappedName;
+		}
+		
+		const FString NameSpace = bUseNameSpace ? LibraryName + TEXT(".") : FString();
+		ShapeNameList.Add(MakeShared<FString>(UControlRigShapeLibrary::GetShapeName(ShapeLibrary.Get(), bUseNameSpace, LibraryNameMap, ShapeLibrary->DefaultShape)));
 		for (const FControlRigShapeDefinition& Shape : ShapeLibrary->Shapes)
 		{
-			ShapeNameList.Add(MakeShared<FString>(NameSpace + Shape.ShapeName.ToString()));
+			ShapeNameList.Add(MakeShared<FString>(UControlRigShapeLibrary::GetShapeName(ShapeLibrary.Get(), bUseNameSpace, LibraryNameMap, Shape)));
 		}
 	}
 }
diff --git a/Engine/Plugins/Animation/ControlRig/Source/ControlRigEditor/Private/EditMode/ControlRigEditMode.cpp b/Engine/Plugins/Animation/ControlRig/Source/ControlRigEditor/Private/EditMode/ControlRigEditMode.cpp
--- a/Engine/Plugins/Animation/ControlRig/Source/ControlRigEditor/Private/EditMode/ControlRigEditMode.cpp
+++ b/Engine/Plugins/Animation/ControlRig/Source/ControlRigEditor/Private/EditMode/ControlRigEditMode.cpp
@@ -3281,7 +3281,7 @@ void FControlRigEditMode::CreateShapeActors(UControlRig* ControlRig)
 			Param.ShapeTransform = ControlRig->GetHierarchy()->GetControlShapeTransform(ControlElement, ERigTransformType::CurrentLocal);
 			Param.bSelectable = ControlElement->Settings.IsSelectable(false);
 
-			if (const FControlRigShapeDefinition* ShapeDef = UControlRigShapeLibrary::GetShapeByName(ControlElement->Settings.ShapeName, ShapeLibraries))
+			if (const FControlRigShapeDefinition* ShapeDef = UControlRigShapeLibrary::GetShapeByName(ControlElement->Settings.ShapeName, ShapeLibraries, ControlRig->ShapeLibraryNameMap))
 			{
 				Param.MeshTransform = ShapeDef->Transform;
 				Param.StaticMesh = ShapeDef->StaticMesh;
@@ -4092,7 +4092,7 @@ void FControlRigEditMode::ChangeControlShapeTransform(AControlRigShapeActor* Sha
 
 				FTransform MeshTransform = FTransform::Identity;
 				FTransform ShapeTransform = CurrentLocalShapeTransform;
-				if (const FControlRigShapeDefinition* Gizmo = UControlRigShapeLibrary::GetShapeByName(ControlElement->Settings.ShapeName, ControlRig->GetShapeLibraries()))
+				if (const FControlRigShapeDefinition* Gizmo = UControlRigShapeLibrary::GetShapeByName(ControlElement->Settings.ShapeName, ControlRig->GetShapeLibraries(), ControlRig->ShapeLibraryNameMap))
 				{
 					MeshTransform = Gizmo->Transform;
 				}
diff --git a/Engine/Plugins/Animation/ControlRig/Source/ControlRigEditor/Private/Editor/ControlRigEditor.cpp b/Engine/Plugins/Animation/ControlRig/Source/ControlRigEditor/Private/Editor/ControlRigEditor.cpp
--- a/Engine/Plugins/Animation/ControlRig/Source/ControlRigEditor/Private/Editor/ControlRigEditor.cpp
+++ b/Engine/Plugins/Animation/ControlRig/Source/ControlRigEditor/Private/Editor/ControlRigEditor.cpp
@@ -4483,7 +4483,13 @@ void FControlRigEditor::CacheNameLists()
 			{
 				continue;
 			}
-			RigGraph->CacheNameLists(Hierarchy, &ControlRigBP->DrawContainer, ControlRigBP->ShapeLibraries);
+
+			const TArray<TSoftObjectPtr<UControlRigShapeLibrary>>* ShapeLibraries = &ControlRigBP->ShapeLibraries;
+			if(const UControlRig* DebuggedControlRig = Hierarchy->GetTypedOuter<UControlRig>())
+			{
+				ShapeLibraries = &DebuggedControlRig->GetShapeLibraries();
+			}
+			RigGraph->CacheNameLists(Hierarchy, &ControlRigBP->DrawContainer, *ShapeLibraries);
 		}
 	}
 }
diff --git a/Engine/Plugins/Animation/ControlRig/Source/ControlRigEditor/Private/Graph/ControlRigGraphPanelPinFactory.cpp b/Engine/Plugins/Animation/ControlRig/Source/ControlRigEditor/Private/Graph/ControlRigGraphPanelPinFactory.cpp
--- a/Engine/Plugins/Animation/ControlRig/Source/ControlRigEditor/Private/Graph/ControlRigGraphPanelPinFactory.cpp
+++ b/Engine/Plugins/Animation/ControlRig/Source/ControlRigEditor/Private/Graph/ControlRigGraphPanelPinFactory.cpp
@@ -8,6 +8,8 @@
 #include "Graph/SControlRigGraphPinCurveFloat.h"
 #include "Graph/SControlRigGraphPinVariableName.h"
 #include "Graph/SControlRigGraphPinVariableBinding.h"
+#include "Graph/SGraphPinUserDataNameSpace.h"
+#include "Graph/SGraphPinUserDataPath.h"
 #include "KismetPins/SGraphPinExec.h"
 #include "SGraphPinComboBox.h"
 #include "ControlRig.h"
@@ -350,6 +352,15 @@ TSharedPtr<SGraphPin> FControlRigGraphPanelPinFactory::CreatePin(UEdGraphPin* In
 							return &EmptyNameList;
 						});
 				}
+				else if (CustomWidgetName == TEXT("UserDataNameSpace"))
+				{
+					return SNew(SGraphPinUserDataNameSpace, InPin);
+				}
+				else if (CustomWidgetName == TEXT("UserDataPath"))
+				{
+					return SNew(SGraphPinUserDataPath, InPin)
+						.ModelPins({ModelPin});
+				}
 			}
 
 			if (InPin->PinType.PinCategory == UEdGraphSchema_K2::PC_Struct)
diff --git a/Engine/Plugins/Animation/ControlRig/Source/ControlRigEditor/Private/Graph/SGraphPinUserDataNameSpace.cpp b/Engine/Plugins/Animation/ControlRig/Source/ControlRigEditor/Private/Graph/SGraphPinUserDataNameSpace.cpp
new file mode 100644
--- /dev/null
+++ b/Engine/Plugins/Animation/ControlRig/Source/ControlRigEditor/Private/Graph/SGraphPinUserDataNameSpace.cpp
@@ -0,0 +1,122 @@
+// Copyright Epic Games, Inc. All Rights Reserved.
+
+
+#include "Graph/SGraphPinUserDataNameSpace.h"
+#include "Graph/ControlRigGraph.h"
+#include "ControlRigAssetUserData.h"
+#include "Widgets/Layout/SBox.h"
+#include "ScopedTransaction.h"
+
+void SGraphPinUserDataNameSpace::Construct(const FArguments& InArgs, UEdGraphPin* InGraphPinObj)
+{
+	SGraphPin::Construct(SGraphPin::FArguments(), InGraphPinObj);
+}
+
+TSharedRef<SWidget>	SGraphPinUserDataNameSpace::GetDefaultValueWidget()
+{
+	TSharedPtr<FString> InitialSelected;
+	TArray<TSharedPtr<FString>>& LocalNameSpaces = GetNameSpaces();
+	for (TSharedPtr<FString> Item : LocalNameSpaces)
+	{
+		if (Item->Equals(GetNameSpaceText().ToString()))
+		{
+			InitialSelected = Item;
+		}
+	}
+
+	return SNew(SBox)
+		.MinDesiredWidth(150)
+		.MaxDesiredWidth(400)
+		[
+			SAssignNew(NameComboBox, SControlRigGraphPinEditableNameValueWidget)
+				.Visibility(this, &SGraphPin::GetDefaultValueVisibility)
+				.OptionsSource(&NameSpaces)
+				.OnGenerateWidget(this, &SGraphPinUserDataNameSpace::MakeNameSpaceItemWidget)
+				.OnSelectionChanged(this, &SGraphPinUserDataNameSpace::OnNameSpaceChanged)
+				.OnComboBoxOpening(this, &SGraphPinUserDataNameSpace::OnNameSpaceComboBox)
+				.InitiallySelectedItem(InitialSelected)
+				.Content()
+				[
+					SNew(STextBlock)
+					.Text(this, &SGraphPinUserDataNameSpace::GetNameSpaceText)
+				]
+		];
+}
+
+FText SGraphPinUserDataNameSpace::GetNameSpaceText() const
+{
+	return FText::FromString( GraphPinObj->GetDefaultAsString() );
+}
+
+void SGraphPinUserDataNameSpace::SetNameSpaceText(const FText& NewTypeInValue, ETextCommit::Type /*CommitInfo*/)
+{
+	if(!GraphPinObj->GetDefaultAsString().Equals(NewTypeInValue.ToString()))
+	{
+		const FScopedTransaction Transaction( NSLOCTEXT("GraphEditor", "ChangeNameSpacePinValue", "Change Bone Name Pin Value" ) );
+		GraphPinObj->Modify();
+		GraphPinObj->GetSchema()->TrySetDefaultValue(*GraphPinObj, NewTypeInValue.ToString());
+	}
+}
+
+TSharedRef<SWidget> SGraphPinUserDataNameSpace::MakeNameSpaceItemWidget(TSharedPtr<FString> InItem)
+{
+	return 	SNew(STextBlock).Text(FText::FromString(*InItem));// .Font(FAppStyle::GetFontStyle(TEXT("PropertyWindow.NormalFont")));
+}
+
+void SGraphPinUserDataNameSpace::OnNameSpaceChanged(TSharedPtr<FString> NewSelection, ESelectInfo::Type SelectInfo)
+{
+	if (SelectInfo != ESelectInfo::Direct)
+	{
+		FString NewValue = *NewSelection.Get();
+		SetNameSpaceText(FText::FromString(NewValue), ETextCommit::OnEnter);
+	}
+}
+
+void SGraphPinUserDataNameSpace::OnNameSpaceComboBox()
+{
+	TSharedPtr<FString> CurrentlySelected;
+	TArray<TSharedPtr<FString>>& LocalNameSpaces = GetNameSpaces();
+	for (TSharedPtr<FString> Item : LocalNameSpaces)
+	{
+		if (Item->Equals(GetNameSpaceText().ToString()))
+		{
+			CurrentlySelected = Item;
+		}
+	}
+
+	NameComboBox->SetSelectedItem(CurrentlySelected);
+}
+
+TArray<TSharedPtr<FString>>& SGraphPinUserDataNameSpace::GetNameSpaces()
+{
+	NameSpaces.Reset();
+
+	TArray<FString> NameSpaceLookup;
+	if(const UBlueprint* Blueprint = GraphPinObj->GetOwningNode()->GetTypedOuter<UBlueprint>())
+	{
+		if(const UObject* DebuggedObject = Blueprint->GetObjectBeingDebugged())
+		{
+			if(DebuggedObject->Implements<UInterface_AssetUserData>())
+			{
+				const IInterface_AssetUserData* AssetUserDataHost = CastChecked<IInterface_AssetUserData>(DebuggedObject);
+				if(const TArray<UAssetUserData*>* UserDataArray = AssetUserDataHost->GetAssetUserDataArray())
+				{
+					for(const UAssetUserData* UserData : *UserDataArray)
+					{
+						if(const UNameSpacedUserData* NameSpacedUserData = Cast<UNameSpacedUserData>(UserData))
+						{
+							const FString& NameSpace = NameSpacedUserData->NameSpace; 
+							if(!NameSpaceLookup.Contains(NameSpace))
+							{
+								NameSpaceLookup.Add(NameSpace);
+								NameSpaces.Add(MakeShared<FString>(NameSpace));
+							}
+						}
+					}
+				}
+			}
+		}
+	}
+
+	return NameSpaces;
+}
diff --git a/Engine/Plugins/Animation/ControlRig/Source/ControlRigEditor/Private/Graph/SGraphPinUserDataNameSpace.h b/Engine/Plugins/Animation/ControlRig/Source/ControlRigEditor/Private/Graph/SGraphPinUserDataNameSpace.h
new file mode 100644
--- /dev/null
+++ b/Engine/Plugins/Animation/ControlRig/Source/ControlRigEditor/Private/Graph/SGraphPinUserDataNameSpace.h
@@ -0,0 +1,35 @@
+// Copyright Epic Games, Inc. All Rights Reserved.
+
+#pragma once
+
+#include "CoreMinimal.h"
+#include "Widgets/DeclarativeSyntaxSupport.h"
+#include "Widgets/SWidget.h"
+#include "SGraphPin.h"
+#include "Graph/SControlRigGraphPinEditableNameValueWidget.h"
+
+class SGraphPinUserDataNameSpace : public SGraphPin
+{
+public:
+	SLATE_BEGIN_ARGS(SGraphPinUserDataNameSpace) {}
+	SLATE_END_ARGS()
+
+	void Construct(const FArguments& InArgs, UEdGraphPin* InGraphPinObj);
+
+protected:
+	//~ Begin SGraphPin Interface
+	virtual TSharedRef<SWidget>	GetDefaultValueWidget() override;
+	//~ End SGraphPin Interface
+
+	FText GetNameSpaceText() const;
+	virtual void SetNameSpaceText(const FText& NewTypeInValue, ETextCommit::Type CommitInfo);
+
+	TSharedRef<SWidget> MakeNameSpaceItemWidget(TSharedPtr<FString> InItem);
+	void OnNameSpaceChanged(TSharedPtr<FString> NewSelection, ESelectInfo::Type SelectInfo);
+	void OnNameSpaceComboBox();
+	TArray<TSharedPtr<FString>>& GetNameSpaces();
+
+	TArray<TSharedPtr<FString>> NameSpaces;
+	TSharedPtr<SControlRigGraphPinEditableNameValueWidget> NameComboBox;
+
+};
diff --git a/Engine/Plugins/Animation/ControlRig/Source/ControlRigEditor/Private/Graph/SGraphPinUserDataPath.cpp b/Engine/Plugins/Animation/ControlRig/Source/ControlRigEditor/Private/Graph/SGraphPinUserDataPath.cpp
new file mode 100644
--- /dev/null
+++ b/Engine/Plugins/Animation/ControlRig/Source/ControlRigEditor/Private/Graph/SGraphPinUserDataPath.cpp
@@ -0,0 +1,430 @@
+// Copyright Epic Games, Inc. All Rights Reserved.
+
+
+#include "Graph/SGraphPinUserDataPath.h"
+#include "Features/IModularFeatures.h"
+#include "Widgets/Layout/SBox.h"
+#include "Widgets/Layout/SSpacer.h"
+#include "Widgets/Images/SImage.h"
+#include "Framework/MultiBox/MultiBoxExtender.h"
+#include "Framework/MultiBox/MultiBoxBuilder.h"
+#include "EdGraphSchema_K2.h"
+#include "DetailLayoutBuilder.h"
+#include "RigVMModel/RigVMController.h"
+#include "ControlRigBlueprintGeneratedClass.h"
+#include "Graph/ControlRigGraphSchema.h"
+#include "Kismet2/BlueprintEditorUtils.h"
+#include "EdGraphSchema_K2.h"
+
+#define LOCTEXT_NAMESPACE "SGraphPinUserDataPath"
+
+static const FText GraphPinUserDataPathMultipleValues = LOCTEXT("MultipleValues", "Multiple Values");
+
+void SRigVMUserDataPath::Construct(const FArguments& InArgs)
+{
+	ModelPins = InArgs._ModelPins;
+	if(ModelPins.IsEmpty())
+	{
+		return;
+	}
+
+	bAllowUObjects = false;
+	for(const URigVMPin* ModelPin : ModelPins)
+	{
+		if(const URigVMUnitNode* UnitNode = Cast<URigVMUnitNode>(ModelPin->GetNode()))
+		{
+			if(const UScriptStruct* ScriptStruct = UnitNode->GetScriptStruct())
+			{
+				if(const FProperty* Property = ScriptStruct->FindPropertyByName(ModelPin->GetFName()))
+				{
+					bAllowUObjects = Property->HasMetaData(TEXT("AllowUObjects"));
+				}
+			}
+		}
+	}
+
+	MenuAnchor = SNew(SMenuAnchor)
+	.OnGetMenuContent(this, &SRigVMUserDataPath::GetTopLevelMenuContent)
+	[
+		SNew(SButton)
+		.ButtonStyle(&FAppStyle::Get().GetWidgetStyle<FComboButtonStyle>("ComboButton").ButtonStyle)
+		.OnClicked_Lambda([this]() 
+		{
+			MenuAnchor->SetIsOpen(true); 
+			return FReply::Handled(); 
+		})
+		[
+			SNew(SHorizontalBox)
+			+SHorizontalBox::Slot()
+			.AutoWidth()
+			.VAlign(VAlign_Center)
+			.Padding(1.0f, 0.0f)
+			[
+				SNew(SImage)
+				.Image(this, &SRigVMUserDataPath::GetUserDataIcon)
+				.ColorAndOpacity_Lambda([this]() -> FSlateColor
+				{
+					return GetUserDataColor();
+				})
+			]
+			+SHorizontalBox::Slot()
+			.AutoWidth()
+			.VAlign(VAlign_Center)
+			.Padding(4.0f, 0.0f)
+			[
+				SNew(STextBlock)
+				.Text(this, &SRigVMUserDataPath::GetUserDataPathText)
+			]
+
+			+ SHorizontalBox::Slot()
+			.Padding(2, 0)
+			.HAlign(HAlign_Right)
+			.VAlign(VAlign_Center)
+			.AutoWidth()
+			[
+				SNew(SImage)
+				.Image(FAppStyle::Get().GetBrush("Icons.ChevronDown"))
+			]
+		]
+	];
+
+	this->ChildSlot
+	[
+		MenuAnchor.ToSharedRef()
+	];
+}
+
+FString SRigVMUserDataPath::GetUserDataPath(URigVMPin* ModelPin) const
+{
+	if (ModelPin)
+	{
+		return ModelPin->GetDefaultValue();
+	}
+	return FString();
+}
+
+FString SRigVMUserDataPath::GetUserDataPath() const
+{
+	if (ModelPins.Num() > 0)
+	{
+		const FString FirstPath = GetUserDataPath(ModelPins[0]);
+		for(int32 Index = 1; Index < ModelPins.Num(); Index++)
+		{
+			if(!GetUserDataPath(ModelPins[Index]).Equals(FirstPath, ESearchCase::CaseSensitive))
+			{
+				return FString();
+			}
+		}
+		return FirstPath;
+	}
+	return FString();
+}
+
+const FSlateBrush* SRigVMUserDataPath::GetUserDataIcon() const
+{
+	if(const UNameSpacedUserData::FUserData* UserData = GetUserData())
+	{
+		return GetUserDataIcon(UserData);
+	}
+	return nullptr;
+}
+
+const FSlateBrush* SRigVMUserDataPath::GetUserDataIcon(const UNameSpacedUserData::FUserData* InUserData) const
+{
+	static FName TypeIcon(TEXT("Kismet.VariableList.TypeIcon"));
+	static FName ArrayTypeIcon(TEXT("Kismet.VariableList.ArrayTypeIcon"));
+
+	if (InUserData)
+	{
+		if(InUserData->IsArray())
+		{
+			return FAppStyle::GetBrush(ArrayTypeIcon);
+		}
+	}
+	return FAppStyle::GetBrush(TypeIcon);
+}
+
+FLinearColor SRigVMUserDataPath::GetUserDataColor() const
+{
+	if(const UNameSpacedUserData::FUserData* UserData = GetUserData())
+	{
+		return GetUserDataColor(UserData);
+	}
+	return FLinearColor::White;
+}
+
+FLinearColor SRigVMUserDataPath::GetUserDataColor(const UNameSpacedUserData::FUserData* InUserData) const
+{
+	if (InUserData)
+	{
+		const FRigVMExternalVariable ExternalVariable = FRigVMExternalVariable::Make(InUserData->GetProperty(), nullptr);
+		if(ExternalVariable.IsValid(true))
+		{
+			const UControlRigGraphSchema* Schema = GetDefault<UControlRigGraphSchema>();
+			const FEdGraphPinType PinType = RigVMTypeUtils::PinTypeFromExternalVariable(ExternalVariable);
+			return Schema->GetPinTypeColor(PinType);
+		}
+		if(bAllowUObjects && InUserData->IsUObject())
+		{
+			return FLinearColor::Blue;
+		}
+	}
+	return FLinearColor::White;
+}
+
+TSharedRef<SWidget> SRigVMUserDataPath::GetTopLevelMenuContent()
+{
+	FMenuBuilder MenuBuilder(true, nullptr, nullptr, true);
+	FillUserDataPathMenu(MenuBuilder, FString());
+	return MenuBuilder.MakeWidget();
+}
+
+void SRigVMUserDataPath::FillUserDataPathMenu(FMenuBuilder& InMenuBuilder, FString InParentPath)
+{
+	const UNameSpacedUserData* UserDataObject = GetUserDataObject();
+	if(UserDataObject == nullptr)
+	{
+		return;
+	}
+
+	if(InParentPath.IsEmpty())
+	{
+		InMenuBuilder.BeginSection("UserData", LOCTEXT("UserData", "User Data"));
+	}
+	{
+		static FName PropertyIcon(TEXT("Kismet.VariableList.TypeIcon"));
+		const UControlRigGraphSchema* Schema = GetDefault<UControlRigGraphSchema>();
+
+		const TArray<const UNameSpacedUserData::FUserData*> UserDataArray = UserDataObject->GetUserDataArray(InParentPath);
+
+		for(const UNameSpacedUserData::FUserData* UserData : UserDataArray)
+		{
+			const FRigVMExternalVariable ExternalVariable = FRigVMExternalVariable::Make(UserData->GetProperty(),nullptr);
+
+			FEdGraphPinType PinType;
+
+			bool bIsCompatible = ExternalVariable.IsValid(true);
+			if(bIsCompatible)
+			{
+				PinType = RigVMTypeUtils::PinTypeFromExternalVariable(ExternalVariable);
+			}
+			if(!bIsCompatible)
+			{
+				if(bAllowUObjects && UserData->IsUObject())
+				{
+					UClass* PropertyClass = CastField<FObjectProperty>(UserData->GetProperty())->PropertyClass;
+					PinType.PinCategory = UEdGraphSchema_K2::PC_Object;
+					PinType.PinSubCategoryObject = TWeakObjectPtr<UClass>(PropertyClass);
+					bIsCompatible = true;
+				}
+			}
+
+			if(!bIsCompatible)
+			{
+				continue;
+			}
+
+			const FUIAction Action(FExecuteAction::CreateSP(this, &SRigVMUserDataPath::HandleSetUserDataPath, UserData->GetPath()));
+
+			TSharedRef<SWidget> Content = 
+				SNew(SHorizontalBox)
+					+SHorizontalBox::Slot()
+					.AutoWidth()
+					[
+						SNew(SSpacer)
+						.Size(FVector2D(2.0f, 0.0f))
+					]
+					+SHorizontalBox::Slot()
+					.AutoWidth()
+					.VAlign(VAlign_Center)
+					.Padding(1.0f, 0.0f)
+					[
+						SNew(SImage)
+						.Image(GetUserDataIcon(UserData))
+						.ColorAndOpacity(GetUserDataColor(UserData))
+					]
+					+SHorizontalBox::Slot()
+					.AutoWidth()
+					.VAlign(VAlign_Center)
+					.Padding(4.0f, 0.0f)
+					[
+						SNew(STextBlock)
+						.Text(FText::FromString(UserData->GetDisplayName()))
+					];
+
+			const FString& UserDataPath = UserData->GetPath();
+			if(UserDataObject->GetUserDataArray(UserDataPath).Num() > 0)
+			{
+				InMenuBuilder.AddSubMenu(Action, Content, FNewMenuDelegate::CreateSP(this, &SRigVMUserDataPath::FillUserDataPathMenu, UserDataPath));
+			}
+			else
+			{
+				InMenuBuilder.AddMenuEntry(Action, Content);
+			}
+		}
+	}
+	if(InParentPath.IsEmpty())
+	{
+		InMenuBuilder.EndSection();
+	}
+}
+
+void SRigVMUserDataPath::HandleSetUserDataPath(FString InUserDataPath)
+{
+	UControlRigBlueprint* Blueprint = GetBlueprint();
+	if(ModelPins.IsEmpty() || (Blueprint == nullptr))
+	{
+		return;
+	}
+
+	for(const URigVMPin* ModelPin : ModelPins)
+	{
+		URigVMGraph* Model = ModelPin->GetGraph();
+		if(Model == nullptr)
+		{
+			continue;
+		}
+
+		URigVMController* Controller = Blueprint->GetOrCreateController(Model);
+		if(Controller == nullptr)
+		{
+			continue;
+		}
+
+		Controller->OpenUndoBracket(TEXT("Set User Data Path"));
+		Controller->SetPinDefaultValue(ModelPin->GetPinPath(), InUserDataPath);
+
+		if(const UNameSpacedUserData::FUserData* UserData = GetUserData())
+		{
+			const FRigVMExternalVariable ExternalVariable = FRigVMExternalVariable::Make(UserData->GetProperty(), nullptr);
+			if(ExternalVariable.IsValid(true))
+			{
+				const TRigVMTypeIndex ExpectedTypeIndex = FRigVMRegistry::Get().GetTypeIndexFromCPPType(ExternalVariable.GetExtendedCPPType().ToString()); 
+				static const TArray<FString> PinPathsToCheck = {TEXT("Default"), TEXT("Result")};
+				if(const URigVMTemplateNode* TemplateNode = Cast<URigVMTemplateNode>(ModelPin->GetNode()))
+				{
+					for(const FString& PinPathToCheck : PinPathsToCheck)
+					{
+						if(const URigVMPin* PinToCheck = TemplateNode->FindPin(PinPathToCheck))
+						{
+							if(!FRigVMRegistry::Get().CanMatchTypes(PinToCheck->GetTypeIndex(), ExpectedTypeIndex, true))
+							{
+								Controller->UnresolveTemplateNodes({TemplateNode->GetFName()});
+								Controller->ResolveWildCardPin(PinToCheck->GetPinPath(), ExpectedTypeIndex);
+							}
+						}
+					}
+				}
+			}
+		}
+
+		Controller->CloseUndoBracket();
+	}
+}
+
+UControlRigBlueprint* SRigVMUserDataPath::GetBlueprint() const
+{
+	if(!ModelPins.IsEmpty())
+	{
+		return ModelPins[0]->GetTypedOuter<UControlRigBlueprint>();
+	}
+	return nullptr;
+}
+
+FString SRigVMUserDataPath::GetUserDataNameSpace(URigVMPin* ModelPin) const
+{
+	if(const URigVMNode* Node = ModelPin->GetNode())
+	{
+		static constexpr TCHAR NameSpaceName[] = TEXT("NameSpace");
+		if(const URigVMPin* NameSpacePin = Node->FindPin(NameSpaceName))
+		{
+			return NameSpacePin->GetDefaultValue();
+		}
+	}
+	return FString();
+}
+
+FString SRigVMUserDataPath::GetUserDataNameSpace() const
+{
+	if(ModelPins.Num() > 0)
+	{
+		const FString FirstNameSpace = GetUserDataNameSpace(ModelPins[0]);
+		for(int32 Index = 1; Index < ModelPins.Num(); Index++)
+		{
+			if(!GetUserDataNameSpace(ModelPins[Index]).Equals(FirstNameSpace, ESearchCase::CaseSensitive))
+			{
+				return FString();
+			}
+		}
+		return FirstNameSpace;
+	}
+	return FString();
+}
+
+const UNameSpacedUserData* SRigVMUserDataPath::GetUserDataObject() const
+{
+	const FString NameSpace = GetUserDataNameSpace();
+	if(NameSpace.IsEmpty())
+	{
+		return nullptr;
+	}
+	
+	if(const UBlueprint* Blueprint = GetBlueprint())
+	{
+		if(const UObject* DebuggedObject = Blueprint->GetObjectBeingDebugged())
+		{
+			if(DebuggedObject->Implements<UInterface_AssetUserData>())
+			{
+				const IInterface_AssetUserData* AssetUserDataHost = CastChecked<IInterface_AssetUserData>(DebuggedObject);
+				if(const TArray<UAssetUserData*>* UserDataArray = AssetUserDataHost->GetAssetUserDataArray())
+				{
+					for(int32 Index = UserDataArray->Num() - 1; UserDataArray->IsValidIndex(Index); Index--)
+					{
+						const UAssetUserData* UserData = (*UserDataArray)[Index];
+						if(const UNameSpacedUserData* NameSpacedUserData = Cast<UNameSpacedUserData>(UserData))
+						{
+							if(NameSpace.Equals(NameSpacedUserData->NameSpace, ESearchCase::CaseSensitive))
+							{
+								return NameSpacedUserData;
+							}
+						}
+					}
+				}
+			}
+		}
+	}
+
+	return nullptr;
+}
+
+const UNameSpacedUserData::FUserData* SRigVMUserDataPath::GetUserData() const
+{
+	const FString UserDataPath = GetUserDataPath();
+	if(UserDataPath.IsEmpty())
+	{
+		return nullptr;
+	}
+	if(const UNameSpacedUserData* UserDataObject = GetUserDataObject())
+	{
+		return UserDataObject->GetUserData(UserDataPath);
+	}
+	return nullptr;
+}
+
+/////////////////////////////////////////////////////////////////////////////////////////////////////////////
+
+void SGraphPinUserDataPath::Construct(const FArguments& InArgs, UEdGraphPin* InGraphPinObj)
+{
+	this->ModelPins = InArgs._ModelPins;
+	SGraphPin::Construct(SGraphPin::FArguments(), InGraphPinObj);
+}
+
+TSharedRef<SWidget>	SGraphPinUserDataPath::GetDefaultValueWidget()
+{
+	return SNew(SRigVMUserDataPath)
+		.ModelPins(ModelPins)
+		.Visibility(this, &SGraphPin::GetDefaultValueVisibility);
+}
+
+
+#undef LOCTEXT_NAMESPACE
diff --git a/Engine/Plugins/Animation/ControlRig/Source/ControlRigEditor/Private/Graph/SGraphPinUserDataPath.h b/Engine/Plugins/Animation/ControlRig/Source/ControlRigEditor/Private/Graph/SGraphPinUserDataPath.h
new file mode 100644
--- /dev/null
+++ b/Engine/Plugins/Animation/ControlRig/Source/ControlRigEditor/Private/Graph/SGraphPinUserDataPath.h
@@ -0,0 +1,72 @@
+// Copyright Epic Games, Inc. All Rights Reserved.
+
+#pragma once
+
+#include "CoreMinimal.h"
+#include "Widgets/DeclarativeSyntaxSupport.h"
+#include "Widgets/SWidget.h"
+#include "Widgets/Input/SMenuAnchor.h"
+#include "SGraphPin.h"
+#include "RigVMModel/RigVMPin.h"
+#include "ControlRigBlueprint.h"
+#include "ControlRigAssetUserData.h"
+#include "IPropertyAccessEditor.h"
+
+class SRigVMUserDataPath : public SCompoundWidget
+{
+public:
+
+	SLATE_BEGIN_ARGS(SRigVMUserDataPath)
+	: _ModelPins()
+	{}
+
+	SLATE_ARGUMENT(TArray<URigVMPin*>, ModelPins)
+
+	SLATE_END_ARGS()
+
+	void Construct(const FArguments& InArgs);
+
+protected:
+
+	FString GetUserDataPath(URigVMPin* ModelPin) const;
+	FString GetUserDataPath() const;
+	FText GetUserDataPathText() const { return FText::FromString(GetUserDataPath()); }
+	const FSlateBrush* GetUserDataIcon() const;
+	FLinearColor GetUserDataColor() const;
+	const FSlateBrush* GetUserDataIcon(const UNameSpacedUserData::FUserData* InUserData) const;
+	FLinearColor GetUserDataColor(const UNameSpacedUserData::FUserData* InUserData) const;
+	TSharedRef<SWidget> GetTopLevelMenuContent();
+	void FillUserDataPathMenu( FMenuBuilder& InMenuBuilder, FString InParentPath );
+	void HandleSetUserDataPath(FString InUserDataPath);
+
+	UControlRigBlueprint* GetBlueprint() const;
+	FString GetUserDataNameSpace(URigVMPin* ModelPin) const;
+	FString GetUserDataNameSpace() const;
+	const UNameSpacedUserData* GetUserDataObject() const;
+	const UNameSpacedUserData::FUserData* GetUserData() const;
+
+	TArray<URigVMPin*> ModelPins;
+	TSharedPtr<SMenuAnchor> MenuAnchor;
+	bool bAllowUObjects = false;
+};
+
+class SGraphPinUserDataPath : public SGraphPin
+{
+public:
+
+	SLATE_BEGIN_ARGS(SGraphPinUserDataPath){}
+
+		SLATE_ARGUMENT(TArray<URigVMPin*>, ModelPins)
+
+	SLATE_END_ARGS()
+
+	void Construct(const FArguments& InArgs, UEdGraphPin* InGraphPinObj);
+
+protected:
+
+	//~ Begin SGraphPin Interface
+	virtual TSharedRef<SWidget>	GetDefaultValueWidget() override;
+	//~ End SGraphPin Interface
+
+	TArray<URigVMPin*> ModelPins;
+};
