ControlRig: follow up for build break in 24890861

diff --git a/Engine/Plugins/Animation/ControlRig/Source/ControlRigEditor/Private/Graph/SGraphPinUserDataPath.h b/Engine/Plugins/Animation/ControlRig/Source/ControlRigEditor/Private/Graph/SGraphPinUserDataPath.h
--- a/Engine/Plugins/Animation/ControlRig/Source/ControlRigEditor/Private/Graph/SGraphPinUserDataPath.h
+++ b/Engine/Plugins/Animation/ControlRig/Source/ControlRigEditor/Private/Graph/SGraphPinUserDataPath.h
@@ -6,6 +6,7 @@
 #include "Widgets/DeclarativeSyntaxSupport.h"
 #include "Widgets/SWidget.h"
 #include "Widgets/Input/SMenuAnchor.h"
+#include "Widgets/Input/SButton.h"
 #include "SGraphPin.h"
 #include "RigVMModel/RigVMPin.h"
 #include "ControlRigBlueprint.h"
