Control Rig: Fix some public functions not available to reference

diff --git a/Engine/Plugins/Animation/ControlRig/Source/ControlRigEditor/Private/ControlRigEditorModule.cpp b/Engine/Plugins/Animation/ControlRig/Source/ControlRigEditor/Private/ControlRigEditorModule.cpp
--- a/Engine/Plugins/Animation/ControlRig/Source/ControlRigEditor/Private/ControlRigEditorModule.cpp
+++ b/Engine/Plugins/Animation/ControlRig/Source/ControlRigEditor/Private/ControlRigEditorModule.cpp
@@ -1171,8 +1171,8 @@ void FControlRigEditorModule::GetTypeActions(UControlRigBlueprint* CRB, FBluepri
 			{
 				PublicGraphFunctionsString = ControlRigAssetData.GetTagValueRef<FString>(PublicGraphFunctionsProperty->GetFName());
 			}
-			// Only look at the deprecated public functions if the PublicGraphFunctionsProperty is not found
-			else if (PublicFunctionsProperty)
+			// Only look at the deprecated public functions if the PublicGraphFunctionsString is empty
+			if (PublicGraphFunctionsString.IsEmpty())
 			{
 				PublicFunctionsString = ControlRigAssetData.GetTagValueRef<FString>(PublicFunctionsProperty->GetFName());
 			}
