Control Rig: Fix crash when accessing curves with a bad index

diff --git a/Engine/Plugins/Animation/ControlRig/Source/ControlRig/Private/AnimNode_ControlRigBase.cpp b/Engine/Plugins/Animation/ControlRig/Source/ControlRig/Private/AnimNode_ControlRigBase.cpp
--- a/Engine/Plugins/Animation/ControlRig/Source/ControlRig/Private/AnimNode_ControlRigBase.cpp
+++ b/Engine/Plugins/Animation/ControlRig/Source/ControlRig/Private/AnimNode_ControlRigBase.cpp
@@ -373,8 +373,12 @@ void FAnimNode_ControlRigBase::UpdateOutput(UControlRig* ControlRig, FPoseContex
 		TArray<FRigCurveElement*> Curves = Hierarchy->GetCurves();
 		UE::Anim::FCurveUtils::BulkSet(InOutput.Curve, CachedBulkCurves, [Hierarchy, &Curves](const UE::Anim::FNamedIndexElement& InBulkElement)
 		{
-			FRigCurveElement* CurveElement = Curves[InBulkElement.Index];
-			return Hierarchy->GetCurveValue(CurveElement);
+			if (Curves.IsValidIndex(InBulkElement.Index))
+			{
+				FRigCurveElement* CurveElement = Curves[InBulkElement.Index];
+				return Hierarchy->GetCurveValue(CurveElement);
+			}
+			return 0.f;
 		});
 	}
 
