Fixes race condition in URigHierarchy when mutating Elements array from multiple threads.

diff --git a/Engine/Plugins/Animation/ControlRig/Source/ControlRig/Private/Rigs/RigHierarchy.cpp b/Engine/Plugins/Animation/ControlRig/Source/ControlRig/Private/Rigs/RigHierarchy.cpp
--- a/Engine/Plugins/Animation/ControlRig/Source/ControlRig/Private/Rigs/RigHierarchy.cpp
+++ b/Engine/Plugins/Animation/ControlRig/Source/ControlRig/Private/Rigs/RigHierarchy.cpp
@@ -161,6 +161,8 @@ void URigHierarchy::Serialize(FArchive& Ar)
 
 void URigHierarchy::Save(FArchive& Ar)
 {
+	FScopeLock Lock(&ElementsLock);
+	
 	if(Ar.IsTransacting())
 	{
 		Ar << TransformStackIndex;
@@ -208,6 +210,8 @@ void URigHierarchy::Save(FArchive& Ar)
 
 void URigHierarchy::Load(FArchive& Ar)
 {
+	FScopeLock Lock(&ElementsLock);
+	
 	TArray<FRigElementKey> SelectedKeys;
 	if(Ar.IsTransacting())
 	{
@@ -354,6 +358,8 @@ void URigHierarchy::Reset()
 
 void URigHierarchy::ResetToDefault()
 {
+	FScopeLock Lock(&ElementsLock);
+	
 	if(DefaultHierarchyPtr.IsValid())
 	{
 		if(URigHierarchy* DefaultHierarchy = DefaultHierarchyPtr.Get())
@@ -370,6 +376,8 @@ void URigHierarchy::CopyHierarchy(URigHierarchy* InHierarchy)
 	check(InHierarchy);
 	
 	LLM_SCOPE_BYNAME(TEXT("Animation/ControlRig"));
+
+	FScopeLock Lock(&ElementsLock);
 	
 	Reset();
 
@@ -445,19 +453,6 @@ void URigHierarchy::CopyHierarchy(URigHierarchy* InHierarchy)
 	EnsureCacheValidity();
 }
 
-uint32 URigHierarchy::GetNameHash() const
-{
-	uint32 Hash = GetTypeHash(GetTopologyVersion());
-
-	for (int32 ElementIndex = 0; ElementIndex < Elements.Num(); ElementIndex++)
-	{
-		const FRigBaseElement* Element = Elements[ElementIndex];
-		Hash = HashCombine(Hash, GetTypeHash(Element->GetName()));
-	}
-
-	return Hash;
-}
-
 uint32 URigHierarchy::GetTopologyHash(bool bIncludeTopologyVersion, bool bIncludeTransientControls) const
 {
 	uint32 Hash = bIncludeTopologyVersion ? TopologyVersion : 0;
@@ -662,7 +657,9 @@ void URigHierarchy::ResetPoseToInitial(ERigElementType InTypeFilter)
 {
 	LLM_SCOPE_BYNAME(TEXT("Animation/ControlRig"));
 	bool bPerformFiltering = InTypeFilter != ERigElementType::All;
-
+	
+	FScopeLock Lock(&ElementsLock);
+	
 	// if we are resetting the pose on some elements, we need to check if
 	// any of affected elements has any children that would not be affected
 	// by resetting the pose. if all children are affected we can use the
diff --git a/Engine/Plugins/Animation/ControlRig/Source/ControlRig/Public/Rigs/RigHierarchy.h b/Engine/Plugins/Animation/ControlRig/Source/ControlRig/Public/Rigs/RigHierarchy.h
--- a/Engine/Plugins/Animation/ControlRig/Source/ControlRig/Public/Rigs/RigHierarchy.h
+++ b/Engine/Plugins/Animation/ControlRig/Source/ControlRig/Public/Rigs/RigHierarchy.h
@@ -3611,6 +3611,7 @@ private:
 
 	// Storage for the elements
 	mutable TArray<FRigBaseElement*> Elements;
+	FCriticalSection ElementsLock;
 
 	// Storage for the elements
 	mutable TArray<TArray<FRigBaseElement*>> ElementsPerType;
