Control Rig: Fix change shape name multiple times does not work

diff --git a/Engine/Plugins/Animation/ControlRig/Source/ControlRig/Private/ControlRigGizmoActor.cpp b/Engine/Plugins/Animation/ControlRig/Source/ControlRig/Private/ControlRigGizmoActor.cpp
--- a/Engine/Plugins/Animation/ControlRig/Source/ControlRig/Private/ControlRigGizmoActor.cpp
+++ b/Engine/Plugins/Animation/ControlRig/Source/ControlRig/Private/ControlRigGizmoActor.cpp
@@ -151,6 +151,7 @@ bool AControlRigShapeActor::UpdateControlSettings(
 						{
 							StaticMeshComponent->SetStaticMesh(StaticMesh);
 							bShapeTransformChanged = true;
+							ShapeName = ControlSettings.ShapeName;
 						}
 					}
 					else
