Control Rig: Fix crash when opening rig with multiple bone roots (some imported and others created by user)

diff --git a/Engine/Plugins/Animation/ControlRig/Source/ControlRigEditor/Private/Editor/ControlRigSkeletalMeshComponent.cpp b/Engine/Plugins/Animation/ControlRig/Source/ControlRigEditor/Private/Editor/ControlRigSkeletalMeshComponent.cpp
--- a/Engine/Plugins/Animation/ControlRig/Source/ControlRigEditor/Private/Editor/ControlRigSkeletalMeshComponent.cpp
+++ b/Engine/Plugins/Animation/ControlRig/Source/ControlRigEditor/Private/Editor/ControlRigSkeletalMeshComponent.cpp
@@ -133,7 +133,7 @@ void UControlRigSkeletalMeshComponent::RebuildDebugDrawSkeleton()
 				NewMeshBoneInfo.Name = BoneElement->GetName();
 				NewMeshBoneInfo.ParentIndex = ParentIndex; 
 				// give ref pose here
-				RefSkelModifier.Add(NewMeshBoneInfo, Hierarchy->GetInitialGlobalTransform(Index));
+				RefSkelModifier.Add(NewMeshBoneInfo, Hierarchy->GetInitialGlobalTransform(Index), true);
 
 				AddedBoneMap.FindOrAdd(BoneElement->GetName(), DebugDrawBones.Num());
 				DebugDrawBones.Add(DebugDrawBones.Num());
