Control Rig: Fix initial transforms not applying correctly on controls

diff --git a/Engine/Plugins/Animation/ControlRig/Source/ControlRig/Private/Rigs/RigHierarchy.cpp b/Engine/Plugins/Animation/ControlRig/Source/ControlRig/Private/Rigs/RigHierarchy.cpp
--- a/Engine/Plugins/Animation/ControlRig/Source/ControlRig/Private/Rigs/RigHierarchy.cpp
+++ b/Engine/Plugins/Animation/ControlRig/Source/ControlRig/Private/Rigs/RigHierarchy.cpp
@@ -774,6 +774,7 @@ void URigHierarchy::ResetPoseToInitial(ERigElementType InTypeFilter)
 		{
 			ControlElement->Offset.Current = ControlElement->Offset.Initial;
 			ControlElement->Shape.Current = ControlElement->Shape.Initial;
+			ControlElement->PreferredEulerAngles.Current = ControlElement->PreferredEulerAngles.Initial;
 		}
 
 		if(FRigTransformElement* TransformElement = Cast<FRigTransformElement>(Elements[ElementIndex]))
diff --git a/Engine/Plugins/Animation/ControlRig/Source/ControlRig/Private/Rigs/RigHierarchyController.cpp b/Engine/Plugins/Animation/ControlRig/Source/ControlRig/Private/Rigs/RigHierarchyController.cpp
--- a/Engine/Plugins/Animation/ControlRig/Source/ControlRig/Private/Rigs/RigHierarchyController.cpp
+++ b/Engine/Plugins/Animation/ControlRig/Source/ControlRig/Private/Rigs/RigHierarchyController.cpp
@@ -359,6 +359,7 @@ FRigElementKey URigHierarchyController::AddControl(
 		NewElement->Shape.MarkDirty(ERigTransformType::InitialGlobal);
 		NewElement->Offset.Current = NewElement->Offset.Initial;
 		NewElement->Pose.Current = NewElement->Pose.Initial;
+		NewElement->PreferredEulerAngles.Current = NewElement->PreferredEulerAngles.Initial;
 		NewElement->Shape.Current = NewElement->Shape.Initial;
 	}
 
diff --git a/Engine/Plugins/Animation/ControlRig/Source/ControlRig/Public/Rigs/RigHierarchyElements.h b/Engine/Plugins/Animation/ControlRig/Source/ControlRig/Public/Rigs/RigHierarchyElements.h
--- a/Engine/Plugins/Animation/ControlRig/Source/ControlRig/Public/Rigs/RigHierarchyElements.h
+++ b/Engine/Plugins/Animation/ControlRig/Source/ControlRig/Public/Rigs/RigHierarchyElements.h
@@ -1326,6 +1326,7 @@ protected:
 
 	friend struct FRigBaseElement;
 	friend class URigHierarchy;
+	friend class URigHierarchyController;
 };
 
 USTRUCT(BlueprintType)
