Fixes Control Rig crash when adding animation channel control.

diff --git a/Engine/Plugins/Animation/ControlRig/Source/ControlRigEditor/Private/ControlRigElementDetails.cpp b/Engine/Plugins/Animation/ControlRig/Source/ControlRigEditor/Private/ControlRigElementDetails.cpp
--- a/Engine/Plugins/Animation/ControlRig/Source/ControlRigEditor/Private/ControlRigElementDetails.cpp
+++ b/Engine/Plugins/Animation/ControlRig/Source/ControlRigEditor/Private/ControlRigElementDetails.cpp
@@ -3091,8 +3091,8 @@ void FRigControlElementDetails::CustomizeControl(IDetailLayoutBuilder& DetailBui
 
 void FRigControlElementDetails::CustomizeAnimationChannels(IDetailLayoutBuilder& DetailBuilder)
 {
-	// only show this for non-animation channels
-	if(!IsAnyControlNotOfAnimationType(ERigControlAnimationType::AnimationChannel))
+	// only show this for animation channels
+	if(IsAnyControlNotOfAnimationType(ERigControlAnimationType::AnimationChannel))
 	{
 		return;
 	}
