Control Rig: Fix some public functions not available in context menu

diff --git a/Engine/Plugins/Animation/ControlRig/Source/ControlRigEditor/Private/ControlRigEditorModule.cpp b/Engine/Plugins/Animation/ControlRig/Source/ControlRigEditor/Private/ControlRigEditorModule.cpp
--- a/Engine/Plugins/Animation/ControlRig/Source/ControlRigEditor/Private/ControlRigEditorModule.cpp
+++ b/Engine/Plugins/Animation/ControlRig/Source/ControlRigEditor/Private/ControlRigEditorModule.cpp
@@ -1172,7 +1172,7 @@ void FControlRigEditorModule::GetTypeActions(UControlRigBlueprint* CRB, FBluepri
 				PublicGraphFunctionsString = ControlRigAssetData.GetTagValueRef<FString>(PublicGraphFunctionsProperty->GetFName());
 			}
 			// Only look at the deprecated public functions if the PublicGraphFunctionsString is empty
-			if (PublicGraphFunctionsString.IsEmpty())
+			if (PublicGraphFunctionsString.IsEmpty() && PublicFunctionsProperty)
 			{
 				PublicFunctionsString = ControlRigAssetData.GetTagValueRef<FString>(PublicFunctionsProperty->GetFName());
 			}
