Refactored animation rig structures to reduce memory padding and runtime memory footprint

diff --git a/Engine/Plugins/Animation/ControlRig/Source/ControlRig/Private/Rigs/RigHierarchy.cpp b/Engine/Plugins/Animation/ControlRig/Source/ControlRig/Private/Rigs/RigHierarchy.cpp
--- a/Engine/Plugins/Animation/ControlRig/Source/ControlRig/Private/Rigs/RigHierarchy.cpp
+++ b/Engine/Plugins/Animation/ControlRig/Source/ControlRig/Private/Rigs/RigHierarchy.cpp
@@ -4438,9 +4438,9 @@ void URigHierarchy::CleanupInvalidCaches()
 
 	struct Local
 	{
-		static bool NeedsCheck(const FRigLocalAndGlobalTransform& InTransform)
+		static bool NeedsCheck(bool bDirty[FRigLocalAndGlobalTransform::EDirtyMax])
 		{
-			return !InTransform.Local.bDirty && !InTransform.Global.bDirty;
+			return !bDirty[FRigLocalAndGlobalTransform::ELocal] && !bDirty[FRigLocalAndGlobalTransform::EGlobal];
 		}
 	};
 
@@ -4450,17 +4450,17 @@ void URigHierarchy::CleanupInvalidCaches()
 		FRigBaseElement* BaseElement = HierarchyForCacheValidation->Elements[ElementIndex];
 		if(FRigControlElement* ControlElement = Cast<FRigControlElement>(BaseElement))
 		{
-			if(Local::NeedsCheck(ControlElement->Offset.Initial))
+			if(Local::NeedsCheck(ControlElement->Offset.bDirtyInitial))
 			{
 				ControlElement->Offset.MarkDirty(ERigTransformType::InitialGlobal);
 			}
 
-			if(Local::NeedsCheck(ControlElement->Pose.Initial))
+			if(Local::NeedsCheck(ControlElement->Pose.bDirtyInitial))
 			{
 				ControlElement->Pose.MarkDirty(ERigTransformType::InitialGlobal);
 			}
 
-			if(Local::NeedsCheck(ControlElement->Shape.Initial))
+			if(Local::NeedsCheck(ControlElement->Shape.bDirtyInitial))
 			{
 				ControlElement->Shape.MarkDirty(ERigTransformType::InitialGlobal);
 			}
@@ -4469,7 +4469,7 @@ void URigHierarchy::CleanupInvalidCaches()
 
 		if(FRigMultiParentElement* MultiParentElement = Cast<FRigMultiParentElement>(BaseElement))
 		{
-			if(Local::NeedsCheck(MultiParentElement->Pose.Initial))
+			if(Local::NeedsCheck(MultiParentElement->Pose.bDirtyInitial))
 			{
 				MultiParentElement->Pose.MarkDirty(ERigTransformType::InitialLocal);
 			}
@@ -4478,7 +4478,7 @@ void URigHierarchy::CleanupInvalidCaches()
 
 		if(FRigTransformElement* TransformElement = Cast<FRigTransformElement>(BaseElement))
 		{
-			if(Local::NeedsCheck(TransformElement->Pose.Initial))
+			if(Local::NeedsCheck(TransformElement->Pose.bDirtyInitial))
 			{
 				TransformElement->Pose.MarkDirty(ERigTransformType::InitialGlobal);
 			}
@@ -4497,7 +4497,7 @@ void URigHierarchy::CleanupInvalidCaches()
 		{
 			FRigControlElement* OtherControlElement = HierarchyForCacheValidation->FindChecked<FRigControlElement>(ControlElement->GetKey());
 			
-			if(Local::NeedsCheck(ControlElement->Offset.Initial))
+			if(Local::NeedsCheck(ControlElement->Offset.bDirtyInitial))
 			{
 				const FTransform CachedGlobalTransform = OtherControlElement->Offset.Get(ERigTransformType::InitialGlobal);
 				const FTransform ComputedGlobalTransform = HierarchyForCacheValidation->GetControlOffsetTransform(OtherControlElement, ERigTransformType::InitialGlobal);
@@ -4508,7 +4508,7 @@ void URigHierarchy::CleanupInvalidCaches()
 				}
 			}
 
-			if(Local::NeedsCheck(ControlElement->Pose.Initial))
+			if(Local::NeedsCheck(ControlElement->Pose.bDirtyInitial))
 			{
 				const FTransform CachedGlobalTransform = ControlElement->Pose.Get(ERigTransformType::InitialGlobal);
 				const FTransform ComputedGlobalTransform = HierarchyForCacheValidation->GetTransform(OtherControlElement, ERigTransformType::InitialGlobal);
@@ -4519,7 +4519,7 @@ void URigHierarchy::CleanupInvalidCaches()
 				}
 			}
 
-			if(Local::NeedsCheck(ControlElement->Shape.Initial))
+			if(Local::NeedsCheck(ControlElement->Shape.bDirtyInitial))
 			{
 				const FTransform CachedGlobalTransform = ControlElement->Shape.Get(ERigTransformType::InitialGlobal);
 				const FTransform ComputedGlobalTransform = HierarchyForCacheValidation->GetControlShapeTransform(OtherControlElement, ERigTransformType::InitialGlobal);
@@ -4536,7 +4536,7 @@ void URigHierarchy::CleanupInvalidCaches()
 		{
 			FRigMultiParentElement* OtherMultiParentElement = HierarchyForCacheValidation->FindChecked<FRigMultiParentElement>(MultiParentElement->GetKey());
 
-			if(Local::NeedsCheck(MultiParentElement->Pose.Initial))
+			if(Local::NeedsCheck(MultiParentElement->Pose.bDirtyInitial))
 			{
 				const FTransform CachedGlobalTransform = MultiParentElement->Pose.Get(ERigTransformType::InitialGlobal);
 				const FTransform ComputedGlobalTransform = HierarchyForCacheValidation->GetTransform(OtherMultiParentElement, ERigTransformType::InitialGlobal);
@@ -4554,7 +4554,7 @@ void URigHierarchy::CleanupInvalidCaches()
 		{
 			FRigTransformElement* OtherTransformElement = HierarchyForCacheValidation->FindChecked<FRigTransformElement>(TransformElement->GetKey());
 
-			if(Local::NeedsCheck(TransformElement->Pose.Initial))
+			if(Local::NeedsCheck(TransformElement->Pose.bDirtyInitial))
 			{
 				const FTransform CachedGlobalTransform = TransformElement->Pose.Get(ERigTransformType::InitialGlobal);
 				const FTransform ComputedGlobalTransform = HierarchyForCacheValidation->GetTransform(OtherTransformElement, ERigTransformType::InitialGlobal);
@@ -5810,7 +5810,7 @@ FTransform URigHierarchy::LazilyComputeParentConstraint(
 {
 	LLM_SCOPE_BYNAME(TEXT("Animation/ControlRig"));
 	const FRigElementParentConstraint& Constraint = InConstraints[InIndex];
-	if(Constraint.Cache.bDirty)
+	if(Constraint.bCacheIsDirty)
 	{
 		FTransform Transform = GetTransform(Constraint.ParentElement, InTransformType);
 		if(bApplyLocalOffsetTransform)
@@ -5824,7 +5824,7 @@ FTransform URigHierarchy::LazilyComputeParentConstraint(
 
 		Transform.NormalizeRotation();
 		Constraint.Cache.Transform = Transform;
-		Constraint.Cache.bDirty = false;
+		Constraint.bCacheIsDirty = false;
 	}
 	return Constraint.Cache.Transform;
 }
@@ -5843,7 +5843,7 @@ void URigHierarchy::ComputeParentConstraintIndices(
 	for(int32 ConstraintIndex = 0; ConstraintIndex < InConstraints.Num(); ConstraintIndex++)
 	{
 		// this is not relying on the cache whatsoever. we might as well remove it.
-		InConstraints[ConstraintIndex].Cache.bDirty = true;
+		InConstraints[ConstraintIndex].bCacheIsDirty = true;
 		
 		const FRigElementWeight& Weight = InConstraints[ConstraintIndex].GetWeight(bInitial);
 		if(Weight.AffectsLocation())
diff --git a/Engine/Plugins/Animation/ControlRig/Source/ControlRig/Private/Rigs/RigHierarchyController.cpp b/Engine/Plugins/Animation/ControlRig/Source/ControlRig/Private/Rigs/RigHierarchyController.cpp
--- a/Engine/Plugins/Animation/ControlRig/Source/ControlRig/Private/Rigs/RigHierarchyController.cpp
+++ b/Engine/Plugins/Animation/ControlRig/Source/ControlRig/Private/Rigs/RigHierarchyController.cpp
@@ -966,10 +966,10 @@ FString URigHierarchyController::ExportToText(TArray<FRigElementKey> InKeys) con
 
 		if(FRigTransformElement* TransformElement = Cast<FRigTransformElement>(Element))
 		{
-			PerElementData.Pose.Initial.Local.Set(Hierarchy->GetTransform(TransformElement, ERigTransformType::InitialLocal));
-			PerElementData.Pose.Initial.Global.Set(Hierarchy->GetTransform(TransformElement, ERigTransformType::InitialGlobal));
-			PerElementData.Pose.Current.Local.Set(Hierarchy->GetTransform(TransformElement, ERigTransformType::CurrentLocal));
-			PerElementData.Pose.Current.Global.Set(Hierarchy->GetTransform(TransformElement, ERigTransformType::CurrentGlobal));
+			PerElementData.Pose.Initial.Local.Set(Hierarchy->GetTransform(TransformElement, ERigTransformType::InitialLocal), PerElementData.Pose.bDirtyInitial[FRigLocalAndGlobalTransform::ELocal]);
+			PerElementData.Pose.Initial.Global.Set(Hierarchy->GetTransform(TransformElement, ERigTransformType::InitialGlobal), PerElementData.Pose.bDirtyInitial[FRigLocalAndGlobalTransform::EGlobal]);
+			PerElementData.Pose.Current.Local.Set(Hierarchy->GetTransform(TransformElement, ERigTransformType::CurrentLocal), PerElementData.Pose.bDirtyCurrent[FRigLocalAndGlobalTransform::ELocal]);
+			PerElementData.Pose.Current.Global.Set(Hierarchy->GetTransform(TransformElement, ERigTransformType::CurrentGlobal), PerElementData.Pose.bDirtyCurrent[FRigLocalAndGlobalTransform::EGlobal]);
 		}
 
 		switch (Key.Type)
diff --git a/Engine/Plugins/Animation/ControlRig/Source/ControlRig/Private/Rigs/RigHierarchyElements.cpp b/Engine/Plugins/Animation/ControlRig/Source/ControlRig/Private/Rigs/RigHierarchyElements.cpp
--- a/Engine/Plugins/Animation/ControlRig/Source/ControlRig/Private/Rigs/RigHierarchyElements.cpp
+++ b/Engine/Plugins/Animation/ControlRig/Source/ControlRig/Private/Rigs/RigHierarchyElements.cpp
@@ -290,32 +290,32 @@ void FRigBaseElement::NotifyMetadataTagChanged(const FName& InTag, bool bAdded)
 // FRigComputedTransform
 ////////////////////////////////////////////////////////////////////////////////
 
-void FRigComputedTransform::Save(FArchive& Ar)
+void FRigComputedTransform::Save(FArchive& Ar, bool& bDirty)
 {
 	Ar << Transform;
 	Ar << bDirty;
 }
 
-void FRigComputedTransform::Load(FArchive& Ar)
+void FRigComputedTransform::Load(FArchive& Ar, bool& bDirty)
 {
 	// load and save are identical
-	Save(Ar);
+	Save(Ar, bDirty);
 }
 
 ////////////////////////////////////////////////////////////////////////////////
 // FRigLocalAndGlobalTransform
 ////////////////////////////////////////////////////////////////////////////////
 
-void FRigLocalAndGlobalTransform::Save(FArchive& Ar)
+void FRigLocalAndGlobalTransform::Save(FArchive& Ar, bool bDirty[EDirtyMax])
 {
-	Local.Save(Ar);
-	Global.Save(Ar);
+	Local.Save(Ar, bDirty[ELocal]);
+	Global.Save(Ar, bDirty[EGlobal]);
 }
 
-void FRigLocalAndGlobalTransform::Load(FArchive& Ar)
+void FRigLocalAndGlobalTransform::Load(FArchive& Ar, bool bDirty[EDirtyMax])
 {
-	Local.Load(Ar);
-	Global.Load(Ar);
+	Local.Load(Ar, bDirty[ELocal]);
+	Global.Load(Ar, bDirty[EGlobal]);
 }
 
 ////////////////////////////////////////////////////////////////////////////////
@@ -324,14 +324,14 @@ void FRigLocalAndGlobalTransform::Load(FArchive& Ar)
 
 void FRigCurrentAndInitialTransform::Save(FArchive& Ar)
 {
-	Current.Save(Ar);
-	Initial.Save(Ar);
+	Current.Save(Ar, bDirtyCurrent);
+	Initial.Save(Ar, bDirtyInitial);
 }
 
 void FRigCurrentAndInitialTransform::Load(FArchive& Ar)
 {
-	Current.Load(Ar);
-	Initial.Load(Ar);
+	Current.Load(Ar, bDirtyCurrent);
+	Initial.Load(Ar, bDirtyInitial);
 }
 
 ////////////////////////////////////////////////////////////////////////////////
@@ -616,7 +616,7 @@ void FRigMultiParentElement::Load(FArchive& Ar, URigHierarchy* Hierarchy, ESeria
 			ensure(ParentKey.IsValid());
 
 			ParentConstraints[ParentIndex].ParentElement = Hierarchy->FindChecked<FRigTransformElement>(ParentKey);
-			ParentConstraints[ParentIndex].Cache.bDirty = true;
+			ParentConstraints[ParentIndex].bCacheIsDirty = true;
 
 			if (Ar.CustomVer(FControlRigObjectVersion::GUID) >= FControlRigObjectVersion::RigHierarchyMultiParentConstraints)
 			{
@@ -684,6 +684,7 @@ void FRigMultiParentElement::CopyPose(FRigBaseElement* InOther, bool bCurrent, b
 ////////////////////////////////////////////////////////////////////////////////
 // FRigBoneElement
 ////////////////////////////////////////////////////////////////////////////////
+static_assert(sizeof(FRigBoneElement) <= 736, "FRigBoneElement was optimized to fit into 736 bytes bin of MallocBinned3");
 
 void FRigBoneElement::Save(FArchive& Ar, URigHierarchy* Hierarchy, ESerializationPhase SerializationPhase)
 {
diff --git a/Engine/Plugins/Animation/ControlRig/Source/ControlRig/Public/Rigs/RigHierarchyElements.h b/Engine/Plugins/Animation/ControlRig/Source/ControlRig/Public/Rigs/RigHierarchyElements.h
--- a/Engine/Plugins/Animation/ControlRig/Source/ControlRig/Public/Rigs/RigHierarchyElements.h
+++ b/Engine/Plugins/Animation/ControlRig/Source/ControlRig/Public/Rigs/RigHierarchyElements.h
@@ -246,13 +246,12 @@ struct CONTROLRIG_API FRigComputedTransform
 
 	FRigComputedTransform()
 	: Transform(FTransform::Identity)
-	, bDirty(false)
 	{}
 
-	void Save(FArchive& Ar);
-	void Load(FArchive& Ar);
+	void Save(FArchive& Ar, bool& bDirty);
+	void Load(FArchive& Ar, bool& bDirty);
 
-	void Set(const FTransform& InTransform)
+	void Set(const FTransform& InTransform, bool& bDirty)
 	{
 #if WITH_EDITOR
 		ensure(InTransform.GetRotation().IsNormalized());
@@ -273,13 +272,11 @@ struct CONTROLRIG_API FRigComputedTransform
 
 	bool operator == (const FRigComputedTransform& Other) const
 	{
-		return bDirty == Other.bDirty && Equals(Transform, Other.Transform);
+		return Equals(Transform, Other.Transform);
     }
 
 	UPROPERTY(BlueprintReadOnly, EditAnywhere, Category = "Pose")
 	FTransform Transform;
-
-	bool bDirty;
 };
 
 USTRUCT(BlueprintType)
@@ -292,13 +289,15 @@ struct CONTROLRIG_API FRigLocalAndGlobalTransform
     , Global()
 	{}
 
-	void Save(FArchive& Ar);
-	void Load(FArchive& Ar);
-
-	bool operator == (const FRigLocalAndGlobalTransform& Other) const
+	enum EDirty
 	{
-		return Local == Other.Local && Global == Other.Global;
-	}
+		ELocal = 0,
+		EGlobal,
+		EDirtyMax
+	};
+
+	void Save(FArchive& Ar, bool bDirty[EDirtyMax]);
+	void Load(FArchive& Ar, bool bDirty[EDirtyMax]);
 
 	UPROPERTY(BlueprintReadOnly, EditAnywhere, Category = "Pose")
 	FRigComputedTransform Local;
@@ -365,6 +364,42 @@ struct CONTROLRIG_API FRigCurrentAndInitialTransform
 		return Initial.Global;
 	}
 
+	bool& GetDirtyFlag(const ERigTransformType::Type InTransformType)
+	{
+		switch (InTransformType)
+		{
+			case ERigTransformType::CurrentLocal:
+				return bDirtyCurrent[FRigLocalAndGlobalTransform::ELocal];
+
+			case ERigTransformType::CurrentGlobal:
+				return bDirtyCurrent[FRigLocalAndGlobalTransform::EGlobal];
+
+			case ERigTransformType::InitialLocal:
+				return bDirtyInitial[FRigLocalAndGlobalTransform::ELocal];
+			
+			default:
+				return bDirtyInitial[FRigLocalAndGlobalTransform::EGlobal];
+		}
+	}
+
+	const bool& GetDirtyFlag(const ERigTransformType::Type InTransformType) const
+	{
+		switch (InTransformType)
+		{
+		case ERigTransformType::CurrentLocal:
+			return bDirtyCurrent[FRigLocalAndGlobalTransform::ELocal];
+
+		case ERigTransformType::CurrentGlobal:
+			return bDirtyCurrent[FRigLocalAndGlobalTransform::EGlobal];
+
+		case ERigTransformType::InitialLocal:
+			return bDirtyInitial[FRigLocalAndGlobalTransform::ELocal];
+
+		default:
+			return bDirtyInitial[FRigLocalAndGlobalTransform::EGlobal];
+		}
+	}
+
 	const FTransform& Get(const ERigTransformType::Type InTransformType) const
 	{
 		return operator[](InTransformType).Transform;
@@ -372,18 +407,18 @@ struct CONTROLRIG_API FRigCurrentAndInitialTransform
 
 	void Set(const ERigTransformType::Type InTransformType, const FTransform& InTransform)
 	{
-		operator[](InTransformType).Set(InTransform);
+		operator[](InTransformType).Set(InTransform, GetDirtyFlag(InTransformType));
 	}
 
 	bool IsDirty(const ERigTransformType::Type InTransformType) const
 	{
-		return operator[](InTransformType).bDirty;
+		return GetDirtyFlag(InTransformType);
 	}
 
 	void MarkDirty(const ERigTransformType::Type InTransformType)
 	{
-		ensure(!(operator[](ERigTransformType::SwapLocalAndGlobal(InTransformType)).bDirty));
-		operator[](InTransformType).bDirty = true;
+		ensure(!(GetDirtyFlag(ERigTransformType::SwapLocalAndGlobal(InTransformType))));
+		GetDirtyFlag(InTransformType) = true;
 	}
 
 	void Save(FArchive& Ar);
@@ -391,7 +426,10 @@ struct CONTROLRIG_API FRigCurrentAndInitialTransform
 
 	bool operator == (const FRigCurrentAndInitialTransform& Other) const
 	{
-		return Current == Other.Current && Initial == Other.Initial;
+		return bDirtyCurrent[FRigLocalAndGlobalTransform::ELocal]  == Other.bDirtyCurrent[FRigLocalAndGlobalTransform::ELocal]  && Current.Local  == Other.Current.Local
+			&& bDirtyCurrent[FRigLocalAndGlobalTransform::EGlobal] == Other.bDirtyCurrent[FRigLocalAndGlobalTransform::EGlobal] && Current.Global == Other.Current.Global
+			&& bDirtyInitial[FRigLocalAndGlobalTransform::ELocal]  == Other.bDirtyInitial[FRigLocalAndGlobalTransform::ELocal]  && Initial.Local  == Other.Initial.Local
+			&& bDirtyInitial[FRigLocalAndGlobalTransform::EGlobal] == Other.bDirtyInitial[FRigLocalAndGlobalTransform::EGlobal] && Initial.Global == Other.Initial.Global;
 	}
 
 	UPROPERTY(BlueprintReadOnly, EditAnywhere, Category = "Pose")
@@ -399,6 +437,9 @@ struct CONTROLRIG_API FRigCurrentAndInitialTransform
 
 	UPROPERTY(BlueprintReadOnly, EditAnywhere, Category = "Pose")
 	FRigLocalAndGlobalTransform Initial;
+
+	bool bDirtyCurrent[FRigLocalAndGlobalTransform::EDirtyMax] = { false, false };
+	bool bDirtyInitial[FRigLocalAndGlobalTransform::EDirtyMax] = { false, false };
 };
 
 USTRUCT(BlueprintType)
@@ -513,11 +554,11 @@ public:
 	, NameString()
     , Index(INDEX_NONE)
 	, SubIndex(INDEX_NONE)
-	, bSelected(false)
 	, CreatedAtInstructionIndex(INDEX_NONE)
+	, OwnedInstances(0)
 	, TopologyVersion(0)
 	, MetadataVersion(0)
-	, OwnedInstances(0)
+	, bSelected(false)
 	{}
 
 	FRigBaseElement(const FRigBaseElement& InOther);
@@ -545,12 +586,6 @@ protected:
 	UPROPERTY(BlueprintReadOnly, VisibleAnywhere, Category = RigElement, meta = (AllowPrivateAccess = "true"))
 	int32 SubIndex;
 
-	UPROPERTY(BlueprintReadOnly, Transient, Category = RigElement, meta = (AllowPrivateAccess = "true"))
-	bool bSelected;
-
-	UPROPERTY(BlueprintReadOnly, Transient, Category = RigElement, meta = (AllowPrivateAccess = "true"))
-	int32 CreatedAtInstructionIndex;
-
 	TArray<FRigBaseMetadata*> Metadata;
 	TMap<FName,int32> MetadataNameToIndex;
 
@@ -674,15 +709,22 @@ protected:
 	void NotifyMetadataChanged(const FName& InName);
 	void NotifyMetadataTagChanged(const FName& InTag, bool bAdded);
 
-	mutable uint16 TopologyVersion;
-	mutable uint16 MetadataVersion;
 	mutable FRigBaseElementChildrenArray CachedChildren;
 
+	FRigElementMetadataChangedDelegate MetadataChangedDelegate;
+	FRigElementMetadataTagChangedDelegate MetadataTagChangedDelegate;
+
+	UPROPERTY(BlueprintReadOnly, Transient, Category = RigElement, meta = (AllowPrivateAccess = "true"))
+	int32 CreatedAtInstructionIndex;
+
 	// used for constructing / destructing the memory. typically == 1
 	int32 OwnedInstances;
 
-	FRigElementMetadataChangedDelegate MetadataChangedDelegate;
-	FRigElementMetadataTagChangedDelegate MetadataTagChangedDelegate;
+	mutable uint16 TopologyVersion;
+	mutable uint16 MetadataVersion;
+
+	UPROPERTY(BlueprintReadOnly, Transient, Category = RigElement, meta = (AllowPrivateAccess = "true"))
+	bool bSelected;
 
 	friend class URigHierarchy;
 	friend class URigHierarchyController;
@@ -873,18 +915,18 @@ public:
 USTRUCT()
 struct CONTROLRIG_API FRigElementParentConstraint
 {
-public:
 	GENERATED_BODY()
 
 	FRigTransformElement* ParentElement;
 	FRigElementWeight Weight;
 	FRigElementWeight InitialWeight;
 	mutable FRigComputedTransform Cache;
+	mutable bool bCacheIsDirty;
 		
 	FRigElementParentConstraint()
 		: ParentElement(nullptr)
 	{
-		Cache.bDirty = true;
+		bCacheIsDirty = true;
 	}
 
 	const FRigElementWeight& GetWeight(bool bInitial = false) const
@@ -902,7 +944,7 @@ public:
 		{
 			InitialWeight = InOther.InitialWeight;
 		}
-		Cache.bDirty = true;
+		bCacheIsDirty = true;
 	}
 };
 
