Control Rig editor layout does not persist across editor runs
#fix *rabbit hole alarm* - capture WeakPtr rather than SharedPtr to SDockTab to prevent bi-directional referencing. This made it so that SDockTab, and its TabIconAttribute (which holds SharedPtr a to the TabMananger) stayed around. Which in turn made it so that saving the editor layout on shutdown could corrupt the ControlRig editor entry. As it tried to save out the current layout using the still-alive tab manager without the window being alive anymore, and thus writing an empty layout.

diff --git a/Engine/Plugins/Animation/ControlRig/Source/ControlRigEditor/Private/Editor/RigHierarchyTabSummoner.cpp b/Engine/Plugins/Animation/ControlRig/Source/ControlRigEditor/Private/Editor/RigHierarchyTabSummoner.cpp
--- a/Engine/Plugins/Animation/ControlRig/Source/ControlRigEditor/Private/Editor/RigHierarchyTabSummoner.cpp
+++ b/Engine/Plugins/Animation/ControlRig/Source/ControlRigEditor/Private/Editor/RigHierarchyTabSummoner.cpp
@@ -44,18 +44,22 @@ TSharedRef<SWidget> FRigHierarchyTabSummoner::CreateTabBody(const FWorkflowTabSp
 TSharedRef<SDockTab> FRigHierarchyTabSummoner::SpawnTab(const FWorkflowTabSpawnInfo& Info) const
 {
 	TSharedRef<SDockTab>  DockTab = FWorkflowTabFactory::SpawnTab(Info);
-	DockTab->SetCanCloseTab(SDockTab::FCanCloseTab::CreateLambda([DockTab]()
+	TWeakPtr<SDockTab> WeakDockTab = DockTab;
+	DockTab->SetCanCloseTab(SDockTab::FCanCloseTab::CreateLambda([WeakDockTab]()
     {
 		int32 HierarchyTabCount = 0;
-		if(SWidget* Content = &DockTab->GetContent().Get())
+		if (TSharedPtr<SDockTab> SharedDocTab = WeakDockTab.Pin())
 		{
-			SRigHierarchy* RigHierarchy = (SRigHierarchy*)Content;
-			if(FControlRigEditor* ControlRigEditorForTab = RigHierarchy->GetControlRigEditor())
+			if(SWidget* Content = &SharedDocTab->GetContent().Get())
 			{
-				HierarchyTabCount = ControlRigEditorForTab->GetRigHierarchyTabCount();
+				SRigHierarchy* RigHierarchy = (SRigHierarchy*)Content;
+				if(FControlRigEditor* ControlRigEditorForTab = RigHierarchy->GetControlRigEditor())
+				{
+					HierarchyTabCount = ControlRigEditorForTab->GetRigHierarchyTabCount();
+				}
 			}
 		}
-        return HierarchyTabCount > 1;
+		return HierarchyTabCount > 1;
     }));
 	DockTab->SetOnTabClosed( SDockTab::FOnTabClosedCallback::CreateLambda([](TSharedRef<SDockTab> DockTab)
 	{
