Control Rig: Fix dirty flags not copying between poses correctly

Submitting on behalf of sara.sara.schvartzman to to unblock teams

diff --git a/Engine/Plugins/Animation/ControlRig/Source/ControlRig/Private/Rigs/RigHierarchy.cpp b/Engine/Plugins/Animation/ControlRig/Source/ControlRig/Private/Rigs/RigHierarchy.cpp
--- a/Engine/Plugins/Animation/ControlRig/Source/ControlRig/Private/Rigs/RigHierarchy.cpp
+++ b/Engine/Plugins/Animation/ControlRig/Source/ControlRig/Private/Rigs/RigHierarchy.cpp
@@ -4450,17 +4450,17 @@ void URigHierarchy::CleanupInvalidCaches()
 		FRigBaseElement* BaseElement = HierarchyForCacheValidation->Elements[ElementIndex];
 		if(FRigControlElement* ControlElement = Cast<FRigControlElement>(BaseElement))
 		{
-			if(Local::NeedsCheck(ControlElement->Offset.bDirtyInitial))
+			if(Local::NeedsCheck(ControlElement->Offset.Initial.bDirty))
 			{
 				ControlElement->Offset.MarkDirty(ERigTransformType::InitialGlobal);
 			}
 
-			if(Local::NeedsCheck(ControlElement->Pose.bDirtyInitial))
+			if(Local::NeedsCheck(ControlElement->Pose.Initial.bDirty))
 			{
 				ControlElement->Pose.MarkDirty(ERigTransformType::InitialGlobal);
 			}
 
-			if(Local::NeedsCheck(ControlElement->Shape.bDirtyInitial))
+			if(Local::NeedsCheck(ControlElement->Shape.Initial.bDirty))
 			{
 				ControlElement->Shape.MarkDirty(ERigTransformType::InitialGlobal);
 			}
@@ -4469,7 +4469,7 @@ void URigHierarchy::CleanupInvalidCaches()
 
 		if(FRigMultiParentElement* MultiParentElement = Cast<FRigMultiParentElement>(BaseElement))
 		{
-			if(Local::NeedsCheck(MultiParentElement->Pose.bDirtyInitial))
+			if(Local::NeedsCheck(MultiParentElement->Pose.Initial.bDirty))
 			{
 				MultiParentElement->Pose.MarkDirty(ERigTransformType::InitialLocal);
 			}
@@ -4478,7 +4478,7 @@ void URigHierarchy::CleanupInvalidCaches()
 
 		if(FRigTransformElement* TransformElement = Cast<FRigTransformElement>(BaseElement))
 		{
-			if(Local::NeedsCheck(TransformElement->Pose.bDirtyInitial))
+			if(Local::NeedsCheck(TransformElement->Pose.Initial.bDirty))
 			{
 				TransformElement->Pose.MarkDirty(ERigTransformType::InitialGlobal);
 			}
@@ -4497,7 +4497,7 @@ void URigHierarchy::CleanupInvalidCaches()
 		{
 			FRigControlElement* OtherControlElement = HierarchyForCacheValidation->FindChecked<FRigControlElement>(ControlElement->GetKey());
 			
-			if(Local::NeedsCheck(ControlElement->Offset.bDirtyInitial))
+			if(Local::NeedsCheck(ControlElement->Offset.Initial.bDirty))
 			{
 				const FTransform CachedGlobalTransform = OtherControlElement->Offset.Get(ERigTransformType::InitialGlobal);
 				const FTransform ComputedGlobalTransform = HierarchyForCacheValidation->GetControlOffsetTransform(OtherControlElement, ERigTransformType::InitialGlobal);
@@ -4508,7 +4508,7 @@ void URigHierarchy::CleanupInvalidCaches()
 				}
 			}
 
-			if(Local::NeedsCheck(ControlElement->Pose.bDirtyInitial))
+			if(Local::NeedsCheck(ControlElement->Pose.Initial.bDirty))
 			{
 				const FTransform CachedGlobalTransform = ControlElement->Pose.Get(ERigTransformType::InitialGlobal);
 				const FTransform ComputedGlobalTransform = HierarchyForCacheValidation->GetTransform(OtherControlElement, ERigTransformType::InitialGlobal);
@@ -4519,7 +4519,7 @@ void URigHierarchy::CleanupInvalidCaches()
 				}
 			}
 
-			if(Local::NeedsCheck(ControlElement->Shape.bDirtyInitial))
+			if(Local::NeedsCheck(ControlElement->Shape.Initial.bDirty))
 			{
 				const FTransform CachedGlobalTransform = ControlElement->Shape.Get(ERigTransformType::InitialGlobal);
 				const FTransform ComputedGlobalTransform = HierarchyForCacheValidation->GetControlShapeTransform(OtherControlElement, ERigTransformType::InitialGlobal);
@@ -4536,7 +4536,7 @@ void URigHierarchy::CleanupInvalidCaches()
 		{
 			FRigMultiParentElement* OtherMultiParentElement = HierarchyForCacheValidation->FindChecked<FRigMultiParentElement>(MultiParentElement->GetKey());
 
-			if(Local::NeedsCheck(MultiParentElement->Pose.bDirtyInitial))
+			if(Local::NeedsCheck(MultiParentElement->Pose.Initial.bDirty))
 			{
 				const FTransform CachedGlobalTransform = MultiParentElement->Pose.Get(ERigTransformType::InitialGlobal);
 				const FTransform ComputedGlobalTransform = HierarchyForCacheValidation->GetTransform(OtherMultiParentElement, ERigTransformType::InitialGlobal);
@@ -4554,7 +4554,7 @@ void URigHierarchy::CleanupInvalidCaches()
 		{
 			FRigTransformElement* OtherTransformElement = HierarchyForCacheValidation->FindChecked<FRigTransformElement>(TransformElement->GetKey());
 
-			if(Local::NeedsCheck(TransformElement->Pose.bDirtyInitial))
+			if(Local::NeedsCheck(TransformElement->Pose.Initial.bDirty))
 			{
 				const FTransform CachedGlobalTransform = TransformElement->Pose.Get(ERigTransformType::InitialGlobal);
 				const FTransform ComputedGlobalTransform = HierarchyForCacheValidation->GetTransform(OtherTransformElement, ERigTransformType::InitialGlobal);
diff --git a/Engine/Plugins/Animation/ControlRig/Source/ControlRig/Private/Rigs/RigHierarchyController.cpp b/Engine/Plugins/Animation/ControlRig/Source/ControlRig/Private/Rigs/RigHierarchyController.cpp
--- a/Engine/Plugins/Animation/ControlRig/Source/ControlRig/Private/Rigs/RigHierarchyController.cpp
+++ b/Engine/Plugins/Animation/ControlRig/Source/ControlRig/Private/Rigs/RigHierarchyController.cpp
@@ -966,10 +966,10 @@ FString URigHierarchyController::ExportToText(TArray<FRigElementKey> InKeys) con
 
 		if(FRigTransformElement* TransformElement = Cast<FRigTransformElement>(Element))
 		{
-			PerElementData.Pose.Initial.Local.Set(Hierarchy->GetTransform(TransformElement, ERigTransformType::InitialLocal), PerElementData.Pose.bDirtyInitial[FRigLocalAndGlobalTransform::ELocal]);
-			PerElementData.Pose.Initial.Global.Set(Hierarchy->GetTransform(TransformElement, ERigTransformType::InitialGlobal), PerElementData.Pose.bDirtyInitial[FRigLocalAndGlobalTransform::EGlobal]);
-			PerElementData.Pose.Current.Local.Set(Hierarchy->GetTransform(TransformElement, ERigTransformType::CurrentLocal), PerElementData.Pose.bDirtyCurrent[FRigLocalAndGlobalTransform::ELocal]);
-			PerElementData.Pose.Current.Global.Set(Hierarchy->GetTransform(TransformElement, ERigTransformType::CurrentGlobal), PerElementData.Pose.bDirtyCurrent[FRigLocalAndGlobalTransform::EGlobal]);
+			PerElementData.Pose.Initial.Local.Set(Hierarchy->GetTransform(TransformElement, ERigTransformType::InitialLocal), PerElementData.Pose.Initial.bDirty[FRigLocalAndGlobalTransform::ELocal]);
+			PerElementData.Pose.Initial.Global.Set(Hierarchy->GetTransform(TransformElement, ERigTransformType::InitialGlobal), PerElementData.Pose.Initial.bDirty[FRigLocalAndGlobalTransform::EGlobal]);
+			PerElementData.Pose.Current.Local.Set(Hierarchy->GetTransform(TransformElement, ERigTransformType::CurrentLocal), PerElementData.Pose.Current.bDirty[FRigLocalAndGlobalTransform::ELocal]);
+			PerElementData.Pose.Current.Global.Set(Hierarchy->GetTransform(TransformElement, ERigTransformType::CurrentGlobal), PerElementData.Pose.Current.bDirty[FRigLocalAndGlobalTransform::EGlobal]);
 		}
 
 		switch (Key.Type)
diff --git a/Engine/Plugins/Animation/ControlRig/Source/ControlRig/Private/Rigs/RigHierarchyElements.cpp b/Engine/Plugins/Animation/ControlRig/Source/ControlRig/Private/Rigs/RigHierarchyElements.cpp
--- a/Engine/Plugins/Animation/ControlRig/Source/ControlRig/Private/Rigs/RigHierarchyElements.cpp
+++ b/Engine/Plugins/Animation/ControlRig/Source/ControlRig/Private/Rigs/RigHierarchyElements.cpp
@@ -306,13 +306,13 @@ void FRigComputedTransform::Load(FArchive& Ar, bool& bDirty)
 // FRigLocalAndGlobalTransform
 ////////////////////////////////////////////////////////////////////////////////
 
-void FRigLocalAndGlobalTransform::Save(FArchive& Ar, bool bDirty[EDirtyMax])
+void FRigLocalAndGlobalTransform::Save(FArchive& Ar)
 {
 	Local.Save(Ar, bDirty[ELocal]);
 	Global.Save(Ar, bDirty[EGlobal]);
 }
 
-void FRigLocalAndGlobalTransform::Load(FArchive& Ar, bool bDirty[EDirtyMax])
+void FRigLocalAndGlobalTransform::Load(FArchive& Ar)
 {
 	Local.Load(Ar, bDirty[ELocal]);
 	Global.Load(Ar, bDirty[EGlobal]);
@@ -324,14 +324,14 @@ void FRigLocalAndGlobalTransform::Load(FArchive& Ar, bool bDirty[EDirtyMax])
 
 void FRigCurrentAndInitialTransform::Save(FArchive& Ar)
 {
-	Current.Save(Ar, bDirtyCurrent);
-	Initial.Save(Ar, bDirtyInitial);
+	Current.Save(Ar);
+	Initial.Save(Ar);
 }
 
 void FRigCurrentAndInitialTransform::Load(FArchive& Ar)
 {
-	Current.Load(Ar, bDirtyCurrent);
-	Initial.Load(Ar, bDirtyInitial);
+	Current.Load(Ar);
+	Initial.Load(Ar);
 }
 
 ////////////////////////////////////////////////////////////////////////////////
diff --git a/Engine/Plugins/Animation/ControlRig/Source/ControlRig/Public/Rigs/RigHierarchyElements.h b/Engine/Plugins/Animation/ControlRig/Source/ControlRig/Public/Rigs/RigHierarchyElements.h
--- a/Engine/Plugins/Animation/ControlRig/Source/ControlRig/Public/Rigs/RigHierarchyElements.h
+++ b/Engine/Plugins/Animation/ControlRig/Source/ControlRig/Public/Rigs/RigHierarchyElements.h
@@ -296,14 +296,27 @@ struct CONTROLRIG_API FRigLocalAndGlobalTransform
 		EDirtyMax
 	};
 
-	void Save(FArchive& Ar, bool bDirty[EDirtyMax]);
-	void Load(FArchive& Ar, bool bDirty[EDirtyMax]);
+	void Save(FArchive& Ar);
+	void Load(FArchive& Ar);
+
+	FRigLocalAndGlobalTransform& operator=(const FRigLocalAndGlobalTransform& Other)
+	{
+		Local = Other.Local;
+		Global = Other.Global;
+		for (int32 i=0; i<FRigLocalAndGlobalTransform::EDirtyMax; ++i)
+		{
+			bDirty[i] = Other.bDirty[i];
+		}
+		return *this;
+	}
 
 	UPROPERTY(BlueprintReadOnly, EditAnywhere, Category = "Pose")
 	FRigComputedTransform Local;
 
 	UPROPERTY(BlueprintReadOnly, EditAnywhere, Category = "Pose")
 	FRigComputedTransform Global;
+
+	bool bDirty[FRigLocalAndGlobalTransform::EDirtyMax] = { false, false };
 };
 
 USTRUCT(BlueprintType)
@@ -369,16 +382,16 @@ struct CONTROLRIG_API FRigCurrentAndInitialTransform
 		switch (InTransformType)
 		{
 			case ERigTransformType::CurrentLocal:
-				return bDirtyCurrent[FRigLocalAndGlobalTransform::ELocal];
+				return Current.bDirty[FRigLocalAndGlobalTransform::ELocal];
 
 			case ERigTransformType::CurrentGlobal:
-				return bDirtyCurrent[FRigLocalAndGlobalTransform::EGlobal];
+				return Current.bDirty[FRigLocalAndGlobalTransform::EGlobal];
 
 			case ERigTransformType::InitialLocal:
-				return bDirtyInitial[FRigLocalAndGlobalTransform::ELocal];
+				return Initial.bDirty[FRigLocalAndGlobalTransform::ELocal];
 			
 			default:
-				return bDirtyInitial[FRigLocalAndGlobalTransform::EGlobal];
+				return Initial.bDirty[FRigLocalAndGlobalTransform::EGlobal];
 		}
 	}
 
@@ -387,16 +400,16 @@ struct CONTROLRIG_API FRigCurrentAndInitialTransform
 		switch (InTransformType)
 		{
 		case ERigTransformType::CurrentLocal:
-			return bDirtyCurrent[FRigLocalAndGlobalTransform::ELocal];
+			return Current.bDirty[FRigLocalAndGlobalTransform::ELocal];
 
 		case ERigTransformType::CurrentGlobal:
-			return bDirtyCurrent[FRigLocalAndGlobalTransform::EGlobal];
+			return Current.bDirty[FRigLocalAndGlobalTransform::EGlobal];
 
 		case ERigTransformType::InitialLocal:
-			return bDirtyInitial[FRigLocalAndGlobalTransform::ELocal];
+			return Initial.bDirty[FRigLocalAndGlobalTransform::ELocal];
 
 		default:
-			return bDirtyInitial[FRigLocalAndGlobalTransform::EGlobal];
+			return Initial.bDirty[FRigLocalAndGlobalTransform::EGlobal];
 		}
 	}
 
@@ -426,10 +439,10 @@ struct CONTROLRIG_API FRigCurrentAndInitialTransform
 
 	bool operator == (const FRigCurrentAndInitialTransform& Other) const
 	{
-		return bDirtyCurrent[FRigLocalAndGlobalTransform::ELocal]  == Other.bDirtyCurrent[FRigLocalAndGlobalTransform::ELocal]  && Current.Local  == Other.Current.Local
-			&& bDirtyCurrent[FRigLocalAndGlobalTransform::EGlobal] == Other.bDirtyCurrent[FRigLocalAndGlobalTransform::EGlobal] && Current.Global == Other.Current.Global
-			&& bDirtyInitial[FRigLocalAndGlobalTransform::ELocal]  == Other.bDirtyInitial[FRigLocalAndGlobalTransform::ELocal]  && Initial.Local  == Other.Initial.Local
-			&& bDirtyInitial[FRigLocalAndGlobalTransform::EGlobal] == Other.bDirtyInitial[FRigLocalAndGlobalTransform::EGlobal] && Initial.Global == Other.Initial.Global;
+		return Current.bDirty[FRigLocalAndGlobalTransform::ELocal]  == Other.Current.bDirty[FRigLocalAndGlobalTransform::ELocal]  && Current.Local  == Other.Current.Local
+			&& Current.bDirty[FRigLocalAndGlobalTransform::EGlobal] == Other.Current.bDirty[FRigLocalAndGlobalTransform::EGlobal] && Current.Global == Other.Current.Global
+			&& Initial.bDirty[FRigLocalAndGlobalTransform::ELocal]  == Other.Initial.bDirty[FRigLocalAndGlobalTransform::ELocal]  && Initial.Local  == Other.Initial.Local
+			&& Initial.bDirty[FRigLocalAndGlobalTransform::EGlobal] == Other.Initial.bDirty[FRigLocalAndGlobalTransform::EGlobal] && Initial.Global == Other.Initial.Global;
 	}
 
 	UPROPERTY(BlueprintReadOnly, EditAnywhere, Category = "Pose")
@@ -437,9 +450,6 @@ struct CONTROLRIG_API FRigCurrentAndInitialTransform
 
 	UPROPERTY(BlueprintReadOnly, EditAnywhere, Category = "Pose")
 	FRigLocalAndGlobalTransform Initial;
-
-	bool bDirtyCurrent[FRigLocalAndGlobalTransform::EDirtyMax] = { false, false };
-	bool bDirtyInitial[FRigLocalAndGlobalTransform::EDirtyMax] = { false, false };
 };
 
 USTRUCT(BlueprintType)
