Sequencer: Control Rig: Moving over selection fix on save from other stream for procedural rigs.

diff --git a/Engine/Plugins/Animation/ControlRig/Source/ControlRigEditor/Private/Sequencer/ControlRigParameterTrackEditor.cpp b/Engine/Plugins/Animation/ControlRig/Source/ControlRigEditor/Private/Sequencer/ControlRigParameterTrackEditor.cpp
--- a/Engine/Plugins/Animation/ControlRig/Source/ControlRigEditor/Private/Sequencer/ControlRigParameterTrackEditor.cpp
+++ b/Engine/Plugins/Animation/ControlRig/Source/ControlRigEditor/Private/Sequencer/ControlRigParameterTrackEditor.cpp
@@ -3145,7 +3145,7 @@ void FControlRigParameterTrackEditor::HandleOnControlRigBound(UControlRig* InCon
 
 void FControlRigParameterTrackEditor::HandleOnObjectBoundToControlRig(UObject* InObject)
 {
-	//reselect these control rigs since selection may get
+	//reselect these control rigs since selection may get lost
 	TMap<TWeakObjectPtr<UControlRig>, TArray<FName>> ReselectIfNeeded;
 	// look for sections to update
 	TArray<UMovieSceneControlRigParameterSection*> SectionsToUpdate;
@@ -3181,7 +3181,7 @@ void FControlRigParameterTrackEditor::HandleOnObjectBoundToControlRig(UObject* I
 	// reconstruct proxies
 	if (!SectionsToUpdate.IsEmpty())
 	{
-		for (UMovieSceneControlRigParameterSection* Section: SectionsToUpdate)
+		for (UMovieSceneControlRigParameterSection* Section : SectionsToUpdate)
 		{
 			Section->ReconstructChannelProxy();
 			Section->MarkAsChanged();
@@ -3190,26 +3190,27 @@ void FControlRigParameterTrackEditor::HandleOnObjectBoundToControlRig(UObject* I
 	if (ReselectIfNeeded.Num() > 0)
 	{
 		GEditor->GetTimerManager()->SetTimerForNextTick([ReselectIfNeeded]()
-		{
-			GEditor->GetTimerManager()->SetTimerForNextTick([ReselectIfNeeded]()
-				{
-					for (const TPair <TWeakObjectPtr<UControlRig>, TArray<FName>>& Pair : ReselectIfNeeded)
+			{
+				GEditor->GetTimerManager()->SetTimerForNextTick([ReselectIfNeeded]()
 					{
-						if (Pair.Key.IsValid())
+						for (const TPair <TWeakObjectPtr<UControlRig>, TArray<FName>>& Pair : ReselectIfNeeded)
 						{
-							Pair.Key->ClearControlSelection();
-							for (const FName& ControlName : Pair.Value)
+							if (Pair.Key.IsValid())
 							{
-								Pair.Key->SelectControl(ControlName, true);
+								Pair.Key->ClearControlSelection();
+								for (const FName& ControlName : Pair.Value)
+								{
+									Pair.Key->SelectControl(ControlName, true);
+								}
 							}
 						}
-					}
-				});
-			
-		});
+					});
+
+			});
 	}
 }
 
+
 void FControlRigParameterTrackEditor::GetControlRigKeys(
 	UControlRig* InControlRig,
 	FName ParameterName,
