Fixing crash in URigHierarchy when interacting with sequencer and BP tools.

diff --git a/Engine/Plugins/Animation/ControlRig/Source/ControlRig/Private/Rigs/RigHierarchy.cpp b/Engine/Plugins/Animation/ControlRig/Source/ControlRig/Private/Rigs/RigHierarchy.cpp
--- a/Engine/Plugins/Animation/ControlRig/Source/ControlRig/Private/Rigs/RigHierarchy.cpp
+++ b/Engine/Plugins/Animation/ControlRig/Source/ControlRig/Private/Rigs/RigHierarchy.cpp
@@ -1909,15 +1909,19 @@ bool URigHierarchy::SwitchToParent(FRigBaseElement* InChild, FRigBaseElement* In
 	const TElementDependencyMap* DependencyMapPtr = &InDependencyMap;
 
 #if WITH_EDITOR
-	TElementDependencyMap DependencyMapFromVM;
-	if(ExecuteContext != nullptr && DependencyMapPtr->IsEmpty())
 	{
-		if(ExecuteContext->VM)
+		FScopeLock Lock(&ExecuteContextLock);
+		TElementDependencyMap DependencyMapFromVM;
+		if(ExecuteContext != nullptr && DependencyMapPtr->IsEmpty())
 		{
-			DependencyMapFromVM = GetDependenciesForVM(ExecuteContext->VM);
-			DependencyMapPtr = &DependencyMapFromVM;
+			if(ExecuteContext->VM)
+			{
+				DependencyMapFromVM = GetDependenciesForVM(ExecuteContext->VM);
+				DependencyMapPtr = &DependencyMapFromVM;
+			}
 		}
 	}
+	
 #endif
 	
 	if(InChild && InParent)
@@ -2459,11 +2463,15 @@ void URigHierarchy::Notify(ERigHierarchyNotification InNotifType, const FRigBase
 	}
 
 	// if we are running a VM right now
-	if(ExecuteContext != nullptr)
 	{
-		QueueNotification(InNotifType, InElement);
-		return;
+		FScopeLock Lock(&ExecuteContextLock);
+		if(ExecuteContext != nullptr)
+		{
+			QueueNotification(InNotifType, InElement);
+			return;
+		}
 	}
+	
 
 	if(QueuedNotifications.IsEmpty())
 	{
@@ -2542,11 +2550,14 @@ void URigHierarchy::SendQueuedNotifications()
 		return;
 	}
 
-	if(ExecuteContext != nullptr)
 	{
-		return;
+		FScopeLock Lock(&ExecuteContextLock);
+    	if(ExecuteContext != nullptr)
+    	{
+    		return;
+    	}
 	}
-
+	
 	// enable access to the controller during this method
 	FRigHierarchyEnableControllerBracket EnableController(this, true);
 
@@ -2687,16 +2698,19 @@ FTransform URigHierarchy::GetTransform(FRigTransformElement* InTransformElement,
 	}
 
 #if WITH_EDITOR
-
-	if(bRecordTransformsAtRuntime && ExecuteContext)
 	{
-		ReadTransformsAtRuntime.Emplace(
-			ExecuteContext->GetPublicData<>().GetInstructionIndex(),
-			ExecuteContext->GetSlice().GetIndex(),
-			InTransformElement->GetIndex(),
-			InTransformType
-		);
+		FScopeLock Lock(&ExecuteContextLock);
+		if(bRecordTransformsAtRuntime && ExecuteContext != nullptr)
+		{
+			ReadTransformsAtRuntime.Emplace(
+				ExecuteContext->GetPublicData<>().GetInstructionIndex(),
+				ExecuteContext->GetSlice().GetIndex(),
+				InTransformElement->GetIndex(),
+				InTransformType
+			);
+		}
 	}
+	
 	TGuardValue<bool> RecordTransformsPerInstructionGuard(bRecordTransformsAtRuntime, false);
 	
 #endif
@@ -2808,46 +2822,53 @@ void URigHierarchy::SetTransform(FRigTransformElement* InTransformElement, const
 
 #if WITH_EDITOR
 
-	if(bRecordTransformsAtRuntime && ExecuteContext)
+	// lock execute context scope
 	{
-		WrittenTransformsAtRuntime.Emplace(
-			ExecuteContext->GetPublicData<>().GetInstructionIndex(),
-			ExecuteContext->GetSlice().GetIndex(),
-			InTransformElement->GetIndex(),
-			InTransformType
-		);
-
-		// if we are setting a control / null parent after a child here - let's let the user know
-		if(InTransformElement->IsA<FRigControlElement>() || InTransformElement->IsA<FRigNullElement>())
+		FScopeLock Lock(&ExecuteContextLock);
+	
+		if(bRecordTransformsAtRuntime && ExecuteContext)
 		{
-			if(const UWorld* World = GetWorld())
+			const FRigVMExecuteContext& PublicData = ExecuteContext->GetPublicData<>();
+			const FRigVMSlice& Slice = ExecuteContext->GetSlice();
+			WrittenTransformsAtRuntime.Emplace(
+				PublicData.GetInstructionIndex(),
+				Slice.GetIndex(),
+				InTransformElement->GetIndex(),
+				InTransformType
+			);
+
+			// if we are setting a control / null parent after a child here - let's let the user know
+			if(InTransformElement->IsA<FRigControlElement>() || InTransformElement->IsA<FRigNullElement>())
 			{
-				// only fire these notes if we are inside the asset editor
-				if(World->WorldType == EWorldType::EditorPreview)
+				if(const UWorld* World = GetWorld())
 				{
-					const FRigBaseElementChildrenArray& Children = GetChildren(InTransformElement);
-					for (FRigBaseElement* Child : Children)
+					// only fire these notes if we are inside the asset editor
+					if(World->WorldType == EWorldType::EditorPreview)
 					{
-						const bool bChildFound = WrittenTransformsAtRuntime.ContainsByPredicate([Child](const TInstructionSliceElement& Entry) -> bool
+						const FRigBaseElementChildrenArray& Children = GetChildren(InTransformElement);
+						for (FRigBaseElement* Child : Children)
 						{
-							return Entry.Get<2>() == Child->GetIndex();
-						});
+							const bool bChildFound = WrittenTransformsAtRuntime.ContainsByPredicate([Child](const TInstructionSliceElement& Entry) -> bool
+							{
+								return Entry.Get<2>() == Child->GetIndex();
+							});
 
-						if(bChildFound)
-						{
-							const FControlRigExecuteContext& CRContext = ExecuteContext->GetPublicData<FControlRigExecuteContext>();
-							if(CRContext.GetLog())
+							if(bChildFound)
 							{
-								static constexpr TCHAR MessageFormat[] = TEXT("Setting transform of parent (%s) after setting child (%s).\nThis may lead to unexpected results.");
-								const FString& Message = FString::Printf(
-									MessageFormat,
-									*InTransformElement->GetName().ToString(),
-									*Child->GetName().ToString());
-								CRContext.GetLog()->Report(
-									EMessageSeverity::Info,
-									ExecuteContext->GetPublicData<>().GetFunctionName(),
-									ExecuteContext->GetPublicData<>().GetInstructionIndex(),
-									Message);
+								const FControlRigExecuteContext& CRContext = ExecuteContext->GetPublicData<FControlRigExecuteContext>();
+								if(CRContext.GetLog())
+								{
+									static constexpr TCHAR MessageFormat[] = TEXT("Setting transform of parent (%s) after setting child (%s).\nThis may lead to unexpected results.");
+									const FString& Message = FString::Printf(
+										MessageFormat,
+										*InTransformElement->GetName().ToString(),
+										*Child->GetName().ToString());
+									CRContext.GetLog()->Report(
+										EMessageSeverity::Info,
+										ExecuteContext->GetPublicData<>().GetFunctionName(),
+										ExecuteContext->GetPublicData<>().GetInstructionIndex(),
+										Message);
+								}
 							}
 						}
 					}
@@ -2855,6 +2876,7 @@ void URigHierarchy::SetTransform(FRigTransformElement* InTransformElement, const
 			}
 		}
 	}
+	
 	TGuardValue<bool> RecordTransformsPerInstructionGuard(bRecordTransformsAtRuntime, false);
 	
 #endif
@@ -2980,15 +3002,20 @@ FTransform URigHierarchy::GetControlOffsetTransform(FRigControlElement* InContro
 
 #if WITH_EDITOR
 
-	if(bRecordTransformsAtRuntime && ExecuteContext)
+	// lock execute context scope
 	{
-		ReadTransformsAtRuntime.Emplace(
-			ExecuteContext->GetPublicData<>().GetInstructionIndex(),
-			ExecuteContext->GetSlice().GetIndex(),
-			InControlElement->GetIndex(),
-			InTransformType
-		);
+		FScopeLock Lock(&ExecuteContextLock);
+		if(bRecordTransformsAtRuntime && ExecuteContext)
+		{
+			ReadTransformsAtRuntime.Emplace(
+				ExecuteContext->GetPublicData<>().GetInstructionIndex(),
+				ExecuteContext->GetSlice().GetIndex(),
+				InControlElement->GetIndex(),
+				InTransformType
+			);
+		}
 	}
+	
 	TGuardValue<bool> RecordTransformsPerInstructionGuard(bRecordTransformsAtRuntime, false);
 	
 #endif
@@ -3063,15 +3090,20 @@ void URigHierarchy::SetControlOffsetTransform(FRigControlElement* InControlEleme
 
 #if WITH_EDITOR
 
-	if(bRecordTransformsAtRuntime && ExecuteContext)
+	// lock execute context scope
 	{
-		WrittenTransformsAtRuntime.Emplace(
-			ExecuteContext->GetPublicData<>().GetInstructionIndex(),
-			ExecuteContext->GetSlice().GetIndex(),
-			InControlElement->GetIndex(),
-			InTransformType
-		);
+		FScopeLock Lock(&ExecuteContextLock);
+		if(bRecordTransformsAtRuntime && ExecuteContext)
+		{
+			WrittenTransformsAtRuntime.Emplace(
+				ExecuteContext->GetPublicData<>().GetInstructionIndex(),
+				ExecuteContext->GetSlice().GetIndex(),
+				InControlElement->GetIndex(),
+				InTransformType
+			);
+		}
 	}
+	
 	TGuardValue<bool> RecordTransformsPerInstructionGuard(bRecordTransformsAtRuntime, false);
 	
 #endif
diff --git a/Engine/Plugins/Animation/ControlRig/Source/ControlRig/Public/Rigs/RigHierarchy.h b/Engine/Plugins/Animation/ControlRig/Source/ControlRig/Public/Rigs/RigHierarchy.h
--- a/Engine/Plugins/Animation/ControlRig/Source/ControlRig/Public/Rigs/RigHierarchy.h
+++ b/Engine/Plugins/Animation/ControlRig/Source/ControlRig/Public/Rigs/RigHierarchy.h
@@ -4055,6 +4055,10 @@ private:
 	
 	void EnsureCacheValidityImpl();
 
+	// NOTE: it is not safe to read or write the execute context without locking it first as the
+	// FRigHierarchyExecuteContextBracket can change the context from the main thread (from sequencer for example)
+	// at the same time that a control rig is being evaluated on an animation thread.
+	mutable FCriticalSection ExecuteContextLock;
 	const FRigVMExtendedExecuteContext* ExecuteContext;
 
 #if WITH_EDITOR
@@ -4146,6 +4150,7 @@ private:
 		: Hierarchy(InHierarchy)
 		, PreviousContext(InHierarchy->ExecuteContext)
 	{
+		Hierarchy->ExecuteContextLock.Lock();
 		Hierarchy->ExecuteContext = InContext;
 	}
 
@@ -4153,6 +4158,7 @@ private:
 	{
 		Hierarchy->ExecuteContext = PreviousContext;
 		Hierarchy->SendQueuedNotifications();
+		Hierarchy->ExecuteContextLock.Unlock();
 	}
 
 	URigHierarchy* Hierarchy;
