Control Rig: Clean up object being debugged delegate

diff --git a/Engine/Plugins/Animation/ControlRig/Source/ControlRigEditor/Private/Editor/SRigHierarchy.cpp b/Engine/Plugins/Animation/ControlRig/Source/ControlRigEditor/Private/Editor/SRigHierarchy.cpp
--- a/Engine/Plugins/Animation/ControlRig/Source/ControlRigEditor/Private/Editor/SRigHierarchy.cpp
+++ b/Engine/Plugins/Animation/ControlRig/Source/ControlRigEditor/Private/Editor/SRigHierarchy.cpp
@@ -286,6 +286,7 @@ void SRigHierarchy::OnEditorClose(const FControlRigEditor* InEditor, UControlRig
 	{
 		InBlueprint->Hierarchy->OnModified().RemoveAll(this);
 		InBlueprint->OnRefreshEditor().RemoveAll(this);
+		InBlueprint->OnSetObjectBeingDebugged().RemoveAll(this);
 	}
 	
 	ControlRigEditor.Reset();
@@ -851,8 +852,11 @@ void SRigHierarchy::HandleSetObjectBeingDebugged(UObject* InObject)
 	if(UControlRig* ControlRig = Cast<UControlRig>(InObject))
 	{
 		ControlRigBeingDebuggedPtr = ControlRig;
-		ControlRig->GetHierarchy()->OnModified().RemoveAll(this);
-		ControlRig->GetHierarchy()->OnModified().AddSP(this, &SRigHierarchy::OnHierarchyModified_AnyThread);
+		if(URigHierarchy* Hierarchy = ControlRig->GetHierarchy())
+		{
+			Hierarchy->OnModified().RemoveAll(this);
+			Hierarchy->OnModified().AddSP(this, &SRigHierarchy::OnHierarchyModified_AnyThread);
+		}
 		ControlRig->OnPreConstructionForUI_AnyThread().RemoveAll(this);
 		ControlRig->OnPreConstructionForUI_AnyThread().AddSP(this, &SRigHierarchy::OnPreConstruction_AnyThread);
 		ControlRig->OnPostConstruction_AnyThread().RemoveAll(this);
