Control Rig: Node errors not disappearing when resolved

diff --git a/Engine/Plugins/Animation/ControlRig/Source/ControlRigEditor/Private/Editor/ControlRigEditor.cpp b/Engine/Plugins/Animation/ControlRig/Source/ControlRigEditor/Private/Editor/ControlRigEditor.cpp
--- a/Engine/Plugins/Animation/ControlRig/Source/ControlRigEditor/Private/Editor/ControlRigEditor.cpp
+++ b/Engine/Plugins/Animation/ControlRig/Source/ControlRigEditor/Private/Editor/ControlRigEditor.cpp
@@ -2285,6 +2285,7 @@ void FControlRigEditor::Compile()
 		
 		if (ControlRig)
 		{
+			ControlRigLog.Reset();
 			ControlRig->RigVMLog = &ControlRigLog;
 
 			URigVMBlueprintGeneratedClass* GeneratedClass = Cast<URigVMBlueprintGeneratedClass>(ControlRig->GetClass());
