Control Rig: Remove redirected, deprecated vector ops.

#rnx

diff --git a/Engine/Plugins/Animation/ControlRig/Source/ControlRig/Private/Units/Deprecated/Math/RigUnit_Vector.cpp b/Engine/Plugins/Animation/ControlRig/Source/ControlRig/Private/Units/Deprecated/Math/RigUnit_Vector.cpp
--- a/Engine/Plugins/Animation/ControlRig/Source/ControlRig/Private/Units/Deprecated/Math/RigUnit_Vector.cpp
+++ b/Engine/Plugins/Animation/ControlRig/Source/ControlRig/Private/Units/Deprecated/Math/RigUnit_Vector.cpp
@@ -90,25 +90,3 @@ FRigVMStructUpgradeInfo FRigUnit_Distance_VectorVector::GetUpgradeInfo() const
 	Info.AddRemappedPin(TEXT("Argument1"), TEXT("B"));
 	return Info;
 }
-
-FRigUnit_MathVectorBezierFourPoint_Execute()
-{
-    DECLARE_SCOPE_HIERARCHICAL_COUNTER_RIGUNIT()
-	FControlRigMathLibrary::FourPointBezier(Bezier, T, Result, Tangent);
-}
-
-FRigVMStructUpgradeInfo FRigUnit_MathVectorBezierFourPoint::GetUpgradeInfo() const
-{
-	// this node is no longer supported
-	return FRigVMStructUpgradeInfo();
-}
-
-FRigUnit_MathVectorMakeBezierFourPoint_Execute()
-{
-}
-
-FRigVMStructUpgradeInfo FRigUnit_MathVectorMakeBezierFourPoint::GetUpgradeInfo() const
-{
-	// this node is no longer supported
-	return FRigVMStructUpgradeInfo();
-}
diff --git a/Engine/Plugins/Animation/ControlRig/Source/ControlRig/Private/Units/Deprecated/Math/RigUnit_Vector.h b/Engine/Plugins/Animation/ControlRig/Source/ControlRig/Private/Units/Deprecated/Math/RigUnit_Vector.h
--- a/Engine/Plugins/Animation/ControlRig/Source/ControlRig/Private/Units/Deprecated/Math/RigUnit_Vector.h
+++ b/Engine/Plugins/Animation/ControlRig/Source/ControlRig/Private/Units/Deprecated/Math/RigUnit_Vector.h
@@ -93,59 +93,3 @@ struct CONTROLRIG_API FRigUnit_Distance_VectorVector : public FRigUnit
 	virtual FRigVMStructUpgradeInfo GetUpgradeInfo() const override;
 };
 
-/**
- * Returns the 4 point bezier interpolation
- */
-USTRUCT(meta=(DisplayName="Bezier Four Point", Deprecated = "5.0"))
-struct CONTROLRIG_API FRigUnit_MathVectorBezierFourPoint : public FRigVMFunction_MathVectorBase
-{
-	GENERATED_BODY()
-
-	FRigUnit_MathVectorBezierFourPoint()
-	{
-		Bezier = FRigVMFourPointBezier();
-		T = 0.f;
-		Result = Tangent = FVector::ZeroVector;
-	}
-
-	RIGVM_METHOD()
-	virtual void Execute() override;
-
-	UPROPERTY(meta=(Input))
-	FRigVMFourPointBezier Bezier;
-
-	UPROPERTY(meta=(Input))
-	float T;
-
-	UPROPERTY(meta=(Output))
-	FVector Result;
-
-	UPROPERTY(meta=(Output))
-	FVector Tangent;
-
-	RIGVM_METHOD()
-	virtual FRigVMStructUpgradeInfo GetUpgradeInfo() const override;
-};
-
-/**
- * Creates a bezier four point
- */
-USTRUCT(meta = (DisplayName = "Make Bezier Four Point", Constant, Deprecated = "5.0"))
-struct CONTROLRIG_API FRigUnit_MathVectorMakeBezierFourPoint : public FRigVMFunction_MathVectorBase
-{
-	GENERATED_BODY()
-
-	FRigUnit_MathVectorMakeBezierFourPoint()
-	{
-		Bezier = FRigVMFourPointBezier();
-	}
-
-	RIGVM_METHOD()
-	virtual void Execute() override;
-
-	UPROPERTY(meta = (Input, Output))
-	FRigVMFourPointBezier Bezier;
-
-	RIGVM_METHOD()
-	virtual FRigVMStructUpgradeInfo GetUpgradeInfo() const override;
-};
