Control Rig: Fix animation channels not showing as child of control in details panel

diff --git a/Engine/Plugins/Animation/ControlRig/Source/ControlRigEditor/Private/ControlRigElementDetails.cpp b/Engine/Plugins/Animation/ControlRig/Source/ControlRigEditor/Private/ControlRigElementDetails.cpp
--- a/Engine/Plugins/Animation/ControlRig/Source/ControlRigEditor/Private/ControlRigElementDetails.cpp
+++ b/Engine/Plugins/Animation/ControlRig/Source/ControlRigEditor/Private/ControlRigElementDetails.cpp
@@ -3091,9 +3091,10 @@ void FRigControlElementDetails::CustomizeControl(IDetailLayoutBuilder& DetailBui
 
 void FRigControlElementDetails::CustomizeAnimationChannels(IDetailLayoutBuilder& DetailBuilder)
 {
-	// only show this for animation channels
-	if(IsAnyControlNotOfAnimationType(ERigControlAnimationType::AnimationChannel))
+	// We only show this section for parents of animation channels
+	if(!IsAnyControlNotOfAnimationType(ERigControlAnimationType::AnimationChannel))
 	{
+		// If all controls are animation channels, just return 
 		return;
 	}
 
@@ -3420,7 +3421,7 @@ FReply FRigControlElementDetails::OnAddAnimationChannelClicked()
 		return FReply::Handled();
 	}
 
-	const FRigElementKey& Key = PerElementInfos[0].GetElement()->GetKey();
+	const FRigElementKey Key = PerElementInfos[0].GetElement()->GetKey();
 	URigHierarchy* HierarchyToChange = PerElementInfos[0].GetDefaultHierarchy();
 
 	static const FString ChannelName = TEXT("Channel");
@@ -4505,6 +4506,35 @@ FDetailWidgetRow& FRigControlElementDetails::CreateEnumValueWidgetRow(
 				}
 			}
 		}
+		else
+		{
+			// If the key was not found for selected elements, it might be a child channel of one of the elements
+			for (const FPerElementInfo ElementInfo : PerElementInfos)
+			{
+				if (const FRigControlElement* ControlElement = ElementInfo.GetElement<FRigControlElement>())
+				{
+					const FRigBaseElementChildrenArray Children = Hierarchy->GetChildren(ControlElement, false);
+					if (FRigBaseElement* const* Child = Children.FindByPredicate([Key](const FRigBaseElement* Info)
+					{
+						return Info->GetKey() == Key;
+					}))
+					{
+						if (const FRigControlElement* ChildElement = Cast<FRigControlElement>(*Child))
+						{
+							Enum = ChildElement->Settings.ControlEnum.Get();
+							if(Enum)
+							{
+								break;
+							}
+						}
+					}
+				}
+			}
+			if(Enum)
+			{
+				break;
+			}
+		}
 	}
 
 	check(Enum != nullptr);
