Control Rig: Debug VM execution button

diff --git a/Engine/Plugins/Animation/ControlRig/Source/ControlRig/Private/ControlRig.cpp b/Engine/Plugins/Animation/ControlRig/Source/ControlRig/Private/ControlRig.cpp
--- a/Engine/Plugins/Animation/ControlRig/Source/ControlRig/Private/ControlRig.cpp
+++ b/Engine/Plugins/Animation/ControlRig/Source/ControlRig/Private/ControlRig.cpp
@@ -423,6 +423,9 @@ bool UControlRig::Execute(const FName& InEventName)
 	PublicContext.SetDeltaTime(DeltaTime);
 	PublicContext.SetAbsoluteTime(AbsoluteTime);
 	PublicContext.SetFramesPerSecond(GetCurrentFramesPerSecond());
+#if UE_RIGVM_DEBUG_EXECUTION
+	PublicContext.bDebugExecution = bDebugExecutionEnabled;
+#endif
 
 	if (VM)
 	{
diff --git a/Engine/Plugins/Animation/ControlRig/Source/ControlRigEditor/Private/ControlRigCompilerDetails.cpp b/Engine/Plugins/Animation/ControlRig/Source/ControlRigEditor/Private/ControlRigCompilerDetails.cpp
--- a/Engine/Plugins/Animation/ControlRig/Source/ControlRigEditor/Private/ControlRigCompilerDetails.cpp
+++ b/Engine/Plugins/Animation/ControlRig/Source/ControlRigEditor/Private/ControlRigCompilerDetails.cpp
@@ -19,6 +19,16 @@
 
 #define LOCTEXT_NAMESPACE "ControlRigCompilerDetails"
 
+#if UE_RIGVM_DEBUG_EXECUTION
+//CVar to specify if we should create a float control for each curve in the curve container
+//By default we don't but it may be useful to do so for debugging
+static TAutoConsoleVariable<int32> CVarControlRigDebugVMExecutionStringEnabled(
+	TEXT("ControlRig.DebugVMExecutionStringEnabled"),
+	0,
+	TEXT("If nonzero we allow to copy the execution of a VM execution."),
+	ECVF_Default);
+#endif
+
 void FRigVMCompileSettingsDetails::CustomizeHeader(TSharedRef<IPropertyHandle> InStructPropertyHandle, FDetailWidgetRow& HeaderRow, IPropertyTypeCustomizationUtils& StructCustomizationUtils)
 {
 	HeaderRow
@@ -92,6 +102,7 @@ void FRigVMCompileSettingsDetails::CustomizeChildren(TSharedRef<IPropertyHandle>
 				]
 			];
 
+		TSharedPtr<SVerticalBox> DebugBox;
 		StructBuilder.AddCustomRow(LOCTEXT("DebuggingTools", "Debugging Tools"))
 			.NameContent()
 			[
@@ -101,7 +112,7 @@ void FRigVMCompileSettingsDetails::CustomizeChildren(TSharedRef<IPropertyHandle>
 			]
 			.ValueContent()
 			[
-				SNew(SVerticalBox)
+				SAssignNew(DebugBox, SVerticalBox)
 				+ SVerticalBox::Slot()
 				[
 					SNew(SButton)
@@ -155,6 +166,23 @@ void FRigVMCompileSettingsDetails::CustomizeChildren(TSharedRef<IPropertyHandle>
 					]
 				]
 			];
+#if UE_RIGVM_DEBUG_EXECUTION
+		if (CVarControlRigDebugVMExecutionStringEnabled->GetBool() == true)
+		{
+			DebugBox->AddSlot()
+			[
+				SNew(SButton)
+				.OnClicked(this, &FRigVMCompileSettingsDetails::OnCopyVMExecutionClicked)
+				.ContentPadding(FMargin(2))
+				.Content()
+				[
+					SNew(STextBlock)
+					.Justification(ETextJustify::Center)
+					.Text(LOCTEXT("CopyVMExecution", "Copy VM Execution"))
+				]
+			];
+		}
+#endif
 	}
 }
 
@@ -251,4 +279,19 @@ FReply FRigVMCompileSettingsDetails::OnCopyGeneratedCodeClicked()
 	return FReply::Handled();
 }
 
+#if UE_RIGVM_DEBUG_EXECUTION
+FReply FRigVMCompileSettingsDetails::OnCopyVMExecutionClicked()
+{
+	if (BlueprintBeingCustomized)
+	{
+		if (UControlRig* ControlRig = Cast<UControlRig>(BlueprintBeingCustomized->GetObjectBeingDebugged()))
+		{
+			FString DebugString = ControlRig->GetDebugExecutionString();
+			FPlatformApplicationMisc::ClipboardCopy(*DebugString);
+		}
+	}
+	return FReply::Handled();
+}
+#endif
+
 #undef LOCTEXT_NAMESPACE
diff --git a/Engine/Plugins/Animation/ControlRig/Source/ControlRigEditor/Private/ControlRigCompilerDetails.h b/Engine/Plugins/Animation/ControlRig/Source/ControlRigEditor/Private/ControlRigCompilerDetails.h
--- a/Engine/Plugins/Animation/ControlRig/Source/ControlRigEditor/Private/ControlRigCompilerDetails.h
+++ b/Engine/Plugins/Animation/ControlRig/Source/ControlRigEditor/Private/ControlRigCompilerDetails.h
@@ -29,6 +29,9 @@ protected:
 	FReply OnCopyByteCodeClicked();
 	FReply OnCopyHierarchyGraphClicked();
 	FReply OnCopyGeneratedCodeClicked();
+#if UE_RIGVM_DEBUG_EXECUTION
+	FReply OnCopyVMExecutionClicked();
+#endif
 	
 	UControlRigBlueprint* BlueprintBeingCustomized = nullptr;
 };
diff --git a/Engine/Plugins/Runtime/RigVM/Source/RigVM/Private/RigVMHost.cpp b/Engine/Plugins/Runtime/RigVM/Source/RigVM/Private/RigVMHost.cpp
--- a/Engine/Plugins/Runtime/RigVM/Source/RigVM/Private/RigVMHost.cpp
+++ b/Engine/Plugins/Runtime/RigVM/Source/RigVM/Private/RigVMHost.cpp
@@ -420,6 +420,9 @@ bool URigVMHost::Execute(const FName& InEventName)
 	PublicContext.SetAbsoluteTime(AbsoluteTime);
 	PublicContext.SetFramesPerSecond(GetCurrentFramesPerSecond());
 	PublicContext.SetOwningComponent(GetOwningSceneComponent());
+#if UE_RIGVM_DEBUG_EXECUTION
+	PublicContext.bDebugExecution = bDebugExecutionEnabled;
+#endif
 
 	if (VM)
 	{
@@ -819,6 +822,19 @@ void URigVMHost::DeclareConstructClasses(TArray<FTopLevelAssetPath>& OutConstruc
 	Super::DeclareConstructClasses(OutConstructClasses, SpecificSubclass);
 	OutConstructClasses.Add(FTopLevelAssetPath(URigVM::StaticClass()));
 }
+
+#if UE_RIGVM_DEBUG_EXECUTION
+const FString URigVMHost::GetDebugExecutionString()
+{
+	TGuardValue<bool> DebugExecutionGuard(bDebugExecutionEnabled, true);
+	FRigVMExecuteContext& PublicContext = GetVM()->GetPublicData<FRigVMExecuteContext>();
+	PublicContext.DebugMemoryString.Reset();
+	
+	Evaluate_AnyThread();
+
+	return PublicContext.DebugMemoryString;
+}
+#endif
 #endif
 
 void URigVMHost::PostInitInstance(URigVMHost* InCDO)
diff --git a/Engine/Plugins/Runtime/RigVM/Source/RigVM/Public/RigVMCore/RigVMExecuteContext.h b/Engine/Plugins/Runtime/RigVM/Source/RigVM/Public/RigVMCore/RigVMExecuteContext.h
--- a/Engine/Plugins/Runtime/RigVM/Source/RigVM/Public/RigVMCore/RigVMExecuteContext.h
+++ b/Engine/Plugins/Runtime/RigVM/Source/RigVM/Public/RigVMCore/RigVMExecuteContext.h
@@ -443,8 +443,10 @@ protected:
 
 
 #if UE_RIGVM_DEBUG_EXECUTION
+public:
+	bool bDebugExecution = false;
 	FString DebugMemoryString;
-
+protected:
 	TArray<FString> PreviousWorkMemory;
 
 	UEnum* InstanceOpCodeEnum;
diff --git a/Engine/Plugins/Runtime/RigVM/Source/RigVM/Public/RigVMHost.h b/Engine/Plugins/Runtime/RigVM/Source/RigVM/Public/RigVMHost.h
--- a/Engine/Plugins/Runtime/RigVM/Source/RigVM/Public/RigVMHost.h
+++ b/Engine/Plugins/Runtime/RigVM/Source/RigVM/Public/RigVMHost.h
@@ -211,6 +211,11 @@ public:
 	static void DeclareConstructClasses(TArray<FTopLevelAssetPath>& OutConstructClasses, const UClass* SpecificSubclass);
 #endif
 
+#if UE_RIGVM_DEBUG_EXECUTION
+	bool bDebugExecutionEnabled = false;
+	const FString GetDebugExecutionString();
+#endif
+
 protected:
 
 	virtual void PostInitInstance(URigVMHost* InCDO);
