Control Rig: Fix invalid hierarchies sending event messages

diff --git a/Engine/Plugins/Animation/ControlRig/Source/ControlRig/Private/Rigs/RigHierarchy.cpp b/Engine/Plugins/Animation/ControlRig/Source/ControlRig/Private/Rigs/RigHierarchy.cpp
--- a/Engine/Plugins/Animation/ControlRig/Source/ControlRig/Private/Rigs/RigHierarchy.cpp
+++ b/Engine/Plugins/Animation/ControlRig/Source/ControlRig/Private/Rigs/RigHierarchy.cpp
@@ -2462,6 +2462,11 @@ void URigHierarchy::Notify(ERigHierarchyNotification InNotifType, const FRigBase
 		return;
 	}
 
+	if (!IsValid(this))
+	{
+		return;
+	}
+
 	// if we are running a VM right now
 	{
 		FScopeLock Lock(&ExecuteContextLock);
