disabled a bunch of static size asserts if delegates race detection is enabled as the detection is a debugging tool but increases the size and static-asserts

diff --git a/Engine/Plugins/Animation/ControlRig/Source/ControlRig/Private/Rigs/RigHierarchyElements.cpp b/Engine/Plugins/Animation/ControlRig/Source/ControlRig/Private/Rigs/RigHierarchyElements.cpp
--- a/Engine/Plugins/Animation/ControlRig/Source/ControlRig/Private/Rigs/RigHierarchyElements.cpp
+++ b/Engine/Plugins/Animation/ControlRig/Source/ControlRig/Private/Rigs/RigHierarchyElements.cpp
@@ -684,7 +684,9 @@ void FRigMultiParentElement::CopyPose(FRigBaseElement* InOther, bool bCurrent, b
 ////////////////////////////////////////////////////////////////////////////////
 // FRigBoneElement
 ////////////////////////////////////////////////////////////////////////////////
+#if !UE_DETECT_DELEGATES_RACE_CONDITIONS // race detector increases mem footprint but is compiled out from test/shipping builds
 static_assert(sizeof(FRigBoneElement) <= 736, "FRigBoneElement was optimized to fit into 736 bytes bin of MallocBinned3");
+#endif
 
 void FRigBoneElement::Save(FArchive& Ar, URigHierarchy* Hierarchy, ESerializationPhase SerializationPhase)
 {
