Control Rig: Re-introduce the STAT_ControlRig_Evaluate stat for CR anim nodes.

diff --git a/Engine/Plugins/Animation/ControlRig/Source/ControlRig/Private/AnimNode_ControlRigBase.cpp b/Engine/Plugins/Animation/ControlRig/Source/ControlRig/Private/AnimNode_ControlRigBase.cpp
--- a/Engine/Plugins/Animation/ControlRig/Source/ControlRig/Private/AnimNode_ControlRigBase.cpp
+++ b/Engine/Plugins/Animation/ControlRig/Source/ControlRig/Private/AnimNode_ControlRigBase.cpp
@@ -15,6 +15,12 @@
 
 #include UE_INLINE_GENERATED_CPP_BY_NAME(AnimNode_ControlRigBase)
 
+
+DECLARE_CYCLE_STAT(TEXT("ControlRig_UpdateInput"), STAT_ControlRig_UpdateInput, STATGROUP_Anim);
+DECLARE_CYCLE_STAT(TEXT("ControlRig_Evaluate"), STAT_ControlRig_Evaluate, STATGROUP_Anim);
+DECLARE_CYCLE_STAT(TEXT("ControlRig_UpdateOutput"), STAT_ControlRig_UpdateOutput, STATGROUP_Anim);
+
+
 #if ENABLE_ANIM_DEBUG
 TAutoConsoleVariable<int32> CVarAnimNodeControlRigDebug(TEXT("a.AnimNode.ControlRig.Debug"), 0, TEXT("Set to 1 to turn on debug drawing for AnimNode_ControlRigBase"));
 #endif
@@ -132,7 +138,7 @@ bool FAnimNode_ControlRigBase::CanExecute()
 
 void FAnimNode_ControlRigBase::UpdateInput(UControlRig* ControlRig, const FPoseContext& InOutput)
 {
-	QUICK_SCOPE_CYCLE_COUNTER(STAT_ControlRig_UpdateInput);
+	SCOPE_CYCLE_COUNTER(STAT_ControlRig_UpdateInput);
 
 
 	if(!CanExecute())
@@ -268,7 +274,7 @@ void FAnimNode_ControlRigBase::UpdateInput(UControlRig* ControlRig, const FPoseC
 
 void FAnimNode_ControlRigBase::UpdateOutput(UControlRig* ControlRig, FPoseContext& InOutput)
 {
-	QUICK_SCOPE_CYCLE_COUNTER(STAT_ControlRig_UpdateOutput);
+	SCOPE_CYCLE_COUNTER(STAT_ControlRig_UpdateOutput);
 
 	if(!CanExecute())
 	{
@@ -443,6 +449,8 @@ void FAnimNode_ControlRigBase::Evaluate_AnyThread(FPoseContext& Output)
 
 void FAnimNode_ControlRigBase::ExecuteControlRig(FPoseContext& InOutput)
 {
+	SCOPE_CYCLE_COUNTER(STAT_ControlRig_Evaluate);
+
 	if (UControlRig* ControlRig = GetControlRig())
 	{
 		// temporarily give control rig access to the stack allocated attribute container
