Sequencer: Control Rig: For metahuman rigs need to also initialize skeleton transforms.

diff --git a/Engine/Plugins/Animation/ControlRig/Source/ControlRig/Private/Sequencer/MovieSceneControlRigParameterTemplate.cpp b/Engine/Plugins/Animation/ControlRig/Source/ControlRig/Private/Sequencer/MovieSceneControlRigParameterTemplate.cpp
--- a/Engine/Plugins/Animation/ControlRig/Source/ControlRig/Private/Sequencer/MovieSceneControlRigParameterTemplate.cpp
+++ b/Engine/Plugins/Animation/ControlRig/Source/ControlRig/Private/Sequencer/MovieSceneControlRigParameterTemplate.cpp
@@ -1171,6 +1171,12 @@ struct FControlRigParameterExecutionToken : IMovieSceneExecutionToken
 					ControlRig->GetObjectBinding()->BindToObject(BoundObject);
 					TArray<FName> SelectedControls = ControlRig->CurrentControlSelection();
 					ControlRig->Initialize();
+					if (USkeletalMeshComponent* SkeletalMeshComponent = Cast<USkeletalMeshComponent>(GetBindableObject(BoundObject)))
+					{
+						ControlRig->RequestInit();
+						ControlRig->SetBoneInitialTransformsFromSkeletalMeshComponent(SkeletalMeshComponent, true);
+						ControlRig->Evaluate_AnyThread();
+					};
 					if (ControlRig->IsA<UFKControlRig>())
 					{
 						UMovieSceneControlRigParameterTrack* Track = Section->GetTypedOuter<UMovieSceneControlRigParameterTrack>();
