Control Rig: Fix crash when creating local variable

diff --git a/Engine/Plugins/Animation/ControlRig/Source/ControlRigEditor/Private/ControlRigLocalVariableDetails.cpp b/Engine/Plugins/Animation/ControlRig/Source/ControlRigEditor/Private/ControlRigLocalVariableDetails.cpp
--- a/Engine/Plugins/Animation/ControlRig/Source/ControlRigEditor/Private/ControlRigLocalVariableDetails.cpp
+++ b/Engine/Plugins/Animation/ControlRig/Source/ControlRigEditor/Private/ControlRigLocalVariableDetails.cpp
@@ -31,7 +31,10 @@ void FRigVMLocalVariableDetails::CustomizeDetails(IDetailLayoutBuilder& DetailBu
 	if(ObjectsBeingCustomized[0].IsValid())
 	{
 		VariableDescription = ObjectsBeingCustomized[0]->GetContent<FRigVMGraphVariableDescription>();
-		GraphBeingCustomized = ObjectsBeingCustomized[0]->GetTypedOuter<URigVMGraph>();
+		if (UObject* Subject = ObjectsBeingCustomized[0]->GetSubject())
+		{
+			GraphBeingCustomized = Cast<URigVMGraph>(ObjectsBeingCustomized[0]->GetSubject());
+		}
 		if (GraphBeingCustomized)
 		{
 			BlueprintBeingCustomized = GraphBeingCustomized->GetTypedOuter<UControlRigBlueprint>();
