Sequencer: Smart Baking for Constraint Baking, Space Baking, Snapper and control rig collapsing. New settings type to hold the per frame, key increment settings when Baking which is shared among these algorithms. Constraints have a new Bake button/dialog now.

diff --git a/Engine/Plugins/Animation/ControlRig/Source/ControlRig/Private/Constraints/ControlRigTransformableHandle.cpp b/Engine/Plugins/Animation/ControlRig/Source/ControlRig/Private/Constraints/ControlRigTransformableHandle.cpp
--- a/Engine/Plugins/Animation/ControlRig/Source/ControlRig/Private/Constraints/ControlRigTransformableHandle.cpp
+++ b/Engine/Plugins/Animation/ControlRig/Source/ControlRig/Private/Constraints/ControlRigTransformableHandle.cpp
@@ -11,6 +11,7 @@
 #include "Channels/MovieSceneChannelProxy.h"
 #include "Sequencer/MovieSceneControlRigParameterSection.h"
 #include "Sections/MovieScene3DTransformSection.h"
+#include "Sequencer/ControlRigSequencerHelpers.h"
 
 #include UE_INLINE_GENERATED_CPP_BY_NAME(ControlRigTransformableHandle)
 
@@ -442,60 +443,10 @@ void UTransformableControlHandle::OnObjectBoundToControlRig(UObject* InObject)
 	}
 }
 
-static TPair<const FChannelMapInfo*, int32> GetInfoAndNumFloatChannels(
-	const UControlRig* InControlRig,
-	const FName& InControlName,
-	const UMovieSceneControlRigParameterSection* InSection)
-{
-	const FRigControlElement* ControlElement = InControlRig ? InControlRig->FindControl(InControlName) : nullptr;
-	auto GetNumFloatChannels = [](const ERigControlType& InControlType)
-	{
-		switch (InControlType)
-		{
-		case ERigControlType::Position:
-		case ERigControlType::Scale:
-		case ERigControlType::Rotator:
-			return 3;
-		case ERigControlType::TransformNoScale:
-			return 6;
-		case ERigControlType::Transform:
-		case ERigControlType::EulerTransform:
-			return 9;
-		default:
-			break;
-		}
-		return 0;
-	};
-
-	const int32 NumFloatChannels = ControlElement ? GetNumFloatChannels(ControlElement->Settings.ControlType) : 0;
-	const FChannelMapInfo* ChannelInfo = InSection ? InSection->ControlChannelMap.Find(InControlName) : nullptr;
-
-	return { ChannelInfo, NumFloatChannels };
-}
 TArrayView<FMovieSceneFloatChannel*>  UTransformableControlHandle::GetFloatChannels(const UMovieSceneSection* InSection) const
 {
-	// no floats for transform sections
-	static const TArrayView<FMovieSceneFloatChannel*> EmptyChannelsView;
-
-	const FChannelMapInfo* ChannelInfo = nullptr;
-	int32 NumChannels = 0;
-	const UMovieSceneControlRigParameterSection* CRSection = Cast<UMovieSceneControlRigParameterSection>(InSection);
-	if (CRSection == nullptr)
-	{
-		return EmptyChannelsView;
-	}
-
-	Tie(ChannelInfo, NumChannels) = GetInfoAndNumFloatChannels(ControlRig.Get(),ControlName, CRSection);
-
-	if (ChannelInfo == nullptr || NumChannels == 0)
-	{
-		return EmptyChannelsView;
-	}
-
-	// return a sub view that just represents the control's channels
-	const TArrayView<FMovieSceneFloatChannel*> FloatChannels = InSection->GetChannelProxy().GetChannels<FMovieSceneFloatChannel>();
-	const int32 ChannelStartIndex = ChannelInfo->ChannelIndex;
-	return FloatChannels.Slice(ChannelStartIndex, NumChannels);
+	return FControlRigSequencerHelpers::GetFloatChannels(ControlRig.Get(),
+		ControlName, InSection);
 }
 
 TArrayView<FMovieSceneDoubleChannel*>  UTransformableControlHandle::GetDoubleChannels(const UMovieSceneSection* InSection) const
diff --git a/Engine/Plugins/Animation/ControlRig/Source/ControlRig/Private/Sequencer/ControlRigSequencerHelpers.cpp b/Engine/Plugins/Animation/ControlRig/Source/ControlRig/Private/Sequencer/ControlRigSequencerHelpers.cpp
new file mode 100644
--- /dev/null
+++ b/Engine/Plugins/Animation/ControlRig/Source/ControlRig/Private/Sequencer/ControlRigSequencerHelpers.cpp
@@ -0,0 +1,66 @@
+// Copyright Epic Games, Inc. All Rights Reserved.
+#include "Sequencer/ControlRigSequencerHelpers.h"
+#include "ControlRig.h"
+#include "Rigs/RigHierarchyElements.h"
+#include "Channels/MovieSceneChannelProxy.h"
+#include "Sequencer/MovieSceneControlRigParameterSection.h"
+
+TPair<const FChannelMapInfo*, int32> FControlRigSequencerHelpers::GetInfoAndNumFloatChannels(
+	const UControlRig* InControlRig,
+	const FName& InControlName,
+	const UMovieSceneControlRigParameterSection* InSection)
+{
+	const FRigControlElement* ControlElement = InControlRig ? InControlRig->FindControl(InControlName) : nullptr;
+	auto GetNumFloatChannels = [](const ERigControlType& InControlType)
+	{
+		switch (InControlType)
+		{
+		case ERigControlType::Position:
+		case ERigControlType::Scale:
+		case ERigControlType::Rotator:
+			return 3;
+		case ERigControlType::TransformNoScale:
+			return 6;
+		case ERigControlType::Transform:
+		case ERigControlType::EulerTransform:
+			return 9;
+		default:
+			break;
+		}
+		return 0;
+	};
+
+	const int32 NumFloatChannels = ControlElement ? GetNumFloatChannels(ControlElement->Settings.ControlType) : 0;
+	const FChannelMapInfo* ChannelInfo = InSection ? InSection->ControlChannelMap.Find(InControlName) : nullptr;
+
+	return { ChannelInfo, NumFloatChannels };
+}
+
+
+TArrayView<FMovieSceneFloatChannel*>  FControlRigSequencerHelpers::GetFloatChannels(const UControlRig* InControlRig,
+	const FName& InControlName, const UMovieSceneSection* InSection)
+{
+	// no floats for transform sections
+	static const TArrayView<FMovieSceneFloatChannel*> EmptyChannelsView;
+
+	const FChannelMapInfo* ChannelInfo = nullptr;
+	int32 NumChannels = 0;
+	const UMovieSceneControlRigParameterSection* CRSection = Cast<UMovieSceneControlRigParameterSection>(InSection);
+	if (CRSection == nullptr)
+	{
+		return EmptyChannelsView;
+	}
+
+	Tie(ChannelInfo, NumChannels) = FControlRigSequencerHelpers::GetInfoAndNumFloatChannels(InControlRig, InControlName, CRSection);
+
+	if (ChannelInfo == nullptr || NumChannels == 0)
+	{
+		return EmptyChannelsView;
+	}
+
+	// return a sub view that just represents the control's channels
+	const TArrayView<FMovieSceneFloatChannel*> FloatChannels = InSection->GetChannelProxy().GetChannels<FMovieSceneFloatChannel>();
+	const int32 ChannelStartIndex = ChannelInfo->ChannelIndex;
+	return FloatChannels.Slice(ChannelStartIndex, NumChannels);
+}
+
diff --git a/Engine/Plugins/Animation/ControlRig/Source/ControlRig/Private/Sequencer/MovieSceneControlRigParameterSection.cpp b/Engine/Plugins/Animation/ControlRig/Source/ControlRig/Private/Sequencer/MovieSceneControlRigParameterSection.cpp
--- a/Engine/Plugins/Animation/ControlRig/Source/ControlRig/Private/Sequencer/MovieSceneControlRigParameterSection.cpp
+++ b/Engine/Plugins/Animation/ControlRig/Source/ControlRig/Private/Sequencer/MovieSceneControlRigParameterSection.cpp
@@ -2583,7 +2583,7 @@ void UMovieSceneControlRigParameterSection::SetControlRig(UControlRig* InControl
 	ControlRigClass = ControlRig ? ControlRig->GetClass() : nullptr;
 }
 
-void UMovieSceneControlRigParameterSection::FixRotationWinding(FName ControlName, FFrameNumber StartFrame, FFrameNumber EndFrame)
+void UMovieSceneControlRigParameterSection::FixRotationWinding(const FName& ControlName, FFrameNumber StartFrame, FFrameNumber EndFrame)
 {
 	FChannelMapInfo* pChannelIndex = ControlChannelMap.Find(ControlName);
 	if (pChannelIndex == nullptr || GetControlRig() == nullptr)
@@ -2630,6 +2630,107 @@ void UMovieSceneControlRigParameterSection::FixRotationWinding(FName ControlName
 	}
 }
 
+void UMovieSceneControlRigParameterSection::OptimizeSection(const FName& ControlName, const FKeyDataOptimizationParams& Params)
+{
+	TArrayView<FMovieSceneFloatChannel*> FloatChannels = GetChannelProxy().GetChannels<FMovieSceneFloatChannel>();
+	FChannelMapInfo* pChannelIndex = ControlChannelMap.Find(ControlName);
+	if (pChannelIndex != nullptr)
+	{
+		int32 ChannelIndex = pChannelIndex->ChannelIndex;
+
+		if (FRigControlElement* ControlElement = ControlRig->FindControl(ControlName))
+		{
+			switch (ControlElement->Settings.ControlType)
+			{
+				case ERigControlType::Position:
+				case ERigControlType::Scale:
+				case ERigControlType::Rotator:
+				{
+					FloatChannels[ChannelIndex]->Optimize(Params);
+					FloatChannels[ChannelIndex + 1]->Optimize(Params);
+					FloatChannels[ChannelIndex + 2]->Optimize(Params);
+					break;
+				}
+
+				case ERigControlType::Transform:
+				case ERigControlType::TransformNoScale:
+				case ERigControlType::EulerTransform:
+				{
+
+					FloatChannels[ChannelIndex]->Optimize(Params);
+					FloatChannels[ChannelIndex + 1]->Optimize(Params);
+					FloatChannels[ChannelIndex + 2]->Optimize(Params);
+					FloatChannels[ChannelIndex + 3]->Optimize(Params);
+					FloatChannels[ChannelIndex + 4]->Optimize(Params);
+					FloatChannels[ChannelIndex + 5]->Optimize(Params);
+
+					if (ControlElement->Settings.ControlType == ERigControlType::Transform ||
+						ControlElement->Settings.ControlType == ERigControlType::EulerTransform)
+					{
+						FloatChannels[ChannelIndex + 6]->Optimize(Params);
+						FloatChannels[ChannelIndex + 7]->Optimize(Params);
+						FloatChannels[ChannelIndex + 8]->Optimize(Params);
+					}
+					break;
+
+				}
+				default:
+					break;
+			}
+		}
+	}
+}
+
+void UMovieSceneControlRigParameterSection::AutoSetTangents(const FName& ControlName)
+{
+	TArrayView<FMovieSceneFloatChannel*> FloatChannels = GetChannelProxy().GetChannels<FMovieSceneFloatChannel>();
+	FChannelMapInfo* pChannelIndex = ControlChannelMap.Find(ControlName);
+	if (pChannelIndex != nullptr)
+	{
+		int32 ChannelIndex = pChannelIndex->ChannelIndex;
+
+		if (FRigControlElement* ControlElement = ControlRig->FindControl(ControlName))
+		{
+			switch (ControlElement->Settings.ControlType)
+			{
+			case ERigControlType::Position:
+			case ERigControlType::Scale:
+			case ERigControlType::Rotator:
+			{
+				FloatChannels[ChannelIndex]->AutoSetTangents();
+				FloatChannels[ChannelIndex + 1]->AutoSetTangents();
+				FloatChannels[ChannelIndex + 2]->AutoSetTangents();
+				break;
+			}
+
+			case ERigControlType::Transform:
+			case ERigControlType::TransformNoScale:
+			case ERigControlType::EulerTransform:
+			{
+
+				FloatChannels[ChannelIndex]->AutoSetTangents();
+				FloatChannels[ChannelIndex + 1]->AutoSetTangents();
+				FloatChannels[ChannelIndex + 2]->AutoSetTangents();
+				FloatChannels[ChannelIndex + 3]->AutoSetTangents();
+				FloatChannels[ChannelIndex + 4]->AutoSetTangents();
+				FloatChannels[ChannelIndex + 5]->AutoSetTangents();
+
+				if (ControlElement->Settings.ControlType == ERigControlType::Transform ||
+					ControlElement->Settings.ControlType == ERigControlType::EulerTransform)
+				{
+					FloatChannels[ChannelIndex + 6]->AutoSetTangents();
+					FloatChannels[ChannelIndex + 7]->AutoSetTangents();
+					FloatChannels[ChannelIndex + 8]->AutoSetTangents();
+				}
+				break;
+
+			}
+			default:
+				break;
+			}
+		}
+	}
+}
 #if WITH_EDITOR
 
 void UMovieSceneControlRigParameterSection::RecordControlRigKey(FFrameNumber FrameNumber, bool bSetDefault, ERichCurveInterpMode InInterpMode)
diff --git a/Engine/Plugins/Animation/ControlRig/Source/ControlRig/Private/Tools/ControlRigSnapSettings.cpp b/Engine/Plugins/Animation/ControlRig/Source/ControlRig/Private/Tools/ControlRigSnapSettings.cpp
deleted file mode 100644
diff --git a/Engine/Plugins/Animation/ControlRig/Source/ControlRig/Public/Sequencer/ControlRigSequencerHelpers.h b/Engine/Plugins/Animation/ControlRig/Source/ControlRig/Public/Sequencer/ControlRigSequencerHelpers.h
new file mode 100644
--- /dev/null
+++ b/Engine/Plugins/Animation/ControlRig/Source/ControlRig/Public/Sequencer/ControlRigSequencerHelpers.h
@@ -0,0 +1,27 @@
+// Copyright Epic Games, Inc. All Rights Reserved.
+
+#pragma once
+
+#include "CoreMinimal.h"
+
+struct FChannelMapInfo;
+class UControlRig;
+class UMovieSceneControlRigParameterSection;
+struct FMovieSceneFloatChannel;
+class UMovieSceneSection;
+class ISequencer;
+
+struct CONTROLRIG_API FControlRigSequencerHelpers
+{
+	static TPair<const FChannelMapInfo*, int32> GetInfoAndNumFloatChannels(
+		const UControlRig* InControlRig,
+		const FName& InControlName,
+		const UMovieSceneControlRigParameterSection* InSection);
+
+	static TArrayView<FMovieSceneFloatChannel*> GetFloatChannels(const UControlRig* InControlRig,
+		const FName& InControlName, const UMovieSceneSection* InSection);
+
+	static UMovieSceneControlRigParameterSection* GetControlRigSection(ISequencer* InSequencer, const UControlRig* InControlRig);
+
+};
+
diff --git a/Engine/Plugins/Animation/ControlRig/Source/ControlRig/Public/Sequencer/MovieSceneControlRigParameterSection.h b/Engine/Plugins/Animation/ControlRig/Source/ControlRig/Public/Sequencer/MovieSceneControlRigParameterSection.h
--- a/Engine/Plugins/Animation/ControlRig/Source/ControlRig/Public/Sequencer/MovieSceneControlRigParameterSection.h
+++ b/Engine/Plugins/Animation/ControlRig/Source/ControlRig/Public/Sequencer/MovieSceneControlRigParameterSection.h
@@ -24,6 +24,7 @@
 
 class UAnimSequence;
 class USkeletalMeshComponent;
+struct FKeyDataOptimizationParams;
 
 struct CONTROLRIG_API FControlRigBindingHelper
 {
@@ -197,7 +198,9 @@ public:
 	TArray<FIntegerParameterNameAndCurve>& GetIntegerParameterNamesAndCurves();
 	const TArray<FIntegerParameterNameAndCurve>& GetIntegerParameterNamesAndCurves() const;
 
-	void FixRotationWinding(FName ControlName, FFrameNumber StartFrame, FFrameNumber EndFrame);
+	void FixRotationWinding(const FName& ControlName, FFrameNumber StartFrame, FFrameNumber EndFrame);
+	void OptimizeSection(const FName& ControlName, const FKeyDataOptimizationParams& InParams);
+	void AutoSetTangents(const FName& ControlName);
 
 	TArray<FSpaceControlNameAndChannel>& GetSpaceChannels();
 	const TArray< FSpaceControlNameAndChannel>& GetSpaceChannels() const;
diff --git a/Engine/Plugins/Animation/ControlRig/Source/ControlRig/Public/Tools/ControlRigSnapSettings.h b/Engine/Plugins/Animation/ControlRig/Source/ControlRig/Public/Tools/ControlRigSnapSettings.h
deleted file mode 100644
diff --git a/Engine/Plugins/Animation/ControlRig/Source/ControlRigEditor/Private/ControlRigSequencerEditorLibrary.cpp b/Engine/Plugins/Animation/ControlRig/Source/ControlRigEditor/Private/ControlRigSequencerEditorLibrary.cpp
--- a/Engine/Plugins/Animation/ControlRig/Source/ControlRigEditor/Private/ControlRigSequencerEditorLibrary.cpp
+++ b/Engine/Plugins/Animation/ControlRig/Source/ControlRigEditor/Private/ControlRigSequencerEditorLibrary.cpp
@@ -47,6 +47,7 @@
 #include "Constraints/ControlRigTransformableHandle.h"
 #include "Constraints/MovieSceneConstraintChannelHelper.h"
 #include "Sections/MovieSceneConstrainedSection.h"
+#include "BakingAnimationKeySettings.h"
 
 #include UE_INLINE_GENERATED_CPP_BY_NAME(ControlRigSequencerEditorLibrary)
 
@@ -765,6 +766,8 @@ bool UControlRigSequencerEditorLibrary::BakeConstraint(UWorld* World, UTickableC
 		UE_LOG(LogControlRig, Error, TEXT("BakeConstraint: Need Valid World"));
 		return false;
 	}
+	TOptional< FBakingAnimationKeySettings> Settings;
+
 	if (UTickableTransformConstraint* TransformConstraint = Cast<UTickableTransformConstraint>(Constraint))
 	{
 		TSharedPtr<ISequencer> Sequencer = GetSequencerFromAsset();
@@ -797,7 +800,7 @@ bool UControlRigSequencerEditorLibrary::BakeConstraint(UWorld* World, UTickableC
 		{
 			FramesToBake = RealFramesToBake;
 		}
-		FConstraintBaker::Bake(World, TransformConstraint, Sequencer, FramesToBake);
+		FConstraintBaker::Bake(World, TransformConstraint, Sequencer, Settings, FramesToBake);
 	}
 	else
 	{
@@ -807,6 +810,48 @@ bool UControlRigSequencerEditorLibrary::BakeConstraint(UWorld* World, UTickableC
 	return true;
 }
 
+bool UControlRigSequencerEditorLibrary::BakeConstraints(UWorld* World, TArray<UTickableConstraint*>& InConstraints, const FBakingAnimationKeySettings& InSettings)
+{
+	if (!World)
+	{
+		UE_LOG(LogControlRig, Error, TEXT("BakeConstraint: Need Valid World"));
+		return false;
+	}
+
+	TSharedPtr<ISequencer> Sequencer = GetSequencerFromAsset();
+	if (!Sequencer || !Sequencer->GetFocusedMovieSceneSequence())
+	{
+		UE_LOG(LogControlRig, Error, TEXT("BakeConstraint: Need loaded level Sequence"));
+		return false;
+	}
+	const UMovieScene* MovieScene = Sequencer->GetFocusedMovieSceneSequence()->GetMovieScene();
+	if (!MovieScene)
+	{
+		UE_LOG(LogControlRig, Error, TEXT("BakeConstraint: Need valid Movie Scene"));
+		return false;
+	}
+	TArray<UTickableTransformConstraint*> TransformConstraints;
+	for (UTickableConstraint* Constraint : InConstraints)
+	{
+		if (UTickableTransformConstraint* TransformConstraint = Cast<UTickableTransformConstraint>(Constraint))
+		{
+			TransformConstraints.Add(TransformConstraint);
+		}
+		else
+		{
+			UE_LOG(LogControlRig, Error, TEXT("BakeConstraint: Need Valid Constraint"));
+			return false;
+		}
+	}
+
+	if(FConstraintBaker::BakeMultiple(World, TransformConstraints, Sequencer, InSettings) == false)
+	{
+		UE_LOG(LogControlRig, Error, TEXT("BakeMultiple: Failed"));
+		return false;
+	}
+	return true;
+}
+
 
 bool UControlRigSequencerEditorLibrary::SnapControlRig(ULevelSequence* LevelSequence, FFrameNumber StartFrame, FFrameNumber EndFrame, const FControlRigSnapperSelection& ChildrenToSnap,
 	const FControlRigSnapperSelection& ParentToSnap, const UControlRigSnapSettings* SnapSettings, ESequenceTimeUnit TimeUnit)
@@ -2542,6 +2587,19 @@ bool UControlRigSequencerEditorLibrary::ImportFBXToControlRigTrack(UWorld* World
 	return bValid;
 }
 
+bool UControlRigSequencerEditorLibrary::CollapseControlRigAnimLayersWithSettings(ULevelSequence* InSequence, UMovieSceneControlRigParameterTrack* InTrack, const FBakingAnimationKeySettings& InSettings)
+{
+	TWeakPtr<ISequencer> WeakSequencer = GetSequencerFromAsset();
+	bool bValid = false;
+	if (WeakSequencer.IsValid() && InTrack)
+	{
+		TSharedPtr<ISequencer>  SequencerPtr = WeakSequencer.Pin();
+		bValid = FControlRigParameterTrackEditor::CollapseAllLayers(SequencerPtr, InTrack, InSettings);
+		
+	}
+	return bValid;
+}
+
 bool UControlRigSequencerEditorLibrary::CollapseControlRigAnimLayers(ULevelSequence* LevelSequence, UMovieSceneControlRigParameterTrack* InTrack, bool bKeyReduce, float Tolerance)
 {
 	TWeakPtr<ISequencer> WeakSequencer = GetSequencerFromAsset();
@@ -2549,13 +2607,17 @@ bool UControlRigSequencerEditorLibrary::CollapseControlRigAnimLayers(ULevelSeque
 
 	if (WeakSequencer.IsValid() && InTrack)
 	{
-		TArray<UMovieSceneSection*> Sections = InTrack->GetAllSections();
-		if (Sections.Num() > 0)
-		{
-			TSharedPtr<ISequencer>  SequencerPtr = WeakSequencer.Pin();
-			UMovieSceneControlRigParameterSection* ParameterSection = Cast<UMovieSceneControlRigParameterSection>(Sections[0]);
-			bValid = FControlRigParameterTrackEditor::CollapseAllLayers(SequencerPtr,InTrack, ParameterSection, bKeyReduce, Tolerance);
-		}
+		TSharedPtr<ISequencer>  SequencerPtr = WeakSequencer.Pin();
+		FBakingAnimationKeySettings CollapseControlsSettings;
+		const FFrameRate TickResolution = SequencerPtr->GetFocusedTickResolution();
+		const FFrameTime FrameTime = SequencerPtr->GetLocalTime().ConvertTo(TickResolution);
+		FFrameNumber CurrentTime = FrameTime.GetFrame();
+
+		TRange<FFrameNumber> Range = SequencerPtr->GetFocusedMovieSceneSequence()->GetMovieScene()->GetPlaybackRange();
+
+		CollapseControlsSettings.StartFrame = Range.GetLowerBoundValue();
+		CollapseControlsSettings.EndFrame = Range.GetUpperBoundValue();
+		bValid = FControlRigParameterTrackEditor::CollapseAllLayers(SequencerPtr,InTrack, CollapseControlsSettings);	
 	}
 	return bValid;
 }
@@ -2605,17 +2667,12 @@ bool UControlRigSequencerEditorLibrary::BakeControlRigSpace(ULevelSequence* InSe
 	{
 		TSharedPtr<ISequencer>  Sequencer = WeakSequencer.Pin();
 		const FFrameRate TickResolution = Sequencer->GetFocusedTickResolution();
-		TArray<FFrameNumber> Frames;
 		const FFrameRate& FrameRate = Sequencer->GetFocusedDisplayRate();
 		FFrameNumber FrameRateInFrameNumber = TickResolution.AsFrameNumber(FrameRate.AsInterval());
 		if (TimeUnit == ESequenceTimeUnit::DisplayRate)
 		{
-			InSettings.StartFrame = FFrameRate::TransformTime(FFrameTime(InSettings.StartFrame, 0), FrameRate, TickResolution).RoundToFrame();
-			InSettings.EndFrame = FFrameRate::TransformTime(FFrameTime(InSettings.EndFrame, 0), FrameRate, TickResolution).RoundToFrame();
-		}
-		for (FFrameNumber& Frame = InSettings.StartFrame; Frame <= InSettings.EndFrame; Frame += FrameRateInFrameNumber)
-		{
-			Frames.Add(Frame);
+			InSettings.Settings.StartFrame = FFrameRate::TransformTime(FFrameTime(InSettings.Settings.StartFrame, 0), FrameRate, TickResolution).RoundToFrame();
+			InSettings.Settings.EndFrame = FFrameRate::TransformTime(FFrameTime(InSettings.Settings.EndFrame, 0), FrameRate, TickResolution).RoundToFrame();
 		}
 
 		FScopedTransaction Transaction(LOCTEXT("BakeControlToSpace", "Bake Control In Space"));
@@ -2627,7 +2684,7 @@ bool UControlRigSequencerEditorLibrary::BakeControlRigSpace(ULevelSequence* InSe
 				if (SpaceChannelAndSection.SpaceChannel)
 				{
 					FControlRigSpaceChannelHelpers::SequencerBakeControlInSpace(InControlRig, Sequencer.Get(), SpaceChannelAndSection.SpaceChannel, SpaceChannelAndSection.SectionToKey,
-						Frames, InControlRig->GetHierarchy(), Element->GetKey(), InSettings);
+						InControlRig->GetHierarchy(), Element->GetKey(), InSettings);
 				}
 			}
 		}
@@ -2714,7 +2771,6 @@ bool UControlRigSequencerEditorLibrary::MoveControlRigSpace(ULevelSequence* Leve
 				}
 				UMovieSceneControlRigParameterSection* ParamSection = Cast<UMovieSceneControlRigParameterSection>(SpaceChannelAndSection.SectionToKey);
 
-
 				TArray<FFrameNumber> OurKeyTimes;
 				TArray<FKeyHandle> OurKeyHandles;
 				TRange<FFrameNumber> CurrentFrameRange;
diff --git a/Engine/Plugins/Animation/ControlRig/Source/ControlRigEditor/Private/EditMode/SControlRigEditModeTools.cpp b/Engine/Plugins/Animation/ControlRig/Source/ControlRigEditor/Private/EditMode/SControlRigEditModeTools.cpp
--- a/Engine/Plugins/Animation/ControlRig/Source/ControlRigEditor/Private/EditMode/SControlRigEditModeTools.cpp
+++ b/Engine/Plugins/Animation/ControlRig/Source/ControlRigEditor/Private/EditMode/SControlRigEditModeTools.cpp
@@ -42,6 +42,8 @@
 
 #define LOCTEXT_NAMESPACE "ControlRigEditModeTools"
 
+//statics to reuse in the UI
+FRigSpacePickerBakeSettings SControlRigEditModeTools::BakeSpaceSettings;
 
 void SControlRigEditModeTools::SetControlRigs(const TArrayView<TWeakObjectPtr<UControlRig>>& InControlRigs)
 {
@@ -401,7 +403,7 @@ void SControlRigEditModeTools::Construct(const FArguments& InArgs, TSharedPtr<FC
 					}
 				})
 			]
-
+			
 			+ SVerticalBox::Slot()
 			.AutoHeight()
 			[
@@ -921,7 +923,6 @@ FReply SControlRigEditModeTools::OnBakeControlsToNewSpaceButtonClicked()
 		if (ControlRig.IsValid() && SpacePickerWidget->GetHierarchy() == ControlRig->GetHierarchy())
 		{
 
-			FRigSpacePickerBakeSettings Settings;
 			//Find default target space, just use first control and find space at current sequencer time
 			//Then Find range
 
@@ -938,42 +939,24 @@ FReply SControlRigEditModeTools::OnBakeControlsToNewSpaceButtonClicked()
 				FFrameNumber CurrentTime = FrameTime.GetFrame();
 				FMovieSceneControlRigSpaceBaseKey Value;
 				using namespace UE::MovieScene;
-				Settings.TargetSpace = URigHierarchy::GetDefaultParentKey();
-
-				TRange<FFrameNumber> Range = Sequencer->GetFocusedMovieSceneSequence()->GetMovieScene()->GetPlaybackRange();
-				TArray<FFrameNumber> Keys;
-				TArray < FKeyHandle> KeyHandles;
-
-				Settings.StartFrame = Range.GetLowerBoundValue();
-				Settings.EndFrame = Range.GetUpperBoundValue();
-				if (Keys.Num() > 0)
+				//set up settings if not setup
+				if (BakeSpaceSettings.TargetSpace == FRigElementKey())
 				{
-					int32 Index = Algo::LowerBound(Keys, CurrentTime);
-					if (Index >= 0 && Index < (Keys.Num() - 1))
-					{
-						Settings.StartFrame = Keys[Index];
-						Settings.EndFrame = Keys[Index + 1];
-
-					}
+					BakeSpaceSettings.TargetSpace = URigHierarchy::GetDefaultParentKey();
+					TRange<FFrameNumber> Range = Sequencer->GetFocusedMovieSceneSequence()->GetMovieScene()->GetPlaybackRange();
+					BakeSpaceSettings.Settings.StartFrame = Range.GetLowerBoundValue();
+					BakeSpaceSettings.Settings.EndFrame = Range.GetUpperBoundValue();
 				}
 
 				TSharedRef<SRigSpacePickerBakeWidget> BakeWidget =
 					SNew(SRigSpacePickerBakeWidget)
-					.Settings(Settings)
+					.Settings(BakeSpaceSettings)
 					.Hierarchy(SpacePickerWidget->GetHierarchy())
 					.Controls(ControlKeys) // use the cached controls here since the selection is not recovered until next tick.
 					.Sequencer(Sequencer)
 					.GetControlCustomization(this, &SControlRigEditModeTools::HandleGetControlElementCustomization)
-					.OnBake_Lambda([Sequencer, ControlRig, TickResolution](URigHierarchy* InHierarchy, TArray<FRigElementKey> InControls, FRigSpacePickerBakeSettings InSettings)
-						{
-							TArray<FFrameNumber> Frames;
-
-							const FFrameRate& FrameRate = Sequencer->GetFocusedDisplayRate();
-							FFrameNumber FrameRateInFrameNumber = TickResolution.AsFrameNumber(FrameRate.AsInterval());
-							for (FFrameNumber& Frame = InSettings.StartFrame; Frame <= InSettings.EndFrame; Frame += FrameRateInFrameNumber)
-							{
-								Frames.Add(Frame);
-							}
+					.OnBake_Lambda([Sequencer, ControlRig, TickResolution](URigHierarchy* InHierarchy, TArray<FRigElementKey> InControls, FRigSpacePickerBakeSettings& InSettings)
+						{		
 							FScopedTransaction Transaction(LOCTEXT("BakeControlToSpace", "Bake Control In Space"));
 							for (const FRigElementKey& ControlKey : InControls)
 							{
@@ -982,8 +965,9 @@ FReply SControlRigEditModeTools::OnBakeControlsToNewSpaceButtonClicked()
 								if (SpaceChannelAndSection.SpaceChannel)
 								{
 									FControlRigSpaceChannelHelpers::SequencerBakeControlInSpace(ControlRig.Get(), Sequencer, SpaceChannelAndSection.SpaceChannel, SpaceChannelAndSection.SectionToKey,
-										Frames, InHierarchy, ControlKey, InSettings);
+										InHierarchy, ControlKey, InSettings);
 								}
+								SControlRigEditModeTools::BakeSpaceSettings = InSettings;
 							}
 							return FReply::Handled();
 						});
diff --git a/Engine/Plugins/Animation/ControlRig/Source/ControlRigEditor/Private/EditMode/SControlRigEditModeTools.h b/Engine/Plugins/Animation/ControlRig/Source/ControlRigEditor/Private/EditMode/SControlRigEditModeTools.h
--- a/Engine/Plugins/Animation/ControlRig/Source/ControlRigEditor/Private/EditMode/SControlRigEditModeTools.h
+++ b/Engine/Plugins/Animation/ControlRig/Source/ControlRigEditor/Private/EditMode/SControlRigEditModeTools.h
@@ -59,7 +59,12 @@ public:
 	virtual bool IsPropertyKeyingEnabled() const override;
 	virtual void OnKeyPropertyClicked(const IPropertyHandle& KeyedPropertyHandle) override;
 	virtual bool IsPropertyAnimated(const class IPropertyHandle& PropertyHandle, UObject *ParentObject) const override;
+
+	//reuse settings for space and constraint baking baking
+	static FRigSpacePickerBakeSettings BakeSpaceSettings;
+
 private:
+
 	/** Sequencer we are currently bound to */
 
 	TWeakPtr<ISequencer> WeakSequencer;
diff --git a/Engine/Plugins/Animation/ControlRig/Source/ControlRigEditor/Private/EditMode/SControlRigSpacePicker.cpp b/Engine/Plugins/Animation/ControlRig/Source/ControlRigEditor/Private/EditMode/SControlRigSpacePicker.cpp
--- a/Engine/Plugins/Animation/ControlRig/Source/ControlRigEditor/Private/EditMode/SControlRigSpacePicker.cpp
+++ b/Engine/Plugins/Animation/ControlRig/Source/ControlRigEditor/Private/EditMode/SControlRigSpacePicker.cpp
@@ -276,21 +276,9 @@ FReply SControlRigSpacePicker::OnBakeControlsToNewSpaceButtonClicked()
 		Settings.TargetSpace = URigHierarchy::GetDefaultParentKey();
 
 		TRange<FFrameNumber> Range = Sequencer->GetFocusedMovieSceneSequence()->GetMovieScene()->GetPlaybackRange();
-		TArray<FFrameNumber> Keys;
-		TArray < FKeyHandle> KeyHandles;
 
-		Settings.StartFrame = Range.GetLowerBoundValue();
-		Settings.EndFrame = Range.GetUpperBoundValue();
-		if (Keys.Num() > 0)
-		{
-			int32 Index = Algo::LowerBound(Keys, CurrentTime);
-			if (Index >= 0 && Index < (Keys.Num() - 1))
-			{
-				Settings.StartFrame = Keys[Index];
-				Settings.EndFrame = Keys[Index + 1];
-
-			}
-		}
+		Settings.Settings.StartFrame = Range.GetLowerBoundValue();
+		Settings.Settings.EndFrame = Range.GetUpperBoundValue();
 
 		TSharedRef<SRigSpacePickerBakeWidget> BakeWidget =
 			SNew(SRigSpacePickerBakeWidget)
@@ -301,14 +289,7 @@ FReply SControlRigSpacePicker::OnBakeControlsToNewSpaceButtonClicked()
 			.GetControlCustomization(this, &SControlRigSpacePicker::HandleGetControlElementCustomization)
 			.OnBake_Lambda([Sequencer, ControlRig, TickResolution](URigHierarchy* InHierarchy, TArray<FRigElementKey> InControls, FRigSpacePickerBakeSettings InSettings)
 		{
-			TArray<FFrameNumber> Frames;
-
-			const FFrameRate& FrameRate = Sequencer->GetFocusedDisplayRate();
-			FFrameNumber FrameRateInFrameNumber = TickResolution.AsFrameNumber(FrameRate.AsInterval());
-			for (FFrameNumber& Frame = InSettings.StartFrame; Frame <= InSettings.EndFrame; Frame += FrameRateInFrameNumber)
-			{
-				Frames.Add(Frame);
-			}
+	
 			FScopedTransaction Transaction(LOCTEXT("BakeControlToSpace", "Bake Control In Space"));
 			for (const FRigElementKey& ControlKey : InControls)
 			{
@@ -317,7 +298,7 @@ FReply SControlRigSpacePicker::OnBakeControlsToNewSpaceButtonClicked()
 				if (SpaceChannelAndSection.SpaceChannel)
 				{
 					FControlRigSpaceChannelHelpers::SequencerBakeControlInSpace(ControlRig, Sequencer, SpaceChannelAndSection.SpaceChannel, SpaceChannelAndSection.SectionToKey,
-						Frames, InHierarchy, ControlKey, InSettings);
+						InHierarchy, ControlKey, InSettings);
 				}
 			}
 			return FReply::Handled();
diff --git a/Engine/Plugins/Animation/ControlRig/Source/ControlRigEditor/Private/Editor/Constraints/SConstraintsWidget.cpp b/Engine/Plugins/Animation/ControlRig/Source/ControlRigEditor/Private/Editor/Constraints/SConstraintsWidget.cpp
--- a/Engine/Plugins/Animation/ControlRig/Source/ControlRigEditor/Private/Editor/Constraints/SConstraintsWidget.cpp
+++ b/Engine/Plugins/Animation/ControlRig/Source/ControlRigEditor/Private/Editor/Constraints/SConstraintsWidget.cpp
@@ -26,6 +26,8 @@
 #include "Tools/BakingHelper.h"
 #include "MovieSceneToolHelpers.h"
 #include "Styling/SlateIconFinder.h"
+#include "Widgets/Input/NumericTypeInterface.h"
+#include "FrameNumberDetailsCustomization.h"
 
 #define LOCTEXT_NAMESPACE "SConstraintsWidget"
 
@@ -109,7 +111,7 @@ TArray<AActor*> GetCurrentSelection()
 	});	
 }
 
-TWeakPtr<ISequencer> GetSequencerChecked()
+static TWeakPtr<ISequencer> GetSequencerChecked()
 {
 	const TWeakPtr<ISequencer> WeakSequencer = FBakingHelper::GetSequencer();
 	if (!WeakSequencer.IsValid() || !WeakSequencer.Pin()->GetFocusedMovieSceneSequence())
@@ -156,7 +158,6 @@ void SDroppableConstraintItem::Construct(
 	.Padding(FMargin(8.f, 2.f, 12.f, 2.f))
 	[
 		SNew(SOverlay)
-
 		+SOverlay::Slot()
 		[
 			SNew(SBorder)
@@ -453,7 +454,6 @@ TSharedRef<ITableRow> SConstraintsCreationWidget::OnGenerateWidgetForItem(
 	const TSharedRef<STableViewBase>& OwnerTable)
 {
 	return SNew(STableRow<ItemSharedPtr>, OwnerTable)
-	.Style(&FAppStyle::Get(), "PlacementBrowser.PlaceableItemRow")
 	[
 		SNew(SDroppableConstraintItem, InItem.ToSharedRef(), SharedThis(this))
 	];
@@ -494,7 +494,13 @@ void SEditableConstraintItem::Construct(
 		}
 		
 		const FConstraintsManagerController& Controller = FConstraintsManagerController::Get(World);
-		return Controller.GetConstraint(ConstraintItem->Name);
+		UTickableConstraint* Constraint = ConstraintItem->Constraint.Get();
+		int32 Index = Controller.GetConstraintsArray().Find(Constraint);
+		if (Index == INDEX_NONE)
+		{
+			return nullptr;
+		}
+		return Constraint;
 	};
 	TWeakObjectPtr<UTickableConstraint> Constraint = GetConstraint();
 
@@ -503,10 +509,10 @@ void SEditableConstraintItem::Construct(
 	
 	// labels
 	FString ParentLabel(TEXT("undefined")), ChildLabel(TEXT("undefined"));
-	if (!InItem->Label.IsEmpty())
+	if (!InItem->GetLabel().IsEmpty())
 	{
-		ChildLabel = InItem->Label;
-		InItem->Label.Split(TEXT("."), &ParentLabel, &ChildLabel);
+		ChildLabel = InItem->GetLabel();
+		InItem->GetLabel().Split(TEXT("."), &ParentLabel, &ChildLabel);
 	}
 	
 	FString ParentFullLabel = ParentLabel, ChildFullLabel = ChildLabel;
@@ -714,43 +720,17 @@ void SEditableConstraintItem::Construct(
 	];
 }
 
-/**
- * SConstraintsEditionWidget
- */
-
-void SConstraintsEditionWidget::Construct(const FArguments& InArgs)
-{
-	ChildSlot
-	[
-		SNew( SVerticalBox )
-
-		+ SVerticalBox::Slot()
-		.Padding(FMargin(0.0f, 3.f))
-		[
-			SNew(SOverlay)
-
-			+ SOverlay::Slot()
-			[
-				SAssignNew(ListView, ConstraintItemListView)
-				.SelectionMode(ESelectionMode::Single)
-				.ListItemsSource( &ListItems )
-				.OnGenerateRow(this, &SConstraintsEditionWidget::OnGenerateWidgetForItem)
-				.OnContextMenuOpening(this, &SConstraintsEditionWidget::CreateContextMenu)
-				.OnMouseButtonDoubleClick(this, &SConstraintsEditionWidget::OnItemDoubleClicked)
-			]
-		]
-	];
+//////////////////////////////////////////////////////////////
+/// FBaseConstraintListWidget
+///////////////////////////////////////////////////////////
 
-	RefreshConstraintList();
-	RegisterSelectionChanged();
-}
 
-SConstraintsEditionWidget::~SConstraintsEditionWidget()
+FBaseConstraintListWidget::~FBaseConstraintListWidget()
 {
 	UnregisterSelectionChanged();
 }
 
-void SConstraintsEditionWidget::RegisterSelectionChanged()
+void FBaseConstraintListWidget::RegisterSelectionChanged()
 {
 	FLevelEditorModule& LevelEditor = FModuleManager::GetModuleChecked<FLevelEditorModule>("LevelEditor");
 	FLevelEditorModule::FActorSelectionChangedEvent& ActorSelectionChangedEvent = LevelEditor.OnActorSelectionChanged();
@@ -763,10 +743,10 @@ void SConstraintsEditionWidget::RegisterSelectionChanged()
 	}
 
 	// register
-	ActorSelectionChangedEvent.AddRaw( this, &SConstraintsEditionWidget::OnActorSelectionChanged );
+	ActorSelectionChangedEvent.AddRaw(this, &FBaseConstraintListWidget::OnActorSelectionChanged);
 }
 
-void SConstraintsEditionWidget::UnregisterSelectionChanged()
+void FBaseConstraintListWidget::UnregisterSelectionChanged()
 {
 	if (OnSelectionChangedHandle.IsValid())
 	{
@@ -776,6 +756,158 @@ void SConstraintsEditionWidget::UnregisterSelectionChanged()
 	}
 }
 
+void FBaseConstraintListWidget::PostUndo(bool bSuccess)
+{
+	InvalidateConstraintList();
+}
+
+void FBaseConstraintListWidget::PostRedo(bool bSuccess)
+{
+	InvalidateConstraintList();
+}
+
+void FBaseConstraintListWidget::InvalidateConstraintList()
+{
+	bNeedsRefresh = true;
+
+	// NOTE: this is a hack to disable the picker if the selection changes. The picker should trigger the escape key
+	// but this is an editor wide change
+	static const FActorPickerModeModule& ActorPickerMode = FModuleManager::Get().GetModuleChecked<FActorPickerModeModule>("ActorPickerMode");
+	if (ActorPickerMode.IsInActorPickingMode())
+	{
+		ActorPickerMode.EndActorPickingMode();
+	}
+}
+
+int32 FBaseConstraintListWidget::RefreshConstraintList()
+{
+	// get constraints
+	UWorld* World = GetCurrentWorld();
+	if (!IsValid(World))
+	{
+		return 0;
+	}
+
+	const TArray<AActor*> Selection = GetCurrentSelection();
+
+	const bool bIsConstraintsActor = Selection.Num() == 1 && Selection[0]->IsA<AConstraintsActor>();
+
+	TArray< TObjectPtr<UTickableConstraint> > Constraints;
+	if (bIsConstraintsActor)
+	{
+		const FConstraintsManagerController& Controller = FConstraintsManagerController::Get(World);
+		static constexpr bool bSorted = true;
+		Constraints = Controller.GetAllConstraints(bSorted);
+	}
+	else
+	{
+		for (const AActor* Actor : Selection)
+		{
+			FTransformConstraintUtils::GetParentConstraints(World, Actor, Constraints);
+		}
+	}
+
+	// rebuild item list
+	ListItems.Empty();
+
+	const UEnum* ETransformConstraintTypeEnum = StaticEnum<ETransformConstraintType>();
+	for (const TObjectPtr<UTickableConstraint>& Constraint : Constraints)
+	{
+		const int8 Type = FConstraintInfo::GetType(Constraint->GetClass());
+		if (ETransformConstraintTypeEnum->IsValidEnumValue(Type))
+		{
+			const ETransformConstraintType ConstraintType = static_cast<ETransformConstraintType>(Type);
+			ListItems.Emplace(FEditableConstraintItem::Make(Constraint, ConstraintType));
+		}
+	}
+
+	// refresh tree
+	ListView->RequestListRefresh();
+
+	return ListItems.Num();
+}
+
+void FBaseConstraintListWidget::OnActorSelectionChanged(const TArray<UObject*>& NewSelection, bool bForceRefresh)
+{
+	// NOTE we use this delegate to trigger an tree update, however, control actors are not selected as they are 
+	// Temporary Editor Actors so NewSelection won't contain the controls
+	InvalidateConstraintList();
+}
+
+
+//////////////////////////////////////////////////////////////
+/// SConstraintsEditionWidget
+///////////////////////////////////////////////////////////
+
+
+void SConstraintsEditionWidget::Construct(const FArguments& InArgs)
+{
+
+	ChildSlot
+	[
+		SNew(SBorder)
+		.Visibility(EVisibility::Visible)
+		.BorderImage(FAppStyle::GetBrush("Menu.Background"))
+		[
+
+			SNew( SVerticalBox )
+
+			+ SVerticalBox::Slot()
+			.AutoHeight()
+			.Padding(FMargin(0.0f, 0.f))
+			[
+				SNew(SOverlay)
+
+				+ SOverlay::Slot()
+				[
+					SAssignNew(ListView, ConstraintItemListView)
+					.SelectionMode(ESelectionMode::Single)
+					.ListItemsSource( &ListItems )
+					.OnGenerateRow(this, &SConstraintsEditionWidget::OnGenerateWidgetForItem)
+					.OnContextMenuOpening(this, &SConstraintsEditionWidget::CreateContextMenu)
+					.OnMouseButtonDoubleClick(this, &SConstraintsEditionWidget::OnItemDoubleClicked)
+				]
+
+			]
+	
+			+ SVerticalBox::Slot()
+				.AutoHeight()
+				.Padding(FMargin(0.0f, 0.f))
+				[
+		
+					SNew(SHorizontalBox)
+					+ SHorizontalBox::Slot()
+					.FillWidth(1.f)
+					.HAlign(HAlign_Fill)
+					[
+						SNew(SSpacer)
+					]
+		
+					+ SHorizontalBox::Slot()
+					.AutoWidth()
+					.VAlign(VAlign_Center)
+					.HAlign(HAlign_Right)
+					.Padding(1.f)
+					[
+						SNew(SButton)
+						.ButtonStyle(FAppStyle::Get(), "FlatButton.Default")
+						.Text(LOCTEXT("BakeButton", "Bake..."))
+						.OnClicked(this, &SConstraintsEditionWidget::OnBakeClicked)
+						.IsEnabled_Lambda([this]()
+						{
+							return (ListItems.Num() > 0);
+						})
+						.ToolTipText(LOCTEXT("BakeConstraintButtonToolTip", "Bake selected contraints"))
+					]
+				]
+		]
+	];
+
+	RefreshConstraintList();
+	RegisterSelectionChanged();
+}
+
+
 void SConstraintsEditionWidget::Tick(const FGeometry& AllottedGeometry, const double InCurrentTime, const float InDeltaTime)
 {
 	SCompoundWidget::Tick(AllottedGeometry, InCurrentTime, InDeltaTime);
@@ -787,15 +919,16 @@ void SConstraintsEditionWidget::Tick(const FGeometry& AllottedGeometry, const do
 	}
 }
 
+
 TSharedRef<ITableRow> SConstraintsEditionWidget::OnGenerateWidgetForItem(
 	ItemSharedPtr InItem,
 	const TSharedRef<STableViewBase>& OwnerTable)
 {
 	return SNew(STableRow<ItemSharedPtr>, OwnerTable)
-	.Style(&FAppStyle::Get(), "PlacementBrowser.PlaceableItemRow")
-	[
-		SNew(SEditableConstraintItem, InItem.ToSharedRef(), SharedThis(this))
-	];
+		.Style(&FAppStyle::Get(), "PlacementBrowser.PlaceableItemRow")
+		[
+			SNew(SEditableConstraintItem, InItem.ToSharedRef(), SharedThis(this))
+		];
 }
 
 bool SConstraintsEditionWidget::CanMoveUp(const TSharedPtr<FEditableConstraintItem>& Item) const
@@ -839,8 +972,8 @@ void SConstraintsEditionWidget::MoveItemUp(const TSharedPtr<FEditableConstraintI
 		const FConstraintsManagerController& Controller = FConstraintsManagerController::Get(World);
 
 		// the current item needs to tick before the previous item 
-		const FName BeforeName(ListItems[Index]->Name.ToString());
-		const FName AfterName(ListItems[Index-1]->Name.ToString());
+		const FName BeforeName(ListItems[Index]->GetName().ToString());
+		const FName AfterName(ListItems[Index-1]->GetName().ToString());
 		Controller.SetConstraintsDependencies(BeforeName, AfterName);
 		
 		RefreshConstraintList();
@@ -866,8 +999,8 @@ void SConstraintsEditionWidget::MoveItemDown(const TSharedPtr<FEditableConstrain
 		const FConstraintsManagerController& Controller = FConstraintsManagerController::Get(World);
 
 		// the current item needs to tick after the next item 
-		const FName BeforeName(ListItems[Index+1]->Name.ToString());
-		const FName AfterName(ListItems[Index]->Name.ToString());
+		const FName BeforeName(ListItems[Index+1]->GetName().ToString());
+		const FName AfterName(ListItems[Index]->GetName().ToString());
 		Controller.SetConstraintsDependencies(BeforeName, AfterName);
 		
 		RefreshConstraintList();
@@ -886,88 +1019,25 @@ void SConstraintsEditionWidget::RemoveItem(const TSharedPtr<FEditableConstraintI
 	
 	const FConstraintsManagerController& Controller = FConstraintsManagerController::Get(World);
 	
-	const FName ConstraintName(Item->Name.ToString());
-	Controller.RemoveConstraint(ConstraintName);
+	Controller.RemoveConstraint(Item->Constraint.Get());
 
 	RefreshConstraintList();
 }
 
-void SConstraintsEditionWidget::PostUndo(bool bSuccess)
-{
-	InvalidateConstraintList();
-}
-
-void SConstraintsEditionWidget::PostRedo(bool bSuccess)
-{
-	InvalidateConstraintList();
-}
-
-void SConstraintsEditionWidget::InvalidateConstraintList()
-{
-	bNeedsRefresh = true;
-
-	// NOTE: this is a hack to disable the picker if the selection changes. The picker should trigger the escape key
-	// but this is an editor wide change
-	static const FActorPickerModeModule& ActorPickerMode = FModuleManager::Get().GetModuleChecked<FActorPickerModeModule>("ActorPickerMode");
-	if (ActorPickerMode.IsInActorPickingMode())
-	{
-		ActorPickerMode.EndActorPickingMode();
-	}
-}
-
-int32 SConstraintsEditionWidget::RefreshConstraintList()
+FReply SConstraintsEditionWidget::OnBakeClicked()
 {
-	// get constraints
-	UWorld* World = GetCurrentWorld();
-	if (!IsValid(World))
-	{
-		return 0;
-	}
-	
-	const TArray<AActor*> Selection = GetCurrentSelection();
-	
-	const bool bIsConstraintsActor = Selection.Num() == 1 && Selection[0]->IsA<AConstraintsActor>();
-	
-	TArray< TObjectPtr<UTickableConstraint> > Constraints;
-	if (bIsConstraintsActor)
-	{
-		const FConstraintsManagerController& Controller = FConstraintsManagerController::Get(World);
-		static constexpr bool bSorted = true;
-		Constraints = Controller.GetAllConstraints(bSorted);
-	}
-	else
-	{
-		for (const AActor* Actor : Selection)
-		{
-			FTransformConstraintUtils::GetParentConstraints(World, Actor, Constraints);
-		}
-	}
-	
-	// rebuild item list
-	ListItems.Empty();
-	
-	const UEnum* ETransformConstraintTypeEnum = StaticEnum<ETransformConstraintType>();
-	for (const TObjectPtr<UTickableConstraint>& Constraint: Constraints)
+	UTickableTransformConstraint* InConstraint = nullptr;
+	TWeakPtr<ISequencer> WeakSequencer = GetSequencerChecked();
+	if (!WeakSequencer.IsValid() || ListItems.Num() < 1)
 	{
-		const int8 Type = FConstraintInfo::GetType(Constraint->GetClass()); 
-		if (ETransformConstraintTypeEnum->IsValidEnumValue(Type))
-		{
-			const ETransformConstraintType ConstraintType = static_cast<ETransformConstraintType>(Type);
-			ListItems.Emplace( FEditableConstraintItem::Make(Constraint->GetFName(), ConstraintType, Constraint->GetLabel()) );
-		}
+		return FReply::Unhandled();
 	}
+	TSharedPtr<ISequencer> Sequencer = WeakSequencer.Pin();
+	TSharedRef<SConstraintBakeWidget> BakeWidget =
+		SNew(SConstraintBakeWidget)
+		.Sequencer(Sequencer);
 
-	// refresh tree
-	ListView->RequestListRefresh();
-	
-	return ListItems.Num();
-}
-
-void SConstraintsEditionWidget::OnActorSelectionChanged(const TArray<UObject*>& NewSelection, bool bForceRefresh)
-{
-	// NOTE we use this delegate to trigger an tree update, however, control actors are not selected as they are 
-	// Temporary Editor Actors so NewSelection won't contain the controls
-	InvalidateConstraintList();
+	return BakeWidget->OpenDialog(true);
 }
 
 TSharedPtr<SWidget> SConstraintsEditionWidget::CreateContextMenu()
@@ -991,8 +1061,9 @@ TSharedPtr<SWidget> SConstraintsEditionWidget::CreateContextMenu()
 	}
 	
 	const FConstraintsManagerController& Controller = FConstraintsManagerController::Get(World);
-	UTickableConstraint* Constraint = Controller.GetConstraint(ListItems[Index]->Name);
-	if (!Constraint)
+	UTickableConstraint* Constraint = ListItems[Index]->Constraint.Get();
+	int32 ContIndex = Controller.GetConstraintsArray().Find(Constraint);
+	if (ContIndex == INDEX_NONE)
 	{
 		return SNullWidget::NullWidget;
 	}
@@ -1060,7 +1131,7 @@ TSharedPtr<SWidget> SConstraintsEditionWidget::CreateContextMenu()
 	// transform constraint options
 	if (UTickableTransformConstraint* TransformConstraint = Cast<UTickableTransformConstraint>(Constraint))
 	{
-		const FText ConstraintLabel = FText::FromName(ListItems[Index]->Name);
+		const FText ConstraintLabel = FText::FromName(ListItems[Index]->GetName());
 	
 		// baking (note that this will probably be moved)
 		MenuBuilder.BeginSection("BakeConstraint", LOCTEXT("BakeConstraintHeader", "Bake Constraint"));
@@ -1078,9 +1149,8 @@ TSharedPtr<SWidget> SConstraintsEditionWidget::CreateContextMenu()
 						{
 							return;
 						}
-						
 						FScopedTransaction Transaction(LOCTEXT("BakeConstraint", "Bake Constraint"));
-						FConstraintBaker::Bake(World, TransformConstraint, WeakSequencer.Pin(), TOptional<TArray<FFrameNumber>>());
+						FConstraintBaker::Bake(World, TransformConstraint, WeakSequencer.Pin(), TOptional< FBakingAnimationKeySettings>(), TOptional<TArray<FFrameNumber>>());
 
 					}
 				})),
@@ -1148,13 +1218,376 @@ void SConstraintsEditionWidget::OnItemDoubleClicked(ItemSharedPtr InItem)
 	}
 	
 	const FConstraintsManagerController& Controller = FConstraintsManagerController::Get(World);
-	UTickableConstraint* Constraint = Controller.GetConstraint(InItem->Name);
-	if (!Constraint)
+	UTickableConstraint* Constraint = InItem->Constraint.Get();
+	int32 ContIndex = Controller.GetConstraintsArray().Find(Constraint);
+	if (ContIndex != INDEX_NONE)
 	{
-		return;
+		Constraint->SetActive(!Constraint->Active);
+	}
+}
+
+//////////////////////////////////////////////////////////////
+/// SConstraintBakeWidget
+///////////////////////////////////////////////////////////
+
+TOptional<FBakingAnimationKeySettings> SConstraintBakeWidget::BakeConstraintSettings;
+
+void SConstraintBakeWidget::Construct(const FArguments& InArgs)
+{
+	check(InArgs._Sequencer);
+	Sequencer = InArgs._Sequencer;
+
+	if (BakeConstraintSettings.IsSet() == false)
+	{
+		BakeConstraintSettings = FBakingAnimationKeySettings();
+		const FFrameRate TickResolution = Sequencer->GetFocusedTickResolution();
+		const FFrameTime FrameTime = Sequencer->GetLocalTime().ConvertTo(TickResolution);
+		FFrameNumber CurrentTime = FrameTime.GetFrame();
+
+		TRange<FFrameNumber> Range = Sequencer->GetFocusedMovieSceneSequence()->GetMovieScene()->GetPlaybackRange();
+		TArray<FFrameNumber> Keys;
+		TArray < FKeyHandle> KeyHandles;
+
+		BakeConstraintSettings.GetValue().StartFrame = Range.GetLowerBoundValue();
+		BakeConstraintSettings.GetValue().EndFrame = Range.GetUpperBoundValue();
+	}
+
+
+	Settings = MakeShared<TStructOnScope<FBakingAnimationKeySettings>>();
+	Settings->InitializeAs<FBakingAnimationKeySettings>(BakeConstraintSettings.GetValue());
+
+	FStructureDetailsViewArgs StructureViewArgs;
+	StructureViewArgs.bShowObjects = true;
+	StructureViewArgs.bShowAssets = true;
+	StructureViewArgs.bShowClasses = true;
+	StructureViewArgs.bShowInterfaces = true;
+
+	FDetailsViewArgs ViewArgs;
+	ViewArgs.bAllowSearch = false;
+	ViewArgs.bHideSelectionTip = false;
+	ViewArgs.bShowObjectLabel = false;
+
+	FPropertyEditorModule& PropertyEditor = FModuleManager::Get().LoadModuleChecked<FPropertyEditorModule>(TEXT("PropertyEditor"));
+
+	DetailsView = PropertyEditor.CreateStructureDetailView(ViewArgs, StructureViewArgs, TSharedPtr<FStructOnScope>());
+	TSharedPtr<INumericTypeInterface<double>> NumericTypeInterface = Sequencer->GetNumericTypeInterface();
+	DetailsView->GetDetailsView()->RegisterInstancedCustomPropertyTypeLayout("FrameNumber",
+		FOnGetPropertyTypeCustomizationInstance::CreateLambda([=]() {return MakeShared<FFrameNumberDetailsCustomization>(NumericTypeInterface); }));
+	DetailsView->SetStructureData(Settings);
+
+
+	ChildSlot
+		[
+			SNew(SBorder)
+			.Visibility(EVisibility::Visible)
+			[
+			SNew(SVerticalBox)
+			+ SVerticalBox::Slot()
+			.Padding(FMargin(0.0f, 0.f))
+			[
+				SAssignNew(ListView, ConstraintItemListView)
+				.SelectionMode(ESelectionMode::Multi)
+				.ListItemsSource(&ListItems)
+				.OnGenerateRow(this, &SConstraintBakeWidget::OnGenerateWidgetForItem)
+
+			]
+			
+				+SVerticalBox::Slot()
+				.AutoHeight()
+				.Padding(0.f, 8.f, 0.f, 0.f)
+				[
+					DetailsView->GetWidget().ToSharedRef()
+				]
+
+				+ SVerticalBox::Slot()
+				.AutoHeight()
+				.Padding(0.f, 16.f, 0.f, 8.f)
+				[
+					SNew(SHorizontalBox)
+
+					+ SHorizontalBox::Slot()
+					.FillWidth(1.f)
+					[
+						SNew(SSpacer)
+					]
+
+				+ SHorizontalBox::Slot()
+				.AutoWidth()
+				.HAlign(HAlign_Left)
+				.VAlign(VAlign_Center)
+				.Padding(8.f, 0.f, 0.f, 0.f)
+				[
+					SNew(SButton)
+					.HAlign(HAlign_Center)
+					.ContentPadding(FAppStyle::GetMargin("StandardDialog.ContentPadding"))
+					.Text(LOCTEXT("OK", "OK"))
+					.OnClicked_Lambda([this, InArgs]()
+						{
+							BakeSelected();
+							CloseDialog();
+							return FReply::Handled();
+
+						})
+						.IsEnabled_Lambda([this]()
+						{
+							return (ListItems.Num() > 0 && Settings.IsValid());
+						})
+			]
+
+			+ SHorizontalBox::Slot()
+			.AutoWidth()
+			.HAlign(HAlign_Left)
+			.VAlign(VAlign_Center)
+			.Padding(8.f, 0.f, 16.f, 0.f)
+			[
+				SNew(SButton)
+				.HAlign(HAlign_Center)
+				.ContentPadding(FAppStyle::GetMargin("StandardDialog.ContentPadding"))
+				.Text(LOCTEXT("Cancel", "Cancel"))
+				.OnClicked_Lambda([this]()
+					{
+						CloseDialog();
+					return FReply::Handled();
+					})
+				]
+			]
+		]
+	];
+	RefreshConstraintList();
+	RegisterSelectionChanged();
+}
+
+TSharedRef<ITableRow> SConstraintBakeWidget::OnGenerateWidgetForItem(
+	ItemSharedPtr InItem,
+	const TSharedRef<STableViewBase>& OwnerTable)
+{
+	return SNew(STableRow<ItemSharedPtr>, OwnerTable)
+		[
+			SNew(SBakeConstraintItem, InItem.ToSharedRef(), SharedThis(this))
+		];
+		
+}
+
+void SConstraintBakeWidget::Tick(const FGeometry& AllottedGeometry, const double InCurrentTime, const float InDeltaTime)
+{
+	SCompoundWidget::Tick(AllottedGeometry, InCurrentTime, InDeltaTime);
+
+	if (bNeedsRefresh)
+	{
+		RefreshConstraintList();
+		bNeedsRefresh = false;
+	}
+}
+
+//for the bake we auto select averything when selection or something chagnes
+int32 SConstraintBakeWidget::RefreshConstraintList()
+{
+	int32 Num = FBaseConstraintListWidget::RefreshConstraintList();
+	ListView->SetItemSelection(ListItems, true, ESelectInfo::Direct);
+	return Num;
+}
+
+void SConstraintBakeWidget::BakeSelected()
+{
+	TArray<TSharedPtr<FEditableConstraintItem>> Selection = ListView->GetSelectedItems();
+	TArray<UTickableTransformConstraint*> Constraints;
+	for (TSharedPtr<FEditableConstraintItem>& Item : Selection)
+	{
+		if (UTickableTransformConstraint* TransformConstraint = Cast<UTickableTransformConstraint>(Item->Constraint.Get()))
+		{
+			Constraints.Add(TransformConstraint);
+		}
+	}
+	UWorld* World = GetCurrentWorld();
+
+	FBakingAnimationKeySettings* BakeSettings = Settings->Get();
+	
+	FScopedTransaction Transaction(LOCTEXT("BakeConstraints", "Bake Constraints"));
+	FConstraintBaker::BakeMultiple(World,Constraints, Sequencer,*BakeSettings );
+	BakeConstraintSettings = *BakeSettings;
+	InvalidateConstraintList();
+}
+
+class SConstraintBakedWidgetWindow : public SWindow
+{
+};
+
+
+
+FReply SConstraintBakeWidget::OpenDialog(bool bModal)
+{
+	check(!DialogWindow.IsValid());
+
+	const FVector2D CursorPos = FSlateApplication::Get().GetCursorPos();
+
+	TSharedRef<SConstraintBakedWidgetWindow> Window = SNew(SConstraintBakedWidgetWindow)
+		.Title(LOCTEXT("SConstraintsBakeWidgetTitle", "Bake Constraints"))
+		.CreateTitleBar(true)
+		.Type(EWindowType::Normal)
+		.SizingRule(ESizingRule::Autosized)
+		.ScreenPosition(CursorPos)
+		.FocusWhenFirstShown(true)
+		.ActivationPolicy(EWindowActivationPolicy::FirstShown)
+		[
+			AsShared()
+		];
+
+	Window->SetWidgetToFocusOnActivate(AsShared());
+
+	DialogWindow = Window;
+
+	Window->MoveWindowTo(CursorPos);
+
+	if (bModal)
+	{
+		GEditor->EditorAddModalWindow(Window);
+	}
+	else
+	{
+		FSlateApplication::Get().AddWindow(Window);
+	}
+
+	return FReply::Handled();
+}
+
+void SConstraintBakeWidget::CloseDialog()
+{
+	if (DialogWindow.IsValid())
+	{
+		DialogWindow.Pin()->RequestDestroyWindow();
+		DialogWindow.Reset();
+	}
+}
+
+//////////////////////////////////////////////////////////////
+/// SBakeConstraintItem
+///////////////////////////////////////////////////////////
+
+
+void SBakeConstraintItem::Construct(
+	const FArguments& InArgs,
+	const TSharedPtr<FEditableConstraintItem>& InItem,
+	TSharedPtr<SConstraintBakeWidget> InConstraintsWidget)
+{
+	ConstraintItem = InItem;
+	ConstraintsWidget = InConstraintsWidget;
+
+	static const FSlateBrush* RoundedBoxBrush = FControlRigEditorStyle::Get().GetBrush(TEXT("ControlRig.SpacePicker.RoundedRect"));
+
+	// enum to string
+	const uint8 ConstraintType = static_cast<uint8>(InItem->Type);
+
+
+	// constraint
+	auto GetConstraint = [this]() -> UTickableConstraint*
+	{
+		UWorld* World = GetCurrentWorld();
+		if (!IsValid(World))
+		{
+			return nullptr;
+		}
+
+		const FConstraintsManagerController& Controller = FConstraintsManagerController::Get(World);
+		UTickableConstraint* Constraint = ConstraintItem->Constraint.Get();
+		int32 Index = Controller.GetConstraintsArray().Find(Constraint);
+		if (Index == INDEX_NONE)
+		{
+			return nullptr;
+		}
+		return Constraint;
+	};
+	TWeakObjectPtr<UTickableConstraint> Constraint = GetConstraint();
+
+	// sequencer
+	TWeakPtr<ISequencer> WeakSequencer = GetSequencerChecked();
+
+	// labels
+	FString ParentLabel(TEXT("undefined")), ChildLabel(TEXT("undefined"));
+	if (!InItem->GetLabel().IsEmpty())
+	{
+		ChildLabel = InItem->GetLabel();
+		InItem->GetLabel().Split(TEXT("."), &ParentLabel, &ChildLabel);
 	}
 
-	Constraint->SetActive(!Constraint->Active);
+	FString ParentFullLabel = ParentLabel, ChildFullLabel = ChildLabel;
+	if (IsValid(Constraint.Get()))
+	{
+		Constraint->GetFullLabel().Split(TEXT("."), &ParentFullLabel, &ChildFullLabel);
+	}
+
+	// widgets
+	ChildSlot
+		.Padding(FMargin(8.f, 2.f, 12.f, 2.f))
+		[
+			SNew(SHorizontalBox)
+
+			+ SHorizontalBox::Slot()
+			.AutoWidth()
+			.VAlign(VAlign_Center)
+			.HAlign(HAlign_Fill)
+			.Padding(0)
+			[
+			SNew(SBorder)
+			.Padding(FMargin(5.0, 2.0, 5.0, 2.0))
+			.BorderImage(RoundedBoxBrush)
+			.BorderBackgroundColor_Lambda([Constraint]()
+				{
+					return FStyleColors::Transparent;
+				})	
+				.Content()
+				[
+					SNew(SHorizontalBox)
+
+					// constraint icon
+					+ SHorizontalBox::Slot()
+					.AutoWidth()
+					.VAlign(VAlign_Center)
+					.HAlign(HAlign_Left)
+					.Padding(FMargin(0.f, 0.f, 3.f, 0.f))
+					[
+						SNew(SImage)
+						.Image(FConstraintInfo::GetBrush(ConstraintType))
+						.ColorAndOpacity(FSlateColor::UseForeground())
+					]
+
+				// constraint name
+				+ SHorizontalBox::Slot()
+					.AutoWidth()
+					.VAlign(VAlign_Center)
+					.HAlign(HAlign_Left)
+					.Padding(0)
+					[
+						SNew(STextBlock)
+						.Text_Lambda([ParentLabel]()
+							{
+								return FText::FromString(ParentLabel);
+							})
+					.Font_Lambda([Constraint]()
+						{
+							if (!Constraint.IsValid() || !IsValid(Constraint.Get()))
+							{
+								return IDetailLayoutBuilder::GetDetailFont();
+							}
+							return Constraint->Active ? IDetailLayoutBuilder::GetDetailFont() : IDetailLayoutBuilder::GetDetailFontItalic();
+						})
+								.ToolTipText_Lambda([Constraint, ParentFullLabel, ChildFullLabel]()
+									{
+										if (!Constraint.IsValid() || !IsValid(Constraint.Get()))
+										{
+											return FText();
+										}
+
+						static constexpr TCHAR ToolTipFormat[] = TEXT("%s constraint between parent '%s' and child '%s'.");
+						const FString TypeLabel = Constraint->GetTypeLabel();
+						const FString FullLabel = FString::Printf(ToolTipFormat, *TypeLabel, *ParentFullLabel, *ChildFullLabel);
+						return FText::FromString(FullLabel);
+									})
+					]
+					]
+		]
+
+	
+		];
 }
 
 END_SLATE_FUNCTION_BUILD_OPTIMIZATION
diff --git a/Engine/Plugins/Animation/ControlRig/Source/ControlRigEditor/Private/Editor/Constraints/SConstraintsWidget.h b/Engine/Plugins/Animation/ControlRig/Source/ControlRigEditor/Private/Editor/Constraints/SConstraintsWidget.h
--- a/Engine/Plugins/Animation/ControlRig/Source/ControlRigEditor/Private/Editor/Constraints/SConstraintsWidget.h
+++ b/Engine/Plugins/Animation/ControlRig/Source/ControlRigEditor/Private/Editor/Constraints/SConstraintsWidget.h
@@ -6,12 +6,14 @@
 #include "Widgets/Views/SListView.h"
 #include "Constraint.h"
 #include "EditorUndoClient.h"
+#include "IStructureDetailsView.h"
+#include "BakingAnimationKeySettings.h"
+#include "ConstraintsManager.h"
 
 class AActor;
 class SConstraintsCreationWidget;
 class SConstraintsEditionWidget;
-class UTickableConstraint;
-
+class ISequencer;
 DECLARE_DELEGATE(FOnConstraintCreated);
 
 /**
@@ -147,22 +149,36 @@
 {
 public:
 	static TSharedRef<FEditableConstraintItem> Make(
-		const FName& InName,
-		ETransformConstraintType InType,
-		const FString& InLabel)
+		UTickableConstraint* InConstraint,
+		ETransformConstraintType InType)
 	{
-		return MakeShareable(new FEditableConstraintItem(InName, InType, InLabel));
+		return MakeShareable(new FEditableConstraintItem(InConstraint, InType));
 	}
 	
-	FName Name = NAME_None;
+	TWeakObjectPtr<UTickableConstraint> Constraint = nullptr;
 	ETransformConstraintType Type = ETransformConstraintType::Parent;
-	FString Label;
 
+	FName GetName() const
+	{
+		if (Constraint.IsValid())
+		{
+			return Constraint->GetFName();
+		}
+		return NAME_None;
+	}
+	FString GetLabel() const
+	{
+		if (Constraint.IsValid())
+		{
+			return Constraint->GetLabel();
+		}
+		return FString();
+	}
+
 private:
-	FEditableConstraintItem(const FName& InName, ETransformConstraintType InType, const FString& InLabel)
-		: Name(InName)
+	FEditableConstraintItem(UTickableConstraint* InConstraint, ETransformConstraintType InType)
+		: Constraint(InConstraint)
 		, Type(InType)
-		, Label(InLabel)
 	{}
 	FEditableConstraintItem() {}
 };
@@ -190,10 +206,55 @@
 };
 
 /**
+* Constraint List Shared by Edit Widget and Bake Widget
+*/
+
+class CONTROLRIGEDITOR_API FBaseConstraintListWidget : public FEditorUndoClient
+{
+public:
+
+	virtual ~FBaseConstraintListWidget() override;
+	/* FEditorUndoClient interface */
+	virtual void PostUndo(bool bSuccess);
+	virtual void PostRedo(bool bSuccess);
+	/* End FEditorUndoClient interface */
+
+	/** Invalidates the constraint list for further rebuild. */
+	void InvalidateConstraintList();
+
+	/** Rebuild the constraint list based on the current selection. */
+	virtual int32 RefreshConstraintList();
+
+	/** Triggers a constraint list invalidation when selection in the level viewport. */
+	void OnActorSelectionChanged(const TArray<UObject*>& NewSelection, bool bForceRefresh);
+
+protected:
+	/** Types */
+	using ItemSharedPtr = TSharedPtr<FEditableConstraintItem>;
+	using ConstraintItemListView = SListView< ItemSharedPtr >;
+
+
+	void RegisterSelectionChanged();
+	void UnregisterSelectionChanged();
+
+	/** List view that shows constraint types*/
+	TSharedPtr< ConstraintItemListView > ListView;
+
+	/** List of constraint types */
+	TArray< ItemSharedPtr > ListItems;
+
+	/** Boolean used to handle the items list refresh. */
+	bool bNeedsRefresh = false;
+
+	FDelegateHandle OnSelectionChangedHandle;
+
+};
+
+/**
  * SConstraintsEditionWidget
  */
 
-class CONTROLRIGEDITOR_API SConstraintsEditionWidget : public SCompoundWidget, public FEditorUndoClient
+class CONTROLRIGEDITOR_API SConstraintsEditionWidget : public SCompoundWidget, public FBaseConstraintListWidget
 {
 public:
 	SLATE_BEGIN_ARGS(SConstraintsEditionWidget)	{}
@@ -201,20 +262,10 @@
 
 	/** Constructs this widget with InArgs */
 	void Construct(const FArguments& InArgs);
-	virtual ~SConstraintsEditionWidget() override;
 
 	/** Override for SWidget::Tick */
 	virtual void Tick(const FGeometry& AllottedGeometry, const double InCurrentTime, const float InDeltaTime) override;
 
-	/** Invalidates the constraint list for further rebuild. */
-	void InvalidateConstraintList();
-	
-	/** Rebuild the constraint list based on the current selection. */
-	int32 RefreshConstraintList();
-
-	/** Triggers a constraint list invalidation when selection in the level viewport. */
-	void OnActorSelectionChanged(const TArray<UObject*>& NewSelection, bool bForceRefresh);
-
 	/**  */
 	bool CanMoveUp(const TSharedPtr<FEditableConstraintItem>& Item) const;
 	bool CanMoveDown(const TSharedPtr<FEditableConstraintItem>& Item) const;
@@ -226,34 +277,80 @@
 	/**  */
 	void RemoveItem(const TSharedPtr<FEditableConstraintItem>& Item);
 
-	/* FEditorUndoClient interface */
-	virtual void PostUndo(bool bSuccess);
-	virtual void PostRedo(bool bSuccess);
-	/* End FEditorUndoClient interface */
+private:
 
-private:
+	/** Generates a widget for the specified item */
+	TSharedRef<ITableRow> OnGenerateWidgetForItem(ItemSharedPtr InItem, const TSharedRef<STableViewBase>& OwnerTable);
+
 	/** Types */
 	using ItemSharedPtr = TSharedPtr<FEditableConstraintItem>;
 	using ConstraintItemListView = SListView< ItemSharedPtr >;
 
-	/** Generates a widget for the specified item */
-	TSharedRef<ITableRow> OnGenerateWidgetForItem(ItemSharedPtr InItem, const TSharedRef<STableViewBase>& OwnerTable);
-
 	/** @todo documentation. */
 	TSharedPtr< SWidget > CreateContextMenu();
 	void OnItemDoubleClicked(ItemSharedPtr InItem);
-	
-	void RegisterSelectionChanged();
-	void UnregisterSelectionChanged();
-	
-	/** List view that shows constraint types*/
-	TSharedPtr< ConstraintItemListView > ListView;
 
-	/** List of constraint types */
-	TArray< ItemSharedPtr > ListItems;
+	FReply OnBakeClicked();
+};
 
-	/** Boolean used to handle the items list refresh. */
-	bool bNeedsRefresh = false;
+/** Widget allowing baking of constraints */
+class CONTROLRIGEDITOR_API SConstraintBakeWidget : public SCompoundWidget, public FBaseConstraintListWidget
+{
+public:
 
-	FDelegateHandle OnSelectionChangedHandle;
+	SLATE_BEGIN_ARGS(SConstraintBakeWidget)
+		: _Sequencer(nullptr)
+	{}
+
+	SLATE_ARGUMENT(TSharedPtr<ISequencer>, Sequencer)
+	SLATE_END_ARGS()
+
+	void Construct(const FArguments& InArgs);
+	virtual ~SConstraintBakeWidget() override {}
+
+	FReply OpenDialog(bool bModal = true);
+	void CloseDialog();
+
+	//SWidget overrides
+	virtual void Tick(const FGeometry& AllottedGeometry, const double InCurrentTime, const float InDeltaTime) override;
+
+	//FBaseConstraintListWidget overrides
+	virtual int32 RefreshConstraintList() override;
+
+private:
+	void BakeSelected();
+
+	/** Generates a widget for the specified item */
+	TSharedRef<ITableRow> OnGenerateWidgetForItem(ItemSharedPtr InItem, const TSharedRef<STableViewBase>& OwnerTable);
+
+	TSharedPtr<ISequencer> Sequencer;
+	//static to be reused
+	static TOptional<FBakingAnimationKeySettings> BakeConstraintSettings;
+	//structonscope for details panel
+	TSharedPtr < TStructOnScope<FBakingAnimationKeySettings>> Settings;
+	TWeakPtr<SWindow> DialogWindow;
+	TSharedPtr<IStructureDetailsView> DetailsView;
+
 };
+
+/**
+ * SBakeConstraintItem
+ */
+
+class SBakeConstraintItem : public SCompoundWidget
+{
+public:
+	SLATE_BEGIN_ARGS(SBakeConstraintItem) {}
+	SLATE_END_ARGS()
+
+		/** Constructs this widget with InArgs and the actual tree item. */
+		void Construct(
+			const FArguments& InArgs,
+			const TSharedPtr<FEditableConstraintItem>& InItem,
+			TSharedPtr<SConstraintBakeWidget> InConstraintsWidget);
+
+private:
+	/** TSharedPtr to the tree item. */
+	TSharedPtr<FEditableConstraintItem> ConstraintItem;
+	TWeakPtr<SConstraintBakeWidget> ConstraintsWidget;
+};
diff --git a/Engine/Plugins/Animation/ControlRig/Source/ControlRigEditor/Private/Editor/SRigSpacePickerWidget.cpp b/Engine/Plugins/Animation/ControlRig/Source/ControlRigEditor/Private/Editor/SRigSpacePickerWidget.cpp
--- a/Engine/Plugins/Animation/ControlRig/Source/ControlRigEditor/Private/Editor/SRigSpacePickerWidget.cpp
+++ b/Engine/Plugins/Animation/ControlRig/Source/ControlRigEditor/Private/Editor/SRigSpacePickerWidget.cpp
@@ -186,6 +186,10 @@ void SRigSpacePickerWidget::Construct(const FArguments& InArgs)
 				.ButtonStyle(FAppStyle::Get(), "FlatButton.Default")
 				.Text(LOCTEXT("BakeButton", "Bake..."))
 				.OnClicked(InArgs._OnBakeButtonClicked)
+				.IsEnabled_Lambda([this]()
+				{
+					return (ControlKeys.Num() > 0);
+				})
 				.ToolTipText(LOCTEXT("BakeButtonToolTip", "Allows to bake the animation of one or more controls to a single space."))
 			];
 		}
@@ -1121,6 +1125,8 @@ void SRigSpacePickerBakeWidget::Construct(const FArguments& InArgs)
 	Settings = MakeShared<TStructOnScope<FRigSpacePickerBakeSettings>>();
 	Settings->InitializeAs<FRigSpacePickerBakeSettings>();
 	*Settings = InArgs._Settings;
+	//always setting space to be parent as default, since stored space may not be available.
+	Settings->Get()->TargetSpace = URigHierarchy::GetDefaultParentKey();
 	Sequencer = InArgs._Sequencer;
 
 	FStructureDetailsViewArgs StructureViewArgs;
@@ -1187,7 +1193,7 @@ void SRigSpacePickerBakeWidget::Construct(const FArguments& InArgs)
 
 			+SVerticalBox::Slot()
 			.AutoHeight()
-			.Padding(0.f, 8.f, 0.f, 0.f)
+			.Padding(0.f, 0.f, 0.f, 0.f)
 			[
 				DetailsView->GetWidget().ToSharedRef()
 			]
diff --git a/Engine/Plugins/Animation/ControlRig/Source/ControlRigEditor/Private/Editor/SRigSpacePickerWidget.h b/Engine/Plugins/Animation/ControlRig/Source/ControlRigEditor/Private/Editor/SRigSpacePickerWidget.h
--- a/Engine/Plugins/Animation/ControlRig/Source/ControlRigEditor/Private/Editor/SRigSpacePickerWidget.h
+++ b/Engine/Plugins/Animation/ControlRig/Source/ControlRigEditor/Private/Editor/SRigSpacePickerWidget.h
@@ -205,7 +205,7 @@ public:
 
 private:
 
-	//FRigSpacePickerBakeSettings Settings;
+	//used for setting up the details
 	TSharedPtr<TStructOnScope<FRigSpacePickerBakeSettings>> Settings;
 
 	ISequencer* Sequencer;
diff --git a/Engine/Plugins/Animation/ControlRig/Source/ControlRigEditor/Private/Sequencer/ControlRigParameterTrackEditor.cpp b/Engine/Plugins/Animation/ControlRig/Source/ControlRigEditor/Private/Sequencer/ControlRigParameterTrackEditor.cpp
--- a/Engine/Plugins/Animation/ControlRig/Source/ControlRigEditor/Private/Sequencer/ControlRigParameterTrackEditor.cpp
+++ b/Engine/Plugins/Animation/ControlRig/Source/ControlRigEditor/Private/Sequencer/ControlRigParameterTrackEditor.cpp
@@ -83,6 +83,9 @@
 #include "Constraints/ControlRigTransformableHandle.h"
 #include "PropertyEditorModule.h"
 #include "Constraints/TransformConstraintChannelInterface.h"
+#include "BakingAnimationKeySettings.h"
+#include "FrameNumberDetailsCustomization.h"
+#include "Sequencer/ControlRigSequencerHelpers.h"
 
 #define LOCTEXT_NAMESPACE "FControlRigParameterTrackEditor"
 
@@ -4242,141 +4245,385 @@ void FControlRigParameterTrackEditor::SelectFKBonesToAnimate(UFKControlRig* Auto
 
 	//reconstruct all channel proxies TODO or not to do that is the question
 }
-bool FControlRigParameterTrackEditor::CollapseAllLayers(TSharedPtr<ISequencer>& SequencerPtr, UMovieSceneTrack* OwnerTrack, UMovieSceneControlRigParameterSection* ParameterSection, bool bKeyReduce, float Tolerance)
+
+//////////////////////////////////////////////////////////////
+/// SCollapseControlsWidget
+///////////////////////////////////////////////////////////
+
+/** Widget allowing collapsing of controls */
+class SCollapseControlsWidget : public SCompoundWidget
+{
+public:
+
+	SLATE_BEGIN_ARGS(SCollapseControlsWidget)
+		: _Sequencer(nullptr), _OwnerTrack(nullptr)
+	{}
+
+	SLATE_ARGUMENT(TSharedPtr<ISequencer>, Sequencer)
+	SLATE_ARGUMENT(UMovieSceneTrack*, OwnerTrack)
+	SLATE_END_ARGS()
+
+	void Construct(const FArguments& InArgs);
+	virtual ~SCollapseControlsWidget() override {}
+
+	FReply OpenDialog(bool bModal = true);
+	void CloseDialog();
+
+
+private:
+	void Collapse();
+
+	TSharedPtr<ISequencer> Sequencer;
+	TWeakObjectPtr<UMovieSceneTrack> OwnerTrack;
+	//static to be reused
+	static TOptional<FBakingAnimationKeySettings> CollapseControlsSettings;
+	//structonscope for details panel
+	TSharedPtr < TStructOnScope<FBakingAnimationKeySettings>> Settings;
+	TWeakPtr<SWindow> DialogWindow;
+	TSharedPtr<IStructureDetailsView> DetailsView;
+};
+
+
+TOptional<FBakingAnimationKeySettings> SCollapseControlsWidget::CollapseControlsSettings;
+
+void SCollapseControlsWidget::Construct(const FArguments& InArgs)
+{
+	check(InArgs._Sequencer);
+	Sequencer = InArgs._Sequencer;
+	OwnerTrack = InArgs._OwnerTrack;
+
+	if (CollapseControlsSettings.IsSet() == false)
+	{
+		CollapseControlsSettings = FBakingAnimationKeySettings();
+		const FFrameRate TickResolution = Sequencer->GetFocusedTickResolution();
+		const FFrameTime FrameTime = Sequencer->GetLocalTime().ConvertTo(TickResolution);
+		FFrameNumber CurrentTime = FrameTime.GetFrame();
+
+		TRange<FFrameNumber> Range = Sequencer->GetFocusedMovieSceneSequence()->GetMovieScene()->GetPlaybackRange();
+		TArray<FFrameNumber> Keys;
+		TArray < FKeyHandle> KeyHandles;
+
+		CollapseControlsSettings.GetValue().StartFrame = Range.GetLowerBoundValue();
+		CollapseControlsSettings.GetValue().EndFrame = Range.GetUpperBoundValue();
+	}
+
+
+	Settings = MakeShared<TStructOnScope<FBakingAnimationKeySettings>>();
+	Settings->InitializeAs<FBakingAnimationKeySettings>(CollapseControlsSettings.GetValue());
+
+	FStructureDetailsViewArgs StructureViewArgs;
+	StructureViewArgs.bShowObjects = true;
+	StructureViewArgs.bShowAssets = true;
+	StructureViewArgs.bShowClasses = true;
+	StructureViewArgs.bShowInterfaces = true;
+
+	FDetailsViewArgs ViewArgs;
+	ViewArgs.bAllowSearch = false;
+	ViewArgs.bHideSelectionTip = false;
+	ViewArgs.bShowObjectLabel = false;
+
+	FPropertyEditorModule& PropertyEditor = FModuleManager::Get().LoadModuleChecked<FPropertyEditorModule>(TEXT("PropertyEditor"));
+
+	DetailsView = PropertyEditor.CreateStructureDetailView(ViewArgs, StructureViewArgs, TSharedPtr<FStructOnScope>());
+	TSharedPtr<INumericTypeInterface<double>> NumericTypeInterface = Sequencer->GetNumericTypeInterface();
+	DetailsView->GetDetailsView()->RegisterInstancedCustomPropertyTypeLayout("FrameNumber",
+		FOnGetPropertyTypeCustomizationInstance::CreateLambda([=]() {return MakeShared<FFrameNumberDetailsCustomization>(NumericTypeInterface); }));
+	DetailsView->SetStructureData(Settings);
+
+	ChildSlot
+		[
+			SNew(SBorder)
+			.Visibility(EVisibility::Visible)
+			[
+				SNew(SVerticalBox)
+				+ SVerticalBox::Slot()
+				.AutoHeight()
+				.Padding(0.f, 8.f, 0.f, 0.f)
+				[
+					DetailsView->GetWidget().ToSharedRef()
+				]
+				+ SVerticalBox::Slot()
+				.AutoHeight()
+				.Padding(FMargin(0.0f, 0.f))
+				[
+
+					SNew(SHorizontalBox)
+					+ SHorizontalBox::Slot()
+					.FillWidth(1.f)
+					.HAlign(HAlign_Fill)
+					[
+						SNew(SSpacer)
+					]
+
+				+ SHorizontalBox::Slot()
+					.AutoWidth()
+					.VAlign(VAlign_Center)
+					.HAlign(HAlign_Right)
+					.Padding(0.f)
+					[
+						SNew(SButton)
+						.HAlign(HAlign_Center)
+						.ContentPadding(FAppStyle::GetMargin("StandardDialog.ContentPadding"))
+						.Text(LOCTEXT("OK", "OK"))
+						.OnClicked_Lambda([this, InArgs]()
+							{
+								Collapse();
+								CloseDialog();
+								return FReply::Handled();
+
+							})
+						.IsEnabled_Lambda([this]()
+							{
+								return (Settings.IsValid());
+							})
+					]
+				]
+			]
+		];
+}
+
+
+void  SCollapseControlsWidget::Collapse()
+{
+	FBakingAnimationKeySettings* BakeSettings = Settings->Get();
+	FControlRigParameterTrackEditor::CollapseAllLayers(Sequencer, OwnerTrack.Get(), *BakeSettings);
+
+	CollapseControlsSettings = *BakeSettings;
+}
+
+class SCollapseControlsWidgetWindow : public SWindow
+{
+};
+
+FReply SCollapseControlsWidget::OpenDialog(bool bModal)
 {
+	check(!DialogWindow.IsValid());
+
+	const FVector2D CursorPos = FSlateApplication::Get().GetCursorPos();
+
+	TSharedRef<SCollapseControlsWidgetWindow> Window = SNew(SCollapseControlsWidgetWindow)
+		.Title(LOCTEXT("CollapseControls", "Collapse Controls"))
+		.CreateTitleBar(true)
+		.Type(EWindowType::Normal)
+		.SizingRule(ESizingRule::Autosized)
+		.ScreenPosition(CursorPos)
+		.FocusWhenFirstShown(true)
+		.ActivationPolicy(EWindowActivationPolicy::FirstShown)
+		[
+			AsShared()
+		];
+
+	Window->SetWidgetToFocusOnActivate(AsShared());
+
+	DialogWindow = Window;
+
+	Window->MoveWindowTo(CursorPos);
 
-	if (SequencerPtr.IsValid() && OwnerTrack && ParameterSection && ParameterSection->GetControlRig())
+	if (bModal)
+	{
+		GEditor->EditorAddModalWindow(Window);
+	}
+	else
+	{
+		FSlateApplication::Get().AddWindow(Window);
+	}
+
+	return FReply::Handled();
+}
+
+void SCollapseControlsWidget::CloseDialog()
+{
+	if (DialogWindow.IsValid())
+	{
+		DialogWindow.Pin()->RequestDestroyWindow();
+		DialogWindow.Reset();
+	}
+}
+///////////////
+//end of SCollapseControlsWidget
+/////////////////////
+struct FKeyAndValuesAtFrame
+{
+	FFrameNumber Frame;
+	TArray<FMovieSceneFloatValue> KeyValues;
+	float FinalValue;
+};
+
+bool CollapseAllLayersPerKey(TSharedPtr<ISequencer>& SequencerPtr, UMovieSceneTrack* OwnerTrack, const FBakingAnimationKeySettings& InSettings)
+{
+	if (SequencerPtr.IsValid() && OwnerTrack)
+	{
+		TArray<UMovieSceneSection*> Sections = OwnerTrack->GetAllSections();
+		return MovieSceneToolHelpers::CollapseSection(SequencerPtr, OwnerTrack, Sections, InSettings);
+	}
+	return false;
+}
+
+bool FControlRigParameterTrackEditor::CollapseAllLayers(TSharedPtr<ISequencer>&SequencerPtr, UMovieSceneTrack * OwnerTrack, const FBakingAnimationKeySettings &InSettings)
+{
+	return CollapseAllLayersPerKey(SequencerPtr, OwnerTrack, InSettings);
+	if (SequencerPtr.IsValid() && OwnerTrack)
 	{
 		TArray<UMovieSceneSection*> Sections = OwnerTrack->GetAllSections();
 		//make sure right type
-		if (ParameterSection->GetBlendType().Get() != EMovieSceneBlendType::Absolute && Sections.Num() > 0 && Sections[0] != ParameterSection)
+		if (Sections.Num() < 1)
 		{
-			UE_LOG(LogControlRigEditor, Log, TEXT("Section wrong type or not first when collapsing layers"));
+			UE_LOG(LogControlRigEditor, Log, TEXT("CollapseAllSections::No sections on track"));
 			return false;
 		}
-		FScopedTransaction Transaction(LOCTEXT("CollapseAllSections", "Collapse All Sections"));
-		ParameterSection->Modify();
-		UControlRig* ControlRig = ParameterSection->GetControlRig();
-		TRange<FFrameNumber> Range = SequencerPtr->GetFocusedMovieSceneSequence()->GetMovieScene()->GetPlaybackRange();
-		FFrameNumber StartFrame = Range.GetLowerBoundValue();
-		FFrameNumber EndFrame = Range.GetUpperBoundValue();
-		const FFrameRate& FrameRate = SequencerPtr->GetFocusedDisplayRate();
-		const FFrameRate& TickResolution = SequencerPtr->GetFocusedTickResolution();
-		FMovieSceneSequenceTransform RootToLocalTransform = SequencerPtr->GetFocusedMovieSceneSequenceTransform();
-
-		FFrameNumber FrameRateInFrameNumber = TickResolution.AsFrameNumber(FrameRate.AsInterval());
-		TArray<FFrameNumber> Frames;
-		for (FFrameNumber& Frame = StartFrame; Frame <= EndFrame; Frame += FrameRateInFrameNumber)
-		{
-			Frames.Add(Frame);
-		}
-		//Store transforms
-		TArray<TPair<FName, TArray<FTransform>>> ControlLocalTransforms;
-		TArray<FRigControlElement*> Controls;
-		ControlRig->GetControlsInOrder(Controls);
-
-		for (FRigControlElement* ControlElement : Controls)
+		if (UMovieSceneControlRigParameterSection* ParameterSection = Cast<UMovieSceneControlRigParameterSection>(Sections[0]))
 		{
-			if (!ControlRig->GetHierarchy()->IsAnimatable(ControlElement))
+			if (ParameterSection->GetBlendType().Get() == EMovieSceneBlendType::Absolute)
 			{
-				continue;
-			}
-			TPair<FName, TArray<FTransform>> NameTransforms;
-			NameTransforms.Key = ControlElement->GetName();
-			NameTransforms.Value.SetNum(Frames.Num());
-			ControlLocalTransforms.Add(NameTransforms);
-		}
+				FScopedTransaction Transaction(LOCTEXT("CollapseAllSections", "Collapse All Sections"));
+				ParameterSection->Modify();
+				UControlRig* ControlRig = ParameterSection->GetControlRig();
+				FMovieSceneSequenceTransform RootToLocalTransform = SequencerPtr->GetFocusedMovieSceneSequenceTransform();
+
+				FFrameNumber StartFrame = InSettings.StartFrame;
+				FFrameNumber EndFrame = InSettings.EndFrame;
+				TRange<FFrameNumber> Range(StartFrame, EndFrame);
+				const FFrameRate& FrameRate = SequencerPtr->GetFocusedDisplayRate();
+				const FFrameRate& TickResolution = SequencerPtr->GetFocusedTickResolution();
+				
+				//frames and (optional) tangents
+				TArray<TPair<FFrameNumber, TArray<FMovieSceneTangentData>>> StoredTangents; //store tangents so we can reset them
+				TArray<FFrameNumber> Frames;
+				FFrameNumber FrameRateInFrameNumber = TickResolution.AsFrameNumber(FrameRate.AsInterval());
+				for (FFrameNumber& Frame = StartFrame; Frame <= EndFrame; Frame += FrameRateInFrameNumber)
+				{
+					Frames.Add(Frame);
+				}
+				
+				//Store transforms
+				TArray<TPair<FName, TArray<FTransform>>> ControlLocalTransforms;
+				TArray<FRigControlElement*> Controls;
+				ControlRig->GetControlsInOrder(Controls);
 
-		//get all of the local 
-		int32 Index = 0;
-		for (Index = 0; Index < Frames.Num(); ++Index)
-		{
-			const FFrameNumber& FrameNumber = Frames[Index];
-			FFrameTime GlobalTime(FrameNumber);
-			GlobalTime = GlobalTime * RootToLocalTransform.InverseLinearOnly();
+				for (FRigControlElement* ControlElement : Controls)
+				{
+					if (!ControlRig->GetHierarchy()->IsAnimatable(ControlElement))
+					{
+						continue;
+					}
+					TPair<FName, TArray<FTransform>> NameTransforms;
+					NameTransforms.Key = ControlElement->GetName();
+					NameTransforms.Value.SetNum(Frames.Num());
+					ControlLocalTransforms.Add(NameTransforms);
+				}
 
-			FMovieSceneContext Context = FMovieSceneContext(FMovieSceneEvaluationRange(GlobalTime, TickResolution), SequencerPtr->GetPlaybackStatus()).SetHasJumped(true);
+				//get all of the local 
+				int32 Index = 0;
+				for (Index = 0; Index < Frames.Num(); ++Index)
+				{
+					const FFrameNumber& FrameNumber = Frames[Index];
+					FFrameTime GlobalTime(FrameNumber);
+					GlobalTime = GlobalTime * RootToLocalTransform.InverseLinearOnly();
 
-			SequencerPtr->GetEvaluationTemplate().EvaluateSynchronousBlocking(Context, *SequencerPtr);
-			ControlRig->Evaluate_AnyThread();
-			for (TPair<FName, TArray<FTransform>>& TrailControlTransform : ControlLocalTransforms)
-			{
-				TrailControlTransform.Value[Index] = ControlRig->GetControlLocalTransform(TrailControlTransform.Key);
-			}
-		}
-		//delete other sections
-		OwnerTrack->Modify();
-		for (Index = Sections.Num() - 1; Index >= 0; --Index)
-		{
-			if (Sections[Index] != ParameterSection)
-			{
-				OwnerTrack->RemoveSectionAt(Index);
-			}
-		}
+					FMovieSceneContext Context = FMovieSceneContext(FMovieSceneEvaluationRange(GlobalTime, TickResolution), SequencerPtr->GetPlaybackStatus()).SetHasJumped(true);
 
-		//remove all keys, except Space Channels, from the Section.
-		ParameterSection->RemoveAllKeys(false /*bIncludedSpaceKeys*/);
+					SequencerPtr->GetEvaluationTemplate().EvaluateSynchronousBlocking(Context, *SequencerPtr);
+					ControlRig->Evaluate_AnyThread();
+					for (TPair<FName, TArray<FTransform>>& TrailControlTransform : ControlLocalTransforms)
+					{
+						TrailControlTransform.Value[Index] = ControlRig->GetControlLocalTransform(TrailControlTransform.Key);
+					}
+				}
+				//delete other sections
+				OwnerTrack->Modify();
+				for (Index = Sections.Num() - 1; Index >= 0; --Index)
+				{
+					if (Sections[Index] != ParameterSection)
+					{
+						OwnerTrack->RemoveSectionAt(Index);
+					}
+				}
 
-		FRigControlModifiedContext Context;
-		Context.SetKey = EControlRigSetKey::Always;
+				//remove all keys, except Space Channels, from the Section.
+				ParameterSection->RemoveAllKeys(false /*bIncludedSpaceKeys*/);
 
+				FRigControlModifiedContext Context;
+				Context.SetKey = EControlRigSetKey::Always;
 
-		FScopedSlowTask Feedback(Frames.Num(), LOCTEXT("CollapsingSections", "Collapsing Sections"));
-		Feedback.MakeDialog(true);
+				FScopedSlowTask Feedback(Frames.Num(), LOCTEXT("CollapsingSections", "Collapsing Sections"));
+				Feedback.MakeDialog(true);
 
-		const ERichCurveInterpMode InterpMode = bKeyReduce ? RCIM_Cubic : RCIM_Linear;
+				const ERichCurveInterpMode InterpMode = InSettings.bReduceKeys ? RCIM_Cubic : RCIM_Linear;
 
-		Index = 0;
-		for (Index = 0; Index < Frames.Num(); ++Index)
-		{
-			Feedback.EnterProgressFrame(1, LOCTEXT("CollapsingSections", "Collapsing Sections"));
-			const FFrameNumber& FrameNumber = Frames[Index];
-			Context.LocalTime = TickResolution.AsSeconds(FFrameTime(FrameNumber));
-			//need to do the twice hack since controls aren't really in order
-			for (int32 TwiceHack = 0; TwiceHack < 2; ++TwiceHack)
-			{
-				for (TPair<FName, TArray<FTransform>>& TrailControlTransform : ControlLocalTransforms)
+				Index = 0;
+				for (Index = 0; Index < Frames.Num(); ++Index)
 				{
-					ControlRig->SetControlLocalTransform(TrailControlTransform.Key, TrailControlTransform.Value[Index],false, Context, false);
+					Feedback.EnterProgressFrame(1, LOCTEXT("CollapsingSections", "Collapsing Sections"));
+					const FFrameNumber& FrameNumber = Frames[Index];
+					Context.LocalTime = TickResolution.AsSeconds(FFrameTime(FrameNumber));
+					//need to do the twice hack since controls aren't really in order
+					for (int32 TwiceHack = 0; TwiceHack < 2; ++TwiceHack)
+					{
+						for (TPair<FName, TArray<FTransform>>& TrailControlTransform : ControlLocalTransforms)
+						{
+							ControlRig->SetControlLocalTransform(TrailControlTransform.Key, TrailControlTransform.Value[Index], false, Context, false);
+						}
+					}
+					ControlRig->Evaluate_AnyThread();
+					ParameterSection->RecordControlRigKey(FrameNumber, true, InterpMode);
+
+					if (Feedback.ShouldCancel())
+					{
+						Transaction.Cancel();
+						SequencerPtr->NotifyMovieSceneDataChanged(EMovieSceneDataChangeType::MovieSceneStructureItemAdded);
+						return false;
+					}
 				}
-			}
-			ControlRig->Evaluate_AnyThread();
-			ParameterSection->RecordControlRigKey(FrameNumber, true, InterpMode);
+				if (InSettings.bReduceKeys)
+				{
+					FKeyDataOptimizationParams Params;
+					Params.bAutoSetInterpolation = true;
+					Params.Tolerance = InSettings.Tolerance;
+					FMovieSceneChannelProxy& ChannelProxy = ParameterSection->GetChannelProxy();
+					TArrayView<FMovieSceneFloatChannel*> FloatChannels = ChannelProxy.GetChannels<FMovieSceneFloatChannel>();
 
-			if (Feedback.ShouldCancel())
-			{
-				Transaction.Cancel();
+					for (FMovieSceneFloatChannel* Channel : FloatChannels)
+					{
+						Channel->Optimize(Params); //should also auto tangent
+					}
+				}
+				//reset everything back
 				SequencerPtr->NotifyMovieSceneDataChanged(EMovieSceneDataChangeType::MovieSceneStructureItemAdded);
+				return true;
+			}
+			else
+			{
+				UE_LOG(LogControlRigEditor, Log, TEXT("CollapseAllSections:: First section is not additive"));
 				return false;
 			}
 		}
-		if (bKeyReduce)
+		else
 		{
-			FKeyDataOptimizationParams Params;
-			Params.bAutoSetInterpolation = true;
-			Params.Tolerance = Tolerance;
-			FMovieSceneChannelProxy& ChannelProxy = ParameterSection->GetChannelProxy();
-			TArrayView<FMovieSceneFloatChannel*> FloatChannels = ChannelProxy.GetChannels<FMovieSceneFloatChannel>();
-
-			for (FMovieSceneFloatChannel* Channel : FloatChannels)
-			{
-				Channel->Optimize(Params); //should also auto tangent
-			}
+			UE_LOG(LogControlRigEditor, Log, TEXT("CollapseAllSections:: No Control Rig section"));
+			return false;
 		}
-		//reset everything back
-		SequencerPtr->NotifyMovieSceneDataChanged(EMovieSceneDataChangeType::MovieSceneStructureItemAdded);
-		return true;
 	}
+	UE_LOG(LogControlRigEditor, Log, TEXT("CollapseAllSections:: Sequencer or track is invalid"));
 	return false;
 }
+
 void FControlRigParameterSection::CollapseAllLayers()
 {
 	TSharedPtr<ISequencer> SequencerPtr = WeakSequencer.Pin();
-	UMovieSceneControlRigParameterSection* ParameterSection = CastChecked<UMovieSceneControlRigParameterSection>(WeakSection.Get());
-	UMovieSceneTrack* OwnerTrack = ParameterSection->GetTypedOuter<UMovieSceneTrack>();
-	FControlRigParameterTrackEditor::CollapseAllLayers(SequencerPtr,OwnerTrack, ParameterSection);
+	if (UMovieSceneControlRigParameterSection* ParameterSection = CastChecked<UMovieSceneControlRigParameterSection>(WeakSection.Get()))
+	{
+		UMovieSceneTrack* OwnerTrack = ParameterSection->GetTypedOuter<UMovieSceneTrack>();
+		TSharedPtr<ISequencer> Sequencer = WeakSequencer.Pin();
+		TSharedRef<SCollapseControlsWidget> BakeWidget =
+			SNew(SCollapseControlsWidget)
+			.Sequencer(Sequencer)
+			.OwnerTrack(OwnerTrack);
 
+		BakeWidget->OpenDialog(true);
+	}
 }
+
+
 void FControlRigParameterSection::KeyZeroValue()
 {
 	UMovieSceneControlRigParameterSection* ParameterSection = CastChecked<UMovieSceneControlRigParameterSection>(WeakSection.Get());
diff --git a/Engine/Plugins/Animation/ControlRig/Source/ControlRigEditor/Private/Sequencer/ControlRigParameterTrackEditor.h b/Engine/Plugins/Animation/ControlRig/Source/ControlRigEditor/Private/Sequencer/ControlRigParameterTrackEditor.h
--- a/Engine/Plugins/Animation/ControlRig/Source/ControlRigEditor/Private/Sequencer/ControlRigParameterTrackEditor.h
+++ b/Engine/Plugins/Animation/ControlRig/Source/ControlRigEditor/Private/Sequencer/ControlRigParameterTrackEditor.h
@@ -33,7 +33,7 @@ struct FMovieSceneControlRigSpaceChannel;
 struct FMovieSceneChannel;
 struct FKeyAddOrDeleteEventItem;
 struct FKeyMoveEventItem;
-
+struct FBakingAnimationKeySettings;
 /**
  * Tools for animation tracks
  */
@@ -229,7 +229,7 @@ public:
 	bool ModifyOurGeneratedKeysByCurrentAndWeight(UObject* Object, UControlRig* InControlRig, FName RigControlName, UMovieSceneTrack *Track, UMovieSceneSection* SectionToKey, FFrameNumber EvaluateTime, FGeneratedTrackKeys& InOutGeneratedTotalKeys, float Weight) const;
 
 	//**Function to collapse all layers from this section onto the first absoluate layer.*/
-	static bool CollapseAllLayers(TSharedPtr<ISequencer>& SequencerPtr, UMovieSceneTrack* OwnerTrack, UMovieSceneControlRigParameterSection* ParameterSection, bool bKeyReduce = false, float Tolerance = 0.001f);
+	static bool CollapseAllLayers(TSharedPtr<ISequencer>& SequencerPtr, UMovieSceneTrack* OwnerTrack, const FBakingAnimationKeySettings& InSettings);
 
 private:
 	FDelegateHandle SelectionChangedHandle;
diff --git a/Engine/Plugins/Animation/ControlRig/Source/ControlRigEditor/Private/Sequencer/ControlRigSpaceChannelEditors.cpp b/Engine/Plugins/Animation/ControlRig/Source/ControlRigEditor/Private/Sequencer/ControlRigSpaceChannelEditors.cpp
--- a/Engine/Plugins/Animation/ControlRig/Source/ControlRigEditor/Private/Sequencer/ControlRigSpaceChannelEditors.cpp
+++ b/Engine/Plugins/Animation/ControlRig/Source/ControlRigEditor/Private/Sequencer/ControlRigSpaceChannelEditors.cpp
@@ -61,6 +61,7 @@
 #include "CurveEditorSettings.h"
 #include "TimeToPixel.h"
 #include "Templates/Tuple.h"
+#include "Sequencer/ControlRigSequencerHelpers.h"
 
 #define LOCTEXT_NAMESPACE "ControlRigEditMode"
 
@@ -860,17 +861,161 @@ void FControlRigSpaceChannelHelpers::GetFramesInThisSpaceAfterThisTime(
 }
 
 void FControlRigSpaceChannelHelpers::SequencerBakeControlInSpace(UControlRig* ControlRig, ISequencer* Sequencer, FMovieSceneControlRigSpaceChannel* Channel, UMovieSceneSection* SectionToKey,
-	TArray<FFrameNumber> Frames, URigHierarchy* RigHierarchy, const FRigElementKey& ControlKey, FRigSpacePickerBakeSettings Settings)
+	URigHierarchy* RigHierarchy, const FRigElementKey& ControlKey, FRigSpacePickerBakeSettings Settings)
 {
 	if (bDoNotCompensate == true)
 	{
 		return;
 	}
-	if (ControlRig && Sequencer && Channel && SectionToKey && Frames.Num() > 0 && RigHierarchy
+	if (ControlRig && Sequencer && Channel && SectionToKey && (Settings.Settings.StartFrame != Settings.Settings.EndFrame) && RigHierarchy
 		&& Sequencer->GetFocusedMovieSceneSequence() && Sequencer->GetFocusedMovieSceneSequence()->GetMovieScene())
 	{
+	
 		if (UMovieSceneControlRigParameterSection* Section = Cast<UMovieSceneControlRigParameterSection>(SectionToKey))
 		{
+			//get the setting as a rig element key to use for comparisons for smart baking, 
+			FMovieSceneControlRigSpaceBaseKey Value;
+			if (Settings.TargetSpace == RigHierarchy->GetWorldSpaceReferenceKey())
+			{
+				Value.SpaceType = EMovieSceneControlRigSpaceType::World;
+			}
+			else if (Settings.TargetSpace == RigHierarchy->GetDefaultParentKey())
+			{
+				Value.SpaceType = EMovieSceneControlRigSpaceType::Parent;
+			}
+			else
+			{
+				FRigElementKey DefaultParent = RigHierarchy->GetFirstParent(ControlKey);
+				if (DefaultParent == Settings.TargetSpace)
+				{
+					Value.SpaceType = EMovieSceneControlRigSpaceType::Parent;
+				}
+				else
+				{
+					Value.SpaceType = EMovieSceneControlRigSpaceType::ControlRig;
+					Value.ControlRigElement = Settings.TargetSpace;
+				}
+			}
+			TArrayView<FMovieSceneFloatChannel*> Channels = FControlRigSequencerHelpers::GetFloatChannels(ControlRig,
+				ControlKey.Name, Section);
+			//if smart baking we 1) Don't do any baking over times where the Settings.TargetSpace is active, 
+			//and 2) we just bake over key frames
+			//need to store these times first since we will still delete all of these space keys
+
+			TArray<FFrameNumber> Frames;
+			const FFrameRate TickResolution = Sequencer->GetFocusedTickResolution();
+			const FFrameRate  FrameRate = Sequencer->GetFocusedDisplayRate();
+			//get space keys so we can delete them and use them for smart baking at them
+			TArray<FFrameNumber> Keys;
+			TArray < FKeyHandle> KeyHandles;
+			FFrameNumber StartFrame = Settings.Settings.StartFrame;
+			FFrameNumber EndFrame = Settings.Settings.EndFrame;
+			TRange<FFrameNumber> Range(StartFrame, EndFrame);
+			Channel->GetKeys(Range, &Keys, &KeyHandles);
+			TArray<FFrameNumber> MinusOneFrames; //frames we need to delete
+			TArray<TPair<FFrameNumber, TArray<FMovieSceneTangentData>>> StoredTangents; //need to set tangents at the space switch locations
+			TSet<FFrameNumber> ParentFrames; // these are parent frames we keep track of so we can 
+			if (Settings.Settings.BakingKeySettings == EBakingKeySettings::KeysOnly)
+			{
+				TSet<FFrameNumber> MinusOneKeys; 
+				TSortedMap<FFrameNumber,FFrameNumber> FrameMap;
+				//add space keys to bake
+				{
+					for (FFrameNumber& Frame : Keys)
+					{
+						FMovieSceneControlRigSpaceBaseKey SpaceValue;
+						using namespace UE::MovieScene;
+						EvaluateChannel(Channel, Frame, SpaceValue);
+						if (SpaceValue != Value)
+						{
+							FChannelMapInfo* pChannelIndex = nullptr;
+							FRigControlElement* ControlElement = nullptr;
+							int32 NumChannels = 0;
+							int32 ChannelIndex = 0;
+							Tie(ControlElement, pChannelIndex, NumChannels, ChannelIndex) = GetControlAndChannelInfos(ControlRig, Section, ControlKey.Name);
+
+							if (pChannelIndex && ControlElement && NumChannels > 0)
+							{
+								TPair<FFrameNumber, TArray<FMovieSceneTangentData>> FrameAndTangent;
+								FrameAndTangent.Key = Frame;
+								EvaluateTangentAtThisTime<FMovieSceneFloatChannel>(ChannelIndex, NumChannels, Section, Frame, FrameAndTangent.Value);
+								StoredTangents.Add(FrameAndTangent);
+							}
+							FrameMap.Add(Frame,Frame);
+						}
+						MinusOneKeys.Add(Frame - 1);
+					}
+				}
+
+				TArray<FFrameNumber> TransformFrameTimes = FMovieSceneConstraintChannelHelper::GetTransformTimes(
+					Channels, Settings.Settings.StartFrame, Settings.Settings.EndFrame);
+				//add transforms keys to bake
+				{
+					for (FFrameNumber& Frame : TransformFrameTimes)
+					{
+						if (MinusOneKeys.Contains(Frame) == false)
+						{
+							FrameMap.Add(Frame);
+						}
+					}
+				}
+				//also need to add transform keys for the parent control rig intervals over the time we bake
+				TArray<FSpaceRange> SpaceRanges = Channel->FindSpaceIntervals();
+				for (const FSpaceRange& SpaceRange : SpaceRanges)
+				{
+					if (SpaceRange.Key.SpaceType == EMovieSceneControlRigSpaceType::ControlRig && SpaceRange.Key.ControlRigElement.IsValid())
+					{
+						TRange<FFrameNumber> Overlap = TRange<FFrameNumber>::Intersection(Range, SpaceRange.Range);
+						if (Overlap.IsEmpty() == false)
+						{
+							TArrayView<FMovieSceneFloatChannel*> ParentChannels = FControlRigSequencerHelpers::GetFloatChannels(ControlRig,
+								SpaceRange.Key.ControlRigElement.Name, Section);
+
+							TArray<FFrameNumber> ParentFrameTimes = FMovieSceneConstraintChannelHelper::GetTransformTimes(
+								ParentChannels, StartFrame, EndFrame);
+							for (FFrameNumber& Frame : ParentFrameTimes)
+							{
+								//if we don't have a key at the parent time, add the frame AND store it's tangent.
+								if (FrameMap.Contains(Frame) == false)
+								{
+									FrameMap.Add(Frame);
+									ParentFrames.Add(Frame);
+
+									FChannelMapInfo* pChannelIndex = nullptr;
+									FRigControlElement* ControlElement = nullptr;
+									int32 NumChannels = 0;
+									int32 ChannelIndex = 0;
+									Tie(ControlElement, pChannelIndex, NumChannels, ChannelIndex) = GetControlAndChannelInfos(ControlRig, Section, SpaceRange.Key.ControlRigElement.Name);
+									TPair<FFrameNumber, TArray<FMovieSceneTangentData>> FrameAndTangent;
+									FrameAndTangent.Key = Frame;
+									EvaluateTangentAtThisTime<FMovieSceneFloatChannel>(ChannelIndex, NumChannels, Section, Frame, FrameAndTangent.Value);
+									StoredTangents.Add(FrameAndTangent);
+								}
+							}
+						}
+					}
+				
+				}
+				FrameMap.GenerateKeyArray(Frames);
+
+			}
+			else
+			{
+				FFrameNumber FrameRateInFrameNumber = TickResolution.AsFrameNumber(FrameRate.AsInterval());
+				if (Settings.Settings.FrameIncrement > 1) //increment per frame by increment
+				{
+					FrameRateInFrameNumber.Value *= Settings.Settings.FrameIncrement;
+				}
+				for (FFrameNumber& Frame = Settings.Settings.StartFrame; Frame <= Settings.Settings.EndFrame; Frame += FrameRateInFrameNumber)
+				{
+					Frames.Add(Frame);
+				}
+
+			}
+			if (Frames.Num() == 0)
+			{
+				return;
+			}
 			TGuardValue<bool> CompensateGuard(bDoNotCompensate, true);
 			TArray<FTransform> ControlRigParentWorldTransforms;
 			ControlRigParentWorldTransforms.SetNum(Frames.Num());
@@ -883,21 +1028,23 @@ void FControlRigSpaceChannelHelpers::SequencerBakeControlInSpace(UControlRig* Co
 			TArray<FTransform> ControlWorldTransforms;
 			Snapper.GetControlRigControlTransforms(Sequencer, ControlRig, ControlKey.Name, Frames, ControlRigParentWorldTransforms, ControlWorldTransforms);
 
-			//Find all space keys in range and delete them since it will get replaced with new space when we components.
+			//Need to delete keys, if smart baking we only delete the -1 keys,
+			//if not, we need to delete all transform
+			//and always delete space keys
 			Section->Modify();
-			FFrameNumber StartFrame = Frames[0];
-			FFrameNumber EndFrame = Frames[Frames.Num() - 1];
-			TArray<FFrameNumber> Keys;
-			TArray < FKeyHandle> KeyHandles;
-			TRange<FFrameNumber> Range(StartFrame, EndFrame);
-			Channel->GetKeys(Range, &Keys, &KeyHandles);
 			Channel->DeleteKeys(KeyHandles);
-			//also delete any transform keys at the deleted key times and times -1.
-			for (const FFrameNumber& DeleteFrame : Keys)
+			if (Settings.Settings.BakingKeySettings == EBakingKeySettings::KeysOnly)
+			{
+				for (const FFrameNumber& DeleteFrame : Keys)
+				{
+					FControlRigSpaceChannelHelpers::DeleteTransformKeysAtThisTime(ControlRig, Section, ControlKey.Name, DeleteFrame - 1);
+				}
+			}
+			else
 			{
-				FControlRigSpaceChannelHelpers::DeleteTransformKeysAtThisTime(ControlRig, Section, ControlKey.Name, DeleteFrame - 1);
-				FControlRigSpaceChannelHelpers::DeleteTransformKeysAtThisTime(ControlRig, Section, ControlKey.Name, DeleteFrame);
+				FMovieSceneConstraintChannelHelper::DeleteTransformTimes(Channels, StartFrame, EndFrame);
 			}
+			
 
 			URigHierarchy* Hierarchy = ControlRig->GetHierarchy();
 
@@ -906,29 +1053,7 @@ void FControlRigSpaceChannelHelpers::SequencerBakeControlInSpace(UControlRig* Co
 			using namespace UE::MovieScene;
 			EvaluateChannel(Channel, StartFrame, StartFrameValue);
 			EvaluateChannel(Channel, EndFrame, EndFrameValue);
-
-			FMovieSceneControlRigSpaceBaseKey Value;
-			if (Settings.TargetSpace == RigHierarchy->GetWorldSpaceReferenceKey())
-			{
-				Value.SpaceType = EMovieSceneControlRigSpaceType::World;
-			}
-			else if (Settings.TargetSpace == RigHierarchy->GetDefaultParentKey())
-			{
-				Value.SpaceType = EMovieSceneControlRigSpaceType::Parent;
-			}
-			else
-			{
-				FRigElementKey DefaultParent = RigHierarchy->GetFirstParent(ControlKey);
-				if (DefaultParent == Settings.TargetSpace)
-				{
-					Value.SpaceType = EMovieSceneControlRigSpaceType::Parent;
-				}
-				else
-				{
-					Value.SpaceType = EMovieSceneControlRigSpaceType::ControlRig;
-					Value.ControlRigElement = Settings.TargetSpace;
-				}
-			}
+			
 			const bool bCompensateStart = StartFrameValue != Value;
 			TRange<FFrameNumber> PlaybackRange = Sequencer->GetFocusedMovieSceneSequence()->GetMovieScene()->GetPlaybackRange();
 			const bool bCompensateEnd = (EndFrameValue != Value && PlaybackRange.GetUpperBoundValue() != EndFrame);
@@ -948,9 +1073,7 @@ void FControlRigSpaceChannelHelpers::SequencerBakeControlInSpace(UControlRig* Co
 			//now set all of the key values
 			FRigControlModifiedContext Context;
 			Context.SetKey = EControlRigSetKey::Always;
-
-			FFrameRate TickResolution = Sequencer->GetFocusedTickResolution();
-
+			Context.KeyMask = (uint32)EControlRigContextChannelToKey::AllTransform;
 			URigHierarchy::TElementDependencyMap Dependencies = RigHierarchy->GetDependenciesForVM(ControlRig->GetVM());
 			RigHierarchy->SwitchToParent(ControlKey, Settings.TargetSpace, false, true, Dependencies, nullptr);
 			ControlRig->Evaluate_AnyThread();
@@ -972,6 +1095,16 @@ void FControlRigSpaceChannelHelpers::SequencerBakeControlInSpace(UControlRig* Co
 				//evaluate control rig
 				ControlRig->Evaluate_AnyThread();
 				Context.LocalTime = TickResolution.AsSeconds(FFrameTime(Frame));
+				
+				//if doing smart baking only set keys on those that exist IF not a key from the parent, in that case we need to key the  whole transform.
+				if (Settings.Settings.BakingKeySettings == EBakingKeySettings::KeysOnly && ParentFrames.Contains(Frame) == false)
+				{
+					Context.KeyMask = (uint32)GetCurrentTransformKeysAtThisTime(ControlRig, ControlKey.Name, SectionToKey, Frame);
+				}
+				else
+				{
+					Context.KeyMask = (uint32)EControlRigContextChannelToKey::AllTransform;
+				}
 				ControlRig->SetControlGlobalTransform(ControlKey.Name, GlobalTransform, true, Context);
 			}
 
@@ -1006,66 +1139,38 @@ void FControlRigSpaceChannelHelpers::SequencerBakeControlInSpace(UControlRig* Co
 				ChannelInterface.AddKey(EndFrame, Forward<FMovieSceneControlRigSpaceBaseKey>(EndFrameValue));
 				const FTransform GlobalTransform = ControlWorldTransforms[Frames.Num() - 1];
 				Context.LocalTime = TickResolution.AsSeconds(FFrameTime(EndFrame));
+				Context.KeyMask = (uint32)EControlRigContextChannelToKey::AllTransform;
 				ControlRig->SetControlGlobalTransform(ControlKey.Name, GlobalTransform, true, Context);
 			}
-		
-			// Fix any Rotation Channels
-			Section->FixRotationWinding(ControlKey.Name, Frames[0], Frames[Frames.Num() - 1]);
-			// Then reduce
-			if (Settings.bReduceKeys)
+			//Fix tangents at removed space switches if needed
+			for (TPair<FFrameNumber, TArray<FMovieSceneTangentData>>& StoredTangent : StoredTangents)
 			{
-				TArrayView<FMovieSceneFloatChannel*> FloatChannels = Section->GetChannelProxy().GetChannels<FMovieSceneFloatChannel>();
-				FChannelMapInfo* pChannelIndex = Section->ControlChannelMap.Find(ControlKey.Name);
-				if (pChannelIndex != nullptr)
-				{
-					int32 ChannelIndex = pChannelIndex->ChannelIndex;
-
-					if (FRigControlElement* ControlElement = ControlRig->FindControl(ControlKey.Name))
-					{
-						FKeyDataOptimizationParams Params;
-						Params.bAutoSetInterpolation = true;
-						Params.Tolerance = Settings.Tolerance;
-						Params.Range = TRange <FFrameNumber>(Frames[0], Frames[Frames.Num() - 1]);
-
-						switch (ControlElement->Settings.ControlType)
-						{
-						case ERigControlType::Position:
-						case ERigControlType::Scale:
-						case ERigControlType::Rotator:
-						{
-							FloatChannels[ChannelIndex]->Optimize(Params);
-							FloatChannels[ChannelIndex + 1]->Optimize(Params);
-							FloatChannels[ChannelIndex + 2]->Optimize(Params);
-							break;
-						}
-
-						case ERigControlType::Transform:
-						case ERigControlType::TransformNoScale:
-						case ERigControlType::EulerTransform:
-						{
-
-							FloatChannels[ChannelIndex]->Optimize(Params);
-							FloatChannels[ChannelIndex + 1]->Optimize(Params);
-							FloatChannels[ChannelIndex + 2]->Optimize(Params);
-							FloatChannels[ChannelIndex + 3]->Optimize(Params);
-							FloatChannels[ChannelIndex + 4]->Optimize(Params);
-							FloatChannels[ChannelIndex + 5]->Optimize(Params);
+				FChannelMapInfo* pChannelIndex = nullptr;
+				FRigControlElement* ControlElement = nullptr;
+				int32 NumChannels = 0;
+				int32 ChannelIndex = 0;
+				Tie(ControlElement, pChannelIndex, NumChannels, ChannelIndex) = GetControlAndChannelInfos(ControlRig, Section, ControlKey.Name);
 
-							if (ControlElement->Settings.ControlType == ERigControlType::Transform ||
-								ControlElement->Settings.ControlType == ERigControlType::EulerTransform)
-							{
-								FloatChannels[ChannelIndex + 6]->Optimize(Params);
-								FloatChannels[ChannelIndex + 7]->Optimize(Params);
-								FloatChannels[ChannelIndex + 8]->Optimize(Params);
-							}
-							break;
-
-						}
-						}
-					}
+				if (pChannelIndex && ControlElement && NumChannels > 0)
+				{
+					SetTangentsAtThisTime<FMovieSceneFloatChannel>(ChannelIndex, NumChannels, Section, StoredTangent.Key, StoredTangent.Value);
 				}
 			}
-			
+			// Fix any Rotation Channels
+			Section->FixRotationWinding(ControlKey.Name, Frames[0], Frames[Frames.Num() - 1]);
+			// Then optimize
+			if (Settings.Settings.BakingKeySettings == EBakingKeySettings::AllFrames && Settings.Settings.bReduceKeys)
+			{
+				FKeyDataOptimizationParams Params;
+				Params.bAutoSetInterpolation = true;
+				Params.Tolerance = Settings.Settings.Tolerance;
+				Params.Range = TRange <FFrameNumber>(Frames[0], Frames[Frames.Num() - 1]);
+				Section->OptimizeSection(ControlKey.Name, Params);
+			}
+			else  //need to auto set tangents, the above section will do that also
+			{
+				Section->AutoSetTangents(ControlKey.Name);
+			}	
 		}
 		Sequencer->NotifyMovieSceneDataChanged(EMovieSceneDataChangeType::MovieSceneStructureItemAdded); //may have added channel
 	}
@@ -1333,16 +1438,16 @@ FReply FControlRigSpaceChannelHelpers::OpenBakeDialog(ISequencer* Sequencer, FMo
 			Settings.TargetSpace = URigHierarchy::GetDefaultParentKey();
 
 			TRange<FFrameNumber> Range = Sequencer->GetFocusedMovieSceneSequence()->GetMovieScene()->GetPlaybackRange();
-			Settings.StartFrame = Range.GetLowerBoundValue();
-			Settings.EndFrame = Range.GetUpperBoundValue();
+			Settings.Settings.StartFrame = Range.GetLowerBoundValue();
+			Settings.Settings.EndFrame = Range.GetUpperBoundValue();
 			TMovieSceneChannelData<FMovieSceneControlRigSpaceBaseKey> ChannelData = Channel->GetData();
 			TArrayView<const FFrameNumber> Times = ChannelData.GetTimes();
 			if (KeyIndex >= 0 && KeyIndex < Times.Num())
 			{
-				Settings.StartFrame = Times[KeyIndex];
+				Settings.Settings.StartFrame = Times[KeyIndex];
 				if (KeyIndex + 1 < Times.Num())
 				{
-					Settings.EndFrame = Times[KeyIndex + 1];
+					Settings.Settings.EndFrame = Times[KeyIndex + 1];
 				}
 			}
 			TArray<FRigElementKey> Controls;
@@ -1362,14 +1467,7 @@ FReply FControlRigSpaceChannelHelpers::OpenBakeDialog(ISequencer* Sequencer, FMo
 				})
 				.OnBake_Lambda([Sequencer, ControlRig, TickResolution](URigHierarchy* InHierarchy, TArray<FRigElementKey> InControls, FRigSpacePickerBakeSettings InSettings)
 				{
-					TArray<FFrameNumber> Frames;
-
-					const FFrameRate& FrameRate = Sequencer->GetFocusedDisplayRate();
-					FFrameNumber FrameRateInFrameNumber = TickResolution.AsFrameNumber(FrameRate.AsInterval());
-					for (FFrameNumber& Frame = InSettings.StartFrame; Frame <= InSettings.EndFrame; Frame += FrameRateInFrameNumber)
-					{
-						Frames.Add(Frame);
-					}
+					
 					FScopedTransaction Transaction(LOCTEXT("BakeControlToSpace", "Bake Control In Space"));
 					for (const FRigElementKey& ControlKey : InControls)
 					{
@@ -1377,7 +1475,7 @@ FReply FControlRigSpaceChannelHelpers::OpenBakeDialog(ISequencer* Sequencer, FMo
 						if (SpaceChannelAndSection.SpaceChannel)
 						{
 							FControlRigSpaceChannelHelpers::SequencerBakeControlInSpace(ControlRig, Sequencer, SpaceChannelAndSection.SpaceChannel, SpaceChannelAndSection.SectionToKey,
-								Frames, InHierarchy, ControlKey, InSettings);
+								InHierarchy, ControlKey, InSettings);
 						}
 					}
 					return FReply::Handled();
diff --git a/Engine/Plugins/Animation/ControlRig/Source/ControlRigEditor/Private/Tools/ConstraintBaker.cpp b/Engine/Plugins/Animation/ControlRig/Source/ControlRigEditor/Private/Tools/ConstraintBaker.cpp
--- a/Engine/Plugins/Animation/ControlRig/Source/ControlRigEditor/Private/Tools/ConstraintBaker.cpp
+++ b/Engine/Plugins/Animation/ControlRig/Source/ControlRigEditor/Private/Tools/ConstraintBaker.cpp
@@ -18,27 +18,30 @@
 
 #define LOCTEXT_NAMESPACE "ConstraintBaker"
 
+
 // at this stage, we suppose that any of these arguments are safe to use. Make sure to test them before using that function
 void FConstraintBaker::GetMinimalFramesToBake(
 	UWorld* InWorld,
 	const UTickableTransformConstraint* InConstraint,
 	const TSharedPtr<ISequencer>& InSequencer,
-	IMovieSceneConstrainedSection* InSection,
+	UMovieSceneSection* InSection,
+	const TOptional< FBakingAnimationKeySettings>& InSettings,
 	TArray<FFrameNumber>& OutFramesToBake)
 {
 	const UMovieSceneSequence* MovieSceneSequence = InSequencer->GetFocusedMovieSceneSequence();
 	const UMovieScene* MovieScene = MovieSceneSequence ? MovieSceneSequence->GetMovieScene() : nullptr;
-	if (!MovieScene)
+	if (!MovieScene || !InConstraint || !InConstraint->ChildTRSHandle || !InSection)
 	{
 		return;
 	}
+	IMovieSceneConstrainedSection* ConstrainedSection = Cast<IMovieSceneConstrainedSection>(InSection);
 
 	// get constraint channel data if any
 	TArrayView<const FFrameNumber> ConstraintFrames;
 	TArrayView<const bool> ConstraintValues;
 
 	// note that we might want to bake a constraint which is not animated
-	FConstraintAndActiveChannel* ThisActiveChannel = InSection->GetConstraintChannel(InConstraint->GetFName());
+	FConstraintAndActiveChannel* ThisActiveChannel = ConstrainedSection->GetConstraintChannel(InConstraint->GetFName());
 	if (ThisActiveChannel)
 	{
 		const TMovieSceneChannelData<const bool> ConstraintChannelData = ThisActiveChannel->ActiveChannel.GetData();
@@ -54,56 +57,86 @@ void FConstraintBaker::GetMinimalFramesToBake(
 	}
 
 	// default init bounds to the scene bounds
-	FFrameNumber FirstBakingFrame = MovieScene->GetPlaybackRange().GetLowerBoundValue();
-	FFrameNumber LastBakingFrame = MovieScene->GetPlaybackRange().GetUpperBoundValue();
+	FFrameNumber FirstBakingFrame = InSettings.IsSet() ? InSettings.GetValue().StartFrame : MovieScene->GetPlaybackRange().GetLowerBoundValue();
+	FFrameNumber LastBakingFrame = InSettings.IsSet() ? InSettings.GetValue().EndFrame : MovieScene->GetPlaybackRange().GetUpperBoundValue();
 	
-	// set start to first active frame if any
-	const int32 FirstActiveIndex = ConstraintValues.IndexOfByKey(true);
-	if (ConstraintValues.IsValidIndex(FirstActiveIndex))
+	// set start to first active frame if any, if we have no range
+	if (InSettings.IsSet() == false)
 	{
-		FirstBakingFrame = ConstraintFrames[FirstActiveIndex];
-	}
+		const int32 FirstActiveIndex = ConstraintValues.IndexOfByKey(true);
+		if (ConstraintValues.IsValidIndex(FirstActiveIndex))
+		{
+			FirstBakingFrame = ConstraintFrames[FirstActiveIndex];
+		}
 
-	// set end to the last key if inactive
-	const bool bIsLastKeyInactive = (ConstraintValues.Last() == false);
-	if (bIsLastKeyInactive)
-	{
-		LastBakingFrame = ConstraintFrames.Last();
+		// set end to the last key if inactive
+		const bool bIsLastKeyInactive = (ConstraintValues.Last() == false);
+		if (bIsLastKeyInactive)
+		{
+			LastBakingFrame = ConstraintFrames.Last();
+		}
 	}
+	// keep frames within active state
+	auto IsConstraintActive = [InConstraint, ThisActiveChannel](const FFrameNumber& Time)
+	{
+		if (ThisActiveChannel)
+		{
+			bool bIsActive = false; ThisActiveChannel->ActiveChannel.Evaluate(Time, bIsActive);
+			return bIsActive;
+		}
+		return InConstraint->IsFullyActive();
+	};
 
-	// then compute range from first to last
-	TArray<FFrameNumber> Frames;
-	MovieSceneToolHelpers::CalculateFramesBetween(MovieScene, FirstBakingFrame, LastBakingFrame, Frames);
-		
-	// Fill the frames we want to bake
+	if (InSettings.IsSet() == false || InSettings.GetValue().BakingKeySettings != EBakingKeySettings::KeysOnly)
 	{
+		TArray<FFrameNumber> Frames;
+		MovieSceneToolHelpers::CalculateFramesBetween(MovieScene, FirstBakingFrame, LastBakingFrame,  1 /* use increment of 1 since we reduce later*/,
+			Frames);
+
+		// Fill the frames we want to bake
+
 		// add constraint keys
-		for (const FFrameNumber& ConstraintFrame: ConstraintFrames)
+		for (const FFrameNumber& ConstraintFrame : ConstraintFrames)
 		{
 			OutFramesToBake.Add(ConstraintFrame);
 		}
-		
-		// keep frames within active state
-		auto IsConstraintActive = [InConstraint, ThisActiveChannel](const FFrameNumber& Time)
+
+		int32 FrameInc = (InSettings.IsSet() == true) ? InSettings.GetValue().FrameIncrement : 1;
+		for (int32 FrameIndex = 0; FrameIndex < Frames.Num(); ++FrameIndex)
 		{
-			if (ThisActiveChannel)
+			if (FrameIndex % FrameInc == 0)
 			{
-				bool bIsActive = false; ThisActiveChannel->ActiveChannel.Evaluate(Time, bIsActive);
-				return bIsActive;
+				const FFrameNumber& InFrame = Frames[FrameIndex];
+				if (IsConstraintActive(InFrame))
+				{
+					OutFramesToBake.Add(InFrame);
+				}
 			}
-			return InConstraint->IsFullyActive();
-		};
+		}
+		if (FrameInc > 1) //if skipping make sure we get all transforms
+		{
+			FMovieSceneConstraintChannelHelper::GetTransformFramesForConstraintHandles(InConstraint, InSequencer,
+				InSettings.GetValue().StartFrame, InSettings.GetValue().EndFrame, OutFramesToBake);
+		}
 		
-		for (const FFrameNumber& InFrame: Frames)
+	}
+	else //doing per frame key get constraint keys and transform keys
+	{
+		for (const FFrameNumber& Frame : ConstraintFrames)
 		{
-			if (IsConstraintActive(InFrame))
+			if (Frame >= FirstBakingFrame && Frame <= LastBakingFrame)
 			{
-				OutFramesToBake.Add(InFrame);
+				OutFramesToBake.Add(Frame);
 			}
 		}
+
+		FMovieSceneConstraintChannelHelper::GetTransformFramesForConstraintHandles(InConstraint, InSequencer,
+			InSettings.GetValue().StartFrame, InSettings.GetValue().EndFrame, OutFramesToBake);
+
 	}
 
 	// we also need to store which T-1 frames need to be kept for other constraints compensation
+	// and also store transform key locations for parents if doing per key or per frame with non-1 increments
 	{
 		// gather the sorted child's constraint
 		static constexpr bool bSorted = true;
@@ -120,10 +153,18 @@ void FConstraintBaker::GetMinimalFramesToBake(
 			if (TransformConstraint && TransformConstraint->NeedsCompensation() && TransformConstraint != InConstraint) 
 			{
 				const FName ConstraintName = TransformConstraint->GetFName();
-				if (FConstraintAndActiveChannel* ConstraintChannel = InSection->GetConstraintChannel(ConstraintName))
+				if (FConstraintAndActiveChannel* ConstraintChannel = ConstrainedSection->GetConstraintChannel(ConstraintName))
 				{
 					OtherChannels.Add(ConstraintChannel);
 				}
+
+				//if doing keys or per frame increments we need to also 
+				if (InSettings.IsSet() == true &&
+					(InSettings.GetValue().BakingKeySettings == EBakingKeySettings::KeysOnly || InSettings.GetValue().FrameIncrement > 1))
+				{
+					FMovieSceneConstraintChannelHelper::GetTransformFramesForConstraintHandles(TransformConstraint, InSequencer,
+						InSettings.GetValue().StartFrame, InSettings.GetValue().EndFrame, OutFramesToBake);
+				}
 			}
 		}
 
@@ -155,9 +196,25 @@ void FConstraintBaker::GetMinimalFramesToBake(
 	OutFramesToBake.SetNum(Algo::Unique(OutFramesToBake));
 }
 
+bool FConstraintBaker::BakeMultiple(
+	UWorld* InWorld,
+	TArray< UTickableTransformConstraint*>& InConstraints,
+	const TSharedPtr<ISequencer>& InSequencer,
+	const FBakingAnimationKeySettings& InSettings)
+{
+	TOptional<TArray<FFrameNumber>> Frames;
+	TOptional< FBakingAnimationKeySettings> Settings = InSettings;
+	for (UTickableTransformConstraint* Constraint : InConstraints)
+	{
+		FConstraintBaker::Bake(InWorld, Constraint, InSequencer, Settings, Frames);
+	}
+	return true;
+}
+
 void FConstraintBaker::Bake(UWorld* InWorld, 
 	UTickableTransformConstraint* InConstraint,
 	const TSharedPtr<ISequencer>& InSequencer, 
+	const TOptional< FBakingAnimationKeySettings>& InSettings,
 	const TOptional<TArray<FFrameNumber>>& InFrames)
 {
 	if ((InFrames && InFrames->IsEmpty()) || (InConstraint == nullptr) || (InConstraint->ChildTRSHandle == nullptr))
@@ -167,29 +224,13 @@ void FConstraintBaker::Bake(UWorld* InWorld,
 
 	const TObjectPtr<UTransformableHandle>& Handle = InConstraint->ChildTRSHandle;
 
-	const FConstraintChannelInterfaceRegistry& InterfaceRegistry = FConstraintChannelInterfaceRegistry::Get();	
-	ITransformConstraintChannelInterface* Interface = InterfaceRegistry.FindConstraintChannelInterface(Handle->GetClass());
-	if (!Interface)
+	FConstraintSections ConstraintSections = FMovieSceneConstraintChannelHelper::GetConstraintSectionAndChannel(
+		InConstraint, InSequencer);
+	IMovieSceneConstrainedSection* ConstrainedSection = Cast<IMovieSceneConstrainedSection>(ConstraintSections.ConstraintSection);
+	if (ConstrainedSection == nullptr) 
 	{
 		return;
 	}
-	
-	//get the section to be used later to delete the extra transform keys at the frame -1 times, abort if not there for some reason
-	UMovieSceneSection* ConstraintSection = Interface->GetHandleConstraintSection(Handle, InSequencer);
-	const UMovieSceneSection* TransformSection = Interface->GetHandleSection(Handle, InSequencer);
-
-	IMovieSceneConstrainedSection* ConstrainedSection = Cast<IMovieSceneConstrainedSection>(ConstraintSection);
-	if (ConstrainedSection == nullptr || TransformSection == nullptr)
-	{
-		return;
-	}
-	
-	FConstraintAndActiveChannel* ActiveChannel = ConstrainedSection->GetConstraintChannel(InConstraint->GetFName());
-	if (ActiveChannel == nullptr)
-	{
-		return;
-	}
-
 	// compute transforms
 	TArray<FFrameNumber> FramesToBake;
 	if (InFrames)
@@ -198,12 +239,12 @@ void FConstraintBaker::Bake(UWorld* InWorld,
 	}
 	else
 	{
-		GetMinimalFramesToBake(InWorld, InConstraint, InSequencer, ConstrainedSection, FramesToBake);
+		GetMinimalFramesToBake(InWorld, InConstraint, InSequencer, ConstraintSections.ConstraintSection, InSettings, FramesToBake);
 
 		// if it needs recomposition 
-		if (ConstraintSection != TransformSection && !FramesToBake.IsEmpty())
+		if (ConstraintSections.ConstraintSection != ConstraintSections.ChildTransformSection && !FramesToBake.IsEmpty())
 		{
-			const TMovieSceneChannelData<const bool> ConstraintChannelData = ActiveChannel->ActiveChannel.GetData();
+			const TMovieSceneChannelData<const bool> ConstraintChannelData = ConstraintSections.ActiveChannel->ActiveChannel.GetData();
 			const TArrayView<const FFrameNumber> ConstraintFrames = ConstraintChannelData.GetTimes();
 			const TArrayView<const bool> ConstraintValues = ConstraintChannelData.GetValues();
 			for (int32 Index = 0; Index < ConstraintFrames.Num(); ++Index)
@@ -236,15 +277,15 @@ void FConstraintBaker::Bake(UWorld* InWorld,
 		return;
 	}
 	
-	ConstraintSection->Modify();
+	ConstraintSections.ConstraintSection->Modify();
 
 	// disable constraint and delete extra transform keys
-	TMovieSceneChannelData<bool> ConstraintChannelData = ActiveChannel->ActiveChannel.GetData();
+	TMovieSceneChannelData<bool> ConstraintChannelData = ConstraintSections.ActiveChannel->ActiveChannel.GetData();
 	const TArrayView<const FFrameNumber> ConstraintFrames = ConstraintChannelData.GetTimes();
 	
 	// get transform channels
-	const TArrayView<FMovieSceneFloatChannel*> FloatTransformChannels = Handle->GetFloatChannels(TransformSection);
-	const TArrayView<FMovieSceneDoubleChannel*> DoubleTransformChannels = Handle->GetDoubleChannels(TransformSection);
+	const TArrayView<FMovieSceneFloatChannel*> FloatTransformChannels = Handle->GetFloatChannels(ConstraintSections.ChildTransformSection);
+	const TArrayView<FMovieSceneDoubleChannel*> DoubleTransformChannels = Handle->GetDoubleChannels(ConstraintSections.ChildTransformSection);
 
 	for (int32 Index = 0; Index < ConstraintFrames.Num(); ++Index)
 	{
@@ -270,8 +311,18 @@ void FConstraintBaker::Bake(UWorld* InWorld,
 
 	// now bake to channel curves
 	const EMovieSceneTransformChannel Channels = InConstraint->GetChannelsToKey();
-	Interface->AddHandleTransformKeys(InSequencer, Handle, FramesToBake, Transforms, Channels);
+	ConstraintSections.Interface->AddHandleTransformKeys(InSequencer, Handle, FramesToBake, Transforms, Channels);
 
+	if (InSettings.IsSet() && InSettings.GetValue().BakingKeySettings != EBakingKeySettings::KeysOnly
+		&& InSettings.GetValue().bReduceKeys == true)
+	{
+		FKeyDataOptimizationParams Param;
+		Param.bAutoSetInterpolation = true;
+		Param.Tolerance = InSettings.GetValue().Tolerance;
+		TRange<FFrameNumber> Range(InSettings.GetValue().StartFrame, InSettings.GetValue().EndFrame);
+		Param.Range = Range;
+		MovieSceneToolHelpers::OptimizeSection(Param, ConstraintSections.ConstraintSection);
+	}
 	// notify
 	InSequencer->RequestEvaluate();
 }
diff --git a/Engine/Plugins/Animation/ControlRig/Source/ControlRigEditor/Private/Tools/ControlRigSnapSettings.cpp b/Engine/Plugins/Animation/ControlRig/Source/ControlRigEditor/Private/Tools/ControlRigSnapSettings.cpp
new file mode 100644
--- /dev/null
+++ b/Engine/Plugins/Animation/ControlRig/Source/ControlRigEditor/Private/Tools/ControlRigSnapSettings.cpp
@@ -0,0 +1,9 @@
+// Copyright Epic Games, Inc. All Rights Reserved.
+#include "Tools/ControlRigSnapSettings.h"
+
+#include UE_INLINE_GENERATED_CPP_BY_NAME(ControlRigSnapSettings)
+
+UControlRigSnapSettings::UControlRigSnapSettings()
+{
+}
+
diff --git a/Engine/Plugins/Animation/ControlRig/Source/ControlRigEditor/Private/Tools/ControlRigSnapper.cpp b/Engine/Plugins/Animation/ControlRig/Source/ControlRigEditor/Private/Tools/ControlRigSnapper.cpp
--- a/Engine/Plugins/Animation/ControlRig/Source/ControlRigEditor/Private/Tools/ControlRigSnapper.cpp
+++ b/Engine/Plugins/Animation/ControlRig/Source/ControlRigEditor/Private/Tools/ControlRigSnapper.cpp
@@ -21,6 +21,8 @@
 #include "Sections/MovieScene3DTransformSection.h"
 #include "Tracks/MovieScene3DTransformTrack.h"
 #include "Channels/MovieSceneChannelProxy.h"
+#include "Channels/MovieSceneFloatChannel.h"
+#include "Channels/MovieSceneDoubleChannel.h"
 #include "Tools/ControlRigSnapSettings.h"
 #include "MovieSceneToolsModule.h"
 #include "Components/SkeletalMeshComponent.h"
@@ -29,6 +31,10 @@
 #include "LevelSequencePlayer.h"
 #include "LevelSequenceActor.h"
 #include "Tools/BakingHelper.h"
+#include "Containers/SortedMap.h"
+#include "Sequencer/ControlRigSequencerHelpers.h"
+#include "Constraints/MovieSceneConstraintChannelHelper.inl"
+
 
 #include UE_INLINE_GENERATED_CPP_BY_NAME(ControlRigSnapper)
 
@@ -190,7 +196,7 @@ struct FGuidAndActor
 		const TArray<FTransform>& WorldTransformsToSnapTo, const UControlRigSnapSettings* SnapSettings) const
 	{
 		// get section
-		const UMovieScene3DTransformSection* TransformSection = MovieSceneToolHelpers::GetTransformSection(Sequencer, Guid);
+		UMovieScene3DTransformSection* TransformSection = MovieSceneToolHelpers::GetTransformSection(Sequencer, Guid);
 		if (!TransformSection)
 		{
 			return false;
@@ -242,6 +248,18 @@ struct FGuidAndActor
 		// add keys
 		MovieSceneToolHelpers::AddTransformKeys(TransformSection, Frames, LocalTransforms, Channels);
 
+		if (SnapSettings && SnapSettings->BakingKeySettings != EBakingKeySettings::KeysOnly
+			&& SnapSettings->bReduceKeys == true && Frames.Num() > 2)
+		{
+			FKeyDataOptimizationParams Param;
+			Param.bAutoSetInterpolation = true;
+			Param.Tolerance = SnapSettings->Tolerance;
+			TRange<FFrameNumber> Range(Frames[0], Frames[Frames.Num() -1]);
+			Param.Range = Range;
+			MovieSceneToolHelpers::OptimizeSection(Param, TransformSection);
+		}
+
+
 		// notify
 		Sequencer->NotifyMovieSceneDataChanged(EMovieSceneDataChangeType::MovieSceneStructureItemsChanged);
 
@@ -305,7 +323,173 @@ static bool CalculateWorldTransformsFromParents(ISequencer* Sequencer, const FCo
 	return false;
 }
 
+//to do handle constraints and spaces, but for now we just output a single actor parent.
+static void GetActorParents(const FActorForWorldTransforms& ActorSelection,
+	TArray<FActorForWorldTransforms>& OutParentActors, TArray<FControlRigForWorldTransforms>& OutControlRigs)
+{
+	if (ActorSelection.Actor.IsValid())
+	{
+		//has component so parent is next component
+		if (ActorSelection.Component.IsValid() && ActorSelection.Component->GetAttachParent())
+		{
+			FActorForWorldTransforms OutParentActor;
+			OutParentActor.Actor = ActorSelection.Actor;
+			OutParentActor.Component = ActorSelection.Component->GetAttachParent();
+			OutParentActors.Add(OutParentActor);
+			GetActorParents(OutParentActor, OutParentActors, OutControlRigs);
+		}
+		else if (AActor* ParentActor = ActorSelection.Actor->GetAttachParentActor())
+		{
+			FActorForWorldTransforms OutParentActor;
+			OutParentActor.Actor = ParentActor;
+			OutParentActors.Add(OutParentActor);
+			GetActorParents(OutParentActor, OutParentActors, OutControlRigs);
+		}
+	}
+}
+
+//todo handle constraints/spaces
+static void GetControlRigParents(const FControlRigForWorldTransforms& ControlRigAndSelection, 
+	TArray<FActorForWorldTransforms>& OutParentActors, TArray<FControlRigForWorldTransforms>& OutControlRigs)
+{
+
+	if (UControlRig* ControlRig = ControlRigAndSelection.ControlRig.Get())
+	{
+		URigHierarchy* Hierarchy = ControlRig->GetHierarchy();
+		for (const FName& ControlName : ControlRigAndSelection.ControlNames)
+		{
+			if (FRigControlElement* ControlElement = ControlRig->FindControl(ControlName))
+			{
+				if (FRigControlElement* ParentControlElement = Cast<FRigControlElement>(Hierarchy->GetFirstParent(ControlElement)))
+				{
+					FControlRigForWorldTransforms OutControlRig;
+					OutControlRig.ControlRig = ControlRig;
+					OutControlRig.ControlNames.Add(ParentControlElement->GetKey().Name);
+					OutControlRigs.Add(OutControlRig);
+					GetControlRigParents(OutControlRig, OutParentActors, OutControlRigs);
+				}
+				else
+				{
+					if (TSharedPtr<IControlRigObjectBinding> ObjectBinding = ControlRig->GetObjectBinding())
+					{
+						USceneComponent* Component = Cast<USceneComponent>(ObjectBinding->GetBoundObject());
+						if (!Component)
+						{
+							continue;
+						}
+						AActor* Actor = Component->GetTypedOuter< AActor >();
+						if (!Actor)
+						{
+							continue;
+						}
+						FActorForWorldTransforms OutParentActor;
+						OutParentActor.Actor = Actor;
+						OutParentActors.Add(OutParentActor);
+						GetActorParents(OutParentActor, OutParentActors, OutControlRigs);
+					}
+				}
+			}
+		}
+	}
+}
+
+static UMovieSceneControlRigParameterSection* GetControlRigSection(ISequencer* Sequencer, const UControlRig* ControlRig)
+{
+
+	if (ControlRig == nullptr || Sequencer == nullptr)
+	{
+		return nullptr;
+	}
+	UMovieScene* MovieScene = Sequencer->GetFocusedMovieSceneSequence()->GetMovieScene();
+	if (!MovieScene)
+	{
+		return nullptr;
+	}
+	const TArray<FMovieSceneBinding>& Bindings = MovieScene->GetBindings();
+	for (const FMovieSceneBinding& Binding : Bindings)
+	{
+		UMovieSceneControlRigParameterTrack* ControlRigParameterTrack = Cast<UMovieSceneControlRigParameterTrack>(MovieScene->FindTrack(UMovieSceneControlRigParameterTrack::StaticClass(), Binding.GetObjectGuid(), NAME_None));
+		if (ControlRigParameterTrack && ControlRigParameterTrack->GetControlRig() == ControlRig)
+		{
+			UMovieSceneControlRigParameterSection* ActiveSection = Cast<UMovieSceneControlRigParameterSection>(ControlRigParameterTrack->GetSectionToKey());
+			if (ActiveSection)
+			{
+				return ActiveSection;
+			}
+		}
+	}
+	return nullptr;
+}
+
+//Get Transform frames for specified snap items, will add them to the OutFrameMap
+static void GetTransformFrames(ISequencer * Sequencer, const FControlRigSnapperSelection& SnapItem, FFrameNumber StartFrame, FFrameNumber EndFrame, TSortedMap<FFrameNumber, FFrameNumber>& OutFrameMap)
+{
+	TArray<FFrameNumber> FramesToUse;
+	TArray<FActorForWorldTransforms> ParentActors;
+	TArray<FControlRigForWorldTransforms> ParentControlRigs;
+
+	for (const FActorForWorldTransforms& ActorsFor : SnapItem.Actors)
+	{
+		FramesToUse.Reset();
+		ParentActors.Reset();
+		ParentControlRigs.Reset();
+		ParentActors.Add(ActorsFor);
+		GetActorParents(ActorsFor, ParentActors, ParentControlRigs);
+
+	}
+	for (const FControlRigForWorldTransforms& ControlRigsFor : SnapItem.ControlRigs)
+	{
+		FramesToUse.Reset();
+		ParentActors.Reset();
+		ParentControlRigs.Reset();
+		ParentControlRigs.Add(ControlRigsFor);
+		GetControlRigParents(ControlRigsFor, ParentActors, ParentControlRigs);
+	}
 
+	for (FActorForWorldTransforms& Actor : ParentActors)
+	{
+		FGuid Guid;
+		if (Actor.Component.Get())
+		{
+			Guid = Sequencer->GetHandleToObject(Actor.Component.Get(), false);
+		}
+		if (!Guid.IsValid())
+		{
+			Guid = Sequencer->GetHandleToObject(Actor.Actor.Get(), false);
+		}
+		if (Guid.IsValid())
+		{
+			if (UMovieScene3DTransformSection* TransformSection = MovieSceneToolHelpers::GetTransformSection(Sequencer, Guid))
+			{
+				TArrayView<FMovieSceneDoubleChannel*> Channels = TransformSection->GetChannelProxy().GetChannels<FMovieSceneDoubleChannel>();
+				TArray<FFrameNumber> TransformFrameTimes = FMovieSceneConstraintChannelHelper::GetTransformTimes(
+					Channels, StartFrame, EndFrame);
+				for (FFrameNumber& FrameNumber : TransformFrameTimes)
+				{
+					OutFrameMap.Add(FrameNumber);
+				}
+			}
+		}
+	}
+	for (FControlRigForWorldTransforms& ControlRig : ParentControlRigs)
+	{
+		UMovieSceneControlRigParameterSection* Section = GetControlRigSection(Sequencer, ControlRig.ControlRig.Get());
+		if (Section)
+		{
+			for (FName& ControlName : ControlRig.ControlNames)
+			{
+				TArrayView<FMovieSceneFloatChannel*> Channels = FControlRigSequencerHelpers::GetFloatChannels(ControlRig.ControlRig.Get(),
+					ControlName, Section);
+				TArray<FFrameNumber> TransformFrameTimes = FMovieSceneConstraintChannelHelper::GetTransformTimes(
+					Channels, StartFrame, EndFrame);
+				for (FFrameNumber& FrameNumber : TransformFrameTimes)
+				{
+					OutFrameMap.Add(FrameNumber);
+				}
+			}
+		}
+	}
+}
 bool FControlRigSnapper::SnapIt(FFrameNumber StartFrame, FFrameNumber EndFrame,const FControlRigSnapperSelection& ActorToSnap,
 	const FControlRigSnapperSelection& ParentToSnap, const UControlRigSnapSettings* SnapSettings)
 {
@@ -320,8 +504,18 @@ bool FControlRigSnapper::SnapIt(FFrameNumber StartFrame, FFrameNumber EndFrame,c
 		MovieScene->Modify();
 		
 		TArray<FFrameNumber> Frames;
-		MovieSceneToolHelpers::CalculateFramesBetween(MovieScene, StartFrame, EndFrame, Frames);
+		if (SnapSettings->BakingKeySettings == EBakingKeySettings::AllFrames)
+		{
+			MovieSceneToolHelpers::CalculateFramesBetween(MovieScene, StartFrame, EndFrame, SnapSettings->FrameIncrement,Frames);
+		}
+		else
+		{
+			TSortedMap<FFrameNumber, FFrameNumber> FrameMap;
+			GetTransformFrames(Sequencer, ParentToSnap, StartFrame, EndFrame, FrameMap);
+			GetTransformFrames(Sequencer, ActorToSnap,  StartFrame, EndFrame, FrameMap);
+			FrameMap.GenerateKeyArray(Frames);
 
+		}
 		TArray<FTransform> WorldTransformToSnap;
 		bool bSnapToFirstFrameNotParents = !CalculateWorldTransformsFromParents(Sequencer, ParentToSnap, Frames, WorldTransformToSnap);
 		if (Frames.Num() != WorldTransformToSnap.Num())
@@ -480,6 +674,18 @@ bool FControlRigSnapper::SnapIt(FFrameNumber StartFrame, FFrameNumber EndFrame,c
 							ControlRig->SetControlGlobalTransform(Name, GlobalTransform, true, Context);
 						}
 					}
+
+					UMovieSceneControlRigParameterSection* ControlRigSection = GetControlRigSection(Sequencer, ControlRig);
+					if (ControlRigSection && SnapSettings && SnapSettings->BakingKeySettings != EBakingKeySettings::KeysOnly
+						&& SnapSettings->bReduceKeys == true && Frames.Num() > 2)
+					{
+						FKeyDataOptimizationParams Param;
+						Param.bAutoSetInterpolation = true;
+						Param.Tolerance = SnapSettings->Tolerance;
+						TRange<FFrameNumber> Range(Frames[0], Frames[Frames.Num() - 1]);
+						Param.Range = Range;
+						MovieSceneToolHelpers::OptimizeSection(Param, ControlRigSection);
+					}
 				}
 			}
 		}
diff --git a/Engine/Plugins/Animation/ControlRig/Source/ControlRigEditor/Public/ControlRigSequencerEditorLibrary.h b/Engine/Plugins/Animation/ControlRig/Source/ControlRigEditor/Public/ControlRigSequencerEditorLibrary.h
--- a/Engine/Plugins/Animation/ControlRig/Source/ControlRigEditor/Public/ControlRigSequencerEditorLibrary.h
+++ b/Engine/Plugins/Animation/ControlRig/Source/ControlRigEditor/Public/ControlRigSequencerEditorLibrary.h
@@ -23,6 +23,7 @@ class ULevelSequence;
 class UTickableConstraint;
 class UTickableTransformConstraint;
 class UTransformableHandle;
+struct FBakingAnimationKeySettings;
 
 USTRUCT(BlueprintType)
 struct FControlRigSequencerBindingProxy
@@ -142,13 +143,22 @@ public:
 	* Bake the constraint to keys based on the passed in frames. This will use the open sequencer to bake. See ConstraintsScriptingLibrary to get the list of available constraints
 	* @param World The active world
 	* @param Constraint The Constraint to bake. After baking it will be keyed to be inactive of the range of frames that are baked
-	* @param Frames The frames to bake, if the array is empty it will use the active time ranges of the constraint to deteremine where it should bake
+	* @param Frames The frames to bake, if the array is empty it will use the active time ranges of the constraint to determine where it should bake
 	* @param TimeUnit Unit for all frame and time values, either in display rate or tick resolution
 	* @return Returns True if successful, False otherwise
 	*/
 	UFUNCTION(BlueprintCallable, Category = "Editor Scripting | Sequencer Tools | Control Rig")
 	static bool BakeConstraint(UWorld* World, UTickableConstraint* Constraint, const TArray<FFrameNumber>& Frames, ESequenceTimeUnit TimeUnit = ESequenceTimeUnit::DisplayRate);
 
+	/**
+	* Bake the constraint to keys based on the passed in settings. This will use the open sequencer to bake. See ConstraintsScriptingLibrary to get the list of available constraints
+	* @param World The active world
+	* @param InConstraints The Constraints tobake.  After baking they will be keyed to be inactive of the range of frames that are baked
+	* @param InSettings Settings to use for baking
+	* @return Returns True if successful, False otherwise
+	*/
+	UFUNCTION(BlueprintCallable, Category = "Editor Scripting | Sequencer Tools | Control Rig")
+	static bool BakeConstraints(UWorld* World, TArray<UTickableConstraint*>& InConstraints,const FBakingAnimationKeySettings& InSettings);
 	/**
 	* Add a constraint possibly adding to sequencer also if one is open.
 	* @param World The active world
@@ -926,6 +936,16 @@ public:
 	UFUNCTION(BlueprintCallable, Category = "Editor Scripting | Sequencer Tools | Control Rig")
 	static bool CollapseControlRigAnimLayers(ULevelSequence* InSequence,UMovieSceneControlRigParameterTrack* InTrack, bool bKeyReduce = false, float Tolerance = 0.001f);
 
+	/*
+	 * Collapse and bake all sections and layers on a control rig track to just one section using passed in settings.
+	 *
+	 * @param InSequence Sequence that has track to collapse
+	 * @param InTrack Track for layers to collapse
+	 * @param InSettings Settings that determine how to collapse
+	 */
+	UFUNCTION(BlueprintCallable, Category = "Editor Scripting | Sequencer Tools | Control Rig")
+	static bool CollapseControlRigAnimLayersWithSettings(ULevelSequence* InSequence, UMovieSceneControlRigParameterTrack* InTrack, const FBakingAnimationKeySettings& InSettings);
+
 	/*
 	 * Get the default parent key, can be used a parent space.
 	 *
diff --git a/Engine/Plugins/Animation/ControlRig/Source/ControlRigEditor/Public/ControlRigSpaceChannelEditors.h b/Engine/Plugins/Animation/ControlRig/Source/ControlRigEditor/Public/ControlRigSpaceChannelEditors.h
--- a/Engine/Plugins/Animation/ControlRig/Source/ControlRigEditor/Public/ControlRigSpaceChannelEditors.h
+++ b/Engine/Plugins/Animation/ControlRig/Source/ControlRigEditor/Public/ControlRigSpaceChannelEditors.h
@@ -49,7 +49,7 @@ struct FControlRigSpaceChannelHelpers
 	static void CompensateIfNeeded(UControlRig* ControlRig, ISequencer* Sequencer, UMovieSceneControlRigParameterSection* Section, FName ControlName, TOptional<FFrameNumber>& Time);
 	static FSpaceChannelAndSection FindSpaceChannelAndSectionForControl(UControlRig* ControlRig, FName ControlName, ISequencer* Sequencer, bool bCreateIfNeeded);
 	static void SequencerBakeControlInSpace(UControlRig* ControlRig, ISequencer* Sequencer, FMovieSceneControlRigSpaceChannel* Channel, UMovieSceneSection* SectionToKey,
-		TArray<FFrameNumber> Frames, URigHierarchy* RigHierarchy, const FRigElementKey& ControlKey, FRigSpacePickerBakeSettings InSettings);
+		URigHierarchy* RigHierarchy, const FRigElementKey& ControlKey, FRigSpacePickerBakeSettings InSettings);
 	static void GetFramesInThisSpaceAfterThisTime(UControlRig* ControlRig, FName ControlName, FMovieSceneControlRigSpaceBaseKey CurrentValue,
 		FMovieSceneControlRigSpaceChannel* Channel, UMovieSceneSection* SectionToKey,
 		FFrameNumber Time, TSortedMap<FFrameNumber,FFrameNumber>& OutMoreFrames);
diff --git a/Engine/Plugins/Animation/ControlRig/Source/ControlRigEditor/Public/RigSpacePickerBakeSettings.h b/Engine/Plugins/Animation/ControlRig/Source/ControlRigEditor/Public/RigSpacePickerBakeSettings.h
--- a/Engine/Plugins/Animation/ControlRig/Source/ControlRigEditor/Public/RigSpacePickerBakeSettings.h
+++ b/Engine/Plugins/Animation/ControlRig/Source/ControlRigEditor/Public/RigSpacePickerBakeSettings.h
@@ -4,8 +4,10 @@
 
 #include "CoreMinimal.h"
 #include "Rigs/RigHierarchyDefines.h"
+#include "BakingAnimationKeySettings.h"
 #include "RigSpacePickerBakeSettings.generated.h"
 
+
 USTRUCT(BlueprintType)
 struct CONTROLRIGEDITOR_API FRigSpacePickerBakeSettings
 {
@@ -14,24 +16,28 @@ struct CONTROLRIGEDITOR_API FRigSpacePickerBakeSettings
 	FRigSpacePickerBakeSettings()
 	{
 		TargetSpace = FRigElementKey();
-		StartFrame = 0;
-		EndFrame = 100;
-		bReduceKeys = false;
-		Tolerance = 0.001f;
+		Settings = FBakingAnimationKeySettings();
 	}
 
 	UPROPERTY(BlueprintReadWrite, Category = "Settings")
 	FRigElementKey TargetSpace;
 
-	UPROPERTY(BlueprintReadWrite,EditAnywhere, Category = "Settings")
-	FFrameNumber StartFrame;
+	UPROPERTY(BlueprintReadWrite, EditAnywhere, Category = "Bake", meta = (ShowOnlyInnerProperties))
+	FBakingAnimationKeySettings Settings;
+
+	/** DEPRECATED 5.3 */
+	UPROPERTY(meta = (DeprecatedProperty, DeprecationMessage = "Use Settings.StartFrame instead"))
+	FFrameNumber StartFrame_DEPRECATED;
 
-	UPROPERTY(BlueprintReadWrite,EditAnywhere, Category = "Settings")
-	FFrameNumber EndFrame;
+	/** DEPRECATED 5.3 */
+	UPROPERTY(meta = (DeprecatedProperty, DeprecationMessage = "Use Settings.EndFrame instead"))
+	FFrameNumber EndFrame_DEPRECATED;
 
-	UPROPERTY(BlueprintReadWrite,EditAnywhere, Category = "Settings")
-	bool bReduceKeys;
+	/** DEPRECATED 5.3 */
+	UPROPERTY(meta = (DeprecatedProperty, DeprecationMessage = "Use Settings.bReduceKeys instead"))
+	bool bReduceKeys_DEPRECATED;
 
-	UPROPERTY(BlueprintReadWrite,EditAnywhere, Category = "Settings")
-	float Tolerance;
-};
+	/** DEPRECATED 5.3 */
+	UPROPERTY(meta = (DeprecatedProperty, DeprecationMessage = "Use Settings.Tolerance instead"))
+	float Tolerance_DEPRECATED;
+};
\ No newline at end of file
diff --git a/Engine/Plugins/Animation/ControlRig/Source/ControlRigEditor/Public/Tools/ConstraintBaker.h b/Engine/Plugins/Animation/ControlRig/Source/ControlRigEditor/Public/Tools/ConstraintBaker.h
--- a/Engine/Plugins/Animation/ControlRig/Source/ControlRigEditor/Public/Tools/ConstraintBaker.h
+++ b/Engine/Plugins/Animation/ControlRig/Source/ControlRigEditor/Public/Tools/ConstraintBaker.h
@@ -4,14 +4,14 @@
 
 #include "Misc/FrameNumber.h"
 #include "Templates/SharedPointer.h"
+#include "BakingAnimationKeySettings.h"
 
-class IMovieSceneConstrainedSection;
+class UMovieSceneSection;
 class UTransformableHandle;
 class UTickableTransformConstraint;
 class ISequencer;
 class UWorld;
 struct FFrameNumber;
-
 struct FConstraintBaker
 {
 public:
@@ -20,8 +20,16 @@ public:
 		UWorld* InWorld,
 		UTickableTransformConstraint* InConstraint,
 		const TSharedPtr<ISequencer>& InSequencer,
+		const TOptional< FBakingAnimationKeySettings>& InSettings,
 		const TOptional<TArray<FFrameNumber>>& InFrames);
 	
+	/** Bake Multiple Constraints using settings*/
+	static bool BakeMultiple(
+		UWorld* InWorld,
+		TArray< UTickableTransformConstraint*>& InConstraints,
+		const TSharedPtr<ISequencer>& InSequencer,
+		const FBakingAnimationKeySettings& Settings);
+	
 private:
 
 	/** Computes the minimal frames to bake for that constraint. */
@@ -29,7 +37,8 @@ private:
 		UWorld* InWorld,
 		const UTickableTransformConstraint* InConstraint,
 		const TSharedPtr<ISequencer>& InSequencer,
-		IMovieSceneConstrainedSection* InSection,
+		UMovieSceneSection* InSection,
+		const TOptional< FBakingAnimationKeySettings>& Settings,
 		TArray<FFrameNumber>& OutFramesToBake);
 };
 
diff --git a/Engine/Plugins/Animation/ControlRig/Source/ControlRigEditor/Public/Tools/ControlRigSnapSettings.h b/Engine/Plugins/Animation/ControlRig/Source/ControlRigEditor/Public/Tools/ControlRigSnapSettings.h
new file mode 100644
--- /dev/null
+++ b/Engine/Plugins/Animation/ControlRig/Source/ControlRigEditor/Public/Tools/ControlRigSnapSettings.h
@@ -0,0 +1,53 @@
+// Copyright Epic Games, Inc. All Rights Reserved.
+#pragma once
+/**
+* Per Project Settings Animation Snapping
+*/
+#include "CoreMinimal.h"
+#include "UObject/ObjectMacros.h"
+#include "UObject/Object.h"
+#include "ControlRig.h"
+#include "BakingAnimationKeySettings.h"
+#include "ControlRigSnapSettings.generated.h"
+
+
+UCLASS(config = EditorPerProjectUserSettings)
+class CONTROLRIGEDITOR_API UControlRigSnapSettings : public UObject
+{
+
+public:
+	GENERATED_BODY()
+
+	UControlRigSnapSettings();
+
+	/** When snapping keep offset, If false snap completely */
+	UPROPERTY(config, EditAnywhere, BlueprintReadWrite, Category = "Snap Settings", meta = (ToolTip = "When snapping keep offset, if false snap completely"))
+	bool bKeepOffset = false;
+
+	/** When snapping, also snap position */
+	UPROPERTY(config, EditAnywhere, BlueprintReadWrite, Category = "Snap Settings", meta = (ToolTip = "When snapping, also snap position"))
+	bool bSnapPosition = true;
+
+	/** When snapping, also snap rotation */
+	UPROPERTY(config, EditAnywhere, BlueprintReadWrite, Category = "Snap Settings", meta = (ToolTip = "When snapping, also snap rotation"))
+	bool bSnapRotation = false;
+
+	/** When snapping, also snap scale */
+	UPROPERTY(config, EditAnywhere, BlueprintReadWrite, Category = "Snap Settings", meta = (ToolTip = "When snapping, also snap scale"))
+	bool bSnapScale = false;
+
+	UPROPERTY(BlueprintReadWrite, EditAnywhere, Category = "Bake")
+	EBakingKeySettings BakingKeySettings = EBakingKeySettings::KeysOnly;
+
+	UPROPERTY(BlueprintReadWrite, EditAnywhere, Category = "Bake", meta = (ClampMin = "1", UIMin = "1", EditCondition = "BakingKeySettings == EBakingKeySettings::AllFrames"))
+	int32 FrameIncrement = 1;
+
+	UPROPERTY(BlueprintReadWrite, EditAnywhere, Category = "Bake", meta = (EditCondition = "BakingKeySettings == EBakingKeySettings::AllFrames"))
+	bool bReduceKeys = false;
+
+	UPROPERTY(BlueprintReadWrite, EditAnywhere, Category = "Bake", meta = (EditCondition = "BakingKeySettings == EBakingKeySettings::AllFrames || bReduceKeys"))
+	float Tolerance = 0.001f;
+
+};
+
+
diff --git a/Engine/Source/Editor/MovieSceneTools/Private/Constraints/MovieSceneConstraintChannelHelper.cpp b/Engine/Source/Editor/MovieSceneTools/Private/Constraints/MovieSceneConstraintChannelHelper.cpp
--- a/Engine/Source/Editor/MovieSceneTools/Private/Constraints/MovieSceneConstraintChannelHelper.cpp
+++ b/Engine/Source/Editor/MovieSceneTools/Private/Constraints/MovieSceneConstraintChannelHelper.cpp
@@ -18,6 +18,7 @@
 #include "MovieSceneSpawnableAnnotation.h"
 #include "Sections/MovieSceneConstrainedSection.h"
 #include "MovieSceneToolsModule.h"
+#include "Channels/MovieSceneChannelTraits.h"
 #include "Constraints/TransformConstraintChannelInterface.h"
 
 
@@ -1002,6 +1003,104 @@ void FMovieSceneConstraintChannelHelper::CompensateIfNeeded(
 		InSequencer->ForceEvaluate();
 	}
 }
+FConstraintSections FMovieSceneConstraintChannelHelper::GetConstraintSectionAndChannel(
+	const UTickableTransformConstraint* InConstraint,
+	const TSharedPtr<ISequencer>& InSequencer)
+{
+	FConstraintSections ReturnValue;
+
+	if (InSequencer.IsValid() == false)
+	{
+		return ReturnValue;
+	}
+	const TObjectPtr<UTransformableHandle>& ChildHandle = InConstraint->ChildTRSHandle;
+
+	const FConstraintChannelInterfaceRegistry& InterfaceRegistry = FConstraintChannelInterfaceRegistry::Get();
+	ReturnValue.Interface = InterfaceRegistry.FindConstraintChannelInterface(ChildHandle->GetClass());
+	if (!ReturnValue.Interface)
+	{
+		return ReturnValue;
+	}
+	//get the section to be used later to delete the extra transform keys at the frame -1 times, abort if not there for some reason
+	ReturnValue.ConstraintSection = ReturnValue.Interface->GetHandleConstraintSection(ChildHandle, InSequencer);
+	ReturnValue.ChildTransformSection = ReturnValue.Interface->GetHandleSection(ChildHandle, InSequencer);
+
+	ITransformConstraintChannelInterface* ParentInterface = InterfaceRegistry.FindConstraintChannelInterface(InConstraint->ParentTRSHandle->GetClass());
+	if (ParentInterface)
+	{
+		ReturnValue.ParentTransformSection = ParentInterface->GetHandleSection(InConstraint->ParentTRSHandle, InSequencer);
+	}
+	IMovieSceneConstrainedSection* ConstrainedSection = Cast<IMovieSceneConstrainedSection>(ReturnValue.ConstraintSection);
+	if (ConstrainedSection == nullptr )
+	{
+		return ReturnValue;
+	}
+
+	ReturnValue.ActiveChannel = ConstrainedSection->GetConstraintChannel(InConstraint->GetFName());
+	return ReturnValue;
+}
+
+ void FMovieSceneConstraintChannelHelper::GetTransformFramesForConstraintHandles(
+	const UTickableTransformConstraint* InConstraint,
+	const TSharedPtr<ISequencer>& InSequencer,
+	const FFrameNumber& StartFrame,
+	const FFrameNumber& EndFrame,
+	TArray<FFrameNumber>& OutFramesToBake)
+{
+	if (!InConstraint && !InConstraint->ChildTRSHandle && !InConstraint->ParentTRSHandle)
+	{
+		return;
+	}
+
+	FConstraintSections ConstraintSections = FMovieSceneConstraintChannelHelper::GetConstraintSectionAndChannel(
+		InConstraint, InSequencer);
+	if (ConstraintSections.ChildTransformSection)
+	{
+		const TArrayView<FMovieSceneFloatChannel*> FloatTransformChannels = InConstraint->ChildTRSHandle->GetFloatChannels(ConstraintSections.ChildTransformSection);
+		TArray<FFrameNumber> TransformFrameTimes = FMovieSceneConstraintChannelHelper::GetTransformTimes(
+			FloatTransformChannels, StartFrame, EndFrame);
+		//add transforms keys to bake
+		{
+			for (FFrameNumber& Frame : TransformFrameTimes)
+			{
+				OutFramesToBake.Add(Frame);
+			}
+		}
+		const TArrayView<FMovieSceneDoubleChannel*> DoubleTransformChannels = InConstraint->ChildTRSHandle->GetDoubleChannels(ConstraintSections.ChildTransformSection);
+		TransformFrameTimes = FMovieSceneConstraintChannelHelper::GetTransformTimes(
+			DoubleTransformChannels, StartFrame, EndFrame);
+		//add transforms keys to bake
+		{
+			for (FFrameNumber& Frame : TransformFrameTimes)
+			{
+				OutFramesToBake.Add(Frame);
+			}
+		}
+	}
+	if (ConstraintSections.ParentTransformSection)
+	{
+		const TArrayView<FMovieSceneFloatChannel*> FloatTransformChannels = InConstraint->ParentTRSHandle->GetFloatChannels(ConstraintSections.ParentTransformSection);
+		TArray<FFrameNumber> TransformFrameTimes = FMovieSceneConstraintChannelHelper::GetTransformTimes(
+			FloatTransformChannels, StartFrame, EndFrame);
+		//add transforms keys to bake
+		{
+			for (FFrameNumber& Frame : TransformFrameTimes)
+			{
+				OutFramesToBake.Add(Frame);
+			}
+		}
+		const TArrayView<FMovieSceneDoubleChannel*> DoubleTransformChannels = InConstraint->ParentTRSHandle->GetDoubleChannels(ConstraintSections.ParentTransformSection);
+		TransformFrameTimes = FMovieSceneConstraintChannelHelper::GetTransformTimes(
+			DoubleTransformChannels, StartFrame, EndFrame);
+		//add transforms keys to bake
+		{
+			for (FFrameNumber& Frame : TransformFrameTimes)
+			{
+				OutFramesToBake.Add(Frame);
+			}
+		}
+	}
+}
 
 ITransformConstraintChannelInterface* FMovieSceneConstraintChannelHelper::GetHandleInterface(const UTransformableHandle* InHandle)
 {
@@ -1014,6 +1113,7 @@ ITransformConstraintChannelInterface* FMovieSceneConstraintChannelHelper::GetHan
 	return InterfaceRegistry.FindConstraintChannelInterface(InHandle->GetClass());
 }
 
+
 void FMovieSceneConstraintChannelHelper::CreateBindingIDForHandle(const TSharedPtr<ISequencer>& InSequencer, UTransformableHandle* InHandle)
 {
 	if (InHandle == nullptr || InSequencer.IsValid() == false)
@@ -1144,4 +1244,4 @@ void FMovieSceneConstraintChannelHelper::CompensateScale(
 
 	// reset scaling to reference value
 	InParentConstraint->SetScaling(bRefScalingValue);
-}
\ No newline at end of file
+}
diff --git a/Engine/Source/Editor/MovieSceneTools/Private/MovieSceneToolHelpers.cpp b/Engine/Source/Editor/MovieSceneTools/Private/MovieSceneToolHelpers.cpp
--- a/Engine/Source/Editor/MovieSceneTools/Private/MovieSceneToolHelpers.cpp
+++ b/Engine/Source/Editor/MovieSceneTools/Private/MovieSceneToolHelpers.cpp
@@ -77,6 +77,8 @@
 #include "FrameNumberDetailsCustomization.h"
 #include "PropertyEditorDelegates.h"
 #include "INodeAndChannelMappings.h"
+#include "BakingAnimationKeySettings.h"
+#include "Channels/MovieSceneChannelTraits.h"
 #include "Channels/MovieSceneIntegerChannel.h"
 #include "Channels/MovieSceneByteChannel.h"
 #include "Framework/MultiBox/MultiBoxBuilder.h"
@@ -4614,41 +4616,63 @@ void MovieSceneToolHelpers::SetOrAddKey(TMovieSceneChannelData<FMovieSceneFloatV
 	}
 }
 
+void MovieSceneToolHelpers::SetOrAddKey(TMovieSceneChannelData<FMovieSceneFloatValue>& ChannelData, FFrameNumber Time, const FMovieSceneFloatValue& Value)
+{
+	int32 ExistingIndex = ChannelData.FindKey(Time);
+	if (ExistingIndex != INDEX_NONE)
+	{
+		FMovieSceneFloatValue& FloatValue = ChannelData.GetValues()[ExistingIndex]; //-V758
+		FloatValue.Value = Value.Value;
+	}
+	else
+	{
+		ChannelData.AddKey(Time, Value);
+	}
+}
+
 void MovieSceneToolHelpers::SetOrAddKey(TMovieSceneChannelData<FMovieSceneFloatValue>& ChannelData, FFrameNumber Time, float Value, 
 		float ArriveTangent, float LeaveTangent, ERichCurveInterpMode InterpMode, ERichCurveTangentMode TangentMode, 
 		FFrameRate FrameRate, ERichCurveTangentWeightMode WeightedMode, float ArriveTangentWeight, float LeaveTangentWeight)
 {
-	if (ChannelData.FindKey(Time) == INDEX_NONE)
-	{
-		FMovieSceneFloatValue NewKey(Value);
-
-		NewKey.InterpMode = InterpMode;
-		NewKey.TangentMode = TangentMode;
-		NewKey.Tangent.ArriveTangent = ArriveTangent / FrameRate.AsDecimal();
-		NewKey.Tangent.LeaveTangent = LeaveTangent / FrameRate.AsDecimal();
-		NewKey.Tangent.TangentWeightMode = WeightedMode;
-		NewKey.Tangent.ArriveTangentWeight = ArriveTangentWeight;
-		NewKey.Tangent.LeaveTangentWeight = LeaveTangentWeight;
-		ChannelData.AddKey(Time, NewKey);
-	}
+	FMovieSceneFloatValue NewKey(Value);
+
+	NewKey.InterpMode = InterpMode;
+	NewKey.TangentMode = TangentMode;
+	NewKey.Tangent.ArriveTangent = ArriveTangent / FrameRate.AsDecimal();
+	NewKey.Tangent.LeaveTangent = LeaveTangent / FrameRate.AsDecimal();
+	NewKey.Tangent.TangentWeightMode = WeightedMode;
+	NewKey.Tangent.ArriveTangentWeight = ArriveTangentWeight;
+	NewKey.Tangent.LeaveTangentWeight = LeaveTangentWeight;
+	ChannelData.AddKey(Time, NewKey);
+	MovieSceneToolHelpers::SetOrAddKey(ChannelData, Time, NewKey);
 }
 
 void MovieSceneToolHelpers::SetOrAddKey(TMovieSceneChannelData<FMovieSceneDoubleValue>& ChannelData, FFrameNumber Time, double Value, 
 		float ArriveTangent, float LeaveTangent, ERichCurveInterpMode InterpMode, ERichCurveTangentMode TangentMode, 
 		FFrameRate FrameRate, ERichCurveTangentWeightMode WeightedMode, float ArriveTangentWeight, float LeaveTangentWeight)
 {
-	if (ChannelData.FindKey(Time) == INDEX_NONE)
-	{
-		FMovieSceneDoubleValue NewKey(Value);
+	FMovieSceneDoubleValue NewKey(Value);
+	NewKey.InterpMode = InterpMode;
+	NewKey.TangentMode = TangentMode;
+	NewKey.Tangent.ArriveTangent = ArriveTangent / FrameRate.AsDecimal();
+	NewKey.Tangent.LeaveTangent = LeaveTangent / FrameRate.AsDecimal();
+	NewKey.Tangent.TangentWeightMode = WeightedMode;
+	NewKey.Tangent.ArriveTangentWeight = ArriveTangentWeight;
+	NewKey.Tangent.LeaveTangentWeight = LeaveTangentWeight;
+	MovieSceneToolHelpers::SetOrAddKey(ChannelData, Time, NewKey);
+}
 
-		NewKey.InterpMode = InterpMode;
-		NewKey.TangentMode = TangentMode;
-		NewKey.Tangent.ArriveTangent = ArriveTangent / FrameRate.AsDecimal();
-		NewKey.Tangent.LeaveTangent = LeaveTangent / FrameRate.AsDecimal();
-		NewKey.Tangent.TangentWeightMode = WeightedMode;
-		NewKey.Tangent.ArriveTangentWeight = ArriveTangentWeight;
-		NewKey.Tangent.LeaveTangentWeight = LeaveTangentWeight;
-		ChannelData.AddKey(Time, NewKey);
+void MovieSceneToolHelpers::SetOrAddKey(TMovieSceneChannelData<FMovieSceneDoubleValue>& ChannelData, FFrameNumber Time, FMovieSceneDoubleValue Value)
+{
+	int32 ExistingIndex = ChannelData.FindKey(Time);
+	if (ExistingIndex != INDEX_NONE)
+	{
+		FMovieSceneDoubleValue& DoubleValue = ChannelData.GetValues()[ExistingIndex]; //-V758
+		DoubleValue.Value = Value.Value;
+	}
+	else
+	{
+		ChannelData.AddKey(Time, Value);
 	}
 }
 
@@ -4716,6 +4740,7 @@ void MovieSceneToolHelpers::CalculateFramesBetween(
 	const UMovieScene* MovieScene,
 	FFrameNumber StartFrame,
 	FFrameNumber EndFrame,
+	int32 FrameInc,
 	TArray<FFrameNumber>& OutFrames)
 {
 	const bool bReverse = StartFrame > EndFrame;
@@ -4739,20 +4764,319 @@ void MovieSceneToolHelpers::CalculateFramesBetween(
 
 	if (bReverse)
 	{
+		int32 FrameIndex = 0;
 		for (FFrameNumber DisplayFrameNumber = EndTimeInDisplay; DisplayFrameNumber >= StartTimeInDisplay; --DisplayFrameNumber)
 		{
-			FFrameNumber TickFrameNumber = FFrameRate::TransformTime(FFrameTime(DisplayFrameNumber), DisplayResolution, TickResolution).FrameNumber;
-			OutFrames.Add(TickFrameNumber);
+			if (FrameIndex % FrameInc == 0)
+			{
+				FFrameNumber TickFrameNumber = FFrameRate::TransformTime(FFrameTime(DisplayFrameNumber), DisplayResolution, TickResolution).FrameNumber;
+				OutFrames.Add(TickFrameNumber);
+			}
+			++FrameIndex;
 		}
 	}
 	else
 	{
 		for (FFrameNumber DisplayFrameNumber = StartTimeInDisplay; DisplayFrameNumber <= EndTimeInDisplay; ++DisplayFrameNumber)
 		{
-			FFrameNumber TickFrameNumber = FFrameRate::TransformTime(FFrameTime(DisplayFrameNumber), DisplayResolution, TickResolution).FrameNumber;
-			OutFrames.Add(TickFrameNumber);
+			int32 FrameIndex = 0;
+			if (FrameIndex % FrameInc == 0)
+			{
+				FFrameNumber TickFrameNumber = FFrameRate::TransformTime(FFrameTime(DisplayFrameNumber), DisplayResolution, TickResolution).FrameNumber;
+				OutFrames.Add(TickFrameNumber);
+			}
+			++FrameIndex;
+		}
+	}
+}
+
+UENUM(BlueprintType)
+enum class FChannelMergeAlgorithm : uint8
+{
+	/**Average values together*/
+	Average,
+	/**Add values together*/
+	Add
+};
+
+//merge the channels onto the first one using the passed in algorithm, we may skip the last channel also (since it may be weight).
+template<typename ChannelType>
+static bool MergeChannels(TArray<ChannelType*>& Channels,const  TArray<UMovieSceneSection*>& Sections, const TRange<FFrameNumber>& Range,
+	FChannelMergeAlgorithm MergeAlgorithm)
+{
+	if (Channels.Num() < 2 && Channels.Num() != Sections.Num())
+	{
+		return false;
+	}
+	using ChannelValueType = typename ChannelType::ChannelValueType;
+	using CurveValueType = typename ChannelType::CurveValueType;
+
+	//base channel we set values on
+	ChannelType* BaseChannel = Channels[0];
+	TMovieSceneChannelData<ChannelValueType> BaseChannelData = BaseChannel->GetData();
+	//iterate over each key
+	TArray<FFrameNumber> KeyTimes;
+	TArray<FKeyHandle> Handles;
+	//cached set that we set at the end
+	TArray<TPair< FFrameNumber, ChannelValueType>> KeysToSet;
+	for (int32 ChannelIndex = 0; ChannelIndex < Channels.Num(); ++ChannelIndex)
+	{
+		KeyTimes.Reset();
+		Handles.Reset();
+		ChannelType* Channel = Channels[ChannelIndex];
+		Channel->GetKeys(Range, &KeyTimes, &Handles);
+		double DNumChannels = (double)(Channels.Num());
+		for (int32 FrameIndex = 0; FrameIndex < KeyTimes.Num(); ++FrameIndex)
+		{
+			const FFrameNumber& Frame = KeyTimes[FrameIndex];
+			const FFrameTime FrameTime(Frame);
+			int32 KeyIndex = Channel->GetData().GetIndex(Handles[FrameIndex]);
+			ChannelValueType Value = Channel->GetData().GetValues()[KeyIndex];
+			//got value with tangents and times, now we perform the operation
+			Value.Value = 0.0; //zero out the value we calculate it 
+			if (MergeAlgorithm == FChannelMergeAlgorithm::Average)
+			{
+				for (int32 WeightIndex = 0; WeightIndex < Sections.Num(); ++WeightIndex)
+				{
+					float Weight = Sections[WeightIndex]->GetTotalWeightValue(FrameTime);
+					CurveValueType WeightedValue = 0.0;
+					ChannelType* EachChannel = Channels[WeightIndex];
+					EachChannel->Evaluate(FrameTime, WeightedValue);
+					WeightedValue *= ((double)Weight / DNumChannels);
+					Value.Value += WeightedValue;
+				}
+			}
+			else if (MergeAlgorithm == FChannelMergeAlgorithm::Add)
+			{
+				for (int32 WeightIndex = 0; WeightIndex < Sections.Num(); ++WeightIndex)
+				{
+					float Weight = Sections[WeightIndex]->GetTotalWeightValue(FrameTime);
+					CurveValueType WeightedValue = 0.0;
+					ChannelType* EachChannel = Channels[WeightIndex];
+					EachChannel->Evaluate(FrameTime, WeightedValue);
+					WeightedValue *= Weight;
+					Value.Value += WeightedValue;
+				}
+			}
+			KeysToSet.Add(TPair<FFrameNumber, ChannelValueType>(Frame, Value));
+		}
+	}
+	for (TPair<FFrameNumber, ChannelValueType>& KeyToSet : KeysToSet)
+	{
+		MovieSceneToolHelpers::SetOrAddKey(BaseChannelData, KeyToSet.Key, KeyToSet.Value);
+	}
+
+	return true;
+}
+
+bool MovieSceneToolHelpers::OptimizeSection(const FKeyDataOptimizationParams& InParams, UMovieSceneSection* InSection)
+{
+	TArrayView<FMovieSceneFloatChannel*> FloatChannels = InSection->GetChannelProxy().GetChannels<FMovieSceneFloatChannel>();
+	for (FMovieSceneFloatChannel* Channel : FloatChannels)
+	{
+		Channel->Optimize(InParams);
+	}
+	TArrayView<FMovieSceneDoubleChannel*> DoubleChannels = InSection->GetChannelProxy().GetChannels<FMovieSceneDoubleChannel>();
+	for (FMovieSceneDoubleChannel* Channel : DoubleChannels)
+	{
+		Channel->Optimize(InParams);
+	}
+	return true;
+}
+
+template<typename ChannelType>
+static bool MergeSections(UMovieSceneSection* BaseSection, TArrayView<ChannelType*> BaseDoubleChannels, TArray<UMovieSceneSection*>& AbsoluteSections, TArray<UMovieSceneSection*>& AdditiveSections,
+	const TRange<FFrameNumber>& Range,bool bSkipLastChannel)
+{
+	TArrayView<ChannelType*> BaseChannels = BaseSection->GetChannelProxy().GetChannels<ChannelType>();
+	if (BaseChannels.Num() > 0)
+	{
+		int32 SectionIndex = 0;
+		//sanity check to make sure channels are the same size
+		for (SectionIndex = 0; SectionIndex < AbsoluteSections.Num(); ++SectionIndex)
+		{
+			TArrayView<ChannelType*>Channels = AbsoluteSections[SectionIndex]->GetChannelProxy().GetChannels<ChannelType>();
+			if (Channels.Num() != BaseChannels.Num())
+			{
+				UE_LOG(LogMovieScene, Warning, TEXT("MergeSections:: Invalid number of channels"));
+				return false;
+			}
+		}
+		for (SectionIndex = 0; SectionIndex < AdditiveSections.Num(); ++SectionIndex)
+		{
+			TArrayView<ChannelType*>Channels = AdditiveSections[SectionIndex]->GetChannelProxy().GetChannels<ChannelType>();
+			if (Channels.Num() != BaseChannels.Num())
+			{
+				UE_LOG(LogMovieScene, Warning, TEXT("MergeSections:: Invalid number of channels"));
+				return false;
+			}
+		}
+		BaseSection->Modify();
+
+		TArray<ChannelType*> Channels;
+		int32 ChannelIndex = 0;
+		bSkipLastChannel = false;
+		int32 BaseChannelsNum = bSkipLastChannel ? BaseChannels.Num() - 1 : BaseChannels.Num();
+		for (ChannelIndex = 0; ChannelIndex < BaseChannels.Num(); ++ChannelIndex)
+		{
+			if (AbsoluteSections.Num() > 0)
+			{
+				Channels.Reset();
+				for (SectionIndex = 0; SectionIndex < AbsoluteSections.Num(); ++SectionIndex)
+				{
+					TArrayView<ChannelType*>OurChannels = AbsoluteSections[SectionIndex]->GetChannelProxy().GetChannels<ChannelType>();
+					Channels.Add(OurChannels[ChannelIndex]);
+				}
+				//now blend them
+				if (MergeChannels(Channels, AbsoluteSections, Range,
+					FChannelMergeAlgorithm::Average) == false)
+				{
+					UE_LOG(LogMovieScene, Warning, TEXT("MergeSections:: Could not merge channels"));
+					return false;
+				}
+			}
+			if (AdditiveSections.Num() > 0)
+			{
+				//now do additives
+				Channels.Reset();
+				for (SectionIndex = 0; SectionIndex < AdditiveSections.Num(); ++SectionIndex)
+				{
+					TArrayView<ChannelType*>OurChannels = AdditiveSections[SectionIndex]->GetChannelProxy().GetChannels<ChannelType>();
+					Channels.Add(OurChannels[ChannelIndex]);
+				}
+				//now blend them
+				if (MergeChannels(Channels, AdditiveSections, Range,
+					FChannelMergeAlgorithm::Add) == false)
+				{
+					UE_LOG(LogMovieScene, Warning, TEXT("MergeSections:: Could not merge channels"));
+					return false;
+				}
+			}
+		}
+		for (ChannelType* Channel : BaseChannels)
+		{
+			Channel->AutoSetTangents();
+		}
+	}
+	else
+	{
+		return false;
+	}
+	return true;
+}
+
+bool MovieSceneToolHelpers::CollapseSection(TSharedPtr<ISequencer>& SequencerPtr, UMovieSceneTrack* OwnerTrack, TArray<UMovieSceneSection*> Sections,
+	const FBakingAnimationKeySettings& InSettings)
+{
+	if (SequencerPtr.IsValid() && Sections.Num() > 1)
+	{
+		TRange<FFrameNumber> Range(InSettings.StartFrame, InSettings.EndFrame);
+		//we get the  first absolute section or if no absolute sections, first additive,
+		// that's the one we collapse onto
+		UMovieSceneSection* BaseSection = nullptr;
+		//get first absolute
+		for (UMovieSceneSection* Section: Sections)
+		{
+			if (Section->IsActive() == false)
+			{
+				continue;
+			}
+			if (Section->GetBlendType().Get() == EMovieSceneBlendType::Absolute)
+			{
+				BaseSection = Section;
+				break;
+			}
+			else if (Section->GetBlendType().Get() == EMovieSceneBlendType::Additive)
+			{
+				if (BaseSection == nullptr)
+				{
+					BaseSection = Section; //don't stop though may have an absolute later
+				}
+			}
+		}
+		if (BaseSection == nullptr)
+		{
+			UE_LOG(LogMovieScene, Warning, TEXT("CollapseSection:: Invalid section(s) to collapse"));
+			return false;
+		}
+		//now find other sections.
+		TArray<UMovieSceneSection*> AbsoluteSections;
+		TArray<UMovieSceneSection*> AdditiveSections;
+		for (UMovieSceneSection* Section : Sections)
+		{
+			if (Section != BaseSection && Section->IsActive())
+			{
+				if (Section->GetBlendType().Get() == EMovieSceneBlendType::Absolute)
+				{
+					AbsoluteSections.Add(Section);
+				}
+				else if (Section->GetBlendType().Get() == EMovieSceneBlendType::Additive)
+				{
+					AdditiveSections.Add(Section);
+				}
+			}
+		}
+		//now make sure we have sections to blend with and if so add base section to them
+		if (AbsoluteSections.Num() > 0 || AdditiveSections.Num() > 0)
+		{
+			if (AbsoluteSections.Num() > 0)
+			{
+				AbsoluteSections.Insert(BaseSection, 0);
+			}
+			if (AdditiveSections.Num() > 0)
+			{
+				AdditiveSections.Insert(BaseSection, 0);
+			}
+		}
+		else
+		{
+			UE_LOG(LogMovieScene, Warning, TEXT("CollapseSection:: Invalid section(s) to collapse"));
+			return false;
+		}
+
+		FScopedTransaction Transaction(NSLOCTEXT("MovieSceneTools", "CollapseAllSections", "Collapse All Sections"));
+
+		TArrayView<FMovieSceneFloatChannel*> BaseFloatChannels = BaseSection->GetChannelProxy().GetChannels<FMovieSceneFloatChannel>();
+		TArrayView<FMovieSceneDoubleChannel*> BaseDoubleChannels = BaseSection->GetChannelProxy().GetChannels<FMovieSceneDoubleChannel>();
+		if (BaseDoubleChannels.Num() > 0) //transforms
+		{
+			if (MergeSections(BaseSection, BaseDoubleChannels, AbsoluteSections, AdditiveSections, Range, false /*bSkipLastChannel*/) == false)
+			{
+				Transaction.Cancel();
+				return false;
+			}
+		}
+		else if(BaseFloatChannels.Num() > 0) //control rig
+		{
+			//skip weight channel for control rig
+			if (MergeSections(BaseSection, BaseFloatChannels, AbsoluteSections, AdditiveSections, Range, true /*bSkipLastChannel*/) == false)
+			{
+				Transaction.Cancel();
+				return false;
+			}
+		}
+		//delete other sections
+		OwnerTrack->Modify();
+		for (int32 SectionIndex = Sections.Num() - 1; SectionIndex >= 0; --SectionIndex)
+		{
+			if (Sections[SectionIndex] != BaseSection)
+			{
+				TArray<UMovieSceneSection*> AllSections = OwnerTrack->GetAllSections();
+				int32 TrackSectionIndex = INDEX_NONE;
+				if (AllSections.Find(Sections[SectionIndex], TrackSectionIndex))
+				{
+					if (TrackSectionIndex != INDEX_NONE)
+					{
+						OwnerTrack->RemoveSectionAt(TrackSectionIndex);
+					}
+				}
+			}
 		}
+		//reset everything back
+		SequencerPtr->NotifyMovieSceneDataChanged(EMovieSceneDataChangeType::MovieSceneStructureItemAdded);
 	}
+	return true;
+
 }
 
 UMovieScene3DTransformSection* MovieSceneToolHelpers::GetTransformSection(
diff --git a/Engine/Source/Editor/MovieSceneTools/Public/BakingAnimationKeySettings.h b/Engine/Source/Editor/MovieSceneTools/Public/BakingAnimationKeySettings.h
new file mode 100644
--- /dev/null
+++ b/Engine/Source/Editor/MovieSceneTools/Public/BakingAnimationKeySettings.h
@@ -0,0 +1,47 @@
+// Copyright Epic Games, Inc. All Rights Reserved.
+#pragma once
+
+#include "CoreMinimal.h"
+#include "BakingAnimationKeySettings.generated.h"
+
+
+UENUM()
+enum class EBakingKeySettings : uint8
+{
+	KeysOnly UMETA(DisplayName = "Keys Only"),
+	AllFrames UMETA(DisplayName = "All Frames"),
+};
+
+USTRUCT(BlueprintType)
+struct MOVIESCENETOOLS_API FBakingAnimationKeySettings
+{
+	GENERATED_BODY();
+
+	FBakingAnimationKeySettings()
+	{
+		StartFrame = 0;
+		EndFrame = 100;
+		BakingKeySettings = EBakingKeySettings::KeysOnly;
+		FrameIncrement = 1;
+		bReduceKeys = false;
+		Tolerance = 0.001f;
+	}
+
+	UPROPERTY(BlueprintReadWrite, EditAnywhere, Category = "Bake")
+	FFrameNumber StartFrame;
+
+	UPROPERTY(BlueprintReadWrite, EditAnywhere, Category = "Bake")
+	FFrameNumber EndFrame;
+
+	UPROPERTY(BlueprintReadWrite, EditAnywhere, Category = "Bake")
+	EBakingKeySettings BakingKeySettings;
+
+	UPROPERTY(BlueprintReadWrite, EditAnywhere, Category = "Bake", meta = (ClampMin = "1", UIMin = "1", EditCondition = "BakingKeySettings == EBakingKeySettings::AllFrames"))
+	int32 FrameIncrement;
+
+	UPROPERTY(BlueprintReadWrite, EditAnywhere, Category = "Bake", meta = (EditCondition = "BakingKeySettings == EBakingKeySettings::AllFrames"))
+	bool bReduceKeys;
+
+	UPROPERTY(BlueprintReadWrite, EditAnywhere, Category = "Bake", meta = (EditCondition = "BakingKeySettings == EBakingKeySettings::AllFrames || bReduceKeys"))
+	float Tolerance;
+};
diff --git a/Engine/Source/Editor/MovieSceneTools/Public/Constraints/MovieSceneConstraintChannelHelper.h b/Engine/Source/Editor/MovieSceneTools/Public/Constraints/MovieSceneConstraintChannelHelper.h
--- a/Engine/Source/Editor/MovieSceneTools/Public/Constraints/MovieSceneConstraintChannelHelper.h
+++ b/Engine/Source/Editor/MovieSceneTools/Public/Constraints/MovieSceneConstraintChannelHelper.h
@@ -16,6 +16,7 @@ class UTickableParentConstraint;
 class UTransformableHandle;
 struct FMovieSceneConstraintChannel;
 struct FFrameNumber;
+struct FConstraintAndActiveChannel;
 
 
 struct MOVIESCENETOOLS_API FCompensationEvaluator
@@ -41,6 +42,14 @@ private:
 	UTransformableHandle* Handle = nullptr;
 };
 
+struct MOVIESCENETOOLS_API FConstraintSections
+{
+	UMovieSceneSection* ConstraintSection = nullptr;
+	UMovieSceneSection* ChildTransformSection = nullptr;
+	UMovieSceneSection* ParentTransformSection = nullptr;
+	FConstraintAndActiveChannel* ActiveChannel = nullptr;
+	ITransformConstraintChannelInterface* Interface = nullptr;
+};
 struct MOVIESCENETOOLS_API FMovieSceneConstraintChannelHelper
 {
 public:
@@ -84,6 +93,20 @@ public:
 		UMovieSceneSection* InSection,
 		const FFrameNumber& InCurrentFrame, const FFrameNumber& InNextFrame);
 
+	/* Get the section and the channel for the given constraint, will be nullptr's if it doesn't exist in Sequencer*/
+	static FConstraintSections  GetConstraintSectionAndChannel(
+		const UTickableTransformConstraint* InConstraint,
+		const TSharedPtr<ISequencer>& InSequencer);
+
+	/* For the given constraint get all of the transform keys for it's child and parent handles*/
+	static void GetTransformFramesForConstraintHandles(
+		const UTickableTransformConstraint* InConstraint,
+		const TSharedPtr<ISequencer>& InSequencer,
+		const FFrameNumber& StartFrame,
+		const FFrameNumber& EndFrame,
+		TArray<FFrameNumber>& OutFramesToBake);
+
+
 	/** @todo documentation. */
 	template<typename ChannelType>
 	static void GetFramesToCompensate(
@@ -128,6 +151,25 @@ public:
 		const TSharedPtr<ISequencer>& InSequencer,
 		UMovieSceneSection* InSection);
 
+	template< typename ChannelType >
+	static TArray<FFrameNumber> GetTransformTimes(
+		const TArrayView<ChannelType*>& InChannels,
+		const FFrameNumber& StartTime, 
+		const FFrameNumber& EndTime);
+
+	template< typename ChannelType >
+	static  void DeleteTransformTimes(
+		const TArrayView<ChannelType*>& InChannels,
+		const FFrameNumber& StartTime,
+		const FFrameNumber& EndTime);
+
+	/** this will only set the value son channels with keys at the specified time, reusing tangent time etc. */
+	template< typename ChannelType >
+	static void SetTransformTimes(
+		const TArrayView<ChannelType*>& InChannels,
+		const TArray<FFrameNumber>& Frames,
+		const TArray<FTransform>& Transforms);
+
 	static bool bDoNotCompensate;
 
 private:
diff --git a/Engine/Source/Editor/MovieSceneTools/Public/Constraints/MovieSceneConstraintChannelHelper.inl b/Engine/Source/Editor/MovieSceneTools/Public/Constraints/MovieSceneConstraintChannelHelper.inl
--- a/Engine/Source/Editor/MovieSceneTools/Public/Constraints/MovieSceneConstraintChannelHelper.inl
+++ b/Engine/Source/Editor/MovieSceneTools/Public/Constraints/MovieSceneConstraintChannelHelper.inl
@@ -9,6 +9,7 @@
 #include "Channels/MovieSceneChannelProxy.h"
 #include "MovieSceneSection.h"
 #include "Algo/Unique.h"
+#include "Containers/SortedMap.h"
 
 template<typename ChannelType>
 void FMovieSceneConstraintChannelHelper::GetFramesToCompensate(
@@ -261,6 +262,82 @@ void FMovieSceneConstraintChannelHelper::DeleteTransformKeys(
 	}
 }
 
+template< typename ChannelType >
+TArray<FFrameNumber> FMovieSceneConstraintChannelHelper::GetTransformTimes(
+	const TArrayView<ChannelType*>& InChannels,
+	const FFrameNumber& StartTime,
+	const FFrameNumber& EndTime)
+{
+	TSortedMap<FFrameNumber, FFrameNumber> FrameSet;
+	TRange<FFrameNumber> WithinRange(0, 0);
+	WithinRange.SetLowerBoundValue(StartTime);
+	WithinRange.SetUpperBoundValue(EndTime);
+	TArray<FFrameNumber> KeyTimes;
+	TArray<FKeyHandle> KeyHandles;
+	for (ChannelType* Channel : InChannels)
+	{
+		TMovieSceneChannelData<typename ChannelType::ChannelValueType> Data = Channel->GetData();
+		KeyTimes.SetNum(0);
+		KeyHandles.SetNum(0);
+		Data.GetKeys(WithinRange, &KeyTimes,&KeyHandles);
+		for (const FFrameNumber& FrameNumber : KeyTimes)
+		{
+			FrameSet.Add(FrameNumber,FrameNumber);
+		}
+	}
+	TArray<FFrameNumber> Frames;
+	FrameSet.GenerateKeyArray(Frames);
+	return Frames;
+}
+
+template< typename ChannelType >
+ void FMovieSceneConstraintChannelHelper::DeleteTransformTimes(
+	 const TArrayView<ChannelType*>& InChannels,
+	 const FFrameNumber& StartTime,
+	 const FFrameNumber& EndTime)
+{
+	 TRange<FFrameNumber> WithinRange(0, 0);
+	 WithinRange.SetLowerBoundValue(StartTime);
+	 WithinRange.SetUpperBoundValue(EndTime);
+	 TArray<FFrameNumber> KeyTimes;
+	 TArray<FKeyHandle> KeyHandles;
+	 for (ChannelType* Channel : InChannels)
+	 {
+		 TMovieSceneChannelData<typename ChannelType::ChannelValueType> Data = Channel->GetData();
+		 KeyTimes.SetNum(0);
+		 KeyHandles.SetNum(0);
+		 Data.GetKeys(WithinRange, &KeyTimes, &KeyHandles);
+		 Data.DeleteKeys(KeyHandles);
+	 }
+}
+
+/** this will only set the values on channels with keys at the specified time, reusing tangents  etc. */
+template< typename ChannelType >
+void FMovieSceneConstraintChannelHelper::SetTransformTimes(
+	const TArrayView<ChannelType*>& InChannels,
+	const TArray<FFrameNumber>& Frames,
+	const TArray<FTransform>& Transforms)
+{
+	if (Frames.Num() != Transforms.Num())
+	{
+		return;
+	}
+	for (ChannelType* Channel : InChannels)
+	{
+		TMovieSceneChannelData<typename ChannelType::ChannelValueType> Data = Channel->GetData();
+		const TArrayView<const FFrameNumber> Times = Data.GetTimes();
+		for (int32 Index = 0; Index < Frames.Num(); ++Index)
+		{
+			FFrameNumber Time = Frames[0];
+			const int32 KeyIndex = Algo::LowerBound(Times, Time);
+			if (Times.IsValidIndex(KeyIndex) && Times[KeyIndex] == Time)
+			{
+				Data.GetValues()[KeyIndex] = 1.0;//todo 
+			}
+		}
+	}
+}
+
 template<typename ChannelType>
 void EvaluateTangentAtThisTime(int32 ChannelIndex, int32 NumChannels,
 	UMovieSceneSection* Section,
diff --git a/Engine/Source/Editor/MovieSceneTools/Public/MovieSceneToolHelpers.h b/Engine/Source/Editor/MovieSceneTools/Public/MovieSceneToolHelpers.h
--- a/Engine/Source/Editor/MovieSceneTools/Public/MovieSceneToolHelpers.h
+++ b/Engine/Source/Editor/MovieSceneTools/Public/MovieSceneToolHelpers.h
@@ -39,6 +39,8 @@ class ACameraActor;
 struct FActorForWorldTransforms;
 class UMovieScene3DTransformSection;
 class UMovieSceneSubTrack;
+struct FBakingAnimationKeySettings;
+struct FKeyDataOptimizationParams;
 class UMovieSceneSubSection;
 enum class EMovieSceneTransformChannel : uint32;
 
@@ -546,6 +548,16 @@ public:
 	*/
 	static void SetOrAddKey(TMovieSceneChannelData<FMovieSceneDoubleValue>& ChannelData, FFrameNumber Time, double Value);
 
+	/*
+	* Set or add a key onto a float channel based on key value.
+	*/
+	static void SetOrAddKey(TMovieSceneChannelData<FMovieSceneFloatValue>& Curve, FFrameNumber Time, const FMovieSceneFloatValue& Value);
+
+	/*
+	* Set or add a key onto a double channel based on key value.
+	*/
+	static void SetOrAddKey(TMovieSceneChannelData<FMovieSceneDoubleValue>& ChannelData, FFrameNumber Time, FMovieSceneDoubleValue Value);
+
 	/* 
 	* Set or add a key onto a float channel based on rich curve data.
 	*/
@@ -578,11 +590,17 @@ public:
 	 */
 	static bool IsValidAsset(UMovieSceneSequence* Sequence, const FAssetData& InAssetData);
 
+	static bool CollapseSection(TSharedPtr<ISequencer>& SequencerPtr, UMovieSceneTrack* OwnerTrack, TArray<UMovieSceneSection*> Sections,
+		const FBakingAnimationKeySettings& InSettings);
+
+	static bool OptimizeSection(const FKeyDataOptimizationParams& InParams, UMovieSceneSection* InSection);
+
 	/** Returns the frame numbers between start and end. */
 	static void CalculateFramesBetween(
 		const UMovieScene* MovieScene,
 		FFrameNumber StartFrame,
 		FFrameNumber EndFrame,
+		int FrameIncrement,
 		TArray<FFrameNumber>& OutFrames);
 
 	/** Returns the transform section for that guid. */
diff --git a/Engine/Source/Runtime/Experimental/Animation/Constraints/Private/TransformableHandle.cpp b/Engine/Source/Runtime/Experimental/Animation/Constraints/Private/TransformableHandle.cpp
--- a/Engine/Source/Runtime/Experimental/Animation/Constraints/Private/TransformableHandle.cpp
+++ b/Engine/Source/Runtime/Experimental/Animation/Constraints/Private/TransformableHandle.cpp
@@ -10,7 +10,6 @@
 #include "Channels/MovieSceneFloatChannel.h"
 #include "Channels/MovieSceneDoubleChannel.h"
 #include "Channels/MovieSceneChannelProxy.h"
-#include "Channels/MovieSceneDoubleChannel.h"
 #include "Components/SkeletalMeshComponent.h"
 
 #include UE_INLINE_GENERATED_CPP_BY_NAME(TransformableHandle)
