[Backout] - CL25091326 - CIS UE Error
Original CL Desc
-----------------------------------------------------------------
Sequencer: Smart Baking for Constraint Baking, Space Baking, Snapper and control rig collapsing. New settings type to hold the per frame, key increment settings when Baking which is shared among these algorithms. Constraints have a new Bake button/dialog now.

diff --git a/Engine/Plugins/Animation/ControlRig/Source/ControlRig/Private/Constraints/ControlRigTransformableHandle.cpp b/Engine/Plugins/Animation/ControlRig/Source/ControlRig/Private/Constraints/ControlRigTransformableHandle.cpp
--- a/Engine/Plugins/Animation/ControlRig/Source/ControlRig/Private/Constraints/ControlRigTransformableHandle.cpp
+++ b/Engine/Plugins/Animation/ControlRig/Source/ControlRig/Private/Constraints/ControlRigTransformableHandle.cpp
@@ -11,7 +11,6 @@
 #include "Channels/MovieSceneChannelProxy.h"
 #include "Sequencer/MovieSceneControlRigParameterSection.h"
 #include "Sections/MovieScene3DTransformSection.h"
-#include "Sequencer/ControlRigSequencerHelpers.h"
 
 #include UE_INLINE_GENERATED_CPP_BY_NAME(ControlRigTransformableHandle)
 
@@ -443,10 +442,60 @@ void UTransformableControlHandle::OnObjectBoundToControlRig(UObject* InObject)
 	}
 }
 
+static TPair<const FChannelMapInfo*, int32> GetInfoAndNumFloatChannels(
+	const UControlRig* InControlRig,
+	const FName& InControlName,
+	const UMovieSceneControlRigParameterSection* InSection)
+{
+	const FRigControlElement* ControlElement = InControlRig ? InControlRig->FindControl(InControlName) : nullptr;
+	auto GetNumFloatChannels = [](const ERigControlType& InControlType)
+	{
+		switch (InControlType)
+		{
+		case ERigControlType::Position:
+		case ERigControlType::Scale:
+		case ERigControlType::Rotator:
+			return 3;
+		case ERigControlType::TransformNoScale:
+			return 6;
+		case ERigControlType::Transform:
+		case ERigControlType::EulerTransform:
+			return 9;
+		default:
+			break;
+		}
+		return 0;
+	};
+
+	const int32 NumFloatChannels = ControlElement ? GetNumFloatChannels(ControlElement->Settings.ControlType) : 0;
+	const FChannelMapInfo* ChannelInfo = InSection ? InSection->ControlChannelMap.Find(InControlName) : nullptr;
+
+	return { ChannelInfo, NumFloatChannels };
+}
 TArrayView<FMovieSceneFloatChannel*>  UTransformableControlHandle::GetFloatChannels(const UMovieSceneSection* InSection) const
 {
-	return FControlRigSequencerHelpers::GetFloatChannels(ControlRig.Get(),
-		ControlName, InSection);
+	// no floats for transform sections
+	static const TArrayView<FMovieSceneFloatChannel*> EmptyChannelsView;
+
+	const FChannelMapInfo* ChannelInfo = nullptr;
+	int32 NumChannels = 0;
+	const UMovieSceneControlRigParameterSection* CRSection = Cast<UMovieSceneControlRigParameterSection>(InSection);
+	if (CRSection == nullptr)
+	{
+		return EmptyChannelsView;
+	}
+
+	Tie(ChannelInfo, NumChannels) = GetInfoAndNumFloatChannels(ControlRig.Get(),ControlName, CRSection);
+
+	if (ChannelInfo == nullptr || NumChannels == 0)
+	{
+		return EmptyChannelsView;
+	}
+
+	// return a sub view that just represents the control's channels
+	const TArrayView<FMovieSceneFloatChannel*> FloatChannels = InSection->GetChannelProxy().GetChannels<FMovieSceneFloatChannel>();
+	const int32 ChannelStartIndex = ChannelInfo->ChannelIndex;
+	return FloatChannels.Slice(ChannelStartIndex, NumChannels);
 }
 
 TArrayView<FMovieSceneDoubleChannel*>  UTransformableControlHandle::GetDoubleChannels(const UMovieSceneSection* InSection) const
diff --git a/Engine/Plugins/Animation/ControlRig/Source/ControlRig/Private/Sequencer/ControlRigSequencerHelpers.cpp b/Engine/Plugins/Animation/ControlRig/Source/ControlRig/Private/Sequencer/ControlRigSequencerHelpers.cpp
deleted file mode 100644
diff --git a/Engine/Plugins/Animation/ControlRig/Source/ControlRig/Private/Sequencer/MovieSceneControlRigParameterSection.cpp b/Engine/Plugins/Animation/ControlRig/Source/ControlRig/Private/Sequencer/MovieSceneControlRigParameterSection.cpp
--- a/Engine/Plugins/Animation/ControlRig/Source/ControlRig/Private/Sequencer/MovieSceneControlRigParameterSection.cpp
+++ b/Engine/Plugins/Animation/ControlRig/Source/ControlRig/Private/Sequencer/MovieSceneControlRigParameterSection.cpp
@@ -2583,7 +2583,7 @@ void UMovieSceneControlRigParameterSection::SetControlRig(UControlRig* InControl
 	ControlRigClass = ControlRig ? ControlRig->GetClass() : nullptr;
 }
 
-void UMovieSceneControlRigParameterSection::FixRotationWinding(const FName& ControlName, FFrameNumber StartFrame, FFrameNumber EndFrame)
+void UMovieSceneControlRigParameterSection::FixRotationWinding(FName ControlName, FFrameNumber StartFrame, FFrameNumber EndFrame)
 {
 	FChannelMapInfo* pChannelIndex = ControlChannelMap.Find(ControlName);
 	if (pChannelIndex == nullptr || GetControlRig() == nullptr)
@@ -2630,107 +2630,6 @@ void UMovieSceneControlRigParameterSection::FixRotationWinding(const FName& Cont
 	}
 }
 
-void UMovieSceneControlRigParameterSection::OptimizeSection(const FName& ControlName, const FKeyDataOptimizationParams& Params)
-{
-	TArrayView<FMovieSceneFloatChannel*> FloatChannels = GetChannelProxy().GetChannels<FMovieSceneFloatChannel>();
-	FChannelMapInfo* pChannelIndex = ControlChannelMap.Find(ControlName);
-	if (pChannelIndex != nullptr)
-	{
-		int32 ChannelIndex = pChannelIndex->ChannelIndex;
-
-		if (FRigControlElement* ControlElement = ControlRig->FindControl(ControlName))
-		{
-			switch (ControlElement->Settings.ControlType)
-			{
-				case ERigControlType::Position:
-				case ERigControlType::Scale:
-				case ERigControlType::Rotator:
-				{
-					FloatChannels[ChannelIndex]->Optimize(Params);
-					FloatChannels[ChannelIndex + 1]->Optimize(Params);
-					FloatChannels[ChannelIndex + 2]->Optimize(Params);
-					break;
-				}
-
-				case ERigControlType::Transform:
-				case ERigControlType::TransformNoScale:
-				case ERigControlType::EulerTransform:
-				{
-
-					FloatChannels[ChannelIndex]->Optimize(Params);
-					FloatChannels[ChannelIndex + 1]->Optimize(Params);
-					FloatChannels[ChannelIndex + 2]->Optimize(Params);
-					FloatChannels[ChannelIndex + 3]->Optimize(Params);
-					FloatChannels[ChannelIndex + 4]->Optimize(Params);
-					FloatChannels[ChannelIndex + 5]->Optimize(Params);
-
-					if (ControlElement->Settings.ControlType == ERigControlType::Transform ||
-						ControlElement->Settings.ControlType == ERigControlType::EulerTransform)
-					{
-						FloatChannels[ChannelIndex + 6]->Optimize(Params);
-						FloatChannels[ChannelIndex + 7]->Optimize(Params);
-						FloatChannels[ChannelIndex + 8]->Optimize(Params);
-					}
-					break;
-
-				}
-				default:
-					break;
-			}
-		}
-	}
-}
-
-void UMovieSceneControlRigParameterSection::AutoSetTangents(const FName& ControlName)
-{
-	TArrayView<FMovieSceneFloatChannel*> FloatChannels = GetChannelProxy().GetChannels<FMovieSceneFloatChannel>();
-	FChannelMapInfo* pChannelIndex = ControlChannelMap.Find(ControlName);
-	if (pChannelIndex != nullptr)
-	{
-		int32 ChannelIndex = pChannelIndex->ChannelIndex;
-
-		if (FRigControlElement* ControlElement = ControlRig->FindControl(ControlName))
-		{
-			switch (ControlElement->Settings.ControlType)
-			{
-			case ERigControlType::Position:
-			case ERigControlType::Scale:
-			case ERigControlType::Rotator:
-			{
-				FloatChannels[ChannelIndex]->AutoSetTangents();
-				FloatChannels[ChannelIndex + 1]->AutoSetTangents();
-				FloatChannels[ChannelIndex + 2]->AutoSetTangents();
-				break;
-			}
-
-			case ERigControlType::Transform:
-			case ERigControlType::TransformNoScale:
-			case ERigControlType::EulerTransform:
-			{
-
-				FloatChannels[ChannelIndex]->AutoSetTangents();
-				FloatChannels[ChannelIndex + 1]->AutoSetTangents();
-				FloatChannels[ChannelIndex + 2]->AutoSetTangents();
-				FloatChannels[ChannelIndex + 3]->AutoSetTangents();
-				FloatChannels[ChannelIndex + 4]->AutoSetTangents();
-				FloatChannels[ChannelIndex + 5]->AutoSetTangents();
-
-				if (ControlElement->Settings.ControlType == ERigControlType::Transform ||
-					ControlElement->Settings.ControlType == ERigControlType::EulerTransform)
-				{
-					FloatChannels[ChannelIndex + 6]->AutoSetTangents();
-					FloatChannels[ChannelIndex + 7]->AutoSetTangents();
-					FloatChannels[ChannelIndex + 8]->AutoSetTangents();
-				}
-				break;
-
-			}
-			default:
-				break;
-			}
-		}
-	}
-}
 #if WITH_EDITOR
 
 void UMovieSceneControlRigParameterSection::RecordControlRigKey(FFrameNumber FrameNumber, bool bSetDefault, ERichCurveInterpMode InInterpMode)
diff --git a/Engine/Plugins/Animation/ControlRig/Source/ControlRig/Private/Tools/ControlRigSnapSettings.cpp b/Engine/Plugins/Animation/ControlRig/Source/ControlRig/Private/Tools/ControlRigSnapSettings.cpp
new file mode 100644
--- /dev/null
+++ b/Engine/Plugins/Animation/ControlRig/Source/ControlRig/Private/Tools/ControlRigSnapSettings.cpp
@@ -0,0 +1,9 @@
+// Copyright Epic Games, Inc. All Rights Reserved.
+#include "Tools/ControlRigSnapSettings.h"
+
+#include UE_INLINE_GENERATED_CPP_BY_NAME(ControlRigSnapSettings)
+
+UControlRigSnapSettings::UControlRigSnapSettings()
+{
+}
+
diff --git a/Engine/Plugins/Animation/ControlRig/Source/ControlRig/Public/Sequencer/ControlRigSequencerHelpers.h b/Engine/Plugins/Animation/ControlRig/Source/ControlRig/Public/Sequencer/ControlRigSequencerHelpers.h
deleted file mode 100644
diff --git a/Engine/Plugins/Animation/ControlRig/Source/ControlRig/Public/Sequencer/MovieSceneControlRigParameterSection.h b/Engine/Plugins/Animation/ControlRig/Source/ControlRig/Public/Sequencer/MovieSceneControlRigParameterSection.h
--- a/Engine/Plugins/Animation/ControlRig/Source/ControlRig/Public/Sequencer/MovieSceneControlRigParameterSection.h
+++ b/Engine/Plugins/Animation/ControlRig/Source/ControlRig/Public/Sequencer/MovieSceneControlRigParameterSection.h
@@ -24,7 +24,6 @@
 
 class UAnimSequence;
 class USkeletalMeshComponent;
-struct FKeyDataOptimizationParams;
 
 struct CONTROLRIG_API FControlRigBindingHelper
 {
@@ -198,9 +197,7 @@ public:
 	TArray<FIntegerParameterNameAndCurve>& GetIntegerParameterNamesAndCurves();
 	const TArray<FIntegerParameterNameAndCurve>& GetIntegerParameterNamesAndCurves() const;
 
-	void FixRotationWinding(const FName& ControlName, FFrameNumber StartFrame, FFrameNumber EndFrame);
-	void OptimizeSection(const FName& ControlName, const FKeyDataOptimizationParams& InParams);
-	void AutoSetTangents(const FName& ControlName);
+	void FixRotationWinding(FName ControlName, FFrameNumber StartFrame, FFrameNumber EndFrame);
 
 	TArray<FSpaceControlNameAndChannel>& GetSpaceChannels();
 	const TArray< FSpaceControlNameAndChannel>& GetSpaceChannels() const;
diff --git a/Engine/Plugins/Animation/ControlRig/Source/ControlRig/Public/Tools/ControlRigSnapSettings.h b/Engine/Plugins/Animation/ControlRig/Source/ControlRig/Public/Tools/ControlRigSnapSettings.h
new file mode 100644
--- /dev/null
+++ b/Engine/Plugins/Animation/ControlRig/Source/ControlRig/Public/Tools/ControlRigSnapSettings.h
@@ -0,0 +1,40 @@
+// Copyright Epic Games, Inc. All Rights Reserved.
+#pragma once
+/**
+* Per Project Settings Animation Snapping
+*/
+#include "CoreMinimal.h"
+#include "UObject/ObjectMacros.h"
+#include "UObject/Object.h"
+#include "ControlRig.h"
+#include "ControlRigSnapSettings.generated.h"
+
+
+UCLASS(config = EditorPerProjectUserSettings)
+class CONTROLRIG_API UControlRigSnapSettings : public UObject
+{
+
+public:
+	GENERATED_BODY()
+
+	UControlRigSnapSettings();
+
+	/** When snapping keep offset, If false snap completely */
+	UPROPERTY(config, EditAnywhere, BlueprintReadWrite, Category = "Snap Settings", meta = (ToolTip = "When snapping keep offset, if false snap completely"))
+	bool bKeepOffset = false;
+
+	/** When snapping, also snap position */
+	UPROPERTY(config, EditAnywhere, BlueprintReadWrite, Category = "Snap Settings", meta = (ToolTip = "When snapping, also snap position"))
+	bool bSnapPosition = true;
+
+	/** When snapping, also snap rotation */
+	UPROPERTY(config, EditAnywhere, BlueprintReadWrite, Category = "Snap Settings", meta = (ToolTip = "When snapping, also snap rotation"))
+	bool bSnapRotation = false;
+
+	/** When snapping, also snap scale */
+	UPROPERTY(config, EditAnywhere, BlueprintReadWrite, Category = "Snap Settings", meta = (ToolTip = "When snapping, also snap scale"))
+	bool bSnapScale = false;
+
+};
+
+
diff --git a/Engine/Plugins/Animation/ControlRig/Source/ControlRigEditor/Private/ControlRigSequencerEditorLibrary.cpp b/Engine/Plugins/Animation/ControlRig/Source/ControlRigEditor/Private/ControlRigSequencerEditorLibrary.cpp
--- a/Engine/Plugins/Animation/ControlRig/Source/ControlRigEditor/Private/ControlRigSequencerEditorLibrary.cpp
+++ b/Engine/Plugins/Animation/ControlRig/Source/ControlRigEditor/Private/ControlRigSequencerEditorLibrary.cpp
@@ -47,7 +47,6 @@
 #include "Constraints/ControlRigTransformableHandle.h"
 #include "Constraints/MovieSceneConstraintChannelHelper.h"
 #include "Sections/MovieSceneConstrainedSection.h"
-#include "BakingAnimationKeySettings.h"
 
 #include UE_INLINE_GENERATED_CPP_BY_NAME(ControlRigSequencerEditorLibrary)
 
@@ -766,8 +765,6 @@ bool UControlRigSequencerEditorLibrary::BakeConstraint(UWorld* World, UTickableC
 		UE_LOG(LogControlRig, Error, TEXT("BakeConstraint: Need Valid World"));
 		return false;
 	}
-	TOptional< FBakingAnimationKeySettings> Settings;
-
 	if (UTickableTransformConstraint* TransformConstraint = Cast<UTickableTransformConstraint>(Constraint))
 	{
 		TSharedPtr<ISequencer> Sequencer = GetSequencerFromAsset();
@@ -800,7 +797,7 @@ bool UControlRigSequencerEditorLibrary::BakeConstraint(UWorld* World, UTickableC
 		{
 			FramesToBake = RealFramesToBake;
 		}
-		FConstraintBaker::Bake(World, TransformConstraint, Sequencer, Settings, FramesToBake);
+		FConstraintBaker::Bake(World, TransformConstraint, Sequencer, FramesToBake);
 	}
 	else
 	{
@@ -810,48 +807,6 @@ bool UControlRigSequencerEditorLibrary::BakeConstraint(UWorld* World, UTickableC
 	return true;
 }
 
-bool UControlRigSequencerEditorLibrary::BakeConstraints(UWorld* World, TArray<UTickableConstraint*>& InConstraints, const FBakingAnimationKeySettings& InSettings)
-{
-	if (!World)
-	{
-		UE_LOG(LogControlRig, Error, TEXT("BakeConstraint: Need Valid World"));
-		return false;
-	}
-
-	TSharedPtr<ISequencer> Sequencer = GetSequencerFromAsset();
-	if (!Sequencer || !Sequencer->GetFocusedMovieSceneSequence())
-	{
-		UE_LOG(LogControlRig, Error, TEXT("BakeConstraint: Need loaded level Sequence"));
-		return false;
-	}
-	const UMovieScene* MovieScene = Sequencer->GetFocusedMovieSceneSequence()->GetMovieScene();
-	if (!MovieScene)
-	{
-		UE_LOG(LogControlRig, Error, TEXT("BakeConstraint: Need valid Movie Scene"));
-		return false;
-	}
-	TArray<UTickableTransformConstraint*> TransformConstraints;
-	for (UTickableConstraint* Constraint : InConstraints)
-	{
-		if (UTickableTransformConstraint* TransformConstraint = Cast<UTickableTransformConstraint>(Constraint))
-		{
-			TransformConstraints.Add(TransformConstraint);
-		}
-		else
-		{
-			UE_LOG(LogControlRig, Error, TEXT("BakeConstraint: Need Valid Constraint"));
-			return false;
-		}
-	}
-
-	if(FConstraintBaker::BakeMultiple(World, TransformConstraints, Sequencer, InSettings) == false)
-	{
-		UE_LOG(LogControlRig, Error, TEXT("BakeMultiple: Failed"));
-		return false;
-	}
-	return true;
-}
-
 
 bool UControlRigSequencerEditorLibrary::SnapControlRig(ULevelSequence* LevelSequence, FFrameNumber StartFrame, FFrameNumber EndFrame, const FControlRigSnapperSelection& ChildrenToSnap,
 	const FControlRigSnapperSelection& ParentToSnap, const UControlRigSnapSettings* SnapSettings, ESequenceTimeUnit TimeUnit)
@@ -2587,19 +2542,6 @@ bool UControlRigSequencerEditorLibrary::ImportFBXToControlRigTrack(UWorld* World
 	return bValid;
 }
 
-bool UControlRigSequencerEditorLibrary::CollapseControlRigAnimLayersWithSettings(ULevelSequence* InSequence, UMovieSceneControlRigParameterTrack* InTrack, const FBakingAnimationKeySettings& InSettings)
-{
-	TWeakPtr<ISequencer> WeakSequencer = GetSequencerFromAsset();
-	bool bValid = false;
-	if (WeakSequencer.IsValid() && InTrack)
-	{
-		TSharedPtr<ISequencer>  SequencerPtr = WeakSequencer.Pin();
-		bValid = FControlRigParameterTrackEditor::CollapseAllLayers(SequencerPtr, InTrack, InSettings);
-		
-	}
-	return bValid;
-}
-
 bool UControlRigSequencerEditorLibrary::CollapseControlRigAnimLayers(ULevelSequence* LevelSequence, UMovieSceneControlRigParameterTrack* InTrack, bool bKeyReduce, float Tolerance)
 {
 	TWeakPtr<ISequencer> WeakSequencer = GetSequencerFromAsset();
@@ -2607,17 +2549,13 @@ bool UControlRigSequencerEditorLibrary::CollapseControlRigAnimLayers(ULevelSeque
 
 	if (WeakSequencer.IsValid() && InTrack)
 	{
-		TSharedPtr<ISequencer>  SequencerPtr = WeakSequencer.Pin();
-		FBakingAnimationKeySettings CollapseControlsSettings;
-		const FFrameRate TickResolution = SequencerPtr->GetFocusedTickResolution();
-		const FFrameTime FrameTime = SequencerPtr->GetLocalTime().ConvertTo(TickResolution);
-		FFrameNumber CurrentTime = FrameTime.GetFrame();
-
-		TRange<FFrameNumber> Range = SequencerPtr->GetFocusedMovieSceneSequence()->GetMovieScene()->GetPlaybackRange();
-
-		CollapseControlsSettings.StartFrame = Range.GetLowerBoundValue();
-		CollapseControlsSettings.EndFrame = Range.GetUpperBoundValue();
-		bValid = FControlRigParameterTrackEditor::CollapseAllLayers(SequencerPtr,InTrack, CollapseControlsSettings);	
+		TArray<UMovieSceneSection*> Sections = InTrack->GetAllSections();
+		if (Sections.Num() > 0)
+		{
+			TSharedPtr<ISequencer>  SequencerPtr = WeakSequencer.Pin();
+			UMovieSceneControlRigParameterSection* ParameterSection = Cast<UMovieSceneControlRigParameterSection>(Sections[0]);
+			bValid = FControlRigParameterTrackEditor::CollapseAllLayers(SequencerPtr,InTrack, ParameterSection, bKeyReduce, Tolerance);
+		}
 	}
 	return bValid;
 }
@@ -2667,12 +2605,17 @@ bool UControlRigSequencerEditorLibrary::BakeControlRigSpace(ULevelSequence* InSe
 	{
 		TSharedPtr<ISequencer>  Sequencer = WeakSequencer.Pin();
 		const FFrameRate TickResolution = Sequencer->GetFocusedTickResolution();
+		TArray<FFrameNumber> Frames;
 		const FFrameRate& FrameRate = Sequencer->GetFocusedDisplayRate();
 		FFrameNumber FrameRateInFrameNumber = TickResolution.AsFrameNumber(FrameRate.AsInterval());
 		if (TimeUnit == ESequenceTimeUnit::DisplayRate)
 		{
-			InSettings.Settings.StartFrame = FFrameRate::TransformTime(FFrameTime(InSettings.Settings.StartFrame, 0), FrameRate, TickResolution).RoundToFrame();
-			InSettings.Settings.EndFrame = FFrameRate::TransformTime(FFrameTime(InSettings.Settings.EndFrame, 0), FrameRate, TickResolution).RoundToFrame();
+			InSettings.StartFrame = FFrameRate::TransformTime(FFrameTime(InSettings.StartFrame, 0), FrameRate, TickResolution).RoundToFrame();
+			InSettings.EndFrame = FFrameRate::TransformTime(FFrameTime(InSettings.EndFrame, 0), FrameRate, TickResolution).RoundToFrame();
+		}
+		for (FFrameNumber& Frame = InSettings.StartFrame; Frame <= InSettings.EndFrame; Frame += FrameRateInFrameNumber)
+		{
+			Frames.Add(Frame);
 		}
 
 		FScopedTransaction Transaction(LOCTEXT("BakeControlToSpace", "Bake Control In Space"));
@@ -2684,7 +2627,7 @@ bool UControlRigSequencerEditorLibrary::BakeControlRigSpace(ULevelSequence* InSe
 				if (SpaceChannelAndSection.SpaceChannel)
 				{
 					FControlRigSpaceChannelHelpers::SequencerBakeControlInSpace(InControlRig, Sequencer.Get(), SpaceChannelAndSection.SpaceChannel, SpaceChannelAndSection.SectionToKey,
-						InControlRig->GetHierarchy(), Element->GetKey(), InSettings);
+						Frames, InControlRig->GetHierarchy(), Element->GetKey(), InSettings);
 				}
 			}
 		}
@@ -2771,6 +2714,7 @@ bool UControlRigSequencerEditorLibrary::MoveControlRigSpace(ULevelSequence* Leve
 				}
 				UMovieSceneControlRigParameterSection* ParamSection = Cast<UMovieSceneControlRigParameterSection>(SpaceChannelAndSection.SectionToKey);
 
+
 				TArray<FFrameNumber> OurKeyTimes;
 				TArray<FKeyHandle> OurKeyHandles;
 				TRange<FFrameNumber> CurrentFrameRange;
diff --git a/Engine/Plugins/Animation/ControlRig/Source/ControlRigEditor/Private/EditMode/SControlRigEditModeTools.cpp b/Engine/Plugins/Animation/ControlRig/Source/ControlRigEditor/Private/EditMode/SControlRigEditModeTools.cpp
--- a/Engine/Plugins/Animation/ControlRig/Source/ControlRigEditor/Private/EditMode/SControlRigEditModeTools.cpp
+++ b/Engine/Plugins/Animation/ControlRig/Source/ControlRigEditor/Private/EditMode/SControlRigEditModeTools.cpp
@@ -42,8 +42,6 @@
 
 #define LOCTEXT_NAMESPACE "ControlRigEditModeTools"
 
-//statics to reuse in the UI
-FRigSpacePickerBakeSettings SControlRigEditModeTools::BakeSpaceSettings;
 
 void SControlRigEditModeTools::SetControlRigs(const TArrayView<TWeakObjectPtr<UControlRig>>& InControlRigs)
 {
@@ -403,7 +401,7 @@ void SControlRigEditModeTools::Construct(const FArguments& InArgs, TSharedPtr<FC
 					}
 				})
 			]
-			
+
 			+ SVerticalBox::Slot()
 			.AutoHeight()
 			[
@@ -923,6 +921,7 @@ FReply SControlRigEditModeTools::OnBakeControlsToNewSpaceButtonClicked()
 		if (ControlRig.IsValid() && SpacePickerWidget->GetHierarchy() == ControlRig->GetHierarchy())
 		{
 
+			FRigSpacePickerBakeSettings Settings;
 			//Find default target space, just use first control and find space at current sequencer time
 			//Then Find range
 
@@ -939,24 +938,42 @@ FReply SControlRigEditModeTools::OnBakeControlsToNewSpaceButtonClicked()
 				FFrameNumber CurrentTime = FrameTime.GetFrame();
 				FMovieSceneControlRigSpaceBaseKey Value;
 				using namespace UE::MovieScene;
-				//set up settings if not setup
-				if (BakeSpaceSettings.TargetSpace == FRigElementKey())
+				Settings.TargetSpace = URigHierarchy::GetDefaultParentKey();
+
+				TRange<FFrameNumber> Range = Sequencer->GetFocusedMovieSceneSequence()->GetMovieScene()->GetPlaybackRange();
+				TArray<FFrameNumber> Keys;
+				TArray < FKeyHandle> KeyHandles;
+
+				Settings.StartFrame = Range.GetLowerBoundValue();
+				Settings.EndFrame = Range.GetUpperBoundValue();
+				if (Keys.Num() > 0)
 				{
-					BakeSpaceSettings.TargetSpace = URigHierarchy::GetDefaultParentKey();
-					TRange<FFrameNumber> Range = Sequencer->GetFocusedMovieSceneSequence()->GetMovieScene()->GetPlaybackRange();
-					BakeSpaceSettings.Settings.StartFrame = Range.GetLowerBoundValue();
-					BakeSpaceSettings.Settings.EndFrame = Range.GetUpperBoundValue();
+					int32 Index = Algo::LowerBound(Keys, CurrentTime);
+					if (Index >= 0 && Index < (Keys.Num() - 1))
+					{
+						Settings.StartFrame = Keys[Index];
+						Settings.EndFrame = Keys[Index + 1];
+
+					}
 				}
 
 				TSharedRef<SRigSpacePickerBakeWidget> BakeWidget =
 					SNew(SRigSpacePickerBakeWidget)
-					.Settings(BakeSpaceSettings)
+					.Settings(Settings)
 					.Hierarchy(SpacePickerWidget->GetHierarchy())
 					.Controls(ControlKeys) // use the cached controls here since the selection is not recovered until next tick.
 					.Sequencer(Sequencer)
 					.GetControlCustomization(this, &SControlRigEditModeTools::HandleGetControlElementCustomization)
-					.OnBake_Lambda([Sequencer, ControlRig, TickResolution](URigHierarchy* InHierarchy, TArray<FRigElementKey> InControls, FRigSpacePickerBakeSettings& InSettings)
-						{		
+					.OnBake_Lambda([Sequencer, ControlRig, TickResolution](URigHierarchy* InHierarchy, TArray<FRigElementKey> InControls, FRigSpacePickerBakeSettings InSettings)
+						{
+							TArray<FFrameNumber> Frames;
+
+							const FFrameRate& FrameRate = Sequencer->GetFocusedDisplayRate();
+							FFrameNumber FrameRateInFrameNumber = TickResolution.AsFrameNumber(FrameRate.AsInterval());
+							for (FFrameNumber& Frame = InSettings.StartFrame; Frame <= InSettings.EndFrame; Frame += FrameRateInFrameNumber)
+							{
+								Frames.Add(Frame);
+							}
 							FScopedTransaction Transaction(LOCTEXT("BakeControlToSpace", "Bake Control In Space"));
 							for (const FRigElementKey& ControlKey : InControls)
 							{
@@ -965,9 +982,8 @@ FReply SControlRigEditModeTools::OnBakeControlsToNewSpaceButtonClicked()
 								if (SpaceChannelAndSection.SpaceChannel)
 								{
 									FControlRigSpaceChannelHelpers::SequencerBakeControlInSpace(ControlRig.Get(), Sequencer, SpaceChannelAndSection.SpaceChannel, SpaceChannelAndSection.SectionToKey,
-										InHierarchy, ControlKey, InSettings);
+										Frames, InHierarchy, ControlKey, InSettings);
 								}
-								SControlRigEditModeTools::BakeSpaceSettings = InSettings;
 							}
 							return FReply::Handled();
 						});
diff --git a/Engine/Plugins/Animation/ControlRig/Source/ControlRigEditor/Private/EditMode/SControlRigEditModeTools.h b/Engine/Plugins/Animation/ControlRig/Source/ControlRigEditor/Private/EditMode/SControlRigEditModeTools.h
--- a/Engine/Plugins/Animation/ControlRig/Source/ControlRigEditor/Private/EditMode/SControlRigEditModeTools.h
+++ b/Engine/Plugins/Animation/ControlRig/Source/ControlRigEditor/Private/EditMode/SControlRigEditModeTools.h
@@ -59,12 +59,7 @@ public:
 	virtual bool IsPropertyKeyingEnabled() const override;
 	virtual void OnKeyPropertyClicked(const IPropertyHandle& KeyedPropertyHandle) override;
 	virtual bool IsPropertyAnimated(const class IPropertyHandle& PropertyHandle, UObject *ParentObject) const override;
-
-	//reuse settings for space and constraint baking baking
-	static FRigSpacePickerBakeSettings BakeSpaceSettings;
-
 private:
-
 	/** Sequencer we are currently bound to */
 
 	TWeakPtr<ISequencer> WeakSequencer;
diff --git a/Engine/Plugins/Animation/ControlRig/Source/ControlRigEditor/Private/EditMode/SControlRigSpacePicker.cpp b/Engine/Plugins/Animation/ControlRig/Source/ControlRigEditor/Private/EditMode/SControlRigSpacePicker.cpp
--- a/Engine/Plugins/Animation/ControlRig/Source/ControlRigEditor/Private/EditMode/SControlRigSpacePicker.cpp
+++ b/Engine/Plugins/Animation/ControlRig/Source/ControlRigEditor/Private/EditMode/SControlRigSpacePicker.cpp
@@ -276,9 +276,21 @@ FReply SControlRigSpacePicker::OnBakeControlsToNewSpaceButtonClicked()
 		Settings.TargetSpace = URigHierarchy::GetDefaultParentKey();
 
 		TRange<FFrameNumber> Range = Sequencer->GetFocusedMovieSceneSequence()->GetMovieScene()->GetPlaybackRange();
+		TArray<FFrameNumber> Keys;
+		TArray < FKeyHandle> KeyHandles;
 
-		Settings.Settings.StartFrame = Range.GetLowerBoundValue();
-		Settings.Settings.EndFrame = Range.GetUpperBoundValue();
+		Settings.StartFrame = Range.GetLowerBoundValue();
+		Settings.EndFrame = Range.GetUpperBoundValue();
+		if (Keys.Num() > 0)
+		{
+			int32 Index = Algo::LowerBound(Keys, CurrentTime);
+			if (Index >= 0 && Index < (Keys.Num() - 1))
+			{
+				Settings.StartFrame = Keys[Index];
+				Settings.EndFrame = Keys[Index + 1];
+
+			}
+		}
 
 		TSharedRef<SRigSpacePickerBakeWidget> BakeWidget =
 			SNew(SRigSpacePickerBakeWidget)
@@ -289,7 +301,14 @@ FReply SControlRigSpacePicker::OnBakeControlsToNewSpaceButtonClicked()
 			.GetControlCustomization(this, &SControlRigSpacePicker::HandleGetControlElementCustomization)
 			.OnBake_Lambda([Sequencer, ControlRig, TickResolution](URigHierarchy* InHierarchy, TArray<FRigElementKey> InControls, FRigSpacePickerBakeSettings InSettings)
 		{
-	
+			TArray<FFrameNumber> Frames;
+
+			const FFrameRate& FrameRate = Sequencer->GetFocusedDisplayRate();
+			FFrameNumber FrameRateInFrameNumber = TickResolution.AsFrameNumber(FrameRate.AsInterval());
+			for (FFrameNumber& Frame = InSettings.StartFrame; Frame <= InSettings.EndFrame; Frame += FrameRateInFrameNumber)
+			{
+				Frames.Add(Frame);
+			}
 			FScopedTransaction Transaction(LOCTEXT("BakeControlToSpace", "Bake Control In Space"));
 			for (const FRigElementKey& ControlKey : InControls)
 			{
@@ -298,7 +317,7 @@ FReply SControlRigSpacePicker::OnBakeControlsToNewSpaceButtonClicked()
 				if (SpaceChannelAndSection.SpaceChannel)
 				{
 					FControlRigSpaceChannelHelpers::SequencerBakeControlInSpace(ControlRig, Sequencer, SpaceChannelAndSection.SpaceChannel, SpaceChannelAndSection.SectionToKey,
-						InHierarchy, ControlKey, InSettings);
+						Frames, InHierarchy, ControlKey, InSettings);
 				}
 			}
 			return FReply::Handled();
diff --git a/Engine/Plugins/Animation/ControlRig/Source/ControlRigEditor/Private/Editor/Constraints/SConstraintsWidget.cpp b/Engine/Plugins/Animation/ControlRig/Source/ControlRigEditor/Private/Editor/Constraints/SConstraintsWidget.cpp
--- a/Engine/Plugins/Animation/ControlRig/Source/ControlRigEditor/Private/Editor/Constraints/SConstraintsWidget.cpp
+++ b/Engine/Plugins/Animation/ControlRig/Source/ControlRigEditor/Private/Editor/Constraints/SConstraintsWidget.cpp
@@ -26,8 +26,6 @@
 #include "Tools/BakingHelper.h"
 #include "MovieSceneToolHelpers.h"
 #include "Styling/SlateIconFinder.h"
-#include "Widgets/Input/NumericTypeInterface.h"
-#include "FrameNumberDetailsCustomization.h"
 
 #define LOCTEXT_NAMESPACE "SConstraintsWidget"
 
@@ -111,7 +109,7 @@ TArray<AActor*> GetCurrentSelection()
 	});	
 }
 
-static TWeakPtr<ISequencer> GetSequencerChecked()
+TWeakPtr<ISequencer> GetSequencerChecked()
 {
 	const TWeakPtr<ISequencer> WeakSequencer = FBakingHelper::GetSequencer();
 	if (!WeakSequencer.IsValid() || !WeakSequencer.Pin()->GetFocusedMovieSceneSequence())
@@ -158,6 +156,7 @@ void SDroppableConstraintItem::Construct(
 	.Padding(FMargin(8.f, 2.f, 12.f, 2.f))
 	[
 		SNew(SOverlay)
+
 		+SOverlay::Slot()
 		[
 			SNew(SBorder)
@@ -454,6 +453,7 @@ TSharedRef<ITableRow> SConstraintsCreationWidget::OnGenerateWidgetForItem(
 	const TSharedRef<STableViewBase>& OwnerTable)
 {
 	return SNew(STableRow<ItemSharedPtr>, OwnerTable)
+	.Style(&FAppStyle::Get(), "PlacementBrowser.PlaceableItemRow")
 	[
 		SNew(SDroppableConstraintItem, InItem.ToSharedRef(), SharedThis(this))
 	];
@@ -494,13 +494,7 @@ void SEditableConstraintItem::Construct(
 		}
 		
 		const FConstraintsManagerController& Controller = FConstraintsManagerController::Get(World);
-		UTickableConstraint* Constraint = ConstraintItem->Constraint.Get();
-		int32 Index = Controller.GetConstraintsArray().Find(Constraint);
-		if (Index == INDEX_NONE)
-		{
-			return nullptr;
-		}
-		return Constraint;
+		return Controller.GetConstraint(ConstraintItem->Name);
 	};
 	TWeakObjectPtr<UTickableConstraint> Constraint = GetConstraint();
 
@@ -509,10 +503,10 @@ void SEditableConstraintItem::Construct(
 	
 	// labels
 	FString ParentLabel(TEXT("undefined")), ChildLabel(TEXT("undefined"));
-	if (!InItem->GetLabel().IsEmpty())
+	if (!InItem->Label.IsEmpty())
 	{
-		ChildLabel = InItem->GetLabel();
-		InItem->GetLabel().Split(TEXT("."), &ParentLabel, &ChildLabel);
+		ChildLabel = InItem->Label;
+		InItem->Label.Split(TEXT("."), &ParentLabel, &ChildLabel);
 	}
 	
 	FString ParentFullLabel = ParentLabel, ChildFullLabel = ChildLabel;
@@ -720,17 +714,43 @@ void SEditableConstraintItem::Construct(
 	];
 }
 
-//////////////////////////////////////////////////////////////
-/// FBaseConstraintListWidget
-///////////////////////////////////////////////////////////
+/**
+ * SConstraintsEditionWidget
+ */
+
+void SConstraintsEditionWidget::Construct(const FArguments& InArgs)
+{
+	ChildSlot
+	[
+		SNew( SVerticalBox )
+
+		+ SVerticalBox::Slot()
+		.Padding(FMargin(0.0f, 3.f))
+		[
+			SNew(SOverlay)
+
+			+ SOverlay::Slot()
+			[
+				SAssignNew(ListView, ConstraintItemListView)
+				.SelectionMode(ESelectionMode::Single)
+				.ListItemsSource( &ListItems )
+				.OnGenerateRow(this, &SConstraintsEditionWidget::OnGenerateWidgetForItem)
+				.OnContextMenuOpening(this, &SConstraintsEditionWidget::CreateContextMenu)
+				.OnMouseButtonDoubleClick(this, &SConstraintsEditionWidget::OnItemDoubleClicked)
+			]
+		]
+	];
 
+	RefreshConstraintList();
+	RegisterSelectionChanged();
+}
 
-FBaseConstraintListWidget::~FBaseConstraintListWidget()
+SConstraintsEditionWidget::~SConstraintsEditionWidget()
 {
 	UnregisterSelectionChanged();
 }
 
-void FBaseConstraintListWidget::RegisterSelectionChanged()
+void SConstraintsEditionWidget::RegisterSelectionChanged()
 {
 	FLevelEditorModule& LevelEditor = FModuleManager::GetModuleChecked<FLevelEditorModule>("LevelEditor");
 	FLevelEditorModule::FActorSelectionChangedEvent& ActorSelectionChangedEvent = LevelEditor.OnActorSelectionChanged();
@@ -743,10 +763,10 @@ void FBaseConstraintListWidget::RegisterSelectionChanged()
 	}
 
 	// register
-	ActorSelectionChangedEvent.AddRaw(this, &FBaseConstraintListWidget::OnActorSelectionChanged);
+	ActorSelectionChangedEvent.AddRaw( this, &SConstraintsEditionWidget::OnActorSelectionChanged );
 }
 
-void FBaseConstraintListWidget::UnregisterSelectionChanged()
+void SConstraintsEditionWidget::UnregisterSelectionChanged()
 {
 	if (OnSelectionChangedHandle.IsValid())
 	{
@@ -756,158 +776,6 @@ void FBaseConstraintListWidget::UnregisterSelectionChanged()
 	}
 }
 
-void FBaseConstraintListWidget::PostUndo(bool bSuccess)
-{
-	InvalidateConstraintList();
-}
-
-void FBaseConstraintListWidget::PostRedo(bool bSuccess)
-{
-	InvalidateConstraintList();
-}
-
-void FBaseConstraintListWidget::InvalidateConstraintList()
-{
-	bNeedsRefresh = true;
-
-	// NOTE: this is a hack to disable the picker if the selection changes. The picker should trigger the escape key
-	// but this is an editor wide change
-	static const FActorPickerModeModule& ActorPickerMode = FModuleManager::Get().GetModuleChecked<FActorPickerModeModule>("ActorPickerMode");
-	if (ActorPickerMode.IsInActorPickingMode())
-	{
-		ActorPickerMode.EndActorPickingMode();
-	}
-}
-
-int32 FBaseConstraintListWidget::RefreshConstraintList()
-{
-	// get constraints
-	UWorld* World = GetCurrentWorld();
-	if (!IsValid(World))
-	{
-		return 0;
-	}
-
-	const TArray<AActor*> Selection = GetCurrentSelection();
-
-	const bool bIsConstraintsActor = Selection.Num() == 1 && Selection[0]->IsA<AConstraintsActor>();
-
-	TArray< TObjectPtr<UTickableConstraint> > Constraints;
-	if (bIsConstraintsActor)
-	{
-		const FConstraintsManagerController& Controller = FConstraintsManagerController::Get(World);
-		static constexpr bool bSorted = true;
-		Constraints = Controller.GetAllConstraints(bSorted);
-	}
-	else
-	{
-		for (const AActor* Actor : Selection)
-		{
-			FTransformConstraintUtils::GetParentConstraints(World, Actor, Constraints);
-		}
-	}
-
-	// rebuild item list
-	ListItems.Empty();
-
-	const UEnum* ETransformConstraintTypeEnum = StaticEnum<ETransformConstraintType>();
-	for (const TObjectPtr<UTickableConstraint>& Constraint : Constraints)
-	{
-		const int8 Type = FConstraintInfo::GetType(Constraint->GetClass());
-		if (ETransformConstraintTypeEnum->IsValidEnumValue(Type))
-		{
-			const ETransformConstraintType ConstraintType = static_cast<ETransformConstraintType>(Type);
-			ListItems.Emplace(FEditableConstraintItem::Make(Constraint, ConstraintType));
-		}
-	}
-
-	// refresh tree
-	ListView->RequestListRefresh();
-
-	return ListItems.Num();
-}
-
-void FBaseConstraintListWidget::OnActorSelectionChanged(const TArray<UObject*>& NewSelection, bool bForceRefresh)
-{
-	// NOTE we use this delegate to trigger an tree update, however, control actors are not selected as they are 
-	// Temporary Editor Actors so NewSelection won't contain the controls
-	InvalidateConstraintList();
-}
-
-
-//////////////////////////////////////////////////////////////
-/// SConstraintsEditionWidget
-///////////////////////////////////////////////////////////
-
-
-void SConstraintsEditionWidget::Construct(const FArguments& InArgs)
-{
-
-	ChildSlot
-	[
-		SNew(SBorder)
-		.Visibility(EVisibility::Visible)
-		.BorderImage(FAppStyle::GetBrush("Menu.Background"))
-		[
-
-			SNew( SVerticalBox )
-
-			+ SVerticalBox::Slot()
-			.AutoHeight()
-			.Padding(FMargin(0.0f, 0.f))
-			[
-				SNew(SOverlay)
-
-				+ SOverlay::Slot()
-				[
-					SAssignNew(ListView, ConstraintItemListView)
-					.SelectionMode(ESelectionMode::Single)
-					.ListItemsSource( &ListItems )
-					.OnGenerateRow(this, &SConstraintsEditionWidget::OnGenerateWidgetForItem)
-					.OnContextMenuOpening(this, &SConstraintsEditionWidget::CreateContextMenu)
-					.OnMouseButtonDoubleClick(this, &SConstraintsEditionWidget::OnItemDoubleClicked)
-				]
-
-			]
-	
-			+ SVerticalBox::Slot()
-				.AutoHeight()
-				.Padding(FMargin(0.0f, 0.f))
-				[
-		
-					SNew(SHorizontalBox)
-					+ SHorizontalBox::Slot()
-					.FillWidth(1.f)
-					.HAlign(HAlign_Fill)
-					[
-						SNew(SSpacer)
-					]
-		
-					+ SHorizontalBox::Slot()
-					.AutoWidth()
-					.VAlign(VAlign_Center)
-					.HAlign(HAlign_Right)
-					.Padding(1.f)
-					[
-						SNew(SButton)
-						.ButtonStyle(FAppStyle::Get(), "FlatButton.Default")
-						.Text(LOCTEXT("BakeButton", "Bake..."))
-						.OnClicked(this, &SConstraintsEditionWidget::OnBakeClicked)
-						.IsEnabled_Lambda([this]()
-						{
-							return (ListItems.Num() > 0);
-						})
-						.ToolTipText(LOCTEXT("BakeConstraintButtonToolTip", "Bake selected contraints"))
-					]
-				]
-		]
-	];
-
-	RefreshConstraintList();
-	RegisterSelectionChanged();
-}
-
-
 void SConstraintsEditionWidget::Tick(const FGeometry& AllottedGeometry, const double InCurrentTime, const float InDeltaTime)
 {
 	SCompoundWidget::Tick(AllottedGeometry, InCurrentTime, InDeltaTime);
@@ -919,16 +787,15 @@ void SConstraintsEditionWidget::Tick(const FGeometry& AllottedGeometry, const do
 	}
 }
 
-
 TSharedRef<ITableRow> SConstraintsEditionWidget::OnGenerateWidgetForItem(
 	ItemSharedPtr InItem,
 	const TSharedRef<STableViewBase>& OwnerTable)
 {
 	return SNew(STableRow<ItemSharedPtr>, OwnerTable)
-		.Style(&FAppStyle::Get(), "PlacementBrowser.PlaceableItemRow")
-		[
-			SNew(SEditableConstraintItem, InItem.ToSharedRef(), SharedThis(this))
-		];
+	.Style(&FAppStyle::Get(), "PlacementBrowser.PlaceableItemRow")
+	[
+		SNew(SEditableConstraintItem, InItem.ToSharedRef(), SharedThis(this))
+	];
 }
 
 bool SConstraintsEditionWidget::CanMoveUp(const TSharedPtr<FEditableConstraintItem>& Item) const
@@ -972,8 +839,8 @@ void SConstraintsEditionWidget::MoveItemUp(const TSharedPtr<FEditableConstraintI
 		const FConstraintsManagerController& Controller = FConstraintsManagerController::Get(World);
 
 		// the current item needs to tick before the previous item 
-		const FName BeforeName(ListItems[Index]->GetName().ToString());
-		const FName AfterName(ListItems[Index-1]->GetName().ToString());
+		const FName BeforeName(ListItems[Index]->Name.ToString());
+		const FName AfterName(ListItems[Index-1]->Name.ToString());
 		Controller.SetConstraintsDependencies(BeforeName, AfterName);
 		
 		RefreshConstraintList();
@@ -999,8 +866,8 @@ void SConstraintsEditionWidget::MoveItemDown(const TSharedPtr<FEditableConstrain
 		const FConstraintsManagerController& Controller = FConstraintsManagerController::Get(World);
 
 		// the current item needs to tick after the next item 
-		const FName BeforeName(ListItems[Index+1]->GetName().ToString());
-		const FName AfterName(ListItems[Index]->GetName().ToString());
+		const FName BeforeName(ListItems[Index+1]->Name.ToString());
+		const FName AfterName(ListItems[Index]->Name.ToString());
 		Controller.SetConstraintsDependencies(BeforeName, AfterName);
 		
 		RefreshConstraintList();
@@ -1019,25 +886,88 @@ void SConstraintsEditionWidget::RemoveItem(const TSharedPtr<FEditableConstraintI
 	
 	const FConstraintsManagerController& Controller = FConstraintsManagerController::Get(World);
 	
-	Controller.RemoveConstraint(Item->Constraint.Get());
+	const FName ConstraintName(Item->Name.ToString());
+	Controller.RemoveConstraint(ConstraintName);
 
 	RefreshConstraintList();
 }
 
-FReply SConstraintsEditionWidget::OnBakeClicked()
+void SConstraintsEditionWidget::PostUndo(bool bSuccess)
 {
-	UTickableTransformConstraint* InConstraint = nullptr;
-	TWeakPtr<ISequencer> WeakSequencer = GetSequencerChecked();
-	if (!WeakSequencer.IsValid() || ListItems.Num() < 1)
+	InvalidateConstraintList();
+}
+
+void SConstraintsEditionWidget::PostRedo(bool bSuccess)
+{
+	InvalidateConstraintList();
+}
+
+void SConstraintsEditionWidget::InvalidateConstraintList()
+{
+	bNeedsRefresh = true;
+
+	// NOTE: this is a hack to disable the picker if the selection changes. The picker should trigger the escape key
+	// but this is an editor wide change
+	static const FActorPickerModeModule& ActorPickerMode = FModuleManager::Get().GetModuleChecked<FActorPickerModeModule>("ActorPickerMode");
+	if (ActorPickerMode.IsInActorPickingMode())
 	{
-		return FReply::Unhandled();
+		ActorPickerMode.EndActorPickingMode();
 	}
-	TSharedPtr<ISequencer> Sequencer = WeakSequencer.Pin();
-	TSharedRef<SConstraintBakeWidget> BakeWidget =
-		SNew(SConstraintBakeWidget)
-		.Sequencer(Sequencer);
+}
 
-	return BakeWidget->OpenDialog(true);
+int32 SConstraintsEditionWidget::RefreshConstraintList()
+{
+	// get constraints
+	UWorld* World = GetCurrentWorld();
+	if (!IsValid(World))
+	{
+		return 0;
+	}
+	
+	const TArray<AActor*> Selection = GetCurrentSelection();
+	
+	const bool bIsConstraintsActor = Selection.Num() == 1 && Selection[0]->IsA<AConstraintsActor>();
+	
+	TArray< TObjectPtr<UTickableConstraint> > Constraints;
+	if (bIsConstraintsActor)
+	{
+		const FConstraintsManagerController& Controller = FConstraintsManagerController::Get(World);
+		static constexpr bool bSorted = true;
+		Constraints = Controller.GetAllConstraints(bSorted);
+	}
+	else
+	{
+		for (const AActor* Actor : Selection)
+		{
+			FTransformConstraintUtils::GetParentConstraints(World, Actor, Constraints);
+		}
+	}
+	
+	// rebuild item list
+	ListItems.Empty();
+	
+	const UEnum* ETransformConstraintTypeEnum = StaticEnum<ETransformConstraintType>();
+	for (const TObjectPtr<UTickableConstraint>& Constraint: Constraints)
+	{
+		const int8 Type = FConstraintInfo::GetType(Constraint->GetClass()); 
+		if (ETransformConstraintTypeEnum->IsValidEnumValue(Type))
+		{
+			const ETransformConstraintType ConstraintType = static_cast<ETransformConstraintType>(Type);
+			ListItems.Emplace( FEditableConstraintItem::Make(Constraint->GetFName(), ConstraintType, Constraint->GetLabel()) );
+		}
+	}
+
+	// refresh tree
+	ListView->RequestListRefresh();
+	
+	return ListItems.Num();
+}
+
+void SConstraintsEditionWidget::OnActorSelectionChanged(const TArray<UObject*>& NewSelection, bool bForceRefresh)
+{
+	// NOTE we use this delegate to trigger an tree update, however, control actors are not selected as they are 
+	// Temporary Editor Actors so NewSelection won't contain the controls
+	InvalidateConstraintList();
 }
 
 TSharedPtr<SWidget> SConstraintsEditionWidget::CreateContextMenu()
@@ -1061,9 +991,8 @@ TSharedPtr<SWidget> SConstraintsEditionWidget::CreateContextMenu()
 	}
 	
 	const FConstraintsManagerController& Controller = FConstraintsManagerController::Get(World);
-	UTickableConstraint* Constraint = ListItems[Index]->Constraint.Get();
-	int32 ContIndex = Controller.GetConstraintsArray().Find(Constraint);
-	if (ContIndex == INDEX_NONE)
+	UTickableConstraint* Constraint = Controller.GetConstraint(ListItems[Index]->Name);
+	if (!Constraint)
 	{
 		return SNullWidget::NullWidget;
 	}
@@ -1131,7 +1060,7 @@ TSharedPtr<SWidget> SConstraintsEditionWidget::CreateContextMenu()
 	// transform constraint options
 	if (UTickableTransformConstraint* TransformConstraint = Cast<UTickableTransformConstraint>(Constraint))
 	{
-		const FText ConstraintLabel = FText::FromName(ListItems[Index]->GetName());
+		const FText ConstraintLabel = FText::FromName(ListItems[Index]->Name);
 	
 		// baking (note that this will probably be moved)
 		MenuBuilder.BeginSection("BakeConstraint", LOCTEXT("BakeConstraintHeader", "Bake Constraint"));
@@ -1149,8 +1078,9 @@ TSharedPtr<SWidget> SConstraintsEditionWidget::CreateContextMenu()
 						{
 							return;
 						}
+						
 						FScopedTransaction Transaction(LOCTEXT("BakeConstraint", "Bake Constraint"));
-						FConstraintBaker::Bake(World, TransformConstraint, WeakSequencer.Pin(), TOptional< FBakingAnimationKeySettings>(), TOptional<TArray<FFrameNumber>>());
+						FConstraintBaker::Bake(World, TransformConstraint, WeakSequencer.Pin(), TOptional<TArray<FFrameNumber>>());
 
 					}
 				})),
@@ -1218,376 +1148,13 @@ void SConstraintsEditionWidget::OnItemDoubleClicked(ItemSharedPtr InItem)
 	}
 	
 	const FConstraintsManagerController& Controller = FConstraintsManagerController::Get(World);
-	UTickableConstraint* Constraint = InItem->Constraint.Get();
-	int32 ContIndex = Controller.GetConstraintsArray().Find(Constraint);
-	if (ContIndex != INDEX_NONE)
-	{
-		Constraint->SetActive(!Constraint->Active);
-	}
-}
-
-//////////////////////////////////////////////////////////////
-/// SConstraintBakeWidget
-///////////////////////////////////////////////////////////
-
-TOptional<FBakingAnimationKeySettings> SConstraintBakeWidget::BakeConstraintSettings;
-
-void SConstraintBakeWidget::Construct(const FArguments& InArgs)
-{
-	check(InArgs._Sequencer);
-	Sequencer = InArgs._Sequencer;
-
-	if (BakeConstraintSettings.IsSet() == false)
-	{
-		BakeConstraintSettings = FBakingAnimationKeySettings();
-		const FFrameRate TickResolution = Sequencer->GetFocusedTickResolution();
-		const FFrameTime FrameTime = Sequencer->GetLocalTime().ConvertTo(TickResolution);
-		FFrameNumber CurrentTime = FrameTime.GetFrame();
-
-		TRange<FFrameNumber> Range = Sequencer->GetFocusedMovieSceneSequence()->GetMovieScene()->GetPlaybackRange();
-		TArray<FFrameNumber> Keys;
-		TArray < FKeyHandle> KeyHandles;
-
-		BakeConstraintSettings.GetValue().StartFrame = Range.GetLowerBoundValue();
-		BakeConstraintSettings.GetValue().EndFrame = Range.GetUpperBoundValue();
-	}
-
-
-	Settings = MakeShared<TStructOnScope<FBakingAnimationKeySettings>>();
-	Settings->InitializeAs<FBakingAnimationKeySettings>(BakeConstraintSettings.GetValue());
-
-	FStructureDetailsViewArgs StructureViewArgs;
-	StructureViewArgs.bShowObjects = true;
-	StructureViewArgs.bShowAssets = true;
-	StructureViewArgs.bShowClasses = true;
-	StructureViewArgs.bShowInterfaces = true;
-
-	FDetailsViewArgs ViewArgs;
-	ViewArgs.bAllowSearch = false;
-	ViewArgs.bHideSelectionTip = false;
-	ViewArgs.bShowObjectLabel = false;
-
-	FPropertyEditorModule& PropertyEditor = FModuleManager::Get().LoadModuleChecked<FPropertyEditorModule>(TEXT("PropertyEditor"));
-
-	DetailsView = PropertyEditor.CreateStructureDetailView(ViewArgs, StructureViewArgs, TSharedPtr<FStructOnScope>());
-	TSharedPtr<INumericTypeInterface<double>> NumericTypeInterface = Sequencer->GetNumericTypeInterface();
-	DetailsView->GetDetailsView()->RegisterInstancedCustomPropertyTypeLayout("FrameNumber",
-		FOnGetPropertyTypeCustomizationInstance::CreateLambda([=]() {return MakeShared<FFrameNumberDetailsCustomization>(NumericTypeInterface); }));
-	DetailsView->SetStructureData(Settings);
-
-
-	ChildSlot
-		[
-			SNew(SBorder)
-			.Visibility(EVisibility::Visible)
-			[
-			SNew(SVerticalBox)
-			+ SVerticalBox::Slot()
-			.Padding(FMargin(0.0f, 0.f))
-			[
-				SAssignNew(ListView, ConstraintItemListView)
-				.SelectionMode(ESelectionMode::Multi)
-				.ListItemsSource(&ListItems)
-				.OnGenerateRow(this, &SConstraintBakeWidget::OnGenerateWidgetForItem)
-
-			]
-			
-				+SVerticalBox::Slot()
-				.AutoHeight()
-				.Padding(0.f, 8.f, 0.f, 0.f)
-				[
-					DetailsView->GetWidget().ToSharedRef()
-				]
-
-				+ SVerticalBox::Slot()
-				.AutoHeight()
-				.Padding(0.f, 16.f, 0.f, 8.f)
-				[
-					SNew(SHorizontalBox)
-
-					+ SHorizontalBox::Slot()
-					.FillWidth(1.f)
-					[
-						SNew(SSpacer)
-					]
-
-				+ SHorizontalBox::Slot()
-				.AutoWidth()
-				.HAlign(HAlign_Left)
-				.VAlign(VAlign_Center)
-				.Padding(8.f, 0.f, 0.f, 0.f)
-				[
-					SNew(SButton)
-					.HAlign(HAlign_Center)
-					.ContentPadding(FAppStyle::GetMargin("StandardDialog.ContentPadding"))
-					.Text(LOCTEXT("OK", "OK"))
-					.OnClicked_Lambda([this, InArgs]()
-						{
-							BakeSelected();
-							CloseDialog();
-							return FReply::Handled();
-
-						})
-						.IsEnabled_Lambda([this]()
-						{
-							return (ListItems.Num() > 0 && Settings.IsValid());
-						})
-			]
-
-			+ SHorizontalBox::Slot()
-			.AutoWidth()
-			.HAlign(HAlign_Left)
-			.VAlign(VAlign_Center)
-			.Padding(8.f, 0.f, 16.f, 0.f)
-			[
-				SNew(SButton)
-				.HAlign(HAlign_Center)
-				.ContentPadding(FAppStyle::GetMargin("StandardDialog.ContentPadding"))
-				.Text(LOCTEXT("Cancel", "Cancel"))
-				.OnClicked_Lambda([this]()
-					{
-						CloseDialog();
-					return FReply::Handled();
-					})
-				]
-			]
-		]
-	];
-	RefreshConstraintList();
-	RegisterSelectionChanged();
-}
-
-TSharedRef<ITableRow> SConstraintBakeWidget::OnGenerateWidgetForItem(
-	ItemSharedPtr InItem,
-	const TSharedRef<STableViewBase>& OwnerTable)
-{
-	return SNew(STableRow<ItemSharedPtr>, OwnerTable)
-		[
-			SNew(SBakeConstraintItem, InItem.ToSharedRef(), SharedThis(this))
-		];
-		
-}
-
-void SConstraintBakeWidget::Tick(const FGeometry& AllottedGeometry, const double InCurrentTime, const float InDeltaTime)
-{
-	SCompoundWidget::Tick(AllottedGeometry, InCurrentTime, InDeltaTime);
-
-	if (bNeedsRefresh)
-	{
-		RefreshConstraintList();
-		bNeedsRefresh = false;
-	}
-}
-
-//for the bake we auto select averything when selection or something chagnes
-int32 SConstraintBakeWidget::RefreshConstraintList()
-{
-	int32 Num = FBaseConstraintListWidget::RefreshConstraintList();
-	ListView->SetItemSelection(ListItems, true, ESelectInfo::Direct);
-	return Num;
-}
-
-void SConstraintBakeWidget::BakeSelected()
-{
-	TArray<TSharedPtr<FEditableConstraintItem>> Selection = ListView->GetSelectedItems();
-	TArray<UTickableTransformConstraint*> Constraints;
-	for (TSharedPtr<FEditableConstraintItem>& Item : Selection)
-	{
-		if (UTickableTransformConstraint* TransformConstraint = Cast<UTickableTransformConstraint>(Item->Constraint.Get()))
-		{
-			Constraints.Add(TransformConstraint);
-		}
-	}
-	UWorld* World = GetCurrentWorld();
-
-	FBakingAnimationKeySettings* BakeSettings = Settings->Get();
-	
-	FScopedTransaction Transaction(LOCTEXT("BakeConstraints", "Bake Constraints"));
-	FConstraintBaker::BakeMultiple(World,Constraints, Sequencer,*BakeSettings );
-	BakeConstraintSettings = *BakeSettings;
-	InvalidateConstraintList();
-}
-
-class SConstraintBakedWidgetWindow : public SWindow
-{
-};
-
-
-
-FReply SConstraintBakeWidget::OpenDialog(bool bModal)
-{
-	check(!DialogWindow.IsValid());
-
-	const FVector2D CursorPos = FSlateApplication::Get().GetCursorPos();
-
-	TSharedRef<SConstraintBakedWidgetWindow> Window = SNew(SConstraintBakedWidgetWindow)
-		.Title(LOCTEXT("SConstraintsBakeWidgetTitle", "Bake Constraints"))
-		.CreateTitleBar(true)
-		.Type(EWindowType::Normal)
-		.SizingRule(ESizingRule::Autosized)
-		.ScreenPosition(CursorPos)
-		.FocusWhenFirstShown(true)
-		.ActivationPolicy(EWindowActivationPolicy::FirstShown)
-		[
-			AsShared()
-		];
-
-	Window->SetWidgetToFocusOnActivate(AsShared());
-
-	DialogWindow = Window;
-
-	Window->MoveWindowTo(CursorPos);
-
-	if (bModal)
-	{
-		GEditor->EditorAddModalWindow(Window);
-	}
-	else
-	{
-		FSlateApplication::Get().AddWindow(Window);
-	}
-
-	return FReply::Handled();
-}
-
-void SConstraintBakeWidget::CloseDialog()
-{
-	if (DialogWindow.IsValid())
-	{
-		DialogWindow.Pin()->RequestDestroyWindow();
-		DialogWindow.Reset();
-	}
-}
-
-//////////////////////////////////////////////////////////////
-/// SBakeConstraintItem
-///////////////////////////////////////////////////////////
-
-
-void SBakeConstraintItem::Construct(
-	const FArguments& InArgs,
-	const TSharedPtr<FEditableConstraintItem>& InItem,
-	TSharedPtr<SConstraintBakeWidget> InConstraintsWidget)
-{
-	ConstraintItem = InItem;
-	ConstraintsWidget = InConstraintsWidget;
-
-	static const FSlateBrush* RoundedBoxBrush = FControlRigEditorStyle::Get().GetBrush(TEXT("ControlRig.SpacePicker.RoundedRect"));
-
-	// enum to string
-	const uint8 ConstraintType = static_cast<uint8>(InItem->Type);
-
-
-	// constraint
-	auto GetConstraint = [this]() -> UTickableConstraint*
-	{
-		UWorld* World = GetCurrentWorld();
-		if (!IsValid(World))
-		{
-			return nullptr;
-		}
-
-		const FConstraintsManagerController& Controller = FConstraintsManagerController::Get(World);
-		UTickableConstraint* Constraint = ConstraintItem->Constraint.Get();
-		int32 Index = Controller.GetConstraintsArray().Find(Constraint);
-		if (Index == INDEX_NONE)
-		{
-			return nullptr;
-		}
-		return Constraint;
-	};
-	TWeakObjectPtr<UTickableConstraint> Constraint = GetConstraint();
-
-	// sequencer
-	TWeakPtr<ISequencer> WeakSequencer = GetSequencerChecked();
-
-	// labels
-	FString ParentLabel(TEXT("undefined")), ChildLabel(TEXT("undefined"));
-	if (!InItem->GetLabel().IsEmpty())
-	{
-		ChildLabel = InItem->GetLabel();
-		InItem->GetLabel().Split(TEXT("."), &ParentLabel, &ChildLabel);
-	}
-
-	FString ParentFullLabel = ParentLabel, ChildFullLabel = ChildLabel;
-	if (IsValid(Constraint.Get()))
+	UTickableConstraint* Constraint = Controller.GetConstraint(InItem->Name);
+	if (!Constraint)
 	{
-		Constraint->GetFullLabel().Split(TEXT("."), &ParentFullLabel, &ChildFullLabel);
+		return;
 	}
 
-	// widgets
-	ChildSlot
-		.Padding(FMargin(8.f, 2.f, 12.f, 2.f))
-		[
-			SNew(SHorizontalBox)
-
-			+ SHorizontalBox::Slot()
-			.AutoWidth()
-			.VAlign(VAlign_Center)
-			.HAlign(HAlign_Fill)
-			.Padding(0)
-			[
-			SNew(SBorder)
-			.Padding(FMargin(5.0, 2.0, 5.0, 2.0))
-			.BorderImage(RoundedBoxBrush)
-			.BorderBackgroundColor_Lambda([Constraint]()
-				{
-					return FStyleColors::Transparent;
-				})	
-				.Content()
-				[
-					SNew(SHorizontalBox)
-
-					// constraint icon
-					+ SHorizontalBox::Slot()
-					.AutoWidth()
-					.VAlign(VAlign_Center)
-					.HAlign(HAlign_Left)
-					.Padding(FMargin(0.f, 0.f, 3.f, 0.f))
-					[
-						SNew(SImage)
-						.Image(FConstraintInfo::GetBrush(ConstraintType))
-						.ColorAndOpacity(FSlateColor::UseForeground())
-					]
-
-				// constraint name
-				+ SHorizontalBox::Slot()
-					.AutoWidth()
-					.VAlign(VAlign_Center)
-					.HAlign(HAlign_Left)
-					.Padding(0)
-					[
-						SNew(STextBlock)
-						.Text_Lambda([ParentLabel]()
-							{
-								return FText::FromString(ParentLabel);
-							})
-					.Font_Lambda([Constraint]()
-						{
-							if (!Constraint.IsValid() || !IsValid(Constraint.Get()))
-							{
-								return IDetailLayoutBuilder::GetDetailFont();
-							}
-							return Constraint->Active ? IDetailLayoutBuilder::GetDetailFont() : IDetailLayoutBuilder::GetDetailFontItalic();
-						})
-								.ToolTipText_Lambda([Constraint, ParentFullLabel, ChildFullLabel]()
-									{
-										if (!Constraint.IsValid() || !IsValid(Constraint.Get()))
-										{
-											return FText();
-										}
-
-						static constexpr TCHAR ToolTipFormat[] = TEXT("%s constraint between parent '%s' and child '%s'.");
-						const FString TypeLabel = Constraint->GetTypeLabel();
-						const FString FullLabel = FString::Printf(ToolTipFormat, *TypeLabel, *ParentFullLabel, *ChildFullLabel);
-						return FText::FromString(FullLabel);
-									})
-					]
-					]
-		]
-
-	
-		];
+	Constraint->SetActive(!Constraint->Active);
 }
 
 END_SLATE_FUNCTION_BUILD_OPTIMIZATION
diff --git a/Engine/Plugins/Animation/ControlRig/Source/ControlRigEditor/Private/Editor/Constraints/SConstraintsWidget.h b/Engine/Plugins/Animation/ControlRig/Source/ControlRigEditor/Private/Editor/Constraints/SConstraintsWidget.h
--- a/Engine/Plugins/Animation/ControlRig/Source/ControlRigEditor/Private/Editor/Constraints/SConstraintsWidget.h
+++ b/Engine/Plugins/Animation/ControlRig/Source/ControlRigEditor/Private/Editor/Constraints/SConstraintsWidget.h
@@ -6,14 +6,12 @@
 #include "Widgets/Views/SListView.h"
 #include "Constraint.h"
 #include "EditorUndoClient.h"
-#include "IStructureDetailsView.h"
-#include "BakingAnimationKeySettings.h"
-#include "ConstraintsManager.h"
 
 class AActor;
 class SConstraintsCreationWidget;
 class SConstraintsEditionWidget;
-class ISequencer;
+class UTickableConstraint;
+
 DECLARE_DELEGATE(FOnConstraintCreated);
 
 /**
@@ -149,36 +147,22 @@ class FEditableConstraintItem
 {
 public:
 	static TSharedRef<FEditableConstraintItem> Make(
-		UTickableConstraint* InConstraint,
-		ETransformConstraintType InType)
+		const FName& InName,
+		ETransformConstraintType InType,
+		const FString& InLabel)
 	{
-		return MakeShareable(new FEditableConstraintItem(InConstraint, InType));
+		return MakeShareable(new FEditableConstraintItem(InName, InType, InLabel));
 	}
 	
-	TWeakObjectPtr<UTickableConstraint> Constraint = nullptr;
+	FName Name = NAME_None;
 	ETransformConstraintType Type = ETransformConstraintType::Parent;
-
-	FName GetName() const
-	{
-		if (Constraint.IsValid())
-		{
-			return Constraint->GetFName();
-		}
-		return NAME_None;
-	}
-	FString GetLabel() const
-	{
-		if (Constraint.IsValid())
-		{
-			return Constraint->GetLabel();
-		}
-		return FString();
-	}
+	FString Label;
 
 private:
-	FEditableConstraintItem(UTickableConstraint* InConstraint, ETransformConstraintType InType)
-		: Constraint(InConstraint)
+	FEditableConstraintItem(const FName& InName, ETransformConstraintType InType, const FString& InLabel)
+		: Name(InName)
 		, Type(InType)
+		, Label(InLabel)
 	{}
 	FEditableConstraintItem() {}
 };
@@ -205,56 +189,11 @@ private:
 	TWeakPtr<SConstraintsEditionWidget> ConstraintsWidget;
 };
 
-/**
-* Constraint List Shared by Edit Widget and Bake Widget
-*/
-
-class CONTROLRIGEDITOR_API FBaseConstraintListWidget : public FEditorUndoClient
-{
-public:
-
-	virtual ~FBaseConstraintListWidget() override;
-	/* FEditorUndoClient interface */
-	virtual void PostUndo(bool bSuccess);
-	virtual void PostRedo(bool bSuccess);
-	/* End FEditorUndoClient interface */
-
-	/** Invalidates the constraint list for further rebuild. */
-	void InvalidateConstraintList();
-
-	/** Rebuild the constraint list based on the current selection. */
-	virtual int32 RefreshConstraintList();
-
-	/** Triggers a constraint list invalidation when selection in the level viewport. */
-	void OnActorSelectionChanged(const TArray<UObject*>& NewSelection, bool bForceRefresh);
-
-protected:
-	/** Types */
-	using ItemSharedPtr = TSharedPtr<FEditableConstraintItem>;
-	using ConstraintItemListView = SListView< ItemSharedPtr >;
-
-
-	void RegisterSelectionChanged();
-	void UnregisterSelectionChanged();
-
-	/** List view that shows constraint types*/
-	TSharedPtr< ConstraintItemListView > ListView;
-
-	/** List of constraint types */
-	TArray< ItemSharedPtr > ListItems;
-
-	/** Boolean used to handle the items list refresh. */
-	bool bNeedsRefresh = false;
-
-	FDelegateHandle OnSelectionChangedHandle;
-
-};
-
 /**
  * SConstraintsEditionWidget
  */
 
-class CONTROLRIGEDITOR_API SConstraintsEditionWidget : public SCompoundWidget, public FBaseConstraintListWidget
+class CONTROLRIGEDITOR_API SConstraintsEditionWidget : public SCompoundWidget, public FEditorUndoClient
 {
 public:
 	SLATE_BEGIN_ARGS(SConstraintsEditionWidget)	{}
@@ -262,10 +201,20 @@ public:
 
 	/** Constructs this widget with InArgs */
 	void Construct(const FArguments& InArgs);
+	virtual ~SConstraintsEditionWidget() override;
 
 	/** Override for SWidget::Tick */
 	virtual void Tick(const FGeometry& AllottedGeometry, const double InCurrentTime, const float InDeltaTime) override;
 
+	/** Invalidates the constraint list for further rebuild. */
+	void InvalidateConstraintList();
+	
+	/** Rebuild the constraint list based on the current selection. */
+	int32 RefreshConstraintList();
+
+	/** Triggers a constraint list invalidation when selection in the level viewport. */
+	void OnActorSelectionChanged(const TArray<UObject*>& NewSelection, bool bForceRefresh);
+
 	/**  */
 	bool CanMoveUp(const TSharedPtr<FEditableConstraintItem>& Item) const;
 	bool CanMoveDown(const TSharedPtr<FEditableConstraintItem>& Item) const;
@@ -277,80 +226,34 @@ public:
 	/**  */
 	void RemoveItem(const TSharedPtr<FEditableConstraintItem>& Item);
 
-private:
-
-	/** Generates a widget for the specified item */
-	TSharedRef<ITableRow> OnGenerateWidgetForItem(ItemSharedPtr InItem, const TSharedRef<STableViewBase>& OwnerTable);
+	/* FEditorUndoClient interface */
+	virtual void PostUndo(bool bSuccess);
+	virtual void PostRedo(bool bSuccess);
+	/* End FEditorUndoClient interface */
 
+private:
 	/** Types */
 	using ItemSharedPtr = TSharedPtr<FEditableConstraintItem>;
 	using ConstraintItemListView = SListView< ItemSharedPtr >;
 
+	/** Generates a widget for the specified item */
+	TSharedRef<ITableRow> OnGenerateWidgetForItem(ItemSharedPtr InItem, const TSharedRef<STableViewBase>& OwnerTable);
+
 	/** @todo documentation. */
 	TSharedPtr< SWidget > CreateContextMenu();
 	void OnItemDoubleClicked(ItemSharedPtr InItem);
+	
+	void RegisterSelectionChanged();
+	void UnregisterSelectionChanged();
+	
+	/** List view that shows constraint types*/
+	TSharedPtr< ConstraintItemListView > ListView;
 
-	FReply OnBakeClicked();
-};
-
-/** Widget allowing baking of constraints */
-class CONTROLRIGEDITOR_API SConstraintBakeWidget : public SCompoundWidget, public FBaseConstraintListWidget
-{
-public:
-
-	SLATE_BEGIN_ARGS(SConstraintBakeWidget)
-		: _Sequencer(nullptr)
-	{}
-
-	SLATE_ARGUMENT(TSharedPtr<ISequencer>, Sequencer)
-	SLATE_END_ARGS()
-
-	void Construct(const FArguments& InArgs);
-	virtual ~SConstraintBakeWidget() override {}
-
-	FReply OpenDialog(bool bModal = true);
-	void CloseDialog();
-
-	//SWidget overrides
-	virtual void Tick(const FGeometry& AllottedGeometry, const double InCurrentTime, const float InDeltaTime) override;
-
-	//FBaseConstraintListWidget overrides
-	virtual int32 RefreshConstraintList() override;
-
-private:
-	void BakeSelected();
-
-	/** Generates a widget for the specified item */
-	TSharedRef<ITableRow> OnGenerateWidgetForItem(ItemSharedPtr InItem, const TSharedRef<STableViewBase>& OwnerTable);
+	/** List of constraint types */
+	TArray< ItemSharedPtr > ListItems;
 
-	TSharedPtr<ISequencer> Sequencer;
-	//static to be reused
-	static TOptional<FBakingAnimationKeySettings> BakeConstraintSettings;
-	//structonscope for details panel
-	TSharedPtr < TStructOnScope<FBakingAnimationKeySettings>> Settings;
-	TWeakPtr<SWindow> DialogWindow;
-	TSharedPtr<IStructureDetailsView> DetailsView;
+	/** Boolean used to handle the items list refresh. */
+	bool bNeedsRefresh = false;
 
+	FDelegateHandle OnSelectionChangedHandle;
 };
-
-/**
- * SBakeConstraintItem
- */
-
-class SBakeConstraintItem : public SCompoundWidget
-{
-public:
-	SLATE_BEGIN_ARGS(SBakeConstraintItem) {}
-	SLATE_END_ARGS()
-
-		/** Constructs this widget with InArgs and the actual tree item. */
-		void Construct(
-			const FArguments& InArgs,
-			const TSharedPtr<FEditableConstraintItem>& InItem,
-			TSharedPtr<SConstraintBakeWidget> InConstraintsWidget);
-
-private:
-	/** TSharedPtr to the tree item. */
-	TSharedPtr<FEditableConstraintItem> ConstraintItem;
-	TWeakPtr<SConstraintBakeWidget> ConstraintsWidget;
-};
\ No newline at end of file
diff --git a/Engine/Plugins/Animation/ControlRig/Source/ControlRigEditor/Private/Editor/SRigSpacePickerWidget.cpp b/Engine/Plugins/Animation/ControlRig/Source/ControlRigEditor/Private/Editor/SRigSpacePickerWidget.cpp
--- a/Engine/Plugins/Animation/ControlRig/Source/ControlRigEditor/Private/Editor/SRigSpacePickerWidget.cpp
+++ b/Engine/Plugins/Animation/ControlRig/Source/ControlRigEditor/Private/Editor/SRigSpacePickerWidget.cpp
@@ -186,10 +186,6 @@ void SRigSpacePickerWidget::Construct(const FArguments& InArgs)
 				.ButtonStyle(FAppStyle::Get(), "FlatButton.Default")
 				.Text(LOCTEXT("BakeButton", "Bake..."))
 				.OnClicked(InArgs._OnBakeButtonClicked)
-				.IsEnabled_Lambda([this]()
-				{
-					return (ControlKeys.Num() > 0);
-				})
 				.ToolTipText(LOCTEXT("BakeButtonToolTip", "Allows to bake the animation of one or more controls to a single space."))
 			];
 		}
@@ -1125,8 +1121,6 @@ void SRigSpacePickerBakeWidget::Construct(const FArguments& InArgs)
 	Settings = MakeShared<TStructOnScope<FRigSpacePickerBakeSettings>>();
 	Settings->InitializeAs<FRigSpacePickerBakeSettings>();
 	*Settings = InArgs._Settings;
-	//always setting space to be parent as default, since stored space may not be available.
-	Settings->Get()->TargetSpace = URigHierarchy::GetDefaultParentKey();
 	Sequencer = InArgs._Sequencer;
 
 	FStructureDetailsViewArgs StructureViewArgs;
@@ -1193,7 +1187,7 @@ void SRigSpacePickerBakeWidget::Construct(const FArguments& InArgs)
 
 			+SVerticalBox::Slot()
 			.AutoHeight()
-			.Padding(0.f, 0.f, 0.f, 0.f)
+			.Padding(0.f, 8.f, 0.f, 0.f)
 			[
 				DetailsView->GetWidget().ToSharedRef()
 			]
diff --git a/Engine/Plugins/Animation/ControlRig/Source/ControlRigEditor/Private/Editor/SRigSpacePickerWidget.h b/Engine/Plugins/Animation/ControlRig/Source/ControlRigEditor/Private/Editor/SRigSpacePickerWidget.h
--- a/Engine/Plugins/Animation/ControlRig/Source/ControlRigEditor/Private/Editor/SRigSpacePickerWidget.h
+++ b/Engine/Plugins/Animation/ControlRig/Source/ControlRigEditor/Private/Editor/SRigSpacePickerWidget.h
@@ -205,7 +205,7 @@ public:
 
 private:
 
-	//used for setting up the details
+	//FRigSpacePickerBakeSettings Settings;
 	TSharedPtr<TStructOnScope<FRigSpacePickerBakeSettings>> Settings;
 
 	ISequencer* Sequencer;
diff --git a/Engine/Plugins/Animation/ControlRig/Source/ControlRigEditor/Private/Sequencer/ControlRigParameterTrackEditor.cpp b/Engine/Plugins/Animation/ControlRig/Source/ControlRigEditor/Private/Sequencer/ControlRigParameterTrackEditor.cpp
--- a/Engine/Plugins/Animation/ControlRig/Source/ControlRigEditor/Private/Sequencer/ControlRigParameterTrackEditor.cpp
+++ b/Engine/Plugins/Animation/ControlRig/Source/ControlRigEditor/Private/Sequencer/ControlRigParameterTrackEditor.cpp
@@ -83,9 +83,6 @@
 #include "Constraints/ControlRigTransformableHandle.h"
 #include "PropertyEditorModule.h"
 #include "Constraints/TransformConstraintChannelInterface.h"
-#include "BakingAnimationKeySettings.h"
-#include "FrameNumberDetailsCustomization.h"
-#include "Sequencer/ControlRigSequencerHelpers.h"
 
 #define LOCTEXT_NAMESPACE "FControlRigParameterTrackEditor"
 
@@ -4245,385 +4242,141 @@ void FControlRigParameterTrackEditor::SelectFKBonesToAnimate(UFKControlRig* Auto
 
 	//reconstruct all channel proxies TODO or not to do that is the question
 }
-
-//////////////////////////////////////////////////////////////
-/// SCollapseControlsWidget
-///////////////////////////////////////////////////////////
-
-/** Widget allowing collapsing of controls */
-class SCollapseControlsWidget : public SCompoundWidget
-{
-public:
-
-	SLATE_BEGIN_ARGS(SCollapseControlsWidget)
-		: _Sequencer(nullptr), _OwnerTrack(nullptr)
-	{}
-
-	SLATE_ARGUMENT(TSharedPtr<ISequencer>, Sequencer)
-	SLATE_ARGUMENT(UMovieSceneTrack*, OwnerTrack)
-	SLATE_END_ARGS()
-
-	void Construct(const FArguments& InArgs);
-	virtual ~SCollapseControlsWidget() override {}
-
-	FReply OpenDialog(bool bModal = true);
-	void CloseDialog();
-
-
-private:
-	void Collapse();
-
-	TSharedPtr<ISequencer> Sequencer;
-	TWeakObjectPtr<UMovieSceneTrack> OwnerTrack;
-	//static to be reused
-	static TOptional<FBakingAnimationKeySettings> CollapseControlsSettings;
-	//structonscope for details panel
-	TSharedPtr < TStructOnScope<FBakingAnimationKeySettings>> Settings;
-	TWeakPtr<SWindow> DialogWindow;
-	TSharedPtr<IStructureDetailsView> DetailsView;
-};
-
-
-TOptional<FBakingAnimationKeySettings> SCollapseControlsWidget::CollapseControlsSettings;
-
-void SCollapseControlsWidget::Construct(const FArguments& InArgs)
-{
-	check(InArgs._Sequencer);
-	Sequencer = InArgs._Sequencer;
-	OwnerTrack = InArgs._OwnerTrack;
-
-	if (CollapseControlsSettings.IsSet() == false)
-	{
-		CollapseControlsSettings = FBakingAnimationKeySettings();
-		const FFrameRate TickResolution = Sequencer->GetFocusedTickResolution();
-		const FFrameTime FrameTime = Sequencer->GetLocalTime().ConvertTo(TickResolution);
-		FFrameNumber CurrentTime = FrameTime.GetFrame();
-
-		TRange<FFrameNumber> Range = Sequencer->GetFocusedMovieSceneSequence()->GetMovieScene()->GetPlaybackRange();
-		TArray<FFrameNumber> Keys;
-		TArray < FKeyHandle> KeyHandles;
-
-		CollapseControlsSettings.GetValue().StartFrame = Range.GetLowerBoundValue();
-		CollapseControlsSettings.GetValue().EndFrame = Range.GetUpperBoundValue();
-	}
-
-
-	Settings = MakeShared<TStructOnScope<FBakingAnimationKeySettings>>();
-	Settings->InitializeAs<FBakingAnimationKeySettings>(CollapseControlsSettings.GetValue());
-
-	FStructureDetailsViewArgs StructureViewArgs;
-	StructureViewArgs.bShowObjects = true;
-	StructureViewArgs.bShowAssets = true;
-	StructureViewArgs.bShowClasses = true;
-	StructureViewArgs.bShowInterfaces = true;
-
-	FDetailsViewArgs ViewArgs;
-	ViewArgs.bAllowSearch = false;
-	ViewArgs.bHideSelectionTip = false;
-	ViewArgs.bShowObjectLabel = false;
-
-	FPropertyEditorModule& PropertyEditor = FModuleManager::Get().LoadModuleChecked<FPropertyEditorModule>(TEXT("PropertyEditor"));
-
-	DetailsView = PropertyEditor.CreateStructureDetailView(ViewArgs, StructureViewArgs, TSharedPtr<FStructOnScope>());
-	TSharedPtr<INumericTypeInterface<double>> NumericTypeInterface = Sequencer->GetNumericTypeInterface();
-	DetailsView->GetDetailsView()->RegisterInstancedCustomPropertyTypeLayout("FrameNumber",
-		FOnGetPropertyTypeCustomizationInstance::CreateLambda([=]() {return MakeShared<FFrameNumberDetailsCustomization>(NumericTypeInterface); }));
-	DetailsView->SetStructureData(Settings);
-
-	ChildSlot
-		[
-			SNew(SBorder)
-			.Visibility(EVisibility::Visible)
-			[
-				SNew(SVerticalBox)
-				+ SVerticalBox::Slot()
-				.AutoHeight()
-				.Padding(0.f, 8.f, 0.f, 0.f)
-				[
-					DetailsView->GetWidget().ToSharedRef()
-				]
-				+ SVerticalBox::Slot()
-				.AutoHeight()
-				.Padding(FMargin(0.0f, 0.f))
-				[
-
-					SNew(SHorizontalBox)
-					+ SHorizontalBox::Slot()
-					.FillWidth(1.f)
-					.HAlign(HAlign_Fill)
-					[
-						SNew(SSpacer)
-					]
-
-				+ SHorizontalBox::Slot()
-					.AutoWidth()
-					.VAlign(VAlign_Center)
-					.HAlign(HAlign_Right)
-					.Padding(0.f)
-					[
-						SNew(SButton)
-						.HAlign(HAlign_Center)
-						.ContentPadding(FAppStyle::GetMargin("StandardDialog.ContentPadding"))
-						.Text(LOCTEXT("OK", "OK"))
-						.OnClicked_Lambda([this, InArgs]()
-							{
-								Collapse();
-								CloseDialog();
-								return FReply::Handled();
-
-							})
-						.IsEnabled_Lambda([this]()
-							{
-								return (Settings.IsValid());
-							})
-					]
-				]
-			]
-		];
-}
-
-
-void  SCollapseControlsWidget::Collapse()
-{
-	FBakingAnimationKeySettings* BakeSettings = Settings->Get();
-	FControlRigParameterTrackEditor::CollapseAllLayers(Sequencer, OwnerTrack.Get(), *BakeSettings);
-
-	CollapseControlsSettings = *BakeSettings;
-}
-
-class SCollapseControlsWidgetWindow : public SWindow
-{
-};
-
-FReply SCollapseControlsWidget::OpenDialog(bool bModal)
+bool FControlRigParameterTrackEditor::CollapseAllLayers(TSharedPtr<ISequencer>& SequencerPtr, UMovieSceneTrack* OwnerTrack, UMovieSceneControlRigParameterSection* ParameterSection, bool bKeyReduce, float Tolerance)
 {
-	check(!DialogWindow.IsValid());
-
-	const FVector2D CursorPos = FSlateApplication::Get().GetCursorPos();
-
-	TSharedRef<SCollapseControlsWidgetWindow> Window = SNew(SCollapseControlsWidgetWindow)
-		.Title(LOCTEXT("CollapseControls", "Collapse Controls"))
-		.CreateTitleBar(true)
-		.Type(EWindowType::Normal)
-		.SizingRule(ESizingRule::Autosized)
-		.ScreenPosition(CursorPos)
-		.FocusWhenFirstShown(true)
-		.ActivationPolicy(EWindowActivationPolicy::FirstShown)
-		[
-			AsShared()
-		];
-
-	Window->SetWidgetToFocusOnActivate(AsShared());
-
-	DialogWindow = Window;
-
-	Window->MoveWindowTo(CursorPos);
 
-	if (bModal)
-	{
-		GEditor->EditorAddModalWindow(Window);
-	}
-	else
-	{
-		FSlateApplication::Get().AddWindow(Window);
-	}
-
-	return FReply::Handled();
-}
-
-void SCollapseControlsWidget::CloseDialog()
-{
-	if (DialogWindow.IsValid())
-	{
-		DialogWindow.Pin()->RequestDestroyWindow();
-		DialogWindow.Reset();
-	}
-}
-///////////////
-//end of SCollapseControlsWidget
-/////////////////////
-struct FKeyAndValuesAtFrame
-{
-	FFrameNumber Frame;
-	TArray<FMovieSceneFloatValue> KeyValues;
-	float FinalValue;
-};
-
-bool CollapseAllLayersPerKey(TSharedPtr<ISequencer>& SequencerPtr, UMovieSceneTrack* OwnerTrack, const FBakingAnimationKeySettings& InSettings)
-{
-	if (SequencerPtr.IsValid() && OwnerTrack)
-	{
-		TArray<UMovieSceneSection*> Sections = OwnerTrack->GetAllSections();
-		return MovieSceneToolHelpers::CollapseSection(SequencerPtr, OwnerTrack, Sections, InSettings);
-	}
-	return false;
-}
-
-bool FControlRigParameterTrackEditor::CollapseAllLayers(TSharedPtr<ISequencer>&SequencerPtr, UMovieSceneTrack * OwnerTrack, const FBakingAnimationKeySettings &InSettings)
-{
-	return CollapseAllLayersPerKey(SequencerPtr, OwnerTrack, InSettings);
-	if (SequencerPtr.IsValid() && OwnerTrack)
+	if (SequencerPtr.IsValid() && OwnerTrack && ParameterSection && ParameterSection->GetControlRig())
 	{
 		TArray<UMovieSceneSection*> Sections = OwnerTrack->GetAllSections();
 		//make sure right type
-		if (Sections.Num() < 1)
+		if (ParameterSection->GetBlendType().Get() != EMovieSceneBlendType::Absolute && Sections.Num() > 0 && Sections[0] != ParameterSection)
 		{
-			UE_LOG(LogControlRigEditor, Log, TEXT("CollapseAllSections::No sections on track"));
+			UE_LOG(LogControlRigEditor, Log, TEXT("Section wrong type or not first when collapsing layers"));
 			return false;
 		}
-		if (UMovieSceneControlRigParameterSection* ParameterSection = Cast<UMovieSceneControlRigParameterSection>(Sections[0]))
+		FScopedTransaction Transaction(LOCTEXT("CollapseAllSections", "Collapse All Sections"));
+		ParameterSection->Modify();
+		UControlRig* ControlRig = ParameterSection->GetControlRig();
+		TRange<FFrameNumber> Range = SequencerPtr->GetFocusedMovieSceneSequence()->GetMovieScene()->GetPlaybackRange();
+		FFrameNumber StartFrame = Range.GetLowerBoundValue();
+		FFrameNumber EndFrame = Range.GetUpperBoundValue();
+		const FFrameRate& FrameRate = SequencerPtr->GetFocusedDisplayRate();
+		const FFrameRate& TickResolution = SequencerPtr->GetFocusedTickResolution();
+		FMovieSceneSequenceTransform RootToLocalTransform = SequencerPtr->GetFocusedMovieSceneSequenceTransform();
+
+		FFrameNumber FrameRateInFrameNumber = TickResolution.AsFrameNumber(FrameRate.AsInterval());
+		TArray<FFrameNumber> Frames;
+		for (FFrameNumber& Frame = StartFrame; Frame <= EndFrame; Frame += FrameRateInFrameNumber)
+		{
+			Frames.Add(Frame);
+		}
+		//Store transforms
+		TArray<TPair<FName, TArray<FTransform>>> ControlLocalTransforms;
+		TArray<FRigControlElement*> Controls;
+		ControlRig->GetControlsInOrder(Controls);
+
+		for (FRigControlElement* ControlElement : Controls)
 		{
-			if (ParameterSection->GetBlendType().Get() == EMovieSceneBlendType::Absolute)
+			if (!ControlRig->GetHierarchy()->IsAnimatable(ControlElement))
 			{
-				FScopedTransaction Transaction(LOCTEXT("CollapseAllSections", "Collapse All Sections"));
-				ParameterSection->Modify();
-				UControlRig* ControlRig = ParameterSection->GetControlRig();
-				FMovieSceneSequenceTransform RootToLocalTransform = SequencerPtr->GetFocusedMovieSceneSequenceTransform();
-
-				FFrameNumber StartFrame = InSettings.StartFrame;
-				FFrameNumber EndFrame = InSettings.EndFrame;
-				TRange<FFrameNumber> Range(StartFrame, EndFrame);
-				const FFrameRate& FrameRate = SequencerPtr->GetFocusedDisplayRate();
-				const FFrameRate& TickResolution = SequencerPtr->GetFocusedTickResolution();
-				
-				//frames and (optional) tangents
-				TArray<TPair<FFrameNumber, TArray<FMovieSceneTangentData>>> StoredTangents; //store tangents so we can reset them
-				TArray<FFrameNumber> Frames;
-				FFrameNumber FrameRateInFrameNumber = TickResolution.AsFrameNumber(FrameRate.AsInterval());
-				for (FFrameNumber& Frame = StartFrame; Frame <= EndFrame; Frame += FrameRateInFrameNumber)
-				{
-					Frames.Add(Frame);
-				}
-				
-				//Store transforms
-				TArray<TPair<FName, TArray<FTransform>>> ControlLocalTransforms;
-				TArray<FRigControlElement*> Controls;
-				ControlRig->GetControlsInOrder(Controls);
-
-				for (FRigControlElement* ControlElement : Controls)
-				{
-					if (!ControlRig->GetHierarchy()->IsAnimatable(ControlElement))
-					{
-						continue;
-					}
-					TPair<FName, TArray<FTransform>> NameTransforms;
-					NameTransforms.Key = ControlElement->GetName();
-					NameTransforms.Value.SetNum(Frames.Num());
-					ControlLocalTransforms.Add(NameTransforms);
-				}
+				continue;
+			}
+			TPair<FName, TArray<FTransform>> NameTransforms;
+			NameTransforms.Key = ControlElement->GetName();
+			NameTransforms.Value.SetNum(Frames.Num());
+			ControlLocalTransforms.Add(NameTransforms);
+		}
 
-				//get all of the local 
-				int32 Index = 0;
-				for (Index = 0; Index < Frames.Num(); ++Index)
-				{
-					const FFrameNumber& FrameNumber = Frames[Index];
-					FFrameTime GlobalTime(FrameNumber);
-					GlobalTime = GlobalTime * RootToLocalTransform.InverseLinearOnly();
+		//get all of the local 
+		int32 Index = 0;
+		for (Index = 0; Index < Frames.Num(); ++Index)
+		{
+			const FFrameNumber& FrameNumber = Frames[Index];
+			FFrameTime GlobalTime(FrameNumber);
+			GlobalTime = GlobalTime * RootToLocalTransform.InverseLinearOnly();
 
-					FMovieSceneContext Context = FMovieSceneContext(FMovieSceneEvaluationRange(GlobalTime, TickResolution), SequencerPtr->GetPlaybackStatus()).SetHasJumped(true);
+			FMovieSceneContext Context = FMovieSceneContext(FMovieSceneEvaluationRange(GlobalTime, TickResolution), SequencerPtr->GetPlaybackStatus()).SetHasJumped(true);
 
-					SequencerPtr->GetEvaluationTemplate().EvaluateSynchronousBlocking(Context, *SequencerPtr);
-					ControlRig->Evaluate_AnyThread();
-					for (TPair<FName, TArray<FTransform>>& TrailControlTransform : ControlLocalTransforms)
-					{
-						TrailControlTransform.Value[Index] = ControlRig->GetControlLocalTransform(TrailControlTransform.Key);
-					}
-				}
-				//delete other sections
-				OwnerTrack->Modify();
-				for (Index = Sections.Num() - 1; Index >= 0; --Index)
-				{
-					if (Sections[Index] != ParameterSection)
-					{
-						OwnerTrack->RemoveSectionAt(Index);
-					}
-				}
+			SequencerPtr->GetEvaluationTemplate().EvaluateSynchronousBlocking(Context, *SequencerPtr);
+			ControlRig->Evaluate_AnyThread();
+			for (TPair<FName, TArray<FTransform>>& TrailControlTransform : ControlLocalTransforms)
+			{
+				TrailControlTransform.Value[Index] = ControlRig->GetControlLocalTransform(TrailControlTransform.Key);
+			}
+		}
+		//delete other sections
+		OwnerTrack->Modify();
+		for (Index = Sections.Num() - 1; Index >= 0; --Index)
+		{
+			if (Sections[Index] != ParameterSection)
+			{
+				OwnerTrack->RemoveSectionAt(Index);
+			}
+		}
 
-				//remove all keys, except Space Channels, from the Section.
-				ParameterSection->RemoveAllKeys(false /*bIncludedSpaceKeys*/);
+		//remove all keys, except Space Channels, from the Section.
+		ParameterSection->RemoveAllKeys(false /*bIncludedSpaceKeys*/);
 
-				FRigControlModifiedContext Context;
-				Context.SetKey = EControlRigSetKey::Always;
+		FRigControlModifiedContext Context;
+		Context.SetKey = EControlRigSetKey::Always;
 
-				FScopedSlowTask Feedback(Frames.Num(), LOCTEXT("CollapsingSections", "Collapsing Sections"));
-				Feedback.MakeDialog(true);
 
-				const ERichCurveInterpMode InterpMode = InSettings.bReduceKeys ? RCIM_Cubic : RCIM_Linear;
+		FScopedSlowTask Feedback(Frames.Num(), LOCTEXT("CollapsingSections", "Collapsing Sections"));
+		Feedback.MakeDialog(true);
 
-				Index = 0;
-				for (Index = 0; Index < Frames.Num(); ++Index)
-				{
-					Feedback.EnterProgressFrame(1, LOCTEXT("CollapsingSections", "Collapsing Sections"));
-					const FFrameNumber& FrameNumber = Frames[Index];
-					Context.LocalTime = TickResolution.AsSeconds(FFrameTime(FrameNumber));
-					//need to do the twice hack since controls aren't really in order
-					for (int32 TwiceHack = 0; TwiceHack < 2; ++TwiceHack)
-					{
-						for (TPair<FName, TArray<FTransform>>& TrailControlTransform : ControlLocalTransforms)
-						{
-							ControlRig->SetControlLocalTransform(TrailControlTransform.Key, TrailControlTransform.Value[Index], false, Context, false);
-						}
-					}
-					ControlRig->Evaluate_AnyThread();
-					ParameterSection->RecordControlRigKey(FrameNumber, true, InterpMode);
+		const ERichCurveInterpMode InterpMode = bKeyReduce ? RCIM_Cubic : RCIM_Linear;
 
-					if (Feedback.ShouldCancel())
-					{
-						Transaction.Cancel();
-						SequencerPtr->NotifyMovieSceneDataChanged(EMovieSceneDataChangeType::MovieSceneStructureItemAdded);
-						return false;
-					}
-				}
-				if (InSettings.bReduceKeys)
+		Index = 0;
+		for (Index = 0; Index < Frames.Num(); ++Index)
+		{
+			Feedback.EnterProgressFrame(1, LOCTEXT("CollapsingSections", "Collapsing Sections"));
+			const FFrameNumber& FrameNumber = Frames[Index];
+			Context.LocalTime = TickResolution.AsSeconds(FFrameTime(FrameNumber));
+			//need to do the twice hack since controls aren't really in order
+			for (int32 TwiceHack = 0; TwiceHack < 2; ++TwiceHack)
+			{
+				for (TPair<FName, TArray<FTransform>>& TrailControlTransform : ControlLocalTransforms)
 				{
-					FKeyDataOptimizationParams Params;
-					Params.bAutoSetInterpolation = true;
-					Params.Tolerance = InSettings.Tolerance;
-					FMovieSceneChannelProxy& ChannelProxy = ParameterSection->GetChannelProxy();
-					TArrayView<FMovieSceneFloatChannel*> FloatChannels = ChannelProxy.GetChannels<FMovieSceneFloatChannel>();
-
-					for (FMovieSceneFloatChannel* Channel : FloatChannels)
-					{
-						Channel->Optimize(Params); //should also auto tangent
-					}
+					ControlRig->SetControlLocalTransform(TrailControlTransform.Key, TrailControlTransform.Value[Index],false, Context, false);
 				}
-				//reset everything back
-				SequencerPtr->NotifyMovieSceneDataChanged(EMovieSceneDataChangeType::MovieSceneStructureItemAdded);
-				return true;
 			}
-			else
+			ControlRig->Evaluate_AnyThread();
+			ParameterSection->RecordControlRigKey(FrameNumber, true, InterpMode);
+
+			if (Feedback.ShouldCancel())
 			{
-				UE_LOG(LogControlRigEditor, Log, TEXT("CollapseAllSections:: First section is not additive"));
+				Transaction.Cancel();
+				SequencerPtr->NotifyMovieSceneDataChanged(EMovieSceneDataChangeType::MovieSceneStructureItemAdded);
 				return false;
 			}
 		}
-		else
+		if (bKeyReduce)
 		{
-			UE_LOG(LogControlRigEditor, Log, TEXT("CollapseAllSections:: No Control Rig section"));
-			return false;
+			FKeyDataOptimizationParams Params;
+			Params.bAutoSetInterpolation = true;
+			Params.Tolerance = Tolerance;
+			FMovieSceneChannelProxy& ChannelProxy = ParameterSection->GetChannelProxy();
+			TArrayView<FMovieSceneFloatChannel*> FloatChannels = ChannelProxy.GetChannels<FMovieSceneFloatChannel>();
+
+			for (FMovieSceneFloatChannel* Channel : FloatChannels)
+			{
+				Channel->Optimize(Params); //should also auto tangent
+			}
 		}
+		//reset everything back
+		SequencerPtr->NotifyMovieSceneDataChanged(EMovieSceneDataChangeType::MovieSceneStructureItemAdded);
+		return true;
 	}
-	UE_LOG(LogControlRigEditor, Log, TEXT("CollapseAllSections:: Sequencer or track is invalid"));
 	return false;
 }
-
 void FControlRigParameterSection::CollapseAllLayers()
 {
 	TSharedPtr<ISequencer> SequencerPtr = WeakSequencer.Pin();
-	if (UMovieSceneControlRigParameterSection* ParameterSection = CastChecked<UMovieSceneControlRigParameterSection>(WeakSection.Get()))
-	{
-		UMovieSceneTrack* OwnerTrack = ParameterSection->GetTypedOuter<UMovieSceneTrack>();
-		TSharedPtr<ISequencer> Sequencer = WeakSequencer.Pin();
-		TSharedRef<SCollapseControlsWidget> BakeWidget =
-			SNew(SCollapseControlsWidget)
-			.Sequencer(Sequencer)
-			.OwnerTrack(OwnerTrack);
+	UMovieSceneControlRigParameterSection* ParameterSection = CastChecked<UMovieSceneControlRigParameterSection>(WeakSection.Get());
+	UMovieSceneTrack* OwnerTrack = ParameterSection->GetTypedOuter<UMovieSceneTrack>();
+	FControlRigParameterTrackEditor::CollapseAllLayers(SequencerPtr,OwnerTrack, ParameterSection);
 
-		BakeWidget->OpenDialog(true);
-	}
 }
-
-
 void FControlRigParameterSection::KeyZeroValue()
 {
 	UMovieSceneControlRigParameterSection* ParameterSection = CastChecked<UMovieSceneControlRigParameterSection>(WeakSection.Get());
diff --git a/Engine/Plugins/Animation/ControlRig/Source/ControlRigEditor/Private/Sequencer/ControlRigParameterTrackEditor.h b/Engine/Plugins/Animation/ControlRig/Source/ControlRigEditor/Private/Sequencer/ControlRigParameterTrackEditor.h
--- a/Engine/Plugins/Animation/ControlRig/Source/ControlRigEditor/Private/Sequencer/ControlRigParameterTrackEditor.h
+++ b/Engine/Plugins/Animation/ControlRig/Source/ControlRigEditor/Private/Sequencer/ControlRigParameterTrackEditor.h
@@ -33,7 +33,7 @@ struct FMovieSceneControlRigSpaceChannel;
 struct FMovieSceneChannel;
 struct FKeyAddOrDeleteEventItem;
 struct FKeyMoveEventItem;
-struct FBakingAnimationKeySettings;
+
 /**
  * Tools for animation tracks
  */
@@ -229,7 +229,7 @@ public:
 	bool ModifyOurGeneratedKeysByCurrentAndWeight(UObject* Object, UControlRig* InControlRig, FName RigControlName, UMovieSceneTrack *Track, UMovieSceneSection* SectionToKey, FFrameNumber EvaluateTime, FGeneratedTrackKeys& InOutGeneratedTotalKeys, float Weight) const;
 
 	//**Function to collapse all layers from this section onto the first absoluate layer.*/
-	static bool CollapseAllLayers(TSharedPtr<ISequencer>& SequencerPtr, UMovieSceneTrack* OwnerTrack, const FBakingAnimationKeySettings& InSettings);
+	static bool CollapseAllLayers(TSharedPtr<ISequencer>& SequencerPtr, UMovieSceneTrack* OwnerTrack, UMovieSceneControlRigParameterSection* ParameterSection, bool bKeyReduce = false, float Tolerance = 0.001f);
 
 private:
 	FDelegateHandle SelectionChangedHandle;
diff --git a/Engine/Plugins/Animation/ControlRig/Source/ControlRigEditor/Private/Sequencer/ControlRigSpaceChannelEditors.cpp b/Engine/Plugins/Animation/ControlRig/Source/ControlRigEditor/Private/Sequencer/ControlRigSpaceChannelEditors.cpp
--- a/Engine/Plugins/Animation/ControlRig/Source/ControlRigEditor/Private/Sequencer/ControlRigSpaceChannelEditors.cpp
+++ b/Engine/Plugins/Animation/ControlRig/Source/ControlRigEditor/Private/Sequencer/ControlRigSpaceChannelEditors.cpp
@@ -61,7 +61,6 @@
 #include "CurveEditorSettings.h"
 #include "TimeToPixel.h"
 #include "Templates/Tuple.h"
-#include "Sequencer/ControlRigSequencerHelpers.h"
 
 #define LOCTEXT_NAMESPACE "ControlRigEditMode"
 
@@ -861,161 +860,17 @@ void FControlRigSpaceChannelHelpers::GetFramesInThisSpaceAfterThisTime(
 }
 
 void FControlRigSpaceChannelHelpers::SequencerBakeControlInSpace(UControlRig* ControlRig, ISequencer* Sequencer, FMovieSceneControlRigSpaceChannel* Channel, UMovieSceneSection* SectionToKey,
-	URigHierarchy* RigHierarchy, const FRigElementKey& ControlKey, FRigSpacePickerBakeSettings Settings)
+	TArray<FFrameNumber> Frames, URigHierarchy* RigHierarchy, const FRigElementKey& ControlKey, FRigSpacePickerBakeSettings Settings)
 {
 	if (bDoNotCompensate == true)
 	{
 		return;
 	}
-	if (ControlRig && Sequencer && Channel && SectionToKey && (Settings.Settings.StartFrame != Settings.Settings.EndFrame) && RigHierarchy
+	if (ControlRig && Sequencer && Channel && SectionToKey && Frames.Num() > 0 && RigHierarchy
 		&& Sequencer->GetFocusedMovieSceneSequence() && Sequencer->GetFocusedMovieSceneSequence()->GetMovieScene())
 	{
-	
 		if (UMovieSceneControlRigParameterSection* Section = Cast<UMovieSceneControlRigParameterSection>(SectionToKey))
 		{
-			//get the setting as a rig element key to use for comparisons for smart baking, 
-			FMovieSceneControlRigSpaceBaseKey Value;
-			if (Settings.TargetSpace == RigHierarchy->GetWorldSpaceReferenceKey())
-			{
-				Value.SpaceType = EMovieSceneControlRigSpaceType::World;
-			}
-			else if (Settings.TargetSpace == RigHierarchy->GetDefaultParentKey())
-			{
-				Value.SpaceType = EMovieSceneControlRigSpaceType::Parent;
-			}
-			else
-			{
-				FRigElementKey DefaultParent = RigHierarchy->GetFirstParent(ControlKey);
-				if (DefaultParent == Settings.TargetSpace)
-				{
-					Value.SpaceType = EMovieSceneControlRigSpaceType::Parent;
-				}
-				else
-				{
-					Value.SpaceType = EMovieSceneControlRigSpaceType::ControlRig;
-					Value.ControlRigElement = Settings.TargetSpace;
-				}
-			}
-			TArrayView<FMovieSceneFloatChannel*> Channels = FControlRigSequencerHelpers::GetFloatChannels(ControlRig,
-				ControlKey.Name, Section);
-			//if smart baking we 1) Don't do any baking over times where the Settings.TargetSpace is active, 
-			//and 2) we just bake over key frames
-			//need to store these times first since we will still delete all of these space keys
-
-			TArray<FFrameNumber> Frames;
-			const FFrameRate TickResolution = Sequencer->GetFocusedTickResolution();
-			const FFrameRate  FrameRate = Sequencer->GetFocusedDisplayRate();
-			//get space keys so we can delete them and use them for smart baking at them
-			TArray<FFrameNumber> Keys;
-			TArray < FKeyHandle> KeyHandles;
-			FFrameNumber StartFrame = Settings.Settings.StartFrame;
-			FFrameNumber EndFrame = Settings.Settings.EndFrame;
-			TRange<FFrameNumber> Range(StartFrame, EndFrame);
-			Channel->GetKeys(Range, &Keys, &KeyHandles);
-			TArray<FFrameNumber> MinusOneFrames; //frames we need to delete
-			TArray<TPair<FFrameNumber, TArray<FMovieSceneTangentData>>> StoredTangents; //need to set tangents at the space switch locations
-			TSet<FFrameNumber> ParentFrames; // these are parent frames we keep track of so we can 
-			if (Settings.Settings.BakingKeySettings == EBakingKeySettings::KeysOnly)
-			{
-				TSet<FFrameNumber> MinusOneKeys; 
-				TSortedMap<FFrameNumber,FFrameNumber> FrameMap;
-				//add space keys to bake
-				{
-					for (FFrameNumber& Frame : Keys)
-					{
-						FMovieSceneControlRigSpaceBaseKey SpaceValue;
-						using namespace UE::MovieScene;
-						EvaluateChannel(Channel, Frame, SpaceValue);
-						if (SpaceValue != Value)
-						{
-							FChannelMapInfo* pChannelIndex = nullptr;
-							FRigControlElement* ControlElement = nullptr;
-							int32 NumChannels = 0;
-							int32 ChannelIndex = 0;
-							Tie(ControlElement, pChannelIndex, NumChannels, ChannelIndex) = GetControlAndChannelInfos(ControlRig, Section, ControlKey.Name);
-
-							if (pChannelIndex && ControlElement && NumChannels > 0)
-							{
-								TPair<FFrameNumber, TArray<FMovieSceneTangentData>> FrameAndTangent;
-								FrameAndTangent.Key = Frame;
-								EvaluateTangentAtThisTime<FMovieSceneFloatChannel>(ChannelIndex, NumChannels, Section, Frame, FrameAndTangent.Value);
-								StoredTangents.Add(FrameAndTangent);
-							}
-							FrameMap.Add(Frame,Frame);
-						}
-						MinusOneKeys.Add(Frame - 1);
-					}
-				}
-
-				TArray<FFrameNumber> TransformFrameTimes = FMovieSceneConstraintChannelHelper::GetTransformTimes(
-					Channels, Settings.Settings.StartFrame, Settings.Settings.EndFrame);
-				//add transforms keys to bake
-				{
-					for (FFrameNumber& Frame : TransformFrameTimes)
-					{
-						if (MinusOneKeys.Contains(Frame) == false)
-						{
-							FrameMap.Add(Frame);
-						}
-					}
-				}
-				//also need to add transform keys for the parent control rig intervals over the time we bake
-				TArray<FSpaceRange> SpaceRanges = Channel->FindSpaceIntervals();
-				for (const FSpaceRange& SpaceRange : SpaceRanges)
-				{
-					if (SpaceRange.Key.SpaceType == EMovieSceneControlRigSpaceType::ControlRig && SpaceRange.Key.ControlRigElement.IsValid())
-					{
-						TRange<FFrameNumber> Overlap = TRange<FFrameNumber>::Intersection(Range, SpaceRange.Range);
-						if (Overlap.IsEmpty() == false)
-						{
-							TArrayView<FMovieSceneFloatChannel*> ParentChannels = FControlRigSequencerHelpers::GetFloatChannels(ControlRig,
-								SpaceRange.Key.ControlRigElement.Name, Section);
-
-							TArray<FFrameNumber> ParentFrameTimes = FMovieSceneConstraintChannelHelper::GetTransformTimes(
-								ParentChannels, StartFrame, EndFrame);
-							for (FFrameNumber& Frame : ParentFrameTimes)
-							{
-								//if we don't have a key at the parent time, add the frame AND store it's tangent.
-								if (FrameMap.Contains(Frame) == false)
-								{
-									FrameMap.Add(Frame);
-									ParentFrames.Add(Frame);
-
-									FChannelMapInfo* pChannelIndex = nullptr;
-									FRigControlElement* ControlElement = nullptr;
-									int32 NumChannels = 0;
-									int32 ChannelIndex = 0;
-									Tie(ControlElement, pChannelIndex, NumChannels, ChannelIndex) = GetControlAndChannelInfos(ControlRig, Section, SpaceRange.Key.ControlRigElement.Name);
-									TPair<FFrameNumber, TArray<FMovieSceneTangentData>> FrameAndTangent;
-									FrameAndTangent.Key = Frame;
-									EvaluateTangentAtThisTime<FMovieSceneFloatChannel>(ChannelIndex, NumChannels, Section, Frame, FrameAndTangent.Value);
-									StoredTangents.Add(FrameAndTangent);
-								}
-							}
-						}
-					}
-				
-				}
-				FrameMap.GenerateKeyArray(Frames);
-
-			}
-			else
-			{
-				FFrameNumber FrameRateInFrameNumber = TickResolution.AsFrameNumber(FrameRate.AsInterval());
-				if (Settings.Settings.FrameIncrement > 1) //increment per frame by increment
-				{
-					FrameRateInFrameNumber.Value *= Settings.Settings.FrameIncrement;
-				}
-				for (FFrameNumber& Frame = Settings.Settings.StartFrame; Frame <= Settings.Settings.EndFrame; Frame += FrameRateInFrameNumber)
-				{
-					Frames.Add(Frame);
-				}
-
-			}
-			if (Frames.Num() == 0)
-			{
-				return;
-			}
 			TGuardValue<bool> CompensateGuard(bDoNotCompensate, true);
 			TArray<FTransform> ControlRigParentWorldTransforms;
 			ControlRigParentWorldTransforms.SetNum(Frames.Num());
@@ -1028,23 +883,21 @@ void FControlRigSpaceChannelHelpers::SequencerBakeControlInSpace(UControlRig* Co
 			TArray<FTransform> ControlWorldTransforms;
 			Snapper.GetControlRigControlTransforms(Sequencer, ControlRig, ControlKey.Name, Frames, ControlRigParentWorldTransforms, ControlWorldTransforms);
 
-			//Need to delete keys, if smart baking we only delete the -1 keys,
-			//if not, we need to delete all transform
-			//and always delete space keys
+			//Find all space keys in range and delete them since it will get replaced with new space when we components.
 			Section->Modify();
+			FFrameNumber StartFrame = Frames[0];
+			FFrameNumber EndFrame = Frames[Frames.Num() - 1];
+			TArray<FFrameNumber> Keys;
+			TArray < FKeyHandle> KeyHandles;
+			TRange<FFrameNumber> Range(StartFrame, EndFrame);
+			Channel->GetKeys(Range, &Keys, &KeyHandles);
 			Channel->DeleteKeys(KeyHandles);
-			if (Settings.Settings.BakingKeySettings == EBakingKeySettings::KeysOnly)
-			{
-				for (const FFrameNumber& DeleteFrame : Keys)
-				{
-					FControlRigSpaceChannelHelpers::DeleteTransformKeysAtThisTime(ControlRig, Section, ControlKey.Name, DeleteFrame - 1);
-				}
-			}
-			else
+			//also delete any transform keys at the deleted key times and times -1.
+			for (const FFrameNumber& DeleteFrame : Keys)
 			{
-				FMovieSceneConstraintChannelHelper::DeleteTransformTimes(Channels, StartFrame, EndFrame);
+				FControlRigSpaceChannelHelpers::DeleteTransformKeysAtThisTime(ControlRig, Section, ControlKey.Name, DeleteFrame - 1);
+				FControlRigSpaceChannelHelpers::DeleteTransformKeysAtThisTime(ControlRig, Section, ControlKey.Name, DeleteFrame);
 			}
-			
 
 			URigHierarchy* Hierarchy = ControlRig->GetHierarchy();
 
@@ -1053,7 +906,29 @@ void FControlRigSpaceChannelHelpers::SequencerBakeControlInSpace(UControlRig* Co
 			using namespace UE::MovieScene;
 			EvaluateChannel(Channel, StartFrame, StartFrameValue);
 			EvaluateChannel(Channel, EndFrame, EndFrameValue);
-			
+
+			FMovieSceneControlRigSpaceBaseKey Value;
+			if (Settings.TargetSpace == RigHierarchy->GetWorldSpaceReferenceKey())
+			{
+				Value.SpaceType = EMovieSceneControlRigSpaceType::World;
+			}
+			else if (Settings.TargetSpace == RigHierarchy->GetDefaultParentKey())
+			{
+				Value.SpaceType = EMovieSceneControlRigSpaceType::Parent;
+			}
+			else
+			{
+				FRigElementKey DefaultParent = RigHierarchy->GetFirstParent(ControlKey);
+				if (DefaultParent == Settings.TargetSpace)
+				{
+					Value.SpaceType = EMovieSceneControlRigSpaceType::Parent;
+				}
+				else
+				{
+					Value.SpaceType = EMovieSceneControlRigSpaceType::ControlRig;
+					Value.ControlRigElement = Settings.TargetSpace;
+				}
+			}
 			const bool bCompensateStart = StartFrameValue != Value;
 			TRange<FFrameNumber> PlaybackRange = Sequencer->GetFocusedMovieSceneSequence()->GetMovieScene()->GetPlaybackRange();
 			const bool bCompensateEnd = (EndFrameValue != Value && PlaybackRange.GetUpperBoundValue() != EndFrame);
@@ -1073,7 +948,9 @@ void FControlRigSpaceChannelHelpers::SequencerBakeControlInSpace(UControlRig* Co
 			//now set all of the key values
 			FRigControlModifiedContext Context;
 			Context.SetKey = EControlRigSetKey::Always;
-			Context.KeyMask = (uint32)EControlRigContextChannelToKey::AllTransform;
+
+			FFrameRate TickResolution = Sequencer->GetFocusedTickResolution();
+
 			URigHierarchy::TElementDependencyMap Dependencies = RigHierarchy->GetDependenciesForVM(ControlRig->GetVM());
 			RigHierarchy->SwitchToParent(ControlKey, Settings.TargetSpace, false, true, Dependencies, nullptr);
 			ControlRig->Evaluate_AnyThread();
@@ -1095,16 +972,6 @@ void FControlRigSpaceChannelHelpers::SequencerBakeControlInSpace(UControlRig* Co
 				//evaluate control rig
 				ControlRig->Evaluate_AnyThread();
 				Context.LocalTime = TickResolution.AsSeconds(FFrameTime(Frame));
-				
-				//if doing smart baking only set keys on those that exist IF not a key from the parent, in that case we need to key the  whole transform.
-				if (Settings.Settings.BakingKeySettings == EBakingKeySettings::KeysOnly && ParentFrames.Contains(Frame) == false)
-				{
-					Context.KeyMask = (uint32)GetCurrentTransformKeysAtThisTime(ControlRig, ControlKey.Name, SectionToKey, Frame);
-				}
-				else
-				{
-					Context.KeyMask = (uint32)EControlRigContextChannelToKey::AllTransform;
-				}
 				ControlRig->SetControlGlobalTransform(ControlKey.Name, GlobalTransform, true, Context);
 			}
 
@@ -1139,38 +1006,66 @@ void FControlRigSpaceChannelHelpers::SequencerBakeControlInSpace(UControlRig* Co
 				ChannelInterface.AddKey(EndFrame, Forward<FMovieSceneControlRigSpaceBaseKey>(EndFrameValue));
 				const FTransform GlobalTransform = ControlWorldTransforms[Frames.Num() - 1];
 				Context.LocalTime = TickResolution.AsSeconds(FFrameTime(EndFrame));
-				Context.KeyMask = (uint32)EControlRigContextChannelToKey::AllTransform;
 				ControlRig->SetControlGlobalTransform(ControlKey.Name, GlobalTransform, true, Context);
 			}
-			//Fix tangents at removed space switches if needed
-			for (TPair<FFrameNumber, TArray<FMovieSceneTangentData>>& StoredTangent : StoredTangents)
-			{
-				FChannelMapInfo* pChannelIndex = nullptr;
-				FRigControlElement* ControlElement = nullptr;
-				int32 NumChannels = 0;
-				int32 ChannelIndex = 0;
-				Tie(ControlElement, pChannelIndex, NumChannels, ChannelIndex) = GetControlAndChannelInfos(ControlRig, Section, ControlKey.Name);
-
-				if (pChannelIndex && ControlElement && NumChannels > 0)
-				{
-					SetTangentsAtThisTime<FMovieSceneFloatChannel>(ChannelIndex, NumChannels, Section, StoredTangent.Key, StoredTangent.Value);
-				}
-			}
+		
 			// Fix any Rotation Channels
 			Section->FixRotationWinding(ControlKey.Name, Frames[0], Frames[Frames.Num() - 1]);
-			// Then optimize
-			if (Settings.Settings.BakingKeySettings == EBakingKeySettings::AllFrames && Settings.Settings.bReduceKeys)
+			// Then reduce
+			if (Settings.bReduceKeys)
 			{
-				FKeyDataOptimizationParams Params;
-				Params.bAutoSetInterpolation = true;
-				Params.Tolerance = Settings.Settings.Tolerance;
-				Params.Range = TRange <FFrameNumber>(Frames[0], Frames[Frames.Num() - 1]);
-				Section->OptimizeSection(ControlKey.Name, Params);
+				TArrayView<FMovieSceneFloatChannel*> FloatChannels = Section->GetChannelProxy().GetChannels<FMovieSceneFloatChannel>();
+				FChannelMapInfo* pChannelIndex = Section->ControlChannelMap.Find(ControlKey.Name);
+				if (pChannelIndex != nullptr)
+				{
+					int32 ChannelIndex = pChannelIndex->ChannelIndex;
+
+					if (FRigControlElement* ControlElement = ControlRig->FindControl(ControlKey.Name))
+					{
+						FKeyDataOptimizationParams Params;
+						Params.bAutoSetInterpolation = true;
+						Params.Tolerance = Settings.Tolerance;
+						Params.Range = TRange <FFrameNumber>(Frames[0], Frames[Frames.Num() - 1]);
+
+						switch (ControlElement->Settings.ControlType)
+						{
+						case ERigControlType::Position:
+						case ERigControlType::Scale:
+						case ERigControlType::Rotator:
+						{
+							FloatChannels[ChannelIndex]->Optimize(Params);
+							FloatChannels[ChannelIndex + 1]->Optimize(Params);
+							FloatChannels[ChannelIndex + 2]->Optimize(Params);
+							break;
+						}
+
+						case ERigControlType::Transform:
+						case ERigControlType::TransformNoScale:
+						case ERigControlType::EulerTransform:
+						{
+
+							FloatChannels[ChannelIndex]->Optimize(Params);
+							FloatChannels[ChannelIndex + 1]->Optimize(Params);
+							FloatChannels[ChannelIndex + 2]->Optimize(Params);
+							FloatChannels[ChannelIndex + 3]->Optimize(Params);
+							FloatChannels[ChannelIndex + 4]->Optimize(Params);
+							FloatChannels[ChannelIndex + 5]->Optimize(Params);
+
+							if (ControlElement->Settings.ControlType == ERigControlType::Transform ||
+								ControlElement->Settings.ControlType == ERigControlType::EulerTransform)
+							{
+								FloatChannels[ChannelIndex + 6]->Optimize(Params);
+								FloatChannels[ChannelIndex + 7]->Optimize(Params);
+								FloatChannels[ChannelIndex + 8]->Optimize(Params);
+							}
+							break;
+
+						}
+						}
+					}
+				}
 			}
-			else  //need to auto set tangents, the above section will do that also
-			{
-				Section->AutoSetTangents(ControlKey.Name);
-			}	
+			
 		}
 		Sequencer->NotifyMovieSceneDataChanged(EMovieSceneDataChangeType::MovieSceneStructureItemAdded); //may have added channel
 	}
@@ -1438,16 +1333,16 @@ FReply FControlRigSpaceChannelHelpers::OpenBakeDialog(ISequencer* Sequencer, FMo
 			Settings.TargetSpace = URigHierarchy::GetDefaultParentKey();
 
 			TRange<FFrameNumber> Range = Sequencer->GetFocusedMovieSceneSequence()->GetMovieScene()->GetPlaybackRange();
-			Settings.Settings.StartFrame = Range.GetLowerBoundValue();
-			Settings.Settings.EndFrame = Range.GetUpperBoundValue();
+			Settings.StartFrame = Range.GetLowerBoundValue();
+			Settings.EndFrame = Range.GetUpperBoundValue();
 			TMovieSceneChannelData<FMovieSceneControlRigSpaceBaseKey> ChannelData = Channel->GetData();
 			TArrayView<const FFrameNumber> Times = ChannelData.GetTimes();
 			if (KeyIndex >= 0 && KeyIndex < Times.Num())
 			{
-				Settings.Settings.StartFrame = Times[KeyIndex];
+				Settings.StartFrame = Times[KeyIndex];
 				if (KeyIndex + 1 < Times.Num())
 				{
-					Settings.Settings.EndFrame = Times[KeyIndex + 1];
+					Settings.EndFrame = Times[KeyIndex + 1];
 				}
 			}
 			TArray<FRigElementKey> Controls;
@@ -1467,7 +1362,14 @@ FReply FControlRigSpaceChannelHelpers::OpenBakeDialog(ISequencer* Sequencer, FMo
 				})
 				.OnBake_Lambda([Sequencer, ControlRig, TickResolution](URigHierarchy* InHierarchy, TArray<FRigElementKey> InControls, FRigSpacePickerBakeSettings InSettings)
 				{
-					
+					TArray<FFrameNumber> Frames;
+
+					const FFrameRate& FrameRate = Sequencer->GetFocusedDisplayRate();
+					FFrameNumber FrameRateInFrameNumber = TickResolution.AsFrameNumber(FrameRate.AsInterval());
+					for (FFrameNumber& Frame = InSettings.StartFrame; Frame <= InSettings.EndFrame; Frame += FrameRateInFrameNumber)
+					{
+						Frames.Add(Frame);
+					}
 					FScopedTransaction Transaction(LOCTEXT("BakeControlToSpace", "Bake Control In Space"));
 					for (const FRigElementKey& ControlKey : InControls)
 					{
@@ -1475,7 +1377,7 @@ FReply FControlRigSpaceChannelHelpers::OpenBakeDialog(ISequencer* Sequencer, FMo
 						if (SpaceChannelAndSection.SpaceChannel)
 						{
 							FControlRigSpaceChannelHelpers::SequencerBakeControlInSpace(ControlRig, Sequencer, SpaceChannelAndSection.SpaceChannel, SpaceChannelAndSection.SectionToKey,
-								InHierarchy, ControlKey, InSettings);
+								Frames, InHierarchy, ControlKey, InSettings);
 						}
 					}
 					return FReply::Handled();
diff --git a/Engine/Plugins/Animation/ControlRig/Source/ControlRigEditor/Private/Tools/ConstraintBaker.cpp b/Engine/Plugins/Animation/ControlRig/Source/ControlRigEditor/Private/Tools/ConstraintBaker.cpp
--- a/Engine/Plugins/Animation/ControlRig/Source/ControlRigEditor/Private/Tools/ConstraintBaker.cpp
+++ b/Engine/Plugins/Animation/ControlRig/Source/ControlRigEditor/Private/Tools/ConstraintBaker.cpp
@@ -18,30 +18,27 @@
 
 #define LOCTEXT_NAMESPACE "ConstraintBaker"
 
-
 // at this stage, we suppose that any of these arguments are safe to use. Make sure to test them before using that function
 void FConstraintBaker::GetMinimalFramesToBake(
 	UWorld* InWorld,
 	const UTickableTransformConstraint* InConstraint,
 	const TSharedPtr<ISequencer>& InSequencer,
-	UMovieSceneSection* InSection,
-	const TOptional< FBakingAnimationKeySettings>& InSettings,
+	IMovieSceneConstrainedSection* InSection,
 	TArray<FFrameNumber>& OutFramesToBake)
 {
 	const UMovieSceneSequence* MovieSceneSequence = InSequencer->GetFocusedMovieSceneSequence();
 	const UMovieScene* MovieScene = MovieSceneSequence ? MovieSceneSequence->GetMovieScene() : nullptr;
-	if (!MovieScene || !InConstraint || !InConstraint->ChildTRSHandle || !InSection)
+	if (!MovieScene)
 	{
 		return;
 	}
-	IMovieSceneConstrainedSection* ConstrainedSection = Cast<IMovieSceneConstrainedSection>(InSection);
 
 	// get constraint channel data if any
 	TArrayView<const FFrameNumber> ConstraintFrames;
 	TArrayView<const bool> ConstraintValues;
 
 	// note that we might want to bake a constraint which is not animated
-	FConstraintAndActiveChannel* ThisActiveChannel = ConstrainedSection->GetConstraintChannel(InConstraint->GetFName());
+	FConstraintAndActiveChannel* ThisActiveChannel = InSection->GetConstraintChannel(InConstraint->GetFName());
 	if (ThisActiveChannel)
 	{
 		const TMovieSceneChannelData<const bool> ConstraintChannelData = ThisActiveChannel->ActiveChannel.GetData();
@@ -57,86 +54,56 @@ void FConstraintBaker::GetMinimalFramesToBake(
 	}
 
 	// default init bounds to the scene bounds
-	FFrameNumber FirstBakingFrame = InSettings.IsSet() ? InSettings.GetValue().StartFrame : MovieScene->GetPlaybackRange().GetLowerBoundValue();
-	FFrameNumber LastBakingFrame = InSettings.IsSet() ? InSettings.GetValue().EndFrame : MovieScene->GetPlaybackRange().GetUpperBoundValue();
+	FFrameNumber FirstBakingFrame = MovieScene->GetPlaybackRange().GetLowerBoundValue();
+	FFrameNumber LastBakingFrame = MovieScene->GetPlaybackRange().GetUpperBoundValue();
 	
-	// set start to first active frame if any, if we have no range
-	if (InSettings.IsSet() == false)
+	// set start to first active frame if any
+	const int32 FirstActiveIndex = ConstraintValues.IndexOfByKey(true);
+	if (ConstraintValues.IsValidIndex(FirstActiveIndex))
 	{
-		const int32 FirstActiveIndex = ConstraintValues.IndexOfByKey(true);
-		if (ConstraintValues.IsValidIndex(FirstActiveIndex))
-		{
-			FirstBakingFrame = ConstraintFrames[FirstActiveIndex];
-		}
-
-		// set end to the last key if inactive
-		const bool bIsLastKeyInactive = (ConstraintValues.Last() == false);
-		if (bIsLastKeyInactive)
-		{
-			LastBakingFrame = ConstraintFrames.Last();
-		}
+		FirstBakingFrame = ConstraintFrames[FirstActiveIndex];
 	}
-	// keep frames within active state
-	auto IsConstraintActive = [InConstraint, ThisActiveChannel](const FFrameNumber& Time)
-	{
-		if (ThisActiveChannel)
-		{
-			bool bIsActive = false; ThisActiveChannel->ActiveChannel.Evaluate(Time, bIsActive);
-			return bIsActive;
-		}
-		return InConstraint->IsFullyActive();
-	};
 
-	if (InSettings.IsSet() == false || InSettings.GetValue().BakingKeySettings != EBakingKeySettings::KeysOnly)
+	// set end to the last key if inactive
+	const bool bIsLastKeyInactive = (ConstraintValues.Last() == false);
+	if (bIsLastKeyInactive)
 	{
-		TArray<FFrameNumber> Frames;
-		MovieSceneToolHelpers::CalculateFramesBetween(MovieScene, FirstBakingFrame, LastBakingFrame,  1 /* use increment of 1 since we reduce later*/,
-			Frames);
-
-		// Fill the frames we want to bake
+		LastBakingFrame = ConstraintFrames.Last();
+	}
 
+	// then compute range from first to last
+	TArray<FFrameNumber> Frames;
+	MovieSceneToolHelpers::CalculateFramesBetween(MovieScene, FirstBakingFrame, LastBakingFrame, Frames);
+		
+	// Fill the frames we want to bake
+	{
 		// add constraint keys
-		for (const FFrameNumber& ConstraintFrame : ConstraintFrames)
+		for (const FFrameNumber& ConstraintFrame: ConstraintFrames)
 		{
 			OutFramesToBake.Add(ConstraintFrame);
 		}
-
-		int32 FrameInc = (InSettings.IsSet() == true) ? InSettings.GetValue().FrameIncrement : 1;
-		for (int32 FrameIndex = 0; FrameIndex < Frames.Num(); ++FrameIndex)
+		
+		// keep frames within active state
+		auto IsConstraintActive = [InConstraint, ThisActiveChannel](const FFrameNumber& Time)
 		{
-			if (FrameIndex % FrameInc == 0)
+			if (ThisActiveChannel)
 			{
-				const FFrameNumber& InFrame = Frames[FrameIndex];
-				if (IsConstraintActive(InFrame))
-				{
-					OutFramesToBake.Add(InFrame);
-				}
+				bool bIsActive = false; ThisActiveChannel->ActiveChannel.Evaluate(Time, bIsActive);
+				return bIsActive;
 			}
-		}
-		if (FrameInc > 1) //if skipping make sure we get all transforms
-		{
-			FMovieSceneConstraintChannelHelper::GetTransformFramesForConstraintHandles(InConstraint, InSequencer,
-				InSettings.GetValue().StartFrame, InSettings.GetValue().EndFrame, OutFramesToBake);
-		}
+			return InConstraint->IsFullyActive();
+		};
 		
-	}
-	else //doing per frame key get constraint keys and transform keys
-	{
-		for (const FFrameNumber& Frame : ConstraintFrames)
+		for (const FFrameNumber& InFrame: Frames)
 		{
-			if (Frame >= FirstBakingFrame && Frame <= LastBakingFrame)
+			if (IsConstraintActive(InFrame))
 			{
-				OutFramesToBake.Add(Frame);
+				OutFramesToBake.Add(InFrame);
 			}
 		}
-
-		FMovieSceneConstraintChannelHelper::GetTransformFramesForConstraintHandles(InConstraint, InSequencer,
-			InSettings.GetValue().StartFrame, InSettings.GetValue().EndFrame, OutFramesToBake);
-
 	}
 
 	// we also need to store which T-1 frames need to be kept for other constraints compensation
-	// and also store transform key locations for parents if doing per key or per frame with non-1 increments
 	{
 		// gather the sorted child's constraint
 		static constexpr bool bSorted = true;
@@ -153,18 +120,10 @@ void FConstraintBaker::GetMinimalFramesToBake(
 			if (TransformConstraint && TransformConstraint->NeedsCompensation() && TransformConstraint != InConstraint) 
 			{
 				const FName ConstraintName = TransformConstraint->GetFName();
-				if (FConstraintAndActiveChannel* ConstraintChannel = ConstrainedSection->GetConstraintChannel(ConstraintName))
+				if (FConstraintAndActiveChannel* ConstraintChannel = InSection->GetConstraintChannel(ConstraintName))
 				{
 					OtherChannels.Add(ConstraintChannel);
 				}
-
-				//if doing keys or per frame increments we need to also 
-				if (InSettings.IsSet() == true &&
-					(InSettings.GetValue().BakingKeySettings == EBakingKeySettings::KeysOnly || InSettings.GetValue().FrameIncrement > 1))
-				{
-					FMovieSceneConstraintChannelHelper::GetTransformFramesForConstraintHandles(TransformConstraint, InSequencer,
-						InSettings.GetValue().StartFrame, InSettings.GetValue().EndFrame, OutFramesToBake);
-				}
 			}
 		}
 
@@ -196,25 +155,9 @@ void FConstraintBaker::GetMinimalFramesToBake(
 	OutFramesToBake.SetNum(Algo::Unique(OutFramesToBake));
 }
 
-bool FConstraintBaker::BakeMultiple(
-	UWorld* InWorld,
-	TArray< UTickableTransformConstraint*>& InConstraints,
-	const TSharedPtr<ISequencer>& InSequencer,
-	const FBakingAnimationKeySettings& InSettings)
-{
-	TOptional<TArray<FFrameNumber>> Frames;
-	TOptional< FBakingAnimationKeySettings> Settings = InSettings;
-	for (UTickableTransformConstraint* Constraint : InConstraints)
-	{
-		FConstraintBaker::Bake(InWorld, Constraint, InSequencer, Settings, Frames);
-	}
-	return true;
-}
-
 void FConstraintBaker::Bake(UWorld* InWorld, 
 	UTickableTransformConstraint* InConstraint,
 	const TSharedPtr<ISequencer>& InSequencer, 
-	const TOptional< FBakingAnimationKeySettings>& InSettings,
 	const TOptional<TArray<FFrameNumber>>& InFrames)
 {
 	if ((InFrames && InFrames->IsEmpty()) || (InConstraint == nullptr) || (InConstraint->ChildTRSHandle == nullptr))
@@ -224,13 +167,29 @@ void FConstraintBaker::Bake(UWorld* InWorld,
 
 	const TObjectPtr<UTransformableHandle>& Handle = InConstraint->ChildTRSHandle;
 
-	FConstraintSections ConstraintSections = FMovieSceneConstraintChannelHelper::GetConstraintSectionAndChannel(
-		InConstraint, InSequencer);
-	IMovieSceneConstrainedSection* ConstrainedSection = Cast<IMovieSceneConstrainedSection>(ConstraintSections.ConstraintSection);
-	if (ConstrainedSection == nullptr) 
+	const FConstraintChannelInterfaceRegistry& InterfaceRegistry = FConstraintChannelInterfaceRegistry::Get();	
+	ITransformConstraintChannelInterface* Interface = InterfaceRegistry.FindConstraintChannelInterface(Handle->GetClass());
+	if (!Interface)
 	{
 		return;
 	}
+	
+	//get the section to be used later to delete the extra transform keys at the frame -1 times, abort if not there for some reason
+	UMovieSceneSection* ConstraintSection = Interface->GetHandleConstraintSection(Handle, InSequencer);
+	const UMovieSceneSection* TransformSection = Interface->GetHandleSection(Handle, InSequencer);
+
+	IMovieSceneConstrainedSection* ConstrainedSection = Cast<IMovieSceneConstrainedSection>(ConstraintSection);
+	if (ConstrainedSection == nullptr || TransformSection == nullptr)
+	{
+		return;
+	}
+	
+	FConstraintAndActiveChannel* ActiveChannel = ConstrainedSection->GetConstraintChannel(InConstraint->GetFName());
+	if (ActiveChannel == nullptr)
+	{
+		return;
+	}
+
 	// compute transforms
 	TArray<FFrameNumber> FramesToBake;
 	if (InFrames)
@@ -239,12 +198,12 @@ void FConstraintBaker::Bake(UWorld* InWorld,
 	}
 	else
 	{
-		GetMinimalFramesToBake(InWorld, InConstraint, InSequencer, ConstraintSections.ConstraintSection, InSettings, FramesToBake);
+		GetMinimalFramesToBake(InWorld, InConstraint, InSequencer, ConstrainedSection, FramesToBake);
 
 		// if it needs recomposition 
-		if (ConstraintSections.ConstraintSection != ConstraintSections.ChildTransformSection && !FramesToBake.IsEmpty())
+		if (ConstraintSection != TransformSection && !FramesToBake.IsEmpty())
 		{
-			const TMovieSceneChannelData<const bool> ConstraintChannelData = ConstraintSections.ActiveChannel->ActiveChannel.GetData();
+			const TMovieSceneChannelData<const bool> ConstraintChannelData = ActiveChannel->ActiveChannel.GetData();
 			const TArrayView<const FFrameNumber> ConstraintFrames = ConstraintChannelData.GetTimes();
 			const TArrayView<const bool> ConstraintValues = ConstraintChannelData.GetValues();
 			for (int32 Index = 0; Index < ConstraintFrames.Num(); ++Index)
@@ -277,15 +236,15 @@ void FConstraintBaker::Bake(UWorld* InWorld,
 		return;
 	}
 	
-	ConstraintSections.ConstraintSection->Modify();
+	ConstraintSection->Modify();
 
 	// disable constraint and delete extra transform keys
-	TMovieSceneChannelData<bool> ConstraintChannelData = ConstraintSections.ActiveChannel->ActiveChannel.GetData();
+	TMovieSceneChannelData<bool> ConstraintChannelData = ActiveChannel->ActiveChannel.GetData();
 	const TArrayView<const FFrameNumber> ConstraintFrames = ConstraintChannelData.GetTimes();
 	
 	// get transform channels
-	const TArrayView<FMovieSceneFloatChannel*> FloatTransformChannels = Handle->GetFloatChannels(ConstraintSections.ChildTransformSection);
-	const TArrayView<FMovieSceneDoubleChannel*> DoubleTransformChannels = Handle->GetDoubleChannels(ConstraintSections.ChildTransformSection);
+	const TArrayView<FMovieSceneFloatChannel*> FloatTransformChannels = Handle->GetFloatChannels(TransformSection);
+	const TArrayView<FMovieSceneDoubleChannel*> DoubleTransformChannels = Handle->GetDoubleChannels(TransformSection);
 
 	for (int32 Index = 0; Index < ConstraintFrames.Num(); ++Index)
 	{
@@ -311,18 +270,8 @@ void FConstraintBaker::Bake(UWorld* InWorld,
 
 	// now bake to channel curves
 	const EMovieSceneTransformChannel Channels = InConstraint->GetChannelsToKey();
-	ConstraintSections.Interface->AddHandleTransformKeys(InSequencer, Handle, FramesToBake, Transforms, Channels);
+	Interface->AddHandleTransformKeys(InSequencer, Handle, FramesToBake, Transforms, Channels);
 
-	if (InSettings.IsSet() && InSettings.GetValue().BakingKeySettings != EBakingKeySettings::KeysOnly
-		&& InSettings.GetValue().bReduceKeys == true)
-	{
-		FKeyDataOptimizationParams Param;
-		Param.bAutoSetInterpolation = true;
-		Param.Tolerance = InSettings.GetValue().Tolerance;
-		TRange<FFrameNumber> Range(InSettings.GetValue().StartFrame, InSettings.GetValue().EndFrame);
-		Param.Range = Range;
-		MovieSceneToolHelpers::OptimizeSection(Param, ConstraintSections.ConstraintSection);
-	}
 	// notify
 	InSequencer->RequestEvaluate();
 }
diff --git a/Engine/Plugins/Animation/ControlRig/Source/ControlRigEditor/Private/Tools/ControlRigSnapSettings.cpp b/Engine/Plugins/Animation/ControlRig/Source/ControlRigEditor/Private/Tools/ControlRigSnapSettings.cpp
deleted file mode 100644
diff --git a/Engine/Plugins/Animation/ControlRig/Source/ControlRigEditor/Private/Tools/ControlRigSnapper.cpp b/Engine/Plugins/Animation/ControlRig/Source/ControlRigEditor/Private/Tools/ControlRigSnapper.cpp
--- a/Engine/Plugins/Animation/ControlRig/Source/ControlRigEditor/Private/Tools/ControlRigSnapper.cpp
+++ b/Engine/Plugins/Animation/ControlRig/Source/ControlRigEditor/Private/Tools/ControlRigSnapper.cpp
@@ -21,8 +21,6 @@
 #include "Sections/MovieScene3DTransformSection.h"
 #include "Tracks/MovieScene3DTransformTrack.h"
 #include "Channels/MovieSceneChannelProxy.h"
-#include "Channels/MovieSceneFloatChannel.h"
-#include "Channels/MovieSceneDoubleChannel.h"
 #include "Tools/ControlRigSnapSettings.h"
 #include "MovieSceneToolsModule.h"
 #include "Components/SkeletalMeshComponent.h"
@@ -31,10 +29,6 @@
 #include "LevelSequencePlayer.h"
 #include "LevelSequenceActor.h"
 #include "Tools/BakingHelper.h"
-#include "Containers/SortedMap.h"
-#include "Sequencer/ControlRigSequencerHelpers.h"
-#include "Constraints/MovieSceneConstraintChannelHelper.inl"
-
 
 #include UE_INLINE_GENERATED_CPP_BY_NAME(ControlRigSnapper)
 
@@ -196,7 +190,7 @@ struct FGuidAndActor
 		const TArray<FTransform>& WorldTransformsToSnapTo, const UControlRigSnapSettings* SnapSettings) const
 	{
 		// get section
-		UMovieScene3DTransformSection* TransformSection = MovieSceneToolHelpers::GetTransformSection(Sequencer, Guid);
+		const UMovieScene3DTransformSection* TransformSection = MovieSceneToolHelpers::GetTransformSection(Sequencer, Guid);
 		if (!TransformSection)
 		{
 			return false;
@@ -248,18 +242,6 @@ struct FGuidAndActor
 		// add keys
 		MovieSceneToolHelpers::AddTransformKeys(TransformSection, Frames, LocalTransforms, Channels);
 
-		if (SnapSettings && SnapSettings->BakingKeySettings != EBakingKeySettings::KeysOnly
-			&& SnapSettings->bReduceKeys == true && Frames.Num() > 2)
-		{
-			FKeyDataOptimizationParams Param;
-			Param.bAutoSetInterpolation = true;
-			Param.Tolerance = SnapSettings->Tolerance;
-			TRange<FFrameNumber> Range(Frames[0], Frames[Frames.Num() -1]);
-			Param.Range = Range;
-			MovieSceneToolHelpers::OptimizeSection(Param, TransformSection);
-		}
-
-
 		// notify
 		Sequencer->NotifyMovieSceneDataChanged(EMovieSceneDataChangeType::MovieSceneStructureItemsChanged);
 
@@ -323,173 +305,7 @@ static bool CalculateWorldTransformsFromParents(ISequencer* Sequencer, const FCo
 	return false;
 }
 
-//to do handle constraints and spaces, but for now we just output a single actor parent.
-static void GetActorParents(const FActorForWorldTransforms& ActorSelection,
-	TArray<FActorForWorldTransforms>& OutParentActors, TArray<FControlRigForWorldTransforms>& OutControlRigs)
-{
-	if (ActorSelection.Actor.IsValid())
-	{
-		//has component so parent is next component
-		if (ActorSelection.Component.IsValid() && ActorSelection.Component->GetAttachParent())
-		{
-			FActorForWorldTransforms OutParentActor;
-			OutParentActor.Actor = ActorSelection.Actor;
-			OutParentActor.Component = ActorSelection.Component->GetAttachParent();
-			OutParentActors.Add(OutParentActor);
-			GetActorParents(OutParentActor, OutParentActors, OutControlRigs);
-		}
-		else if (AActor* ParentActor = ActorSelection.Actor->GetAttachParentActor())
-		{
-			FActorForWorldTransforms OutParentActor;
-			OutParentActor.Actor = ParentActor;
-			OutParentActors.Add(OutParentActor);
-			GetActorParents(OutParentActor, OutParentActors, OutControlRigs);
-		}
-	}
-}
-
-//todo handle constraints/spaces
-static void GetControlRigParents(const FControlRigForWorldTransforms& ControlRigAndSelection, 
-	TArray<FActorForWorldTransforms>& OutParentActors, TArray<FControlRigForWorldTransforms>& OutControlRigs)
-{
-
-	if (UControlRig* ControlRig = ControlRigAndSelection.ControlRig.Get())
-	{
-		URigHierarchy* Hierarchy = ControlRig->GetHierarchy();
-		for (const FName& ControlName : ControlRigAndSelection.ControlNames)
-		{
-			if (FRigControlElement* ControlElement = ControlRig->FindControl(ControlName))
-			{
-				if (FRigControlElement* ParentControlElement = Cast<FRigControlElement>(Hierarchy->GetFirstParent(ControlElement)))
-				{
-					FControlRigForWorldTransforms OutControlRig;
-					OutControlRig.ControlRig = ControlRig;
-					OutControlRig.ControlNames.Add(ParentControlElement->GetKey().Name);
-					OutControlRigs.Add(OutControlRig);
-					GetControlRigParents(OutControlRig, OutParentActors, OutControlRigs);
-				}
-				else
-				{
-					if (TSharedPtr<IControlRigObjectBinding> ObjectBinding = ControlRig->GetObjectBinding())
-					{
-						USceneComponent* Component = Cast<USceneComponent>(ObjectBinding->GetBoundObject());
-						if (!Component)
-						{
-							continue;
-						}
-						AActor* Actor = Component->GetTypedOuter< AActor >();
-						if (!Actor)
-						{
-							continue;
-						}
-						FActorForWorldTransforms OutParentActor;
-						OutParentActor.Actor = Actor;
-						OutParentActors.Add(OutParentActor);
-						GetActorParents(OutParentActor, OutParentActors, OutControlRigs);
-					}
-				}
-			}
-		}
-	}
-}
-
-static UMovieSceneControlRigParameterSection* GetControlRigSection(ISequencer* Sequencer, const UControlRig* ControlRig)
-{
-
-	if (ControlRig == nullptr || Sequencer == nullptr)
-	{
-		return nullptr;
-	}
-	UMovieScene* MovieScene = Sequencer->GetFocusedMovieSceneSequence()->GetMovieScene();
-	if (!MovieScene)
-	{
-		return nullptr;
-	}
-	const TArray<FMovieSceneBinding>& Bindings = MovieScene->GetBindings();
-	for (const FMovieSceneBinding& Binding : Bindings)
-	{
-		UMovieSceneControlRigParameterTrack* ControlRigParameterTrack = Cast<UMovieSceneControlRigParameterTrack>(MovieScene->FindTrack(UMovieSceneControlRigParameterTrack::StaticClass(), Binding.GetObjectGuid(), NAME_None));
-		if (ControlRigParameterTrack && ControlRigParameterTrack->GetControlRig() == ControlRig)
-		{
-			UMovieSceneControlRigParameterSection* ActiveSection = Cast<UMovieSceneControlRigParameterSection>(ControlRigParameterTrack->GetSectionToKey());
-			if (ActiveSection)
-			{
-				return ActiveSection;
-			}
-		}
-	}
-	return nullptr;
-}
-
-//Get Transform frames for specified snap items, will add them to the OutFrameMap
-static void GetTransformFrames(ISequencer * Sequencer, const FControlRigSnapperSelection& SnapItem, FFrameNumber StartFrame, FFrameNumber EndFrame, TSortedMap<FFrameNumber, FFrameNumber>& OutFrameMap)
-{
-	TArray<FFrameNumber> FramesToUse;
-	TArray<FActorForWorldTransforms> ParentActors;
-	TArray<FControlRigForWorldTransforms> ParentControlRigs;
-
-	for (const FActorForWorldTransforms& ActorsFor : SnapItem.Actors)
-	{
-		FramesToUse.Reset();
-		ParentActors.Reset();
-		ParentControlRigs.Reset();
-		ParentActors.Add(ActorsFor);
-		GetActorParents(ActorsFor, ParentActors, ParentControlRigs);
-
-	}
-	for (const FControlRigForWorldTransforms& ControlRigsFor : SnapItem.ControlRigs)
-	{
-		FramesToUse.Reset();
-		ParentActors.Reset();
-		ParentControlRigs.Reset();
-		ParentControlRigs.Add(ControlRigsFor);
-		GetControlRigParents(ControlRigsFor, ParentActors, ParentControlRigs);
-	}
 
-	for (FActorForWorldTransforms& Actor : ParentActors)
-	{
-		FGuid Guid;
-		if (Actor.Component.Get())
-		{
-			Guid = Sequencer->GetHandleToObject(Actor.Component.Get(), false);
-		}
-		if (!Guid.IsValid())
-		{
-			Guid = Sequencer->GetHandleToObject(Actor.Actor.Get(), false);
-		}
-		if (Guid.IsValid())
-		{
-			if (UMovieScene3DTransformSection* TransformSection = MovieSceneToolHelpers::GetTransformSection(Sequencer, Guid))
-			{
-				TArrayView<FMovieSceneDoubleChannel*> Channels = TransformSection->GetChannelProxy().GetChannels<FMovieSceneDoubleChannel>();
-				TArray<FFrameNumber> TransformFrameTimes = FMovieSceneConstraintChannelHelper::GetTransformTimes(
-					Channels, StartFrame, EndFrame);
-				for (FFrameNumber& FrameNumber : TransformFrameTimes)
-				{
-					OutFrameMap.Add(FrameNumber);
-				}
-			}
-		}
-	}
-	for (FControlRigForWorldTransforms& ControlRig : ParentControlRigs)
-	{
-		UMovieSceneControlRigParameterSection* Section = GetControlRigSection(Sequencer, ControlRig.ControlRig.Get());
-		if (Section)
-		{
-			for (FName& ControlName : ControlRig.ControlNames)
-			{
-				TArrayView<FMovieSceneFloatChannel*> Channels = FControlRigSequencerHelpers::GetFloatChannels(ControlRig.ControlRig.Get(),
-					ControlName, Section);
-				TArray<FFrameNumber> TransformFrameTimes = FMovieSceneConstraintChannelHelper::GetTransformTimes(
-					Channels, StartFrame, EndFrame);
-				for (FFrameNumber& FrameNumber : TransformFrameTimes)
-				{
-					OutFrameMap.Add(FrameNumber);
-				}
-			}
-		}
-	}
-}
 bool FControlRigSnapper::SnapIt(FFrameNumber StartFrame, FFrameNumber EndFrame,const FControlRigSnapperSelection& ActorToSnap,
 	const FControlRigSnapperSelection& ParentToSnap, const UControlRigSnapSettings* SnapSettings)
 {
@@ -504,18 +320,8 @@ bool FControlRigSnapper::SnapIt(FFrameNumber StartFrame, FFrameNumber EndFrame,c
 		MovieScene->Modify();
 		
 		TArray<FFrameNumber> Frames;
-		if (SnapSettings->BakingKeySettings == EBakingKeySettings::AllFrames)
-		{
-			MovieSceneToolHelpers::CalculateFramesBetween(MovieScene, StartFrame, EndFrame, SnapSettings->FrameIncrement,Frames);
-		}
-		else
-		{
-			TSortedMap<FFrameNumber, FFrameNumber> FrameMap;
-			GetTransformFrames(Sequencer, ParentToSnap, StartFrame, EndFrame, FrameMap);
-			GetTransformFrames(Sequencer, ActorToSnap,  StartFrame, EndFrame, FrameMap);
-			FrameMap.GenerateKeyArray(Frames);
+		MovieSceneToolHelpers::CalculateFramesBetween(MovieScene, StartFrame, EndFrame, Frames);
 
-		}
 		TArray<FTransform> WorldTransformToSnap;
 		bool bSnapToFirstFrameNotParents = !CalculateWorldTransformsFromParents(Sequencer, ParentToSnap, Frames, WorldTransformToSnap);
 		if (Frames.Num() != WorldTransformToSnap.Num())
@@ -674,18 +480,6 @@ bool FControlRigSnapper::SnapIt(FFrameNumber StartFrame, FFrameNumber EndFrame,c
 							ControlRig->SetControlGlobalTransform(Name, GlobalTransform, true, Context);
 						}
 					}
-
-					UMovieSceneControlRigParameterSection* ControlRigSection = GetControlRigSection(Sequencer, ControlRig);
-					if (ControlRigSection && SnapSettings && SnapSettings->BakingKeySettings != EBakingKeySettings::KeysOnly
-						&& SnapSettings->bReduceKeys == true && Frames.Num() > 2)
-					{
-						FKeyDataOptimizationParams Param;
-						Param.bAutoSetInterpolation = true;
-						Param.Tolerance = SnapSettings->Tolerance;
-						TRange<FFrameNumber> Range(Frames[0], Frames[Frames.Num() - 1]);
-						Param.Range = Range;
-						MovieSceneToolHelpers::OptimizeSection(Param, ControlRigSection);
-					}
 				}
 			}
 		}
diff --git a/Engine/Plugins/Animation/ControlRig/Source/ControlRigEditor/Public/ControlRigSequencerEditorLibrary.h b/Engine/Plugins/Animation/ControlRig/Source/ControlRigEditor/Public/ControlRigSequencerEditorLibrary.h
--- a/Engine/Plugins/Animation/ControlRig/Source/ControlRigEditor/Public/ControlRigSequencerEditorLibrary.h
+++ b/Engine/Plugins/Animation/ControlRig/Source/ControlRigEditor/Public/ControlRigSequencerEditorLibrary.h
@@ -23,7 +23,6 @@ class ULevelSequence;
 class UTickableConstraint;
 class UTickableTransformConstraint;
 class UTransformableHandle;
-struct FBakingAnimationKeySettings;
 
 USTRUCT(BlueprintType)
 struct FControlRigSequencerBindingProxy
@@ -143,22 +142,13 @@ public:
 	* Bake the constraint to keys based on the passed in frames. This will use the open sequencer to bake. See ConstraintsScriptingLibrary to get the list of available constraints
 	* @param World The active world
 	* @param Constraint The Constraint to bake. After baking it will be keyed to be inactive of the range of frames that are baked
-	* @param Frames The frames to bake, if the array is empty it will use the active time ranges of the constraint to determine where it should bake
+	* @param Frames The frames to bake, if the array is empty it will use the active time ranges of the constraint to deteremine where it should bake
 	* @param TimeUnit Unit for all frame and time values, either in display rate or tick resolution
 	* @return Returns True if successful, False otherwise
 	*/
 	UFUNCTION(BlueprintCallable, Category = "Editor Scripting | Sequencer Tools | Control Rig")
 	static bool BakeConstraint(UWorld* World, UTickableConstraint* Constraint, const TArray<FFrameNumber>& Frames, ESequenceTimeUnit TimeUnit = ESequenceTimeUnit::DisplayRate);
 
-	/**
-	* Bake the constraint to keys based on the passed in settings. This will use the open sequencer to bake. See ConstraintsScriptingLibrary to get the list of available constraints
-	* @param World The active world
-	* @param InConstraints The Constraints tobake.  After baking they will be keyed to be inactive of the range of frames that are baked
-	* @param InSettings Settings to use for baking
-	* @return Returns True if successful, False otherwise
-	*/
-	UFUNCTION(BlueprintCallable, Category = "Editor Scripting | Sequencer Tools | Control Rig")
-	static bool BakeConstraints(UWorld* World, TArray<UTickableConstraint*>& InConstraints,const FBakingAnimationKeySettings& InSettings);
 	/**
 	* Add a constraint possibly adding to sequencer also if one is open.
 	* @param World The active world
@@ -936,16 +926,6 @@ public:
 	UFUNCTION(BlueprintCallable, Category = "Editor Scripting | Sequencer Tools | Control Rig")
 	static bool CollapseControlRigAnimLayers(ULevelSequence* InSequence,UMovieSceneControlRigParameterTrack* InTrack, bool bKeyReduce = false, float Tolerance = 0.001f);
 
-	/*
-	 * Collapse and bake all sections and layers on a control rig track to just one section using passed in settings.
-	 *
-	 * @param InSequence Sequence that has track to collapse
-	 * @param InTrack Track for layers to collapse
-	 * @param InSettings Settings that determine how to collapse
-	 */
-	UFUNCTION(BlueprintCallable, Category = "Editor Scripting | Sequencer Tools | Control Rig")
-	static bool CollapseControlRigAnimLayersWithSettings(ULevelSequence* InSequence, UMovieSceneControlRigParameterTrack* InTrack, const FBakingAnimationKeySettings& InSettings);
-
 	/*
 	 * Get the default parent key, can be used a parent space.
 	 *
diff --git a/Engine/Plugins/Animation/ControlRig/Source/ControlRigEditor/Public/ControlRigSpaceChannelEditors.h b/Engine/Plugins/Animation/ControlRig/Source/ControlRigEditor/Public/ControlRigSpaceChannelEditors.h
--- a/Engine/Plugins/Animation/ControlRig/Source/ControlRigEditor/Public/ControlRigSpaceChannelEditors.h
+++ b/Engine/Plugins/Animation/ControlRig/Source/ControlRigEditor/Public/ControlRigSpaceChannelEditors.h
@@ -49,7 +49,7 @@ struct FControlRigSpaceChannelHelpers
 	static void CompensateIfNeeded(UControlRig* ControlRig, ISequencer* Sequencer, UMovieSceneControlRigParameterSection* Section, FName ControlName, TOptional<FFrameNumber>& Time);
 	static FSpaceChannelAndSection FindSpaceChannelAndSectionForControl(UControlRig* ControlRig, FName ControlName, ISequencer* Sequencer, bool bCreateIfNeeded);
 	static void SequencerBakeControlInSpace(UControlRig* ControlRig, ISequencer* Sequencer, FMovieSceneControlRigSpaceChannel* Channel, UMovieSceneSection* SectionToKey,
-		URigHierarchy* RigHierarchy, const FRigElementKey& ControlKey, FRigSpacePickerBakeSettings InSettings);
+		TArray<FFrameNumber> Frames, URigHierarchy* RigHierarchy, const FRigElementKey& ControlKey, FRigSpacePickerBakeSettings InSettings);
 	static void GetFramesInThisSpaceAfterThisTime(UControlRig* ControlRig, FName ControlName, FMovieSceneControlRigSpaceBaseKey CurrentValue,
 		FMovieSceneControlRigSpaceChannel* Channel, UMovieSceneSection* SectionToKey,
 		FFrameNumber Time, TSortedMap<FFrameNumber,FFrameNumber>& OutMoreFrames);
diff --git a/Engine/Plugins/Animation/ControlRig/Source/ControlRigEditor/Public/RigSpacePickerBakeSettings.h b/Engine/Plugins/Animation/ControlRig/Source/ControlRigEditor/Public/RigSpacePickerBakeSettings.h
--- a/Engine/Plugins/Animation/ControlRig/Source/ControlRigEditor/Public/RigSpacePickerBakeSettings.h
+++ b/Engine/Plugins/Animation/ControlRig/Source/ControlRigEditor/Public/RigSpacePickerBakeSettings.h
@@ -4,10 +4,8 @@
 
 #include "CoreMinimal.h"
 #include "Rigs/RigHierarchyDefines.h"
-#include "BakingAnimationKeySettings.h"
 #include "RigSpacePickerBakeSettings.generated.h"
 
-
 USTRUCT(BlueprintType)
 struct CONTROLRIGEDITOR_API FRigSpacePickerBakeSettings
 {
@@ -16,28 +14,24 @@ struct CONTROLRIGEDITOR_API FRigSpacePickerBakeSettings
 	FRigSpacePickerBakeSettings()
 	{
 		TargetSpace = FRigElementKey();
-		Settings = FBakingAnimationKeySettings();
+		StartFrame = 0;
+		EndFrame = 100;
+		bReduceKeys = false;
+		Tolerance = 0.001f;
 	}
 
 	UPROPERTY(BlueprintReadWrite, Category = "Settings")
 	FRigElementKey TargetSpace;
 
-	UPROPERTY(BlueprintReadWrite, EditAnywhere, Category = "Bake", meta = (ShowOnlyInnerProperties))
-	FBakingAnimationKeySettings Settings;
-
-	/** DEPRECATED 5.3 */
-	UPROPERTY(meta = (DeprecatedProperty, DeprecationMessage = "Use Settings.StartFrame instead"))
-	FFrameNumber StartFrame_DEPRECATED;
+	UPROPERTY(BlueprintReadWrite,EditAnywhere, Category = "Settings")
+	FFrameNumber StartFrame;
 
-	/** DEPRECATED 5.3 */
-	UPROPERTY(meta = (DeprecatedProperty, DeprecationMessage = "Use Settings.EndFrame instead"))
-	FFrameNumber EndFrame_DEPRECATED;
+	UPROPERTY(BlueprintReadWrite,EditAnywhere, Category = "Settings")
+	FFrameNumber EndFrame;
 
-	/** DEPRECATED 5.3 */
-	UPROPERTY(meta = (DeprecatedProperty, DeprecationMessage = "Use Settings.bReduceKeys instead"))
-	bool bReduceKeys_DEPRECATED;
+	UPROPERTY(BlueprintReadWrite,EditAnywhere, Category = "Settings")
+	bool bReduceKeys;
 
-	/** DEPRECATED 5.3 */
-	UPROPERTY(meta = (DeprecatedProperty, DeprecationMessage = "Use Settings.Tolerance instead"))
-	float Tolerance_DEPRECATED;
-};
\ No newline at end of file
+	UPROPERTY(BlueprintReadWrite,EditAnywhere, Category = "Settings")
+	float Tolerance;
+};
diff --git a/Engine/Plugins/Animation/ControlRig/Source/ControlRigEditor/Public/Tools/ConstraintBaker.h b/Engine/Plugins/Animation/ControlRig/Source/ControlRigEditor/Public/Tools/ConstraintBaker.h
--- a/Engine/Plugins/Animation/ControlRig/Source/ControlRigEditor/Public/Tools/ConstraintBaker.h
+++ b/Engine/Plugins/Animation/ControlRig/Source/ControlRigEditor/Public/Tools/ConstraintBaker.h
@@ -4,14 +4,14 @@
 
 #include "Misc/FrameNumber.h"
 #include "Templates/SharedPointer.h"
-#include "BakingAnimationKeySettings.h"
 
-class UMovieSceneSection;
+class IMovieSceneConstrainedSection;
 class UTransformableHandle;
 class UTickableTransformConstraint;
 class ISequencer;
 class UWorld;
 struct FFrameNumber;
+
 struct FConstraintBaker
 {
 public:
@@ -20,16 +20,8 @@ public:
 		UWorld* InWorld,
 		UTickableTransformConstraint* InConstraint,
 		const TSharedPtr<ISequencer>& InSequencer,
-		const TOptional< FBakingAnimationKeySettings>& InSettings,
 		const TOptional<TArray<FFrameNumber>>& InFrames);
 	
-	/** Bake Multiple Constraints using settings*/
-	static bool BakeMultiple(
-		UWorld* InWorld,
-		TArray< UTickableTransformConstraint*>& InConstraints,
-		const TSharedPtr<ISequencer>& InSequencer,
-		const FBakingAnimationKeySettings& Settings);
-	
 private:
 
 	/** Computes the minimal frames to bake for that constraint. */
@@ -37,8 +29,7 @@ private:
 		UWorld* InWorld,
 		const UTickableTransformConstraint* InConstraint,
 		const TSharedPtr<ISequencer>& InSequencer,
-		UMovieSceneSection* InSection,
-		const TOptional< FBakingAnimationKeySettings>& Settings,
+		IMovieSceneConstrainedSection* InSection,
 		TArray<FFrameNumber>& OutFramesToBake);
 };
 
diff --git a/Engine/Plugins/Animation/ControlRig/Source/ControlRigEditor/Public/Tools/ControlRigSnapSettings.h b/Engine/Plugins/Animation/ControlRig/Source/ControlRigEditor/Public/Tools/ControlRigSnapSettings.h
deleted file mode 100644
diff --git a/Engine/Source/Editor/MovieSceneTools/Private/Constraints/MovieSceneConstraintChannelHelper.cpp b/Engine/Source/Editor/MovieSceneTools/Private/Constraints/MovieSceneConstraintChannelHelper.cpp
--- a/Engine/Source/Editor/MovieSceneTools/Private/Constraints/MovieSceneConstraintChannelHelper.cpp
+++ b/Engine/Source/Editor/MovieSceneTools/Private/Constraints/MovieSceneConstraintChannelHelper.cpp
@@ -18,7 +18,6 @@
 #include "MovieSceneSpawnableAnnotation.h"
 #include "Sections/MovieSceneConstrainedSection.h"
 #include "MovieSceneToolsModule.h"
-#include "Channels/MovieSceneChannelTraits.h"
 #include "Constraints/TransformConstraintChannelInterface.h"
 
 
@@ -1003,104 +1002,6 @@ void FMovieSceneConstraintChannelHelper::CompensateIfNeeded(
 		InSequencer->ForceEvaluate();
 	}
 }
-FConstraintSections FMovieSceneConstraintChannelHelper::GetConstraintSectionAndChannel(
-	const UTickableTransformConstraint* InConstraint,
-	const TSharedPtr<ISequencer>& InSequencer)
-{
-	FConstraintSections ReturnValue;
-
-	if (InSequencer.IsValid() == false)
-	{
-		return ReturnValue;
-	}
-	const TObjectPtr<UTransformableHandle>& ChildHandle = InConstraint->ChildTRSHandle;
-
-	const FConstraintChannelInterfaceRegistry& InterfaceRegistry = FConstraintChannelInterfaceRegistry::Get();
-	ReturnValue.Interface = InterfaceRegistry.FindConstraintChannelInterface(ChildHandle->GetClass());
-	if (!ReturnValue.Interface)
-	{
-		return ReturnValue;
-	}
-	//get the section to be used later to delete the extra transform keys at the frame -1 times, abort if not there for some reason
-	ReturnValue.ConstraintSection = ReturnValue.Interface->GetHandleConstraintSection(ChildHandle, InSequencer);
-	ReturnValue.ChildTransformSection = ReturnValue.Interface->GetHandleSection(ChildHandle, InSequencer);
-
-	ITransformConstraintChannelInterface* ParentInterface = InterfaceRegistry.FindConstraintChannelInterface(InConstraint->ParentTRSHandle->GetClass());
-	if (ParentInterface)
-	{
-		ReturnValue.ParentTransformSection = ParentInterface->GetHandleSection(InConstraint->ParentTRSHandle, InSequencer);
-	}
-	IMovieSceneConstrainedSection* ConstrainedSection = Cast<IMovieSceneConstrainedSection>(ReturnValue.ConstraintSection);
-	if (ConstrainedSection == nullptr )
-	{
-		return ReturnValue;
-	}
-
-	ReturnValue.ActiveChannel = ConstrainedSection->GetConstraintChannel(InConstraint->GetFName());
-	return ReturnValue;
-}
-
- void FMovieSceneConstraintChannelHelper::GetTransformFramesForConstraintHandles(
-	const UTickableTransformConstraint* InConstraint,
-	const TSharedPtr<ISequencer>& InSequencer,
-	const FFrameNumber& StartFrame,
-	const FFrameNumber& EndFrame,
-	TArray<FFrameNumber>& OutFramesToBake)
-{
-	if (!InConstraint && !InConstraint->ChildTRSHandle && !InConstraint->ParentTRSHandle)
-	{
-		return;
-	}
-
-	FConstraintSections ConstraintSections = FMovieSceneConstraintChannelHelper::GetConstraintSectionAndChannel(
-		InConstraint, InSequencer);
-	if (ConstraintSections.ChildTransformSection)
-	{
-		const TArrayView<FMovieSceneFloatChannel*> FloatTransformChannels = InConstraint->ChildTRSHandle->GetFloatChannels(ConstraintSections.ChildTransformSection);
-		TArray<FFrameNumber> TransformFrameTimes = FMovieSceneConstraintChannelHelper::GetTransformTimes(
-			FloatTransformChannels, StartFrame, EndFrame);
-		//add transforms keys to bake
-		{
-			for (FFrameNumber& Frame : TransformFrameTimes)
-			{
-				OutFramesToBake.Add(Frame);
-			}
-		}
-		const TArrayView<FMovieSceneDoubleChannel*> DoubleTransformChannels = InConstraint->ChildTRSHandle->GetDoubleChannels(ConstraintSections.ChildTransformSection);
-		TransformFrameTimes = FMovieSceneConstraintChannelHelper::GetTransformTimes(
-			DoubleTransformChannels, StartFrame, EndFrame);
-		//add transforms keys to bake
-		{
-			for (FFrameNumber& Frame : TransformFrameTimes)
-			{
-				OutFramesToBake.Add(Frame);
-			}
-		}
-	}
-	if (ConstraintSections.ParentTransformSection)
-	{
-		const TArrayView<FMovieSceneFloatChannel*> FloatTransformChannels = InConstraint->ParentTRSHandle->GetFloatChannels(ConstraintSections.ParentTransformSection);
-		TArray<FFrameNumber> TransformFrameTimes = FMovieSceneConstraintChannelHelper::GetTransformTimes(
-			FloatTransformChannels, StartFrame, EndFrame);
-		//add transforms keys to bake
-		{
-			for (FFrameNumber& Frame : TransformFrameTimes)
-			{
-				OutFramesToBake.Add(Frame);
-			}
-		}
-		const TArrayView<FMovieSceneDoubleChannel*> DoubleTransformChannels = InConstraint->ParentTRSHandle->GetDoubleChannels(ConstraintSections.ParentTransformSection);
-		TransformFrameTimes = FMovieSceneConstraintChannelHelper::GetTransformTimes(
-			DoubleTransformChannels, StartFrame, EndFrame);
-		//add transforms keys to bake
-		{
-			for (FFrameNumber& Frame : TransformFrameTimes)
-			{
-				OutFramesToBake.Add(Frame);
-			}
-		}
-	}
-}
 
 ITransformConstraintChannelInterface* FMovieSceneConstraintChannelHelper::GetHandleInterface(const UTransformableHandle* InHandle)
 {
@@ -1113,7 +1014,6 @@ ITransformConstraintChannelInterface* FMovieSceneConstraintChannelHelper::GetHan
 	return InterfaceRegistry.FindConstraintChannelInterface(InHandle->GetClass());
 }
 
-
 void FMovieSceneConstraintChannelHelper::CreateBindingIDForHandle(const TSharedPtr<ISequencer>& InSequencer, UTransformableHandle* InHandle)
 {
 	if (InHandle == nullptr || InSequencer.IsValid() == false)
@@ -1244,4 +1144,4 @@ void FMovieSceneConstraintChannelHelper::CompensateScale(
 
 	// reset scaling to reference value
 	InParentConstraint->SetScaling(bRefScalingValue);
-}
+}
\ No newline at end of file
diff --git a/Engine/Source/Editor/MovieSceneTools/Private/MovieSceneToolHelpers.cpp b/Engine/Source/Editor/MovieSceneTools/Private/MovieSceneToolHelpers.cpp
--- a/Engine/Source/Editor/MovieSceneTools/Private/MovieSceneToolHelpers.cpp
+++ b/Engine/Source/Editor/MovieSceneTools/Private/MovieSceneToolHelpers.cpp
@@ -77,8 +77,6 @@
 #include "FrameNumberDetailsCustomization.h"
 #include "PropertyEditorDelegates.h"
 #include "INodeAndChannelMappings.h"
-#include "BakingAnimationKeySettings.h"
-#include "Channels/MovieSceneChannelTraits.h"
 #include "Channels/MovieSceneIntegerChannel.h"
 #include "Channels/MovieSceneByteChannel.h"
 #include "Framework/MultiBox/MultiBoxBuilder.h"
@@ -4616,63 +4614,41 @@ void MovieSceneToolHelpers::SetOrAddKey(TMovieSceneChannelData<FMovieSceneFloatV
 	}
 }
 
-void MovieSceneToolHelpers::SetOrAddKey(TMovieSceneChannelData<FMovieSceneFloatValue>& ChannelData, FFrameNumber Time, const FMovieSceneFloatValue& Value)
-{
-	int32 ExistingIndex = ChannelData.FindKey(Time);
-	if (ExistingIndex != INDEX_NONE)
-	{
-		FMovieSceneFloatValue& FloatValue = ChannelData.GetValues()[ExistingIndex]; //-V758
-		FloatValue.Value = Value.Value;
-	}
-	else
-	{
-		ChannelData.AddKey(Time, Value);
-	}
-}
-
 void MovieSceneToolHelpers::SetOrAddKey(TMovieSceneChannelData<FMovieSceneFloatValue>& ChannelData, FFrameNumber Time, float Value, 
 		float ArriveTangent, float LeaveTangent, ERichCurveInterpMode InterpMode, ERichCurveTangentMode TangentMode, 
 		FFrameRate FrameRate, ERichCurveTangentWeightMode WeightedMode, float ArriveTangentWeight, float LeaveTangentWeight)
 {
-	FMovieSceneFloatValue NewKey(Value);
-
-	NewKey.InterpMode = InterpMode;
-	NewKey.TangentMode = TangentMode;
-	NewKey.Tangent.ArriveTangent = ArriveTangent / FrameRate.AsDecimal();
-	NewKey.Tangent.LeaveTangent = LeaveTangent / FrameRate.AsDecimal();
-	NewKey.Tangent.TangentWeightMode = WeightedMode;
-	NewKey.Tangent.ArriveTangentWeight = ArriveTangentWeight;
-	NewKey.Tangent.LeaveTangentWeight = LeaveTangentWeight;
-	ChannelData.AddKey(Time, NewKey);
-	MovieSceneToolHelpers::SetOrAddKey(ChannelData, Time, NewKey);
+	if (ChannelData.FindKey(Time) == INDEX_NONE)
+	{
+		FMovieSceneFloatValue NewKey(Value);
+
+		NewKey.InterpMode = InterpMode;
+		NewKey.TangentMode = TangentMode;
+		NewKey.Tangent.ArriveTangent = ArriveTangent / FrameRate.AsDecimal();
+		NewKey.Tangent.LeaveTangent = LeaveTangent / FrameRate.AsDecimal();
+		NewKey.Tangent.TangentWeightMode = WeightedMode;
+		NewKey.Tangent.ArriveTangentWeight = ArriveTangentWeight;
+		NewKey.Tangent.LeaveTangentWeight = LeaveTangentWeight;
+		ChannelData.AddKey(Time, NewKey);
+	}
 }
 
 void MovieSceneToolHelpers::SetOrAddKey(TMovieSceneChannelData<FMovieSceneDoubleValue>& ChannelData, FFrameNumber Time, double Value, 
 		float ArriveTangent, float LeaveTangent, ERichCurveInterpMode InterpMode, ERichCurveTangentMode TangentMode, 
 		FFrameRate FrameRate, ERichCurveTangentWeightMode WeightedMode, float ArriveTangentWeight, float LeaveTangentWeight)
 {
-	FMovieSceneDoubleValue NewKey(Value);
-	NewKey.InterpMode = InterpMode;
-	NewKey.TangentMode = TangentMode;
-	NewKey.Tangent.ArriveTangent = ArriveTangent / FrameRate.AsDecimal();
-	NewKey.Tangent.LeaveTangent = LeaveTangent / FrameRate.AsDecimal();
-	NewKey.Tangent.TangentWeightMode = WeightedMode;
-	NewKey.Tangent.ArriveTangentWeight = ArriveTangentWeight;
-	NewKey.Tangent.LeaveTangentWeight = LeaveTangentWeight;
-	MovieSceneToolHelpers::SetOrAddKey(ChannelData, Time, NewKey);
-}
-
-void MovieSceneToolHelpers::SetOrAddKey(TMovieSceneChannelData<FMovieSceneDoubleValue>& ChannelData, FFrameNumber Time, FMovieSceneDoubleValue Value)
-{
-	int32 ExistingIndex = ChannelData.FindKey(Time);
-	if (ExistingIndex != INDEX_NONE)
-	{
-		FMovieSceneDoubleValue& DoubleValue = ChannelData.GetValues()[ExistingIndex]; //-V758
-		DoubleValue.Value = Value.Value;
-	}
-	else
+	if (ChannelData.FindKey(Time) == INDEX_NONE)
 	{
-		ChannelData.AddKey(Time, Value);
+		FMovieSceneDoubleValue NewKey(Value);
+
+		NewKey.InterpMode = InterpMode;
+		NewKey.TangentMode = TangentMode;
+		NewKey.Tangent.ArriveTangent = ArriveTangent / FrameRate.AsDecimal();
+		NewKey.Tangent.LeaveTangent = LeaveTangent / FrameRate.AsDecimal();
+		NewKey.Tangent.TangentWeightMode = WeightedMode;
+		NewKey.Tangent.ArriveTangentWeight = ArriveTangentWeight;
+		NewKey.Tangent.LeaveTangentWeight = LeaveTangentWeight;
+		ChannelData.AddKey(Time, NewKey);
 	}
 }
 
@@ -4740,7 +4716,6 @@ void MovieSceneToolHelpers::CalculateFramesBetween(
 	const UMovieScene* MovieScene,
 	FFrameNumber StartFrame,
 	FFrameNumber EndFrame,
-	int32 FrameInc,
 	TArray<FFrameNumber>& OutFrames)
 {
 	const bool bReverse = StartFrame > EndFrame;
@@ -4764,319 +4739,20 @@ void MovieSceneToolHelpers::CalculateFramesBetween(
 
 	if (bReverse)
 	{
-		int32 FrameIndex = 0;
 		for (FFrameNumber DisplayFrameNumber = EndTimeInDisplay; DisplayFrameNumber >= StartTimeInDisplay; --DisplayFrameNumber)
 		{
-			if (FrameIndex % FrameInc == 0)
-			{
-				FFrameNumber TickFrameNumber = FFrameRate::TransformTime(FFrameTime(DisplayFrameNumber), DisplayResolution, TickResolution).FrameNumber;
-				OutFrames.Add(TickFrameNumber);
-			}
-			++FrameIndex;
+			FFrameNumber TickFrameNumber = FFrameRate::TransformTime(FFrameTime(DisplayFrameNumber), DisplayResolution, TickResolution).FrameNumber;
+			OutFrames.Add(TickFrameNumber);
 		}
 	}
 	else
 	{
 		for (FFrameNumber DisplayFrameNumber = StartTimeInDisplay; DisplayFrameNumber <= EndTimeInDisplay; ++DisplayFrameNumber)
 		{
-			int32 FrameIndex = 0;
-			if (FrameIndex % FrameInc == 0)
-			{
-				FFrameNumber TickFrameNumber = FFrameRate::TransformTime(FFrameTime(DisplayFrameNumber), DisplayResolution, TickResolution).FrameNumber;
-				OutFrames.Add(TickFrameNumber);
-			}
-			++FrameIndex;
-		}
-	}
-}
-
-UENUM(BlueprintType)
-enum class FChannelMergeAlgorithm : uint8
-{
-	/**Average values together*/
-	Average,
-	/**Add values together*/
-	Add
-};
-
-//merge the channels onto the first one using the passed in algorithm, we may skip the last channel also (since it may be weight).
-template<typename ChannelType>
-static bool MergeChannels(TArray<ChannelType*>& Channels,const  TArray<UMovieSceneSection*>& Sections, const TRange<FFrameNumber>& Range,
-	FChannelMergeAlgorithm MergeAlgorithm)
-{
-	if (Channels.Num() < 2 && Channels.Num() != Sections.Num())
-	{
-		return false;
-	}
-	using ChannelValueType = typename ChannelType::ChannelValueType;
-	using CurveValueType = typename ChannelType::CurveValueType;
-
-	//base channel we set values on
-	ChannelType* BaseChannel = Channels[0];
-	TMovieSceneChannelData<ChannelValueType> BaseChannelData = BaseChannel->GetData();
-	//iterate over each key
-	TArray<FFrameNumber> KeyTimes;
-	TArray<FKeyHandle> Handles;
-	//cached set that we set at the end
-	TArray<TPair< FFrameNumber, ChannelValueType>> KeysToSet;
-	for (int32 ChannelIndex = 0; ChannelIndex < Channels.Num(); ++ChannelIndex)
-	{
-		KeyTimes.Reset();
-		Handles.Reset();
-		ChannelType* Channel = Channels[ChannelIndex];
-		Channel->GetKeys(Range, &KeyTimes, &Handles);
-		double DNumChannels = (double)(Channels.Num());
-		for (int32 FrameIndex = 0; FrameIndex < KeyTimes.Num(); ++FrameIndex)
-		{
-			const FFrameNumber& Frame = KeyTimes[FrameIndex];
-			const FFrameTime FrameTime(Frame);
-			int32 KeyIndex = Channel->GetData().GetIndex(Handles[FrameIndex]);
-			ChannelValueType Value = Channel->GetData().GetValues()[KeyIndex];
-			//got value with tangents and times, now we perform the operation
-			Value.Value = 0.0; //zero out the value we calculate it 
-			if (MergeAlgorithm == FChannelMergeAlgorithm::Average)
-			{
-				for (int32 WeightIndex = 0; WeightIndex < Sections.Num(); ++WeightIndex)
-				{
-					float Weight = Sections[WeightIndex]->GetTotalWeightValue(FrameTime);
-					CurveValueType WeightedValue = 0.0;
-					ChannelType* EachChannel = Channels[WeightIndex];
-					EachChannel->Evaluate(FrameTime, WeightedValue);
-					WeightedValue *= ((double)Weight / DNumChannels);
-					Value.Value += WeightedValue;
-				}
-			}
-			else if (MergeAlgorithm == FChannelMergeAlgorithm::Add)
-			{
-				for (int32 WeightIndex = 0; WeightIndex < Sections.Num(); ++WeightIndex)
-				{
-					float Weight = Sections[WeightIndex]->GetTotalWeightValue(FrameTime);
-					CurveValueType WeightedValue = 0.0;
-					ChannelType* EachChannel = Channels[WeightIndex];
-					EachChannel->Evaluate(FrameTime, WeightedValue);
-					WeightedValue *= Weight;
-					Value.Value += WeightedValue;
-				}
-			}
-			KeysToSet.Add(TPair<FFrameNumber, ChannelValueType>(Frame, Value));
-		}
-	}
-	for (TPair<FFrameNumber, ChannelValueType>& KeyToSet : KeysToSet)
-	{
-		MovieSceneToolHelpers::SetOrAddKey(BaseChannelData, KeyToSet.Key, KeyToSet.Value);
-	}
-
-	return true;
-}
-
-bool MovieSceneToolHelpers::OptimizeSection(const FKeyDataOptimizationParams& InParams, UMovieSceneSection* InSection)
-{
-	TArrayView<FMovieSceneFloatChannel*> FloatChannels = InSection->GetChannelProxy().GetChannels<FMovieSceneFloatChannel>();
-	for (FMovieSceneFloatChannel* Channel : FloatChannels)
-	{
-		Channel->Optimize(InParams);
-	}
-	TArrayView<FMovieSceneDoubleChannel*> DoubleChannels = InSection->GetChannelProxy().GetChannels<FMovieSceneDoubleChannel>();
-	for (FMovieSceneDoubleChannel* Channel : DoubleChannels)
-	{
-		Channel->Optimize(InParams);
-	}
-	return true;
-}
-
-template<typename ChannelType>
-static bool MergeSections(UMovieSceneSection* BaseSection, TArrayView<ChannelType*> BaseDoubleChannels, TArray<UMovieSceneSection*>& AbsoluteSections, TArray<UMovieSceneSection*>& AdditiveSections,
-	const TRange<FFrameNumber>& Range,bool bSkipLastChannel)
-{
-	TArrayView<ChannelType*> BaseChannels = BaseSection->GetChannelProxy().GetChannels<ChannelType>();
-	if (BaseChannels.Num() > 0)
-	{
-		int32 SectionIndex = 0;
-		//sanity check to make sure channels are the same size
-		for (SectionIndex = 0; SectionIndex < AbsoluteSections.Num(); ++SectionIndex)
-		{
-			TArrayView<ChannelType*>Channels = AbsoluteSections[SectionIndex]->GetChannelProxy().GetChannels<ChannelType>();
-			if (Channels.Num() != BaseChannels.Num())
-			{
-				UE_LOG(LogMovieScene, Warning, TEXT("MergeSections:: Invalid number of channels"));
-				return false;
-			}
-		}
-		for (SectionIndex = 0; SectionIndex < AdditiveSections.Num(); ++SectionIndex)
-		{
-			TArrayView<ChannelType*>Channels = AdditiveSections[SectionIndex]->GetChannelProxy().GetChannels<ChannelType>();
-			if (Channels.Num() != BaseChannels.Num())
-			{
-				UE_LOG(LogMovieScene, Warning, TEXT("MergeSections:: Invalid number of channels"));
-				return false;
-			}
-		}
-		BaseSection->Modify();
-
-		TArray<ChannelType*> Channels;
-		int32 ChannelIndex = 0;
-		bSkipLastChannel = false;
-		int32 BaseChannelsNum = bSkipLastChannel ? BaseChannels.Num() - 1 : BaseChannels.Num();
-		for (ChannelIndex = 0; ChannelIndex < BaseChannels.Num(); ++ChannelIndex)
-		{
-			if (AbsoluteSections.Num() > 0)
-			{
-				Channels.Reset();
-				for (SectionIndex = 0; SectionIndex < AbsoluteSections.Num(); ++SectionIndex)
-				{
-					TArrayView<ChannelType*>OurChannels = AbsoluteSections[SectionIndex]->GetChannelProxy().GetChannels<ChannelType>();
-					Channels.Add(OurChannels[ChannelIndex]);
-				}
-				//now blend them
-				if (MergeChannels(Channels, AbsoluteSections, Range,
-					FChannelMergeAlgorithm::Average) == false)
-				{
-					UE_LOG(LogMovieScene, Warning, TEXT("MergeSections:: Could not merge channels"));
-					return false;
-				}
-			}
-			if (AdditiveSections.Num() > 0)
-			{
-				//now do additives
-				Channels.Reset();
-				for (SectionIndex = 0; SectionIndex < AdditiveSections.Num(); ++SectionIndex)
-				{
-					TArrayView<ChannelType*>OurChannels = AdditiveSections[SectionIndex]->GetChannelProxy().GetChannels<ChannelType>();
-					Channels.Add(OurChannels[ChannelIndex]);
-				}
-				//now blend them
-				if (MergeChannels(Channels, AdditiveSections, Range,
-					FChannelMergeAlgorithm::Add) == false)
-				{
-					UE_LOG(LogMovieScene, Warning, TEXT("MergeSections:: Could not merge channels"));
-					return false;
-				}
-			}
-		}
-		for (ChannelType* Channel : BaseChannels)
-		{
-			Channel->AutoSetTangents();
-		}
-	}
-	else
-	{
-		return false;
-	}
-	return true;
-}
-
-bool MovieSceneToolHelpers::CollapseSection(TSharedPtr<ISequencer>& SequencerPtr, UMovieSceneTrack* OwnerTrack, TArray<UMovieSceneSection*> Sections,
-	const FBakingAnimationKeySettings& InSettings)
-{
-	if (SequencerPtr.IsValid() && Sections.Num() > 1)
-	{
-		TRange<FFrameNumber> Range(InSettings.StartFrame, InSettings.EndFrame);
-		//we get the  first absolute section or if no absolute sections, first additive,
-		// that's the one we collapse onto
-		UMovieSceneSection* BaseSection = nullptr;
-		//get first absolute
-		for (UMovieSceneSection* Section: Sections)
-		{
-			if (Section->IsActive() == false)
-			{
-				continue;
-			}
-			if (Section->GetBlendType().Get() == EMovieSceneBlendType::Absolute)
-			{
-				BaseSection = Section;
-				break;
-			}
-			else if (Section->GetBlendType().Get() == EMovieSceneBlendType::Additive)
-			{
-				if (BaseSection == nullptr)
-				{
-					BaseSection = Section; //don't stop though may have an absolute later
-				}
-			}
-		}
-		if (BaseSection == nullptr)
-		{
-			UE_LOG(LogMovieScene, Warning, TEXT("CollapseSection:: Invalid section(s) to collapse"));
-			return false;
-		}
-		//now find other sections.
-		TArray<UMovieSceneSection*> AbsoluteSections;
-		TArray<UMovieSceneSection*> AdditiveSections;
-		for (UMovieSceneSection* Section : Sections)
-		{
-			if (Section != BaseSection && Section->IsActive())
-			{
-				if (Section->GetBlendType().Get() == EMovieSceneBlendType::Absolute)
-				{
-					AbsoluteSections.Add(Section);
-				}
-				else if (Section->GetBlendType().Get() == EMovieSceneBlendType::Additive)
-				{
-					AdditiveSections.Add(Section);
-				}
-			}
-		}
-		//now make sure we have sections to blend with and if so add base section to them
-		if (AbsoluteSections.Num() > 0 || AdditiveSections.Num() > 0)
-		{
-			if (AbsoluteSections.Num() > 0)
-			{
-				AbsoluteSections.Insert(BaseSection, 0);
-			}
-			if (AdditiveSections.Num() > 0)
-			{
-				AdditiveSections.Insert(BaseSection, 0);
-			}
-		}
-		else
-		{
-			UE_LOG(LogMovieScene, Warning, TEXT("CollapseSection:: Invalid section(s) to collapse"));
-			return false;
-		}
-
-		FScopedTransaction Transaction(NSLOCTEXT("MovieSceneTools", "CollapseAllSections", "Collapse All Sections"));
-
-		TArrayView<FMovieSceneFloatChannel*> BaseFloatChannels = BaseSection->GetChannelProxy().GetChannels<FMovieSceneFloatChannel>();
-		TArrayView<FMovieSceneDoubleChannel*> BaseDoubleChannels = BaseSection->GetChannelProxy().GetChannels<FMovieSceneDoubleChannel>();
-		if (BaseDoubleChannels.Num() > 0) //transforms
-		{
-			if (MergeSections(BaseSection, BaseDoubleChannels, AbsoluteSections, AdditiveSections, Range, false /*bSkipLastChannel*/) == false)
-			{
-				Transaction.Cancel();
-				return false;
-			}
-		}
-		else if(BaseFloatChannels.Num() > 0) //control rig
-		{
-			//skip weight channel for control rig
-			if (MergeSections(BaseSection, BaseFloatChannels, AbsoluteSections, AdditiveSections, Range, true /*bSkipLastChannel*/) == false)
-			{
-				Transaction.Cancel();
-				return false;
-			}
-		}
-		//delete other sections
-		OwnerTrack->Modify();
-		for (int32 SectionIndex = Sections.Num() - 1; SectionIndex >= 0; --SectionIndex)
-		{
-			if (Sections[SectionIndex] != BaseSection)
-			{
-				TArray<UMovieSceneSection*> AllSections = OwnerTrack->GetAllSections();
-				int32 TrackSectionIndex = INDEX_NONE;
-				if (AllSections.Find(Sections[SectionIndex], TrackSectionIndex))
-				{
-					if (TrackSectionIndex != INDEX_NONE)
-					{
-						OwnerTrack->RemoveSectionAt(TrackSectionIndex);
-					}
-				}
-			}
+			FFrameNumber TickFrameNumber = FFrameRate::TransformTime(FFrameTime(DisplayFrameNumber), DisplayResolution, TickResolution).FrameNumber;
+			OutFrames.Add(TickFrameNumber);
 		}
-		//reset everything back
-		SequencerPtr->NotifyMovieSceneDataChanged(EMovieSceneDataChangeType::MovieSceneStructureItemAdded);
 	}
-	return true;
-
 }
 
 UMovieScene3DTransformSection* MovieSceneToolHelpers::GetTransformSection(
diff --git a/Engine/Source/Editor/MovieSceneTools/Public/BakingAnimationKeySettings.h b/Engine/Source/Editor/MovieSceneTools/Public/BakingAnimationKeySettings.h
deleted file mode 100644
diff --git a/Engine/Source/Editor/MovieSceneTools/Public/Constraints/MovieSceneConstraintChannelHelper.h b/Engine/Source/Editor/MovieSceneTools/Public/Constraints/MovieSceneConstraintChannelHelper.h
--- a/Engine/Source/Editor/MovieSceneTools/Public/Constraints/MovieSceneConstraintChannelHelper.h
+++ b/Engine/Source/Editor/MovieSceneTools/Public/Constraints/MovieSceneConstraintChannelHelper.h
@@ -16,7 +16,6 @@ class UTickableParentConstraint;
 class UTransformableHandle;
 struct FMovieSceneConstraintChannel;
 struct FFrameNumber;
-struct FConstraintAndActiveChannel;
 
 
 struct MOVIESCENETOOLS_API FCompensationEvaluator
@@ -42,14 +41,6 @@ private:
 	UTransformableHandle* Handle = nullptr;
 };
 
-struct MOVIESCENETOOLS_API FConstraintSections
-{
-	UMovieSceneSection* ConstraintSection = nullptr;
-	UMovieSceneSection* ChildTransformSection = nullptr;
-	UMovieSceneSection* ParentTransformSection = nullptr;
-	FConstraintAndActiveChannel* ActiveChannel = nullptr;
-	ITransformConstraintChannelInterface* Interface = nullptr;
-};
 struct MOVIESCENETOOLS_API FMovieSceneConstraintChannelHelper
 {
 public:
@@ -93,20 +84,6 @@ public:
 		UMovieSceneSection* InSection,
 		const FFrameNumber& InCurrentFrame, const FFrameNumber& InNextFrame);
 
-	/* Get the section and the channel for the given constraint, will be nullptr's if it doesn't exist in Sequencer*/
-	static FConstraintSections  GetConstraintSectionAndChannel(
-		const UTickableTransformConstraint* InConstraint,
-		const TSharedPtr<ISequencer>& InSequencer);
-
-	/* For the given constraint get all of the transform keys for it's child and parent handles*/
-	static void GetTransformFramesForConstraintHandles(
-		const UTickableTransformConstraint* InConstraint,
-		const TSharedPtr<ISequencer>& InSequencer,
-		const FFrameNumber& StartFrame,
-		const FFrameNumber& EndFrame,
-		TArray<FFrameNumber>& OutFramesToBake);
-
-
 	/** @todo documentation. */
 	template<typename ChannelType>
 	static void GetFramesToCompensate(
@@ -151,25 +128,6 @@ public:
 		const TSharedPtr<ISequencer>& InSequencer,
 		UMovieSceneSection* InSection);
 
-	template< typename ChannelType >
-	static TArray<FFrameNumber> GetTransformTimes(
-		const TArrayView<ChannelType*>& InChannels,
-		const FFrameNumber& StartTime, 
-		const FFrameNumber& EndTime);
-
-	template< typename ChannelType >
-	static  void DeleteTransformTimes(
-		const TArrayView<ChannelType*>& InChannels,
-		const FFrameNumber& StartTime,
-		const FFrameNumber& EndTime);
-
-	/** this will only set the value son channels with keys at the specified time, reusing tangent time etc. */
-	template< typename ChannelType >
-	static void SetTransformTimes(
-		const TArrayView<ChannelType*>& InChannels,
-		const TArray<FFrameNumber>& Frames,
-		const TArray<FTransform>& Transforms);
-
 	static bool bDoNotCompensate;
 
 private:
diff --git a/Engine/Source/Editor/MovieSceneTools/Public/Constraints/MovieSceneConstraintChannelHelper.inl b/Engine/Source/Editor/MovieSceneTools/Public/Constraints/MovieSceneConstraintChannelHelper.inl
--- a/Engine/Source/Editor/MovieSceneTools/Public/Constraints/MovieSceneConstraintChannelHelper.inl
+++ b/Engine/Source/Editor/MovieSceneTools/Public/Constraints/MovieSceneConstraintChannelHelper.inl
@@ -9,7 +9,6 @@
 #include "Channels/MovieSceneChannelProxy.h"
 #include "MovieSceneSection.h"
 #include "Algo/Unique.h"
-#include "Containers/SortedMap.h"
 
 template<typename ChannelType>
 void FMovieSceneConstraintChannelHelper::GetFramesToCompensate(
@@ -262,82 +261,6 @@ void FMovieSceneConstraintChannelHelper::DeleteTransformKeys(
 	}
 }
 
-template< typename ChannelType >
-TArray<FFrameNumber> FMovieSceneConstraintChannelHelper::GetTransformTimes(
-	const TArrayView<ChannelType*>& InChannels,
-	const FFrameNumber& StartTime,
-	const FFrameNumber& EndTime)
-{
-	TSortedMap<FFrameNumber, FFrameNumber> FrameSet;
-	TRange<FFrameNumber> WithinRange(0, 0);
-	WithinRange.SetLowerBoundValue(StartTime);
-	WithinRange.SetUpperBoundValue(EndTime);
-	TArray<FFrameNumber> KeyTimes;
-	TArray<FKeyHandle> KeyHandles;
-	for (ChannelType* Channel : InChannels)
-	{
-		TMovieSceneChannelData<typename ChannelType::ChannelValueType> Data = Channel->GetData();
-		KeyTimes.SetNum(0);
-		KeyHandles.SetNum(0);
-		Data.GetKeys(WithinRange, &KeyTimes,&KeyHandles);
-		for (const FFrameNumber& FrameNumber : KeyTimes)
-		{
-			FrameSet.Add(FrameNumber,FrameNumber);
-		}
-	}
-	TArray<FFrameNumber> Frames;
-	FrameSet.GenerateKeyArray(Frames);
-	return Frames;
-}
-
-template< typename ChannelType >
- void FMovieSceneConstraintChannelHelper::DeleteTransformTimes(
-	 const TArrayView<ChannelType*>& InChannels,
-	 const FFrameNumber& StartTime,
-	 const FFrameNumber& EndTime)
-{
-	 TRange<FFrameNumber> WithinRange(0, 0);
-	 WithinRange.SetLowerBoundValue(StartTime);
-	 WithinRange.SetUpperBoundValue(EndTime);
-	 TArray<FFrameNumber> KeyTimes;
-	 TArray<FKeyHandle> KeyHandles;
-	 for (ChannelType* Channel : InChannels)
-	 {
-		 TMovieSceneChannelData<typename ChannelType::ChannelValueType> Data = Channel->GetData();
-		 KeyTimes.SetNum(0);
-		 KeyHandles.SetNum(0);
-		 Data.GetKeys(WithinRange, &KeyTimes, &KeyHandles);
-		 Data.DeleteKeys(KeyHandles);
-	 }
-}
-
-/** this will only set the values on channels with keys at the specified time, reusing tangents  etc. */
-template< typename ChannelType >
-void FMovieSceneConstraintChannelHelper::SetTransformTimes(
-	const TArrayView<ChannelType*>& InChannels,
-	const TArray<FFrameNumber>& Frames,
-	const TArray<FTransform>& Transforms)
-{
-	if (Frames.Num() != Transforms.Num())
-	{
-		return;
-	}
-	for (ChannelType* Channel : InChannels)
-	{
-		TMovieSceneChannelData<typename ChannelType::ChannelValueType> Data = Channel->GetData();
-		const TArrayView<const FFrameNumber> Times = Data.GetTimes();
-		for (int32 Index = 0; Index < Frames.Num(); ++Index)
-		{
-			FFrameNumber Time = Frames[0];
-			const int32 KeyIndex = Algo::LowerBound(Times, Time);
-			if (Times.IsValidIndex(KeyIndex) && Times[KeyIndex] == Time)
-			{
-				Data.GetValues()[KeyIndex] = 1.0;//todo 
-			}
-		}
-	}
-}
-
 template<typename ChannelType>
 void EvaluateTangentAtThisTime(int32 ChannelIndex, int32 NumChannels,
 	UMovieSceneSection* Section,
diff --git a/Engine/Source/Editor/MovieSceneTools/Public/MovieSceneToolHelpers.h b/Engine/Source/Editor/MovieSceneTools/Public/MovieSceneToolHelpers.h
--- a/Engine/Source/Editor/MovieSceneTools/Public/MovieSceneToolHelpers.h
+++ b/Engine/Source/Editor/MovieSceneTools/Public/MovieSceneToolHelpers.h
@@ -39,8 +39,6 @@ class ACameraActor;
 struct FActorForWorldTransforms;
 class UMovieScene3DTransformSection;
 class UMovieSceneSubTrack;
-struct FBakingAnimationKeySettings;
-struct FKeyDataOptimizationParams;
 class UMovieSceneSubSection;
 enum class EMovieSceneTransformChannel : uint32;
 
@@ -548,16 +546,6 @@ public:
 	*/
 	static void SetOrAddKey(TMovieSceneChannelData<FMovieSceneDoubleValue>& ChannelData, FFrameNumber Time, double Value);
 
-	/*
-	* Set or add a key onto a float channel based on key value.
-	*/
-	static void SetOrAddKey(TMovieSceneChannelData<FMovieSceneFloatValue>& Curve, FFrameNumber Time, const FMovieSceneFloatValue& Value);
-
-	/*
-	* Set or add a key onto a double channel based on key value.
-	*/
-	static void SetOrAddKey(TMovieSceneChannelData<FMovieSceneDoubleValue>& ChannelData, FFrameNumber Time, FMovieSceneDoubleValue Value);
-
 	/* 
 	* Set or add a key onto a float channel based on rich curve data.
 	*/
@@ -590,17 +578,11 @@ public:
 	 */
 	static bool IsValidAsset(UMovieSceneSequence* Sequence, const FAssetData& InAssetData);
 
-	static bool CollapseSection(TSharedPtr<ISequencer>& SequencerPtr, UMovieSceneTrack* OwnerTrack, TArray<UMovieSceneSection*> Sections,
-		const FBakingAnimationKeySettings& InSettings);
-
-	static bool OptimizeSection(const FKeyDataOptimizationParams& InParams, UMovieSceneSection* InSection);
-
 	/** Returns the frame numbers between start and end. */
 	static void CalculateFramesBetween(
 		const UMovieScene* MovieScene,
 		FFrameNumber StartFrame,
 		FFrameNumber EndFrame,
-		int FrameIncrement,
 		TArray<FFrameNumber>& OutFrames);
 
 	/** Returns the transform section for that guid. */
diff --git a/Engine/Source/Runtime/Experimental/Animation/Constraints/Private/TransformableHandle.cpp b/Engine/Source/Runtime/Experimental/Animation/Constraints/Private/TransformableHandle.cpp
--- a/Engine/Source/Runtime/Experimental/Animation/Constraints/Private/TransformableHandle.cpp
+++ b/Engine/Source/Runtime/Experimental/Animation/Constraints/Private/TransformableHandle.cpp
@@ -10,6 +10,7 @@
 #include "Channels/MovieSceneFloatChannel.h"
 #include "Channels/MovieSceneDoubleChannel.h"
 #include "Channels/MovieSceneChannelProxy.h"
+#include "Channels/MovieSceneDoubleChannel.h"
 #include "Components/SkeletalMeshComponent.h"
 
 #include UE_INLINE_GENERATED_CPP_BY_NAME(TransformableHandle)
