Control Rig: Re-add URigHierarchy::GetNameHash that was lost in CL 24895265

#rnx

diff --git a/Engine/Plugins/Animation/ControlRig/Source/ControlRig/Private/Rigs/RigHierarchy.cpp b/Engine/Plugins/Animation/ControlRig/Source/ControlRig/Private/Rigs/RigHierarchy.cpp
--- a/Engine/Plugins/Animation/ControlRig/Source/ControlRig/Private/Rigs/RigHierarchy.cpp
+++ b/Engine/Plugins/Animation/ControlRig/Source/ControlRig/Private/Rigs/RigHierarchy.cpp
@@ -453,8 +453,23 @@ void URigHierarchy::CopyHierarchy(URigHierarchy* InHierarchy)
 	EnsureCacheValidity();
 }
 
+uint32 URigHierarchy::GetNameHash() const
+{
+	FScopeLock Lock(&ElementsLock);
+	
+	uint32 Hash = GetTypeHash(GetTopologyVersion());
+	for (int32 ElementIndex = 0; ElementIndex < Elements.Num(); ElementIndex++)
+	{
+		const FRigBaseElement* Element = Elements[ElementIndex];
+		Hash = HashCombine(Hash, GetTypeHash(Element->GetName()));
+	}
+	return Hash;	
+}
+
 uint32 URigHierarchy::GetTopologyHash(bool bIncludeTopologyVersion, bool bIncludeTransientControls) const
 {
+	FScopeLock Lock(&ElementsLock);
+	
 	uint32 Hash = bIncludeTopologyVersion ? TopologyVersion : 0;
 
 	for (int32 ElementIndex = 0; ElementIndex < Elements.Num(); ElementIndex++)
diff --git a/Engine/Plugins/Animation/ControlRig/Source/ControlRig/Public/Rigs/RigHierarchy.h b/Engine/Plugins/Animation/ControlRig/Source/ControlRig/Public/Rigs/RigHierarchy.h
--- a/Engine/Plugins/Animation/ControlRig/Source/ControlRig/Public/Rigs/RigHierarchy.h
+++ b/Engine/Plugins/Animation/ControlRig/Source/ControlRig/Public/Rigs/RigHierarchy.h
@@ -3611,7 +3611,7 @@ private:
 
 	// Storage for the elements
 	mutable TArray<FRigBaseElement*> Elements;
-	FCriticalSection ElementsLock;
+	mutable FCriticalSection ElementsLock;
 
 	// Storage for the elements
 	mutable TArray<TArray<FRigBaseElement*>> ElementsPerType;
