Control Rig: 5.2 Fix for Broken rig after BP duplication. Fixed function headers in function ref nodes still soft reference the original BP after BP duplication. It was causing rig compilation errors. Now we make sure we replace every single reference to the original BP with that to the duplicated BP.

diff --git a/Engine/Plugins/Animation/ControlRig/Source/ControlRigDeveloper/Private/ControlRigBlueprint.cpp b/Engine/Plugins/Animation/ControlRig/Source/ControlRigDeveloper/Private/ControlRigBlueprint.cpp
--- a/Engine/Plugins/Animation/ControlRig/Source/ControlRigDeveloper/Private/ControlRigBlueprint.cpp
+++ b/Engine/Plugins/Animation/ControlRig/Source/ControlRigDeveloper/Private/ControlRigBlueprint.cpp
@@ -2659,9 +2659,19 @@ void UControlRigBlueprint::PreDuplicate(FObjectDuplicationParameters& DupParams)
 
 void UControlRigBlueprint::PostDuplicate(bool bDuplicateForPIE)
 {
-	RefreshAllModels();
+	// assuming PostDuplicate is always followed by a PostLoad:
+	// so theoretically, PostDuplicate just makes corrections to the serialized data and does nothing more,
+	// while PostLoad looks at whatever is serialized and load it into memory according to the version of the editor used
+	// note: how to know if we have corrected everything?
+	// ans: check the reference viewer for the duplicated BP and make sure that the original BP does not appear in there
 	
-	Super::PostDuplicate(bDuplicateForPIE);
+	{
+		// pause compilation because we need to patch some stuff first
+		TGuardValue<bool> CompilingGuard(bIsCompiling, true);
+		// this will create the new EMPTY generated class to be used as the function store for this BP
+		// it will be filled during PostLoad based on the graph model
+		Super::PostDuplicate(bDuplicateForPIE);
+	}
 	
 	if (URigHierarchyController* Controller = Hierarchy->GetController(true))
 	{
@@ -2680,13 +2690,20 @@ void UControlRigBlueprint::PostDuplicate(bool bDuplicateForPIE)
 			if(URigVMBlueprintGeneratedClass* CRGeneratedClass = GetControlRigBlueprintGeneratedClass())
 			{
 				FRigVMGraphFunctionStore& Store = CRGeneratedClass->GraphFunctionStore;
-				Store.PostDuplicateHost(InOldPath, InNewPath);
+				// technically not needed, the store should be empty, it will not be populated until PostLoad
+				// this code is kept here just in case things change in the future
+				if (!(ensure(Store.PublicFunctions.Num() == 0) && ensure(Store.PrivateFunctions.Num() == 0)))
+				{
+					Store.PostDuplicateHost(InOldPath, InNewPath);
+				}
 			}
 			RigVMClient.PostDuplicateHost(InOldPath, InNewPath);
 		}
 	};
 
 	// update the paths once for the blueprint and once for the generated class
+	// make sure all function headers pointing to things in the old BP are changed to
+	// point to their duplicates in the new BP
 	UpdateFunctionHeaders(PreDuplicateAssetPath.ToString(), GetPathName());
 	if(const URigVMBlueprintGeneratedClass* CRGeneratedClass = GetControlRigBlueprintGeneratedClass())
 	{
@@ -2697,7 +2714,10 @@ void UControlRigBlueprint::PostDuplicate(bool bDuplicateForPIE)
 	PreDuplicateHostPath.Reset();
 
 	FBlueprintEditorUtils::MarkBlueprintAsStructurallyModified(this);
-	RecompileVM();
+
+	// now that the data in the new BP is correct, make sure there is a call to PostLoad()
+	// to complete the loading and get the BP ready, including refresh all models, patch function store, and recompile VM
+	check(HasAnyFlags(RF_NeedPostLoad));
 }
 
 FRigVMGraphModifiedEvent& UControlRigBlueprint::OnModified()
diff --git a/Engine/Plugins/Runtime/RigVM/Source/RigVM/Private/RigVMCore/RigVMGraphFunctionDefinition.cpp b/Engine/Plugins/Runtime/RigVM/Source/RigVM/Private/RigVMCore/RigVMGraphFunctionDefinition.cpp
--- a/Engine/Plugins/Runtime/RigVM/Source/RigVM/Private/RigVMCore/RigVMGraphFunctionDefinition.cpp
+++ b/Engine/Plugins/Runtime/RigVM/Source/RigVM/Private/RigVMCore/RigVMGraphFunctionDefinition.cpp
@@ -82,6 +82,11 @@ void FRigVMGraphFunctionHeader::PostDuplicateHost(const FString& InOldPathName,
 
 	ReplacePathName(LibraryPointer.LibraryNode);
 	ReplacePathName(LibraryPointer.HostObject);
+	for (TPair<FRigVMGraphFunctionIdentifier, uint32>& Pair : Dependencies)
+	{
+		ReplacePathName(Pair.Key.LibraryNode);
+		ReplacePathName(Pair.Key.HostObject);
+	}
 }
 
 bool FRigVMGraphFunctionData::IsMutable() const
