Control Rig: Output log error fix
#trivial

diff --git a/Engine/Plugins/Animation/ControlRig/Source/ControlRigEditor/Public/RigSpacePickerBakeSettings.h b/Engine/Plugins/Animation/ControlRig/Source/ControlRigEditor/Public/RigSpacePickerBakeSettings.h
--- a/Engine/Plugins/Animation/ControlRig/Source/ControlRigEditor/Public/RigSpacePickerBakeSettings.h
+++ b/Engine/Plugins/Animation/ControlRig/Source/ControlRigEditor/Public/RigSpacePickerBakeSettings.h
@@ -17,6 +17,10 @@ struct CONTROLRIGEDITOR_API FRigSpacePickerBakeSettings
 	{
 		TargetSpace = FRigElementKey();
 		Settings = FBakingAnimationKeySettings();
+		StartFrame_DEPRECATED = 0;
+		EndFrame_DEPRECATED = 100;
+		bReduceKeys_DEPRECATED = false;
+		Tolerance_DEPRECATED = 0.001f;
 	}
 
 	UPROPERTY(BlueprintReadWrite, Category = "Settings")
