Control Rig Editor: do not try to open tab if the editor has already been closed

diff --git a/Engine/Plugins/Animation/ControlRig/Source/ControlRigEditor/Private/Editor/ControlRigEditor.cpp b/Engine/Plugins/Animation/ControlRig/Source/ControlRigEditor/Private/Editor/ControlRigEditor.cpp
--- a/Engine/Plugins/Animation/ControlRig/Source/ControlRigEditor/Private/Editor/ControlRigEditor.cpp
+++ b/Engine/Plugins/Animation/ControlRig/Source/ControlRigEditor/Private/Editor/ControlRigEditor.cpp
@@ -575,17 +575,30 @@ void FControlRigEditor::InitControlRigEditor(const EToolkitMode::Type Mode, cons
 			Compile();
 		}
 	}
-
-	FFunctionGraphTask::CreateAndDispatchWhenReady([this]()
+	
+	FFunctionGraphTask::CreateAndDispatchWhenReady([this, InControlRigBlueprint]()
 	{
+		// no need to do anything if the the CR is not opened anymore
+		// (i.e. destructor has been called before that task actually got a chance to start) 
+		if (!UControlRigBlueprint::sCurrentlyOpenedRigBlueprints.Contains(InControlRigBlueprint))
+		{
+			return;		
+		}
+		
+		TSharedPtr<FTabManager> TabManager = GetTabManager();
+		if (!TabManager.IsValid())
+		{
+			return;
+		}
+		
 		// Always show the myblueprint tab
-		FTabId MyBlueprintTabId(FBlueprintEditorTabs::MyBlueprintID);
-		if(!GetTabManager()->FindExistingLiveTab(MyBlueprintTabId).IsValid())
+		static const FTabId MyBlueprintTabId(FBlueprintEditorTabs::MyBlueprintID);
+		if (!TabManager->FindExistingLiveTab(MyBlueprintTabId).IsValid())
 		{
-			GetTabManager()->TryInvokeTab(MyBlueprintTabId);
+			TabManager->TryInvokeTab(MyBlueprintTabId);
 		}
 		
-	}, TStatId(), NULL, ENamedThreads::GameThread);
+	}, TStatId(), nullptr, ENamedThreads::GameThread);
 	
 	CreateRigHierarchyToGraphDragAndDropMenu();
 
