Fixes interpolation bug in spherical pose reader when scaling active region to 0.

diff --git a/Engine/Plugins/Animation/ControlRig/Source/ControlRig/Private/Units/Highlevel/Hierarchy/RigUnit_SphericalPoseReader.cpp b/Engine/Plugins/Animation/ControlRig/Source/ControlRig/Private/Units/Highlevel/Hierarchy/RigUnit_SphericalPoseReader.cpp
--- a/Engine/Plugins/Animation/ControlRig/Source/ControlRig/Private/Units/Highlevel/Hierarchy/RigUnit_SphericalPoseReader.cpp
+++ b/Engine/Plugins/Animation/ControlRig/Source/ControlRig/Private/Units/Highlevel/Hierarchy/RigUnit_SphericalPoseReader.cpp
@@ -203,9 +203,14 @@ void FRigUnit_SphericalPoseReader::DistanceToEllipse(
 	const float SizeY,
 	FEllipseQuery& OutEllipseQuery)
 {
+	// degenerate ellipse
 	if (SizeX <= KINDA_SMALL_NUMBER || SizeY <= KINDA_SMALL_NUMBER)
 	{
-		return; // degenerate ellipse
+		OutEllipseQuery.ClosestX = 0.0f;
+		OutEllipseQuery.ClosestY = 0.0f;
+		OutEllipseQuery.DistSq = FMath::Pow(InX, 2.f) + FMath::Pow(InY, 2.f);
+		OutEllipseQuery.IsInside = OutEllipseQuery.DistSq <= KINDA_SMALL_NUMBER;
+		return;
 	}
 	
     const float PX = FMath::Abs(InX);
diff --git a/Engine/Plugins/Animation/ControlRig/Source/ControlRig/Private/Units/Highlevel/Hierarchy/RigUnit_SphericalPoseReader.h b/Engine/Plugins/Animation/ControlRig/Source/ControlRig/Private/Units/Highlevel/Hierarchy/RigUnit_SphericalPoseReader.h
--- a/Engine/Plugins/Animation/ControlRig/Source/ControlRig/Private/Units/Highlevel/Hierarchy/RigUnit_SphericalPoseReader.h
+++ b/Engine/Plugins/Animation/ControlRig/Source/ControlRig/Private/Units/Highlevel/Hierarchy/RigUnit_SphericalPoseReader.h
@@ -8,10 +8,10 @@
 
 struct FEllipseQuery
 {
-	float ClosestX;
-	float ClosestY;
-	float DistSq;
-	bool IsInside;
+	float ClosestX = 0.0f;
+	float ClosestY = 0.0f;
+	float DistSq = 0.0f;
+	bool IsInside = false;
 };
 
 USTRUCT()
