Fixed static analysis

diff --git a/Engine/Plugins/Animation/ControlRig/Source/ControlRigEditor/Private/Tools/ControlRigSnapper.cpp b/Engine/Plugins/Animation/ControlRig/Source/ControlRigEditor/Private/Tools/ControlRigSnapper.cpp
--- a/Engine/Plugins/Animation/ControlRig/Source/ControlRigEditor/Private/Tools/ControlRigSnapper.cpp
+++ b/Engine/Plugins/Animation/ControlRig/Source/ControlRigEditor/Private/Tools/ControlRigSnapper.cpp
@@ -222,17 +222,20 @@ struct FGuidAndActor
 
 		// fill channels to key
 		EMovieSceneTransformChannel Channels;
-		if (SnapSettings->bSnapPosition)
+		if (SnapSettings)
 		{
-			EnumAddFlags(Channels, EMovieSceneTransformChannel::Translation);
-		}
-		if (SnapSettings->bSnapRotation)
-		{
-			EnumAddFlags(Channels, EMovieSceneTransformChannel::Rotation);
-		}
-		if (SnapSettings->bSnapScale)
-		{
-			EnumAddFlags(Channels, EMovieSceneTransformChannel::Scale);
+			if (SnapSettings->bSnapPosition)
+			{
+				EnumAddFlags(Channels, EMovieSceneTransformChannel::Translation);
+			}
+			if (SnapSettings->bSnapRotation)
+			{
+				EnumAddFlags(Channels, EMovieSceneTransformChannel::Rotation);
+			}
+			if (SnapSettings->bSnapScale)
+			{
+				EnumAddFlags(Channels, EMovieSceneTransformChannel::Scale);
+			}
 		}
 
 		// compute local transforms
@@ -647,7 +650,7 @@ bool FControlRigSnapper::SnapIt(FFrameNumber StartFrame, FFrameNumber EndFrame,c
 							CurrentParentWorldTransform[0] = ControlRigParentWorldTransforms[Index];
 							//this will evaluate at the current frame which we want
 							GetControlRigControlTransforms(Sequencer, ControlRig, Name, OneFrame, CurrentParentWorldTransform, CurrentControlRigTransform);
-							if (SnapSettings->bSnapPosition == false || SnapSettings->bSnapRotation == false || SnapSettings->bSnapScale == false)
+							if (SnapSettings)
 							{
 								FTransform& Transform = WorldTransformToSnap[Index];
 								const FTransform& CurrentTransform = CurrentControlRigTransform[0];
@@ -666,7 +669,6 @@ bool FControlRigSnapper::SnapIt(FFrameNumber StartFrame, FFrameNumber EndFrame,c
 									FVector Scale = CurrentTransform.GetScale3D();
 									Transform.SetScale3D(Scale);
 								}
-								
 							}
 							const FFrameNumber& FrameNumber = Frames[Index];
 							Context.LocalTime = TickResolution.AsSeconds(FFrameTime(FrameNumber));
