ControlRigBP: Disable a content browser menu option that have been allowing the creation of child BP from a Control Rig BP. This is not something that ControlRig supports and child CRBP assets created this way can have unknown problems

diff --git a/Engine/Plugins/Animation/ControlRig/Source/ControlRigEditor/Private/ControlRigBlueprintActions.h b/Engine/Plugins/Animation/ControlRig/Source/ControlRigEditor/Private/ControlRigBlueprintActions.h
--- a/Engine/Plugins/Animation/ControlRig/Source/ControlRigEditor/Private/ControlRigBlueprintActions.h
+++ b/Engine/Plugins/Animation/ControlRig/Source/ControlRigEditor/Private/ControlRigBlueprintActions.h
@@ -27,6 +27,7 @@ public:
 
 	// FAssetTypeActions_Blueprint interface
 	virtual UFactory* GetFactoryForBlueprintType(UBlueprint* InBlueprint) const override;
+	virtual bool CanCreateNewDerivedBlueprint() const override { return false; };
 
 	static void ExtendSketalMeshToolMenu();
 
