Sequencer:Control Rig: Shift now moves camera to stay with control rig when moving.
#trivial

diff --git a/Engine/Plugins/Animation/ControlRig/Source/ControlRigEditor/Private/EditMode/ControlRigEditMode.cpp b/Engine/Plugins/Animation/ControlRig/Source/ControlRigEditor/Private/EditMode/ControlRigEditMode.cpp
--- a/Engine/Plugins/Animation/ControlRig/Source/ControlRigEditor/Private/EditMode/ControlRigEditMode.cpp
+++ b/Engine/Plugins/Animation/ControlRig/Source/ControlRigEditor/Private/EditMode/ControlRigEditMode.cpp
@@ -1767,11 +1767,15 @@ bool FControlRigEditMode::InputDelta(FEditorViewportClient* InViewportClient, FV
 		//handled by processed mouse clicks
 		return true;
 	}
+
+
 	FVector Drag = InDrag;
 	FRotator Rot = InRot;
 	FVector Scale = InScale;
 
 	const bool bCtrlDown = InViewport->KeyState(EKeys::LeftControl) || InViewport->KeyState(EKeys::RightControl);
+	const bool bShiftDown = InViewport->KeyState(EKeys::LeftShift) || InViewport->KeyState(EKeys::RightShift);
+
 	//button down if left and ctrl and right is down, needed for indirect posting
 	const bool bMouseButtonDown = InViewport->KeyState(EKeys::LeftMouseButton) || (bCtrlDown && InViewport->KeyState(EKeys::RightMouseButton));
 
@@ -1783,6 +1787,19 @@ bool FControlRigEditMode::InputDelta(FEditorViewportClient* InViewportClient, FV
 	const bool bDoTranslation = !Drag.IsZero() && (WidgetMode == UE::Widget::WM_Translate || WidgetMode == UE::Widget::WM_TranslateRotateZ);
 	const bool bDoScale = !Scale.IsZero() && WidgetMode == UE::Widget::WM_Scale;
 
+	if (!InDrag.IsZero() && bShiftDown)
+	{
+		FVector CameraDelta(InDrag);
+
+		// Only apply camera speed modifiers to the drag if we aren't zooming in an ortho viewport.
+		if (!InViewportClient->IsOrtho() || !(InViewport->KeyState(EKeys::LeftMouseButton) && InViewport->KeyState(EKeys::RightMouseButton)))
+		{
+			const float CameraSpeed = InViewportClient->GetCameraSpeed();
+			CameraDelta *= CameraSpeed;
+		}
+
+		InViewportClient->MoveViewportCamera(CameraDelta, InRot);
+	}
 
 	if (InteractionScopes.Num() > 0 && bMouseButtonDown && CurrentAxis != EAxisList::None
 		&& (bDoRotation || bDoTranslation || bDoScale))
