Fixed mismatched copy constructor/assignment operator in FRigLocalAndGlobalTransform.

diff --git a/Engine/Plugins/Animation/ControlRig/Source/ControlRig/Public/Rigs/RigHierarchyElements.h b/Engine/Plugins/Animation/ControlRig/Source/ControlRig/Public/Rigs/RigHierarchyElements.h
--- a/Engine/Plugins/Animation/ControlRig/Source/ControlRig/Public/Rigs/RigHierarchyElements.h
+++ b/Engine/Plugins/Animation/ControlRig/Source/ControlRig/Public/Rigs/RigHierarchyElements.h
@@ -299,17 +299,6 @@ struct CONTROLRIG_API FRigLocalAndGlobalTransform
 	void Save(FArchive& Ar);
 	void Load(FArchive& Ar);
 
-	FRigLocalAndGlobalTransform& operator=(const FRigLocalAndGlobalTransform& Other)
-	{
-		Local = Other.Local;
-		Global = Other.Global;
-		for (int32 i=0; i<FRigLocalAndGlobalTransform::EDirtyMax; ++i)
-		{
-			bDirty[i] = Other.bDirty[i];
-		}
-		return *this;
-	}
-
 	UPROPERTY(BlueprintReadOnly, EditAnywhere, Category = "Pose")
 	FRigComputedTransform Local;
 
