Fixed static analysis

diff --git a/Engine/Plugins/Animation/ControlRig/Source/ControlRigEditor/Private/Tools/ControlRigSnapper.cpp b/Engine/Plugins/Animation/ControlRig/Source/ControlRigEditor/Private/Tools/ControlRigSnapper.cpp
--- a/Engine/Plugins/Animation/ControlRig/Source/ControlRigEditor/Private/Tools/ControlRigSnapper.cpp
+++ b/Engine/Plugins/Animation/ControlRig/Source/ControlRigEditor/Private/Tools/ControlRigSnapper.cpp
@@ -507,7 +507,7 @@ bool FControlRigSnapper::SnapIt(FFrameNumber StartFrame, FFrameNumber EndFrame,c
 		MovieScene->Modify();
 		
 		TArray<FFrameNumber> Frames;
-		if (SnapSettings->BakingKeySettings == EBakingKeySettings::AllFrames)
+		if (SnapSettings && SnapSettings->BakingKeySettings == EBakingKeySettings::AllFrames)
 		{
 			MovieSceneToolHelpers::CalculateFramesBetween(MovieScene, StartFrame, EndFrame, SnapSettings->FrameIncrement,Frames);
 		}
@@ -554,10 +554,13 @@ bool FControlRigSnapper::SnapIt(FFrameNumber StartFrame, FFrameNumber EndFrame,c
 				}
 			}
 		}
+
+		const bool bKeepOffset = SnapSettings && SnapSettings->bKeepOffset;
+		
 		//set transforms on these actors
 		for (FGuidAndActor& GuidActor : ActorsToSnap)
 		{
-			if (bSnapToFirstFrameNotParents || SnapSettings->bKeepOffset) //if we are snapping to the first frame or keep offset we just don't set the parent transforms
+			if (bSnapToFirstFrameNotParents || bKeepOffset) //if we are snapping to the first frame or keep offset we just don't set the parent transforms
 			{
 				TArray<FFrameNumber> OneFrame;
 				OneFrame.Add(Frames[0]);
@@ -582,7 +585,7 @@ bool FControlRigSnapper::SnapIt(FFrameNumber StartFrame, FFrameNumber EndFrame,c
 				}
 
 			}
-			GuidActor.SetTransform(Sequencer, Frames, WorldTransformToSnap,SnapSettings);
+			GuidActor.SetTransform(Sequencer, Frames, WorldTransformToSnap, SnapSettings);
 		}
 
 		//Now do Control Rigs
@@ -622,7 +625,7 @@ bool FControlRigSnapper::SnapIt(FFrameNumber StartFrame, FFrameNumber EndFrame,c
 						TArray<FTransform> CurrentControlRigTransform, CurrentParentWorldTransform;
 						CurrentControlRigTransform.SetNum(1);
 						CurrentParentWorldTransform.SetNum(1);
-						if (bSnapToFirstFrameNotParents || SnapSettings->bKeepOffset)
+						if (bSnapToFirstFrameNotParents || bKeepOffset)
 						{
 							OneFrame[0] = Frames[0];
 							CurrentParentWorldTransform[0] = ControlRigParentWorldTransforms[0];
@@ -634,7 +637,7 @@ bool FControlRigSnapper::SnapIt(FFrameNumber StartFrame, FFrameNumber EndFrame,c
 									Transform = CurrentControlRigTransform[0];
 								}
 							}
-							else if (SnapSettings->bKeepOffset)
+							else if (bKeepOffset)
 							{
 								FTransform FirstWorld = WorldTransformToSnap[0];
 								for (FTransform& Transform : WorldTransformToSnap)
