ControlRig: Fix some execute pins greyed out

diff --git a/Engine/Plugins/Animation/ControlRig/Source/ControlRigEditor/Private/Graph/SControlRigGraphNode.cpp b/Engine/Plugins/Animation/ControlRig/Source/ControlRigEditor/Private/Graph/SControlRigGraphNode.cpp
--- a/Engine/Plugins/Animation/ControlRig/Source/ControlRigEditor/Private/Graph/SControlRigGraphNode.cpp
+++ b/Engine/Plugins/Animation/ControlRig/Source/ControlRigEditor/Private/Graph/SControlRigGraphNode.cpp
@@ -401,7 +401,14 @@ void SControlRigGraphNode::AddPin(const TSharedRef<SGraphPin>& PinToAdd)
 						{
 							const URigVMPin* RootPin = ModelPin->GetRootPin();
 							FLinearColor PinColorAndOpacity = PinToAdd->GetColorAndOpacity();
-							if (Template->FindArgument(RootPin->GetFName()) == nullptr)
+
+							FRigVMDispatchContext DispatchContext;
+							if(const URigVMDispatchNode* DispatchNode = Cast<URigVMDispatchNode>(ModelNode))
+							{
+								DispatchContext = DispatchNode->GetDispatchContext();
+							}
+
+							if (Template->FindArgument(RootPin->GetFName()) == nullptr && Template->FindExecuteArgument(RootPin->GetFName(), DispatchContext) == nullptr)
 							{
 								PinColorAndOpacity.A = 0.2f;
 							}
@@ -1602,6 +1609,12 @@ void SControlRigGraphNode::UpdatePinTreeView()
 	}
 
 	const UControlRigGraphSchema* RigSchema = Cast<UControlRigGraphSchema>(RigGraphNode->GetSchema());
+
+	FRigVMDispatchContext DispatchContext;
+	if(const URigVMDispatchNode* DispatchNode = Cast<URigVMDispatchNode>(ModelNode))
+	{
+		DispatchContext = DispatchNode->GetDispatchContext();
+	}
 	
 	TMap<URigVMPin*, int32> ModelPinToInfoIndex;
 	for(URigVMPin* ModelPin : ModelPins)
@@ -1674,7 +1687,7 @@ void SControlRigGraphNode::UpdatePinTreeView()
 					if (URigVMPin* RootPin = ModelPin->GetRootPin())
 					{
 						FLinearColor PinColorAndOpacity = PinInfo.OutputPinWidget->GetColorAndOpacity();
-						if (Template->FindArgument(RootPin->GetFName()) == nullptr)
+						if (Template->FindArgument(RootPin->GetFName()) == nullptr && Template->FindExecuteArgument(RootPin->GetFName(), DispatchContext) == nullptr)
 						{
 							PinColorAndOpacity.A = 0.2f;
 						}
@@ -1702,7 +1715,7 @@ void SControlRigGraphNode::UpdatePinTreeView()
 					if (URigVMPin* RootPin = ModelPin->GetRootPin())
 					{
 						FLinearColor PinColorAndOpacity = PinInfo.InputPinWidget->GetColorAndOpacity();
-						if (Template->FindArgument(RootPin->GetFName()) == nullptr)
+						if (Template->FindArgument(RootPin->GetFName()) == nullptr && Template->FindExecuteArgument(RootPin->GetFName(), DispatchContext) == nullptr)
 						{
 							PinColorAndOpacity.A = 0.2f;
 						}
