Fixed mismatched copy constructor/assignment operator in FRigBoneHierarchy, FRigCurveContainer, FRigBoneHierarchy and FRigSpaceHierarchy.

diff --git a/Engine/Plugins/Animation/ControlRig/Source/ControlRig/Public/Rigs/RigBoneHierarchy.h b/Engine/Plugins/Animation/ControlRig/Source/ControlRig/Public/Rigs/RigBoneHierarchy.h
--- a/Engine/Plugins/Animation/ControlRig/Source/ControlRig/Public/Rigs/RigBoneHierarchy.h
+++ b/Engine/Plugins/Animation/ControlRig/Source/ControlRig/Public/Rigs/RigBoneHierarchy.h
@@ -77,12 +77,14 @@ struct CONTROLRIG_API FRigBoneHierarchy
 	TArray<FRigBone>::RangedForConstIteratorType end() const   { return Bones.end();   }
 
 	FRigBone& Add(const FName& InNewName, const FName& InParentName, ERigBoneType InType, const FTransform& InInitTransform, const FTransform& InLocalTransform, const FTransform& InGlobalTransform);
-	
-private:
 
-	// disable copy constructor
-	FRigBoneHierarchy(const FRigBoneHierarchy& InOther) {}
+	// Pretty weird that this type is copy/move assignable (needed for USTRUCTs) but not copy/move constructible
+	FRigBoneHierarchy(FRigBoneHierarchy&& InOther) = delete;
+	FRigBoneHierarchy(const FRigBoneHierarchy& InOther) = delete;
+	FRigBoneHierarchy& operator=(FRigBoneHierarchy&& InOther) = default;
+	FRigBoneHierarchy& operator=(const FRigBoneHierarchy& InOther) = default;
 
+private:
 	UPROPERTY(EditAnywhere, Category = FRigBoneHierarchy)
 	TArray<FRigBone> Bones;
 
diff --git a/Engine/Plugins/Animation/ControlRig/Source/ControlRig/Public/Rigs/RigControlHierarchy.h b/Engine/Plugins/Animation/ControlRig/Source/ControlRig/Public/Rigs/RigControlHierarchy.h
--- a/Engine/Plugins/Animation/ControlRig/Source/ControlRig/Public/Rigs/RigControlHierarchy.h
+++ b/Engine/Plugins/Animation/ControlRig/Source/ControlRig/Public/Rigs/RigControlHierarchy.h
@@ -279,12 +279,14 @@ struct CONTROLRIG_API FRigControlHierarchy
 	);
 
 	void PostLoad();
-	
-private:
 
-	// disable copy constructor
-	FRigControlHierarchy(const FRigControlHierarchy& InOther) {}
+	// Pretty weird that this type is copy/move assignable (needed for USTRUCTs) but not copy/move constructible
+	FRigControlHierarchy(FRigControlHierarchy&& InOther) = delete;
+	FRigControlHierarchy(const FRigControlHierarchy& InOther) = delete;
+	FRigControlHierarchy& operator=(FRigControlHierarchy&& InOther) = default;
+	FRigControlHierarchy& operator=(const FRigControlHierarchy& InOther) = default;
 
+private:
 	UPROPERTY(EditAnywhere, Category = FRigControlHierarchy)
 	TArray<FRigControl> Controls;
 
diff --git a/Engine/Plugins/Animation/ControlRig/Source/ControlRig/Public/Rigs/RigCurveContainer.h b/Engine/Plugins/Animation/ControlRig/Source/ControlRig/Public/Rigs/RigCurveContainer.h
--- a/Engine/Plugins/Animation/ControlRig/Source/ControlRig/Public/Rigs/RigCurveContainer.h
+++ b/Engine/Plugins/Animation/ControlRig/Source/ControlRig/Public/Rigs/RigCurveContainer.h
@@ -48,12 +48,14 @@ public:
 	TArray<FRigCurve>::RangedForConstIteratorType end() const   { return Curves.end();   }
 
 	FRigCurve& Add(const FName& InNewName, float InValue);
-	
-private:
 
-	// disable copy constructor
-	FRigCurveContainer(const FRigCurveContainer& InOther) {}
+	// Pretty weird that this type is copy/move assignable (needed for USTRUCTs) but not copy/move constructible
+	FRigCurveContainer(FRigCurveContainer&& InOther) = delete;
+	FRigCurveContainer(const FRigCurveContainer& InOther) = delete;
+	FRigCurveContainer& operator=(FRigCurveContainer&& InOther) = default;
+	FRigCurveContainer& operator=(const FRigCurveContainer& InOther) = default;
 
+private:
 	UPROPERTY(EditAnywhere, Category = FRigCurveContainer)
 	TArray<FRigCurve> Curves;
 
diff --git a/Engine/Plugins/Animation/ControlRig/Source/ControlRig/Public/Rigs/RigSpaceHierarchy.h b/Engine/Plugins/Animation/ControlRig/Source/ControlRig/Public/Rigs/RigSpaceHierarchy.h
--- a/Engine/Plugins/Animation/ControlRig/Source/ControlRig/Public/Rigs/RigSpaceHierarchy.h
+++ b/Engine/Plugins/Animation/ControlRig/Source/ControlRig/Public/Rigs/RigSpaceHierarchy.h
@@ -100,12 +100,14 @@ struct CONTROLRIG_API FRigSpaceHierarchy
 	TArray<FRigSpace>::RangedForConstIteratorType end() const   { return Spaces.end();   }
 
 	FRigSpace& Add(const FName& InNewName, ERigSpaceType InSpaceType, const FName& InParentName, const FTransform& InTransform);
-	
-private:
 
-	// disable copy constructor
-	FRigSpaceHierarchy(const FRigSpaceHierarchy& InOther) {}
+	// Pretty weird that this type is copy/move assignable (needed for USTRUCTs) but not copy/move constructible
+	FRigSpaceHierarchy(FRigSpaceHierarchy&& InOther) = delete;
+	FRigSpaceHierarchy(const FRigSpaceHierarchy& InOther) = delete;
+	FRigSpaceHierarchy& operator=(FRigSpaceHierarchy&& InOther) = default;
+	FRigSpaceHierarchy& operator=(const FRigSpaceHierarchy& InOther) = default;
 
+private:
 	UPROPERTY(EditAnywhere, Category = FRigSpaceHierarchy)
 	TArray<FRigSpace> Spaces;
 
