Sequencer: Control Rig: When looking for a binding that containts a control rig we take into account the control rig so we pick the correct actor or it's component.  For certain setups the face and body rig both bind to the same component on the control rig, but bind differently on the sequencer side.

diff --git a/Engine/Plugins/Animation/ControlRig/Source/ControlRigEditor/Private/Sequencer/ControlRigParameterTrackEditor.cpp b/Engine/Plugins/Animation/ControlRig/Source/ControlRigEditor/Private/Sequencer/ControlRigParameterTrackEditor.cpp
--- a/Engine/Plugins/Animation/ControlRig/Source/ControlRigEditor/Private/Sequencer/ControlRigParameterTrackEditor.cpp
+++ b/Engine/Plugins/Animation/ControlRig/Source/ControlRigEditor/Private/Sequencer/ControlRigParameterTrackEditor.cpp
@@ -2307,7 +2307,7 @@ void FControlRigParameterTrackEditor::SelectRigsAndControls(UControlRig* Control
 }
 
 
-FMovieSceneTrackEditor::FFindOrCreateHandleResult FControlRigParameterTrackEditor::FindOrCreateHandleToSceneCompOrOwner(USceneComponent* InComp)
+FMovieSceneTrackEditor::FFindOrCreateHandleResult FControlRigParameterTrackEditor::FindOrCreateHandleToSceneCompOrOwner(USceneComponent* InComp, UControlRig* InControlRig)
 {
 	const bool bCreateHandleIfMissing = false;
 	FName CreatedFolderName = NAME_None;
@@ -2320,12 +2320,15 @@ FMovieSceneTrackEditor::FFindOrCreateHandleResult FControlRigParameterTrackEdito
 
 	UMovieScene* MovieScene = GetSequencer()->GetFocusedMovieSceneSequence()->GetMovieScene();
 
-	// Prioritize a control rig parameter track on this component
+	// Prioritize a control rig parameter track on this component if it matches the handle
 	if (Result.Handle.IsValid())
 	{
-		if (MovieScene->FindTrack(UMovieSceneControlRigParameterTrack::StaticClass(), Result.Handle, NAME_None))
+		if (UMovieSceneControlRigParameterTrack* Track = Cast<UMovieSceneControlRigParameterTrack>(MovieScene->FindTrack(UMovieSceneControlRigParameterTrack::StaticClass(), Result.Handle, NAME_None)))
 		{
-			return Result;
+			if (InControlRig == nullptr || (Track->GetControlRig() == InControlRig))
+			{
+				return Result;
+			}
 		}
 	}
 
@@ -2335,11 +2338,14 @@ FMovieSceneTrackEditor::FFindOrCreateHandleResult FControlRigParameterTrackEdito
 	bHandleWasValid = OwnerHandle.IsValid();
 	if (OwnerHandle.IsValid())
 	{
-		if (MovieScene->FindTrack(UMovieSceneControlRigParameterTrack::StaticClass(), OwnerHandle, NAME_None))
+		if (UMovieSceneControlRigParameterTrack* Track = Cast<UMovieSceneControlRigParameterTrack>(MovieScene->FindTrack(UMovieSceneControlRigParameterTrack::StaticClass(), OwnerHandle, NAME_None)))
 		{
-			Result.Handle = OwnerHandle;
-			Result.bWasCreated = bHandleWasValid == false && Result.Handle.IsValid();
-			return Result;
+			if (InControlRig == nullptr || (Track->GetControlRig() == InControlRig))
+			{
+				Result.Handle = OwnerHandle;
+				Result.bWasCreated = bHandleWasValid == false && Result.Handle.IsValid();
+				return Result;
+			}
 		}
 	}
 
@@ -2998,7 +3004,7 @@ void FControlRigParameterTrackEditor::HandleControlSelected(UControlRig* Subject
 		}
 		ActorObject = Component->GetOwner();
 		bool bCreateTrack = false;
-		FFindOrCreateHandleResult HandleResult = FindOrCreateHandleToSceneCompOrOwner(Component);
+		FFindOrCreateHandleResult HandleResult = FindOrCreateHandleToSceneCompOrOwner(Component, Subject);
 		FGuid ObjectHandle = HandleResult.Handle;
 		if (!ObjectHandle.IsValid())
 		{
@@ -3562,7 +3568,7 @@ FKeyPropertyResult FControlRigParameterTrackEditor::AddKeysToControlRig(
 		KeyMode == ESequencerKeyMode::ManualKeyForced ||
 		AllowEditsMode == EAllowEditsMode::AllowSequencerEditsOnly;
 
-	FFindOrCreateHandleResult HandleResult = FindOrCreateHandleToSceneCompOrOwner(InSceneComp);
+	FFindOrCreateHandleResult HandleResult = FindOrCreateHandleToSceneCompOrOwner(InSceneComp, InControlRig);
 	FGuid ObjectHandle = HandleResult.Handle;
 	KeyPropertyResult.bHandleCreated = HandleResult.bWasCreated;
 	if (ObjectHandle.IsValid())
@@ -3589,7 +3595,7 @@ void FControlRigParameterTrackEditor::AddControlKeys(
 	}
 	bool bCreateTrack = false;
 	bool bCreateHandle = false;
-	FFindOrCreateHandleResult HandleResult = FindOrCreateHandleToSceneCompOrOwner(InSceneComp);
+	FFindOrCreateHandleResult HandleResult = FindOrCreateHandleToSceneCompOrOwner(InSceneComp, InControlRig);
 	FGuid ObjectHandle = HandleResult.Handle;
 	if (!ObjectHandle.IsValid())
 	{
diff --git a/Engine/Plugins/Animation/ControlRig/Source/ControlRigEditor/Private/Sequencer/ControlRigParameterTrackEditor.h b/Engine/Plugins/Animation/ControlRig/Source/ControlRigEditor/Private/Sequencer/ControlRigParameterTrackEditor.h
--- a/Engine/Plugins/Animation/ControlRig/Source/ControlRigEditor/Private/Sequencer/ControlRigParameterTrackEditor.h
+++ b/Engine/Plugins/Animation/ControlRig/Source/ControlRigEditor/Private/Sequencer/ControlRigParameterTrackEditor.h
@@ -166,8 +166,8 @@ private:
 	/** Select control rig if not selected, select controls from key areas */
 	void SelectRigsAndControls(UControlRig* Subject, const TArray<const IKeyArea*>& KeyAreas);
 
-	/** Handle Creation for SkelMeshComp or Actor Owner, either may have a binding*/
-	FMovieSceneTrackEditor::FFindOrCreateHandleResult FindOrCreateHandleToSceneCompOrOwner(USceneComponent* InComp);
+	/** Handle Creation for SkelMeshComp or Actor Owner, either may have a binding, use optional control rig to pick one*/
+	FMovieSceneTrackEditor::FFindOrCreateHandleResult FindOrCreateHandleToSceneCompOrOwner(USceneComponent* InComp, UControlRig* InControlRig = nullptr);
 
 	/** Handle Creation for control rig track given the object binding and the control rig */	
 	FMovieSceneTrackEditor::FFindOrCreateTrackResult FindOrCreateControlRigTrackForObject(FGuid ObjectBinding, UControlRig* ControlRig, FName PropertyName = NAME_None, bool bCreateTrackIfMissing = true);
