ControlRig: fixed a local variable detail customization crash due to a previous detail wrapper object refactor that changed how the owner graph is referenced by the detail object. previously we were getting it from the Outer of the detail object, after the refactor we should get it from via the Subject member var

diff --git a/Engine/Plugins/Animation/ControlRig/Source/ControlRigEditor/Private/ControlRigLocalVariableDetails.cpp b/Engine/Plugins/Animation/ControlRig/Source/ControlRigEditor/Private/ControlRigLocalVariableDetails.cpp
--- a/Engine/Plugins/Animation/ControlRig/Source/ControlRigEditor/Private/ControlRigLocalVariableDetails.cpp
+++ b/Engine/Plugins/Animation/ControlRig/Source/ControlRigEditor/Private/ControlRigLocalVariableDetails.cpp
@@ -35,7 +35,7 @@ void FRigVMLocalVariableDetails::CustomizeDetails(IDetailLayoutBuilder& DetailBu
 		{
 			GraphBeingCustomized = Cast<URigVMGraph>(ObjectsBeingCustomized[0]->GetSubject());
 		}
-		if (GraphBeingCustomized)
+		if (ensure(GraphBeingCustomized))
 		{
 			BlueprintBeingCustomized = GraphBeingCustomized->GetTypedOuter<UControlRigBlueprint>();
 		}
