Added check for singularity in spherical pose reader.

diff --git a/Engine/Plugins/Animation/ControlRig/Source/ControlRig/Private/Units/Highlevel/Hierarchy/RigUnit_SphericalPoseReader.cpp b/Engine/Plugins/Animation/ControlRig/Source/ControlRig/Private/Units/Highlevel/Hierarchy/RigUnit_SphericalPoseReader.cpp
--- a/Engine/Plugins/Animation/ControlRig/Source/ControlRig/Private/Units/Highlevel/Hierarchy/RigUnit_SphericalPoseReader.cpp
+++ b/Engine/Plugins/Animation/ControlRig/Source/ControlRig/Private/Units/Highlevel/Hierarchy/RigUnit_SphericalPoseReader.cpp
@@ -73,29 +73,45 @@ FRigUnit_SphericalPoseReader_Execute()
 	FTransform WorldOffset = LocalDriverTransformInit * GlobalDriverParentTransform;
 	WorldOffset.SetLocation(GlobalDriverTransform.GetLocation());
 
-	// get driver axis
+	// get driver axis in local space of sphere
 	const FVector CurrentGlobalDriverAxis = GlobalDriverTransform.GetRotation().RotateVector(DriverAxis);
 	DriverNormal = WorldOffset.InverseTransformVectorNoScale(CurrentGlobalDriverAxis).GetSafeNormal();
 
+	//
 	// evaluate the interpolation of the output param
-	const float AngleFromForward = FMath::Acos(FVector::DotProduct(DriverNormal, FVector::ZAxisVector));
-	const float Mag = RemapRange(AngleFromForward, 0.0f, PI, 0.0f, 1.0f);
-	Driver2D = DriverNormal;
-	if (Mag > SMALL_NUMBER)
+	//
+
+	// get angle between the driver direction and the local Z-axis of the sphere
+	float DriverDotZAxis = FVector::DotProduct(DriverNormal, FVector::ZAxisVector);
+	if (FMath::IsNearlyEqual(DriverDotZAxis, -1.0f, KINDA_SMALL_NUMBER))
 	{
-		Driver2D.Z = 0.0f;
-		Driver2D.Normalize();
-		Driver2D *= Mag;
+		// avoid singularity at negative pole (back of sphere)
+		OutputParam = 0.0f;
+		Debug.DrawDebug(WorldOffset, ExecuteContext.GetDrawInterface(), InnerRegion, OuterRegion, DriverNormal, Driver2D, OutputParam);
+		return;
 	}
-	Driver2D.Z = -1.0f;
 
+	// remap angle from 0-PI, to 0-1
+	const float AngleFromForward = FMath::Acos(DriverDotZAxis);
+	const float Mag = RemapRange(AngleFromForward, 0.0f, PI, 0.0f, 1.0f);
 	if (FMath::IsNearlyZero(Mag))
 	{
-		// avoid NaNs from DistanceToEllipse, and guaranteed to be inside inner ellipse
+		// avoid singularity at positive pole (guaranteed to be inside inner ellipse)
+		// causes NaNs in DistanceToEllipse
 		OutputParam = 1.0f;
 		Debug.DrawDebug(WorldOffset, ExecuteContext.GetDrawInterface(), InnerRegion, OuterRegion, DriverNormal, Driver2D, OutputParam);
 		return;
 	}
+
+	// convert to 2d polar coordinates, with magnitude normalized by range 0-PI
+	Driver2D = DriverNormal;
+	if (Mag > SMALL_NUMBER)
+	{
+		Driver2D.Z = 0.0f;
+		Driver2D.Normalize();
+		Driver2D *= Mag;
+	}
+	Driver2D.Z = -1.0f;
 	
 	const float PointX = Driver2D.X;
 	const float PointY = Driver2D.Y;
diff --git a/Engine/Plugins/Animation/ControlRig/Source/ControlRig/Private/Units/Highlevel/Hierarchy/RigUnit_SphericalPoseReader.h b/Engine/Plugins/Animation/ControlRig/Source/ControlRig/Private/Units/Highlevel/Hierarchy/RigUnit_SphericalPoseReader.h
--- a/Engine/Plugins/Animation/ControlRig/Source/ControlRig/Private/Units/Highlevel/Hierarchy/RigUnit_SphericalPoseReader.h
+++ b/Engine/Plugins/Animation/ControlRig/Source/ControlRig/Private/Units/Highlevel/Hierarchy/RigUnit_SphericalPoseReader.h
@@ -267,7 +267,7 @@ struct CONTROLRIG_API FSphericalPoseReaderDebugSettings
 /*
  * Outputs a float value between 0-1 based off of the driver item's rotation in a specified region.
  */
-USTRUCT(meta=(DisplayName="Spherical Pose Reader", Category="Hierarchy", Keywords="Pose Reader"))
+USTRUCT(meta=(DisplayName="Spherical Pose Reader", Category="Hierarchy", Keywords="Pose Reader, SPR"))
 struct CONTROLRIG_API FRigUnit_SphericalPoseReader: public FRigUnit_HighlevelBaseMutable
 {
 	GENERATED_BODY()
