Control Rig: Add node to check if a shape is available in the shape libraries

diff --git a/Engine/Plugins/Animation/ControlRig/Source/ControlRig/Private/ControlRig.cpp b/Engine/Plugins/Animation/ControlRig/Source/ControlRig/Private/ControlRig.cpp
--- a/Engine/Plugins/Animation/ControlRig/Source/ControlRig/Private/ControlRig.cpp
+++ b/Engine/Plugins/Animation/ControlRig/Source/ControlRig/Private/ControlRig.cpp
@@ -271,6 +271,17 @@ void UControlRig::OnAddShapeLibrary(const FControlRigExecuteContext* InContext,
 #endif
 }
 
+bool UControlRig::OnShapeExists(const FName& InShapeName) const
+{
+	const TArray<TSoftObjectPtr<UControlRigShapeLibrary>>& Libraries = GetShapeLibraries();
+	if (UControlRigShapeLibrary::GetShapeByName(InShapeName, GetShapeLibraries(), ShapeLibraryNameMap))
+	{
+		return true;
+	}
+	
+	return false;
+}
+
 bool UControlRig::InitializeVM(const FName& InEventName)
 {
 	if(!Super::InitializeVM(InEventName))
@@ -520,6 +531,7 @@ bool UControlRig::Execute(const FName& InEventName)
 
 	// allow access to the shape libraries
 	PublicContext.OnAddShapeLibraryDelegate.BindUObject(this, &UControlRig::OnAddShapeLibrary);
+	PublicContext.OnShapeExistsDelegate.BindUObject(this, &UControlRig::OnShapeExists);
 
 	// allow access to the default hierarchy to allow to reset
 	if(!HasAnyFlags(RF_ClassDefaultObject))
diff --git a/Engine/Plugins/Animation/ControlRig/Source/ControlRig/Private/Units/Core/RigUnit_UserData.cpp b/Engine/Plugins/Animation/ControlRig/Source/ControlRig/Private/Units/Core/RigUnit_UserData.cpp
--- a/Engine/Plugins/Animation/ControlRig/Source/ControlRig/Private/Units/Core/RigUnit_UserData.cpp
+++ b/Engine/Plugins/Animation/ControlRig/Source/ControlRig/Private/Units/Core/RigUnit_UserData.cpp
@@ -227,4 +227,13 @@ FRigUnit_SetupShapeLibraryFromUserData_Execute()
 		Log(FString::Printf(Format, *NameSpace));
 	}
 #endif
+}
+
+FRigUnit_ShapeExists_Execute()
+{
+	Result = false;
+	if (ExecuteContext.OnShapeExistsDelegate.IsBound())
+	{
+		Result = ExecuteContext.OnShapeExistsDelegate.Execute(ShapeName);
+	}
 }
\ No newline at end of file
diff --git a/Engine/Plugins/Animation/ControlRig/Source/ControlRig/Private/Units/Core/RigUnit_UserData.h b/Engine/Plugins/Animation/ControlRig/Source/ControlRig/Private/Units/Core/RigUnit_UserData.h
--- a/Engine/Plugins/Animation/ControlRig/Source/ControlRig/Private/Units/Core/RigUnit_UserData.h
+++ b/Engine/Plugins/Animation/ControlRig/Source/ControlRig/Private/Units/Core/RigUnit_UserData.h
@@ -85,6 +85,36 @@ struct CONTROLRIG_API FRigUnit_SetupShapeLibraryFromUserData : public FRigUnitMu
 	UPROPERTY(meta = (Input))
 	bool LogShapeLibraries;
 
+	RIGVM_METHOD()
+	virtual void Execute() override;
+};
+
+/**
+ * Checks whether or not a shape is available in the rig's shape libraries
+ */
+USTRUCT(meta=(DisplayName="Shape Exists", Keywords="Shape,Gizmo,Library,Exists,Contains"))
+struct CONTROLRIG_API FRigUnit_ShapeExists : public FRigUnit
+{
+	GENERATED_BODY()
+
+	FRigUnit_ShapeExists()
+	{
+		ShapeName = NAME_None;
+		Result = false;
+	}
+
+	/*
+	 * The name of the shape to search for
+	 */
+	UPROPERTY(meta = (Input))
+	FName ShapeName;
+
+	/*
+	 * True if the shape name exists in any of the shape libraries 
+	 */
+	UPROPERTY(meta = (Output))
+	bool Result;
+
 	RIGVM_METHOD()
 	virtual void Execute() override;
 };
\ No newline at end of file
diff --git a/Engine/Plugins/Animation/ControlRig/Source/ControlRig/Public/ControlRig.h b/Engine/Plugins/Animation/ControlRig/Source/ControlRig/Public/ControlRig.h
--- a/Engine/Plugins/Animation/ControlRig/Source/ControlRig/Public/ControlRig.h
+++ b/Engine/Plugins/Animation/ControlRig/Source/ControlRig/Public/ControlRig.h
@@ -461,6 +461,7 @@ private:
 	FRigEventDelegate RigEventDelegate;
 
 	void OnAddShapeLibrary(const FControlRigExecuteContext* InContext, const FString& InLibraryName, UControlRigShapeLibrary* InShapeLibrary, bool bReplaceExisting, bool bLogResults);
+	bool OnShapeExists(const FName& InShapeName) const;
 	virtual void InitializeFromCDO() override;
 
 
diff --git a/Engine/Plugins/Animation/ControlRig/Source/ControlRig/Public/Units/RigUnitContext.h b/Engine/Plugins/Animation/ControlRig/Source/ControlRig/Public/Units/RigUnitContext.h
--- a/Engine/Plugins/Animation/ControlRig/Source/ControlRig/Public/Units/RigUnitContext.h
+++ b/Engine/Plugins/Animation/ControlRig/Source/ControlRig/Public/Units/RigUnitContext.h
@@ -103,6 +103,13 @@ struct FControlRigExecuteContext : public FRigVMExecuteContext
 	{
 	}
 
+	virtual void Reset() override
+	{
+		Super::Reset();
+		OnAddShapeLibraryDelegate.Unbind();
+		OnShapeExistsDelegate.Unbind();
+	}
+
 	virtual void Copy(const FRigVMExecuteContext* InOtherContext) override
 	{
 		Super::Copy(InOtherContext);
@@ -139,7 +146,11 @@ struct FControlRigExecuteContext : public FRigVMExecuteContext
 	TArray<const UAssetUserData*> AssetUserData;
 
 	DECLARE_DELEGATE_FiveParams(FOnAddShapeLibrary, const FControlRigExecuteContext* InContext, const FString&, UControlRigShapeLibrary*, bool /* replace? */, bool /* log results */);
-	FOnAddShapeLibrary OnAddShapeLibraryDelegate; 
+	FOnAddShapeLibrary OnAddShapeLibraryDelegate;
+
+	DECLARE_DELEGATE_RetVal_OneParam(bool, FOnShapeExists, const FName&);
+	FOnShapeExists OnShapeExistsDelegate;
+	
 	FRigUnitContext UnitContext;
 	URigHierarchy* Hierarchy;
 };
diff --git a/Engine/Plugins/Animation/ControlRig/Source/ControlRigDeveloper/Private/ControlRigBlueprint.cpp b/Engine/Plugins/Animation/ControlRig/Source/ControlRigDeveloper/Private/ControlRigBlueprint.cpp
--- a/Engine/Plugins/Animation/ControlRig/Source/ControlRigDeveloper/Private/ControlRigBlueprint.cpp
+++ b/Engine/Plugins/Animation/ControlRig/Source/ControlRigDeveloper/Private/ControlRigBlueprint.cpp
@@ -313,6 +313,26 @@ void UControlRigBlueprint::PostEditChangeChainProperty(FPropertyChangedChainEven
 {
 	Super::PostEditChangeChainProperty(PropertyChangedEvent);
 	PostEditChangeChainPropertyEvent.Broadcast(PropertyChangedEvent);
+
+	// Propagate shape libraries
+	if (PropertyChangedEvent.Property && PropertyChangedEvent.Property->GetFName() == GET_MEMBER_NAME_CHECKED(UControlRigBlueprint, ShapeLibraries))
+	{
+		URigVMBlueprintGeneratedClass* RigClass = GetControlRigBlueprintGeneratedClass();
+		UControlRig* CDO = Cast<UControlRig>(RigClass->GetDefaultObject(false /* create if needed */));
+
+		TArray<UObject*> ArchetypeInstances;
+		CDO->GetArchetypeInstances(ArchetypeInstances);
+		ArchetypeInstances.Add(CDO);
+
+		// Propagate libraries to archetypes
+		for (UObject* Instance : ArchetypeInstances)
+		{
+			if (UControlRig* InstanceRig = Cast<UControlRig>(Instance))
+			{
+				InstanceRig->ShapeLibraries = ShapeLibraries;
+			}
+		}
+	}
 }
 
 void UControlRigBlueprint::PostRename(UObject* OldOuter, const FName OldName)
diff --git a/Engine/Plugins/Runtime/RigVM/Source/RigVM/Public/RigVMCore/RigVMExecuteContext.h b/Engine/Plugins/Runtime/RigVM/Source/RigVM/Public/RigVMCore/RigVMExecuteContext.h
--- a/Engine/Plugins/Runtime/RigVM/Source/RigVM/Public/RigVMCore/RigVMExecuteContext.h
+++ b/Engine/Plugins/Runtime/RigVM/Source/RigVM/Public/RigVMCore/RigVMExecuteContext.h
@@ -512,7 +512,10 @@ struct RIGVM_API FRigVMExtendedExecuteContext
 		*this = InOther;
 	}
 
-	virtual ~FRigVMExtendedExecuteContext() {}
+	virtual ~FRigVMExtendedExecuteContext()
+	{
+		Reset();
+	}
 
 	void Reset()
 	{
