Control Rig: Fix unbinding context delegates breaks rig selection

diff --git a/Engine/Plugins/Animation/ControlRig/Source/ControlRig/Public/Units/RigUnitContext.h b/Engine/Plugins/Animation/ControlRig/Source/ControlRig/Public/Units/RigUnitContext.h
--- a/Engine/Plugins/Animation/ControlRig/Source/ControlRig/Public/Units/RigUnitContext.h
+++ b/Engine/Plugins/Animation/ControlRig/Source/ControlRig/Public/Units/RigUnitContext.h
@@ -103,13 +103,6 @@ struct FControlRigExecuteContext : public FRigVMExecuteContext
 	{
 	}
 
-	virtual void Reset() override
-	{
-		Super::Reset();
-		OnAddShapeLibraryDelegate.Unbind();
-		OnShapeExistsDelegate.Unbind();
-	}
-
 	virtual void Copy(const FRigVMExecuteContext* InOtherContext) override
 	{
 		Super::Copy(InOtherContext);
