Control Rig: Fix replace shape library not removing previous library with same name

diff --git a/Engine/Plugins/Animation/ControlRig/Source/ControlRig/Private/ControlRig.cpp b/Engine/Plugins/Animation/ControlRig/Source/ControlRig/Private/ControlRig.cpp
--- a/Engine/Plugins/Animation/ControlRig/Source/ControlRig/Private/ControlRig.cpp
+++ b/Engine/Plugins/Animation/ControlRig/Source/ControlRig/Private/ControlRig.cpp
@@ -220,8 +220,14 @@ void UControlRig::OnAddShapeLibrary(const FControlRigExecuteContext* InContext,
 			{
 				continue;
 			}
-			if(ExistingShapeLibrary->GetName().Equals(LibraryName, ESearchCase::IgnoreCase))
+			FString ExistingName = ExistingShapeLibrary->GetName();
+			if (FString* MapName = ShapeLibraryNameMap.Find(ExistingName))
 			{
+				ExistingName = *MapName;
+			}
+			if(ExistingName.Equals(LibraryName, ESearchCase::IgnoreCase))
+			{
+				ShapeLibraryNameMap.Remove(ExistingShapeLibrary->GetName());
 				ShapeLibraries[Index] = InShapeLibrary;
 				break;
 			}
