Details Keyframe Handler:
- Fixed IsPropertyAnimated implementations so that no new handles are created if not found.

diff --git a/Engine/Plugins/Animation/ControlRig/Source/ControlRigEditor/Private/EditMode/SControlRigDetails.cpp b/Engine/Plugins/Animation/ControlRig/Source/ControlRigEditor/Private/EditMode/SControlRigDetails.cpp
--- a/Engine/Plugins/Animation/ControlRig/Source/ControlRigEditor/Private/EditMode/SControlRigDetails.cpp
+++ b/Engine/Plugins/Animation/ControlRig/Source/ControlRigEditor/Private/EditMode/SControlRigDetails.cpp
@@ -706,7 +706,8 @@ bool SControlRigDetails::IsPropertyAnimated(const IPropertyHandle& PropertyHandl
 	ISequencer* Sequencer = GetSequencer();
 	if (Sequencer && Sequencer->GetFocusedMovieSceneSequence())
 	{
-		FGuid ObjectHandle = Sequencer->GetHandleToObject(ParentObject);
+		constexpr bool bCreateHandleIfMissing = false;
+		FGuid ObjectHandle = Sequencer->GetHandleToObject(ParentObject, bCreateHandleIfMissing);
 		if (ObjectHandle.IsValid())
 		{
 			UMovieScene* MovieScene = Sequencer->GetFocusedMovieSceneSequence()->GetMovieScene();
diff --git a/Engine/Plugins/Animation/ControlRig/Source/ControlRigEditor/Private/EditMode/SControlRigEditModeTools.cpp b/Engine/Plugins/Animation/ControlRig/Source/ControlRigEditor/Private/EditMode/SControlRigEditModeTools.cpp
--- a/Engine/Plugins/Animation/ControlRig/Source/ControlRigEditor/Private/EditMode/SControlRigEditModeTools.cpp
+++ b/Engine/Plugins/Animation/ControlRig/Source/ControlRigEditor/Private/EditMode/SControlRigEditModeTools.cpp
@@ -546,7 +546,8 @@ bool SControlRigEditModeTools::IsPropertyAnimated(const IPropertyHandle& Propert
 	TSharedPtr<ISequencer> Sequencer = WeakSequencer.Pin();
 	if (Sequencer.IsValid() && Sequencer->GetFocusedMovieSceneSequence())
 	{
-		FGuid ObjectHandle = Sequencer->GetHandleToObject(ParentObject);
+		constexpr bool bCreateHandleIfMissing = false;
+		FGuid ObjectHandle = Sequencer->GetHandleToObject(ParentObject, bCreateHandleIfMissing);
 		if (ObjectHandle.IsValid()) 
 		{
 			UMovieScene* MovieScene = Sequencer->GetFocusedMovieSceneSequence()->GetMovieScene();
diff --git a/Engine/Source/Editor/UMGEditor/Private/Animation/UMGDetailKeyframeHandler.cpp b/Engine/Source/Editor/UMGEditor/Private/Animation/UMGDetailKeyframeHandler.cpp
--- a/Engine/Source/Editor/UMGEditor/Private/Animation/UMGDetailKeyframeHandler.cpp
+++ b/Engine/Source/Editor/UMGEditor/Private/Animation/UMGDetailKeyframeHandler.cpp
@@ -27,7 +27,8 @@ bool FUMGDetailKeyframeHandler::IsPropertyAnimated(const IPropertyHandle& Proper
 	TSharedPtr<ISequencer> Sequencer = BlueprintEditor.Pin()->GetSequencer();
 	if (Sequencer.IsValid() && Sequencer->GetFocusedMovieSceneSequence())
 	{
-		FGuid ObjectHandle = Sequencer->GetHandleToObject(ParentObject);
+		constexpr bool bCreateHandleIfMissing = false;
+		FGuid ObjectHandle = Sequencer->GetHandleToObject(ParentObject, bCreateHandleIfMissing);
 		if (ObjectHandle.IsValid()) 
 		{
 			UMovieScene* MovieScene = Sequencer->GetFocusedMovieSceneSequence()->GetMovieScene();
