Control Rig Editor: enable drawing as long as there are bones in the hierachy

diff --git a/Engine/Plugins/Animation/ControlRig/Source/ControlRigEditor/Private/Editor/ControlRigEditor.cpp b/Engine/Plugins/Animation/ControlRig/Source/ControlRigEditor/Private/Editor/ControlRigEditor.cpp
--- a/Engine/Plugins/Animation/ControlRig/Source/ControlRigEditor/Private/Editor/ControlRigEditor.cpp
+++ b/Engine/Plugins/Animation/ControlRig/Source/ControlRigEditor/Private/Editor/ControlRigEditor.cpp
@@ -7331,6 +7331,14 @@ void FControlRigEditor::OnPostConstruction_AnyThread(UControlRig* InRig, const F
 			{
 				SetDetailViewForRigElements();
 			}
+			
+			if (FControlRigEditorEditMode* EditMode = GetEditMode())
+            {
+				if (InRig)
+            	{
+            		EditMode->bDrawHierarchyBones = !InRig->GetHierarchy()->GetBones().IsEmpty();
+            	}
+            }
 		};
 				
 		if(IsInGameThread())
