Make sure xray materials are loaded before using them

diff --git a/Engine/Plugins/Animation/ControlRig/Source/ControlRigEditor/Private/EditMode/ControlRigEditMode.cpp b/Engine/Plugins/Animation/ControlRig/Source/ControlRigEditor/Private/EditMode/ControlRigEditMode.cpp
--- a/Engine/Plugins/Animation/ControlRig/Source/ControlRigEditor/Private/EditMode/ControlRigEditMode.cpp
+++ b/Engine/Plugins/Animation/ControlRig/Source/ControlRigEditor/Private/EditMode/ControlRigEditMode.cpp
@@ -3308,8 +3308,19 @@ void FControlRigEditMode::CreateShapeActors(UControlRig* ControlRig)
 			{
 				Param.MeshTransform = ShapeDef->Transform;
 				Param.StaticMesh = ShapeDef->StaticMesh;
-				Param.Material = (bShowControlsAsOverlay && ShapeDef->Library->XRayMaterial.IsValid()) ?
-					ShapeDef->Library->XRayMaterial : ShapeDef->Library->DefaultMaterial;
+				Param.Material = ShapeDef->Library->DefaultMaterial;
+				if (bShowControlsAsOverlay)
+				{
+					TSoftObjectPtr<UMaterial> XRayMaterial = ShapeDef->Library->XRayMaterial;
+					if (XRayMaterial.IsPending())
+					{
+						XRayMaterial.LoadSynchronous();
+					}
+					if (XRayMaterial.IsValid())
+					{
+						Param.Material = XRayMaterial;
+					}
+				}
 				Param.ColorParameterName = ShapeDef->Library->MaterialColorParameter;
 			}
 
