Disable gizmo manipulation for channels + disable transformation of non-visible controls even if selected

diff --git a/Engine/Plugins/Animation/ControlRig/Source/ControlRigEditor/Private/EditMode/ControlRigEditMode.cpp b/Engine/Plugins/Animation/ControlRig/Source/ControlRigEditor/Private/EditMode/ControlRigEditMode.cpp
--- a/Engine/Plugins/Animation/ControlRig/Source/ControlRigEditor/Private/EditMode/ControlRigEditMode.cpp
+++ b/Engine/Plugins/Animation/ControlRig/Source/ControlRigEditor/Private/EditMode/ControlRigEditMode.cpp
@@ -1840,7 +1840,7 @@ bool FControlRigEditMode::InputDelta(FEditorViewportClient* InViewportClient, FV
 
 					for (AControlRigShapeActor* ShapeActor : Pairs.Value)
 					{
-						if (ShapeActor->IsSelected())
+						if (ShapeActor->IsEnabled() && ShapeActor->IsSelected())
 						{
 							// test local vs global
 							if (bManipulatorMadeChange == false)
@@ -3448,7 +3448,37 @@ bool FControlRigEditMode::AreRigElementSelectedAndMovable(UControlRig* ControlRi
 {
 	const UControlRigEditModeSettings* Settings = GetDefault<UControlRigEditModeSettings>();
 
-	if (!Settings || !ControlRig || !AreRigElementsSelected(ValidControlTypeMask(), ControlRig))
+	if (!Settings || !ControlRig)
+	{
+		return false;
+	}
+
+	auto IsAnySelectedControlMovable = [this, ControlRig]()
+	{
+		URigHierarchy* Hierarchy = ControlRig->GetHierarchy();
+		if (!Hierarchy)
+		{
+			return false;
+		}
+		
+		const TArray<FRigElementKey> SelectedRigElements = GetSelectedRigElements(ControlRig);
+		return SelectedRigElements.ContainsByPredicate([Hierarchy](const FRigElementKey& Element)
+		{
+			if (!FRigElementTypeHelper::DoesHave(ValidControlTypeMask(), Element.Type))
+			{
+				return false;
+			}
+			const FRigControlElement* ControlElement = Hierarchy->Find<FRigControlElement>(Element);
+			if (!ControlElement)
+			{
+				return false;
+			}
+			// can a control non selectable in the viewport be movable?  
+			return ControlElement->Settings.IsSelectable();
+		});
+	};
+	
+	if (!IsAnySelectedControlMovable())
 	{
 		return false;
 	}
