Control Rig: Convert combo box on pin type selection to list view selection

diff --git a/Engine/Plugins/Animation/ControlRig/Source/ControlRigEditor/Private/Editor/ControlRigEditor.cpp b/Engine/Plugins/Animation/ControlRig/Source/ControlRigEditor/Private/Editor/ControlRigEditor.cpp
--- a/Engine/Plugins/Animation/ControlRig/Source/ControlRigEditor/Private/Editor/ControlRigEditor.cpp
+++ b/Engine/Plugins/Animation/ControlRig/Source/ControlRigEditor/Private/Editor/ControlRigEditor.cpp
@@ -5188,6 +5188,34 @@ TRigVMTypeIndex FControlRigEditor::OnRequestPinTypeSelectionDialog(const TArray<
 
 	TRigVMTypeIndex Answer = INDEX_NONE;
 
+	const FRigVMRegistry& Registry = FRigVMRegistry::Get();
+
+	TArray<TSharedPtr<FName>> TypeNames;
+	TMap<FName, uint8> TypeNameToIndex;
+	TypeNames.Reserve(InTypes.Num());
+	for (int32 i=0; i<InTypes.Num(); ++i)
+	{
+		const TRigVMTypeIndex& TypeIndex = InTypes[i];
+		TRigVMTypeIndex FinalType = TypeIndex;
+		if (FinalType == RigVMTypeUtils::TypeIndex::Float)
+		{
+			FinalType = RigVMTypeUtils::TypeIndex::Double;
+		}
+		if (FinalType == RigVMTypeUtils::TypeIndex::FloatArray)
+		{
+			FinalType = RigVMTypeUtils::TypeIndex::DoubleArray;
+		}
+
+		const FRigVMTemplateArgumentType& ArgumentType = Registry.GetType(FinalType);
+		if (!TypeNames.ContainsByPredicate([&ArgumentType](const TSharedPtr<FName>& InName)
+		{
+			return *InName.Get() == ArgumentType.CPPType;
+		}))
+		{
+			TypeNames.AddUnique(MakeShared<FName>(ArgumentType.CPPType));
+			TypeNameToIndex.Add(ArgumentType.CPPType, i);
+		}
+	}
 	TSharedPtr< SWindow > Window = SNew(SWindow)
 		.Title(LOCTEXT("SelectPinType", "Select Pin Type"))
 		.ScreenPosition(FSlateApplication::Get().GetCursorPos())
@@ -5212,12 +5240,47 @@ TRigVMTypeIndex FControlRigEditor::OnRequestPinTypeSelectionDialog(const TArray<
 							SNew(SScrollBox)
 							+SScrollBox::Slot()
 							[
-								SNew(SControlRigChangePinType)
-									.Blueprint(GetControlRigBlueprint())
-									.Types(InTypes)
-									.OnTypeSelected_Lambda([&Answer](const TRigVMTypeIndex& TypeSelected)
+								SNew(SListView<TSharedPtr<FName>>)
+									.ListItemsSource(&TypeNames)
+									.OnGenerateRow_Lambda([&Registry, &TypeNameToIndex, &InTypes](const TSharedPtr<FName> InItem, const TSharedRef<STableViewBase>& Owner)
+									{
+										TRigVMTypeIndex TypeIndex = InTypes[TypeNameToIndex.FindChecked(*InItem.Get())];
+										const FRigVMTemplateArgumentType Type = FRigVMRegistry::Get().GetType(TypeIndex);
+										const bool bIsArray = Type.IsArray();
+										static const FName TypeIcon(TEXT("Kismet.VariableList.TypeIcon"));
+										static const FName ArrayTypeIcon(TEXT("Kismet.VariableList.ArrayTypeIcon"));
+
+										const FEdGraphPinType PinType = RigVMTypeUtils::PinTypeFromTypeIndex(TypeIndex);
+										const UControlRigGraphSchema* Schema = GetDefault<UControlRigGraphSchema>();
+										const FLinearColor Color = Schema->GetPinTypeColor(PinType);
+										
+										return SNew(STableRow<TSharedPtr<FString>>, Owner)
+												.Padding(FMargin(16, 4, 16, 4))
+												[
+													
+													SNew(SHorizontalBox)
+													+ SHorizontalBox::Slot()
+													.AutoWidth()
+													.VAlign(VAlign_Center)
+													[
+														SNew(SBox)
+														.HeightOverride(16.0f)
+														[
+															SNew(SImage)
+															.Image(bIsArray ? FAppStyle::GetBrush(ArrayTypeIcon) : FAppStyle::GetBrush(TypeIcon))
+															.ColorAndOpacity(Color)
+														]
+													]
+
+													+ SHorizontalBox::Slot()
+													[
+														SNew(STextBlock).Text(FText::FromName(*InItem.Get()))
+													]
+												];
+									})
+									.OnSelectionChanged_Lambda([&Answer, &TypeNames, &TypeNameToIndex, &InTypes](const TSharedPtr<FName> InName, ESelectInfo::Type)
 									{
-										Answer = TypeSelected;
+										Answer = InTypes[TypeNameToIndex.FindChecked(*InName.Get())];
 										FSlateApplication::Get().GetActiveModalWindow()->RequestDestroyWindow();
 									})
 							]
