Control Rig: Introduce Rotation Order on Controls

This is guarded (and disabled by default) using the CVar 'ControlRig.Hierarchy.EnableRotationOrder'

diff --git a/Engine/Plugins/Animation/ControlRig/Source/ControlRig/Private/Rigs/RigHierarchy.cpp b/Engine/Plugins/Animation/ControlRig/Source/ControlRig/Private/Rigs/RigHierarchy.cpp
--- a/Engine/Plugins/Animation/ControlRig/Source/ControlRig/Private/Rigs/RigHierarchy.cpp
+++ b/Engine/Plugins/Animation/ControlRig/Source/ControlRig/Private/Rigs/RigHierarchy.cpp
@@ -44,6 +44,7 @@ static TAutoConsoleVariable<int32> CVarControlRigHierarchyTraceAlways(TEXT("Cont
 static TAutoConsoleVariable<int32> CVarControlRigHierarchyTraceCallstack(TEXT("ControlRig.Hierarchy.TraceCallstack"), 0, TEXT("if nonzero we will record the callstack for any trace entry.\nOnly works if(ControlRig.Hierarchy.TraceEnabled != 0)"));
 static TAutoConsoleVariable<int32> CVarControlRigHierarchyTracePrecision(TEXT("ControlRig.Hierarchy.TracePrecision"), 3, TEXT("sets the number digits in a float when tracing hierarchies."));
 static TAutoConsoleVariable<int32> CVarControlRigHierarchyTraceOnSpawn(TEXT("ControlRig.Hierarchy.TraceOnSpawn"), 0, TEXT("sets the number of frames to trace when a new hierarchy is spawned"));
+TAutoConsoleVariable<bool> CVarControlRigHierarchyEnableRotationOrder(TEXT("ControlRig.Hierarchy.EnableRotationOrder"), false, TEXT("enables the rotation order for controls"));
 static int32 sRigHierarchyLastTrace = INDEX_NONE;
 static TCHAR sRigHierarchyTraceFormat[16];
 
diff --git a/Engine/Plugins/Animation/ControlRig/Source/ControlRig/Private/Rigs/RigHierarchyElements.cpp b/Engine/Plugins/Animation/ControlRig/Source/ControlRig/Private/Rigs/RigHierarchyElements.cpp
--- a/Engine/Plugins/Animation/ControlRig/Source/ControlRig/Private/Rigs/RigHierarchyElements.cpp
+++ b/Engine/Plugins/Animation/ControlRig/Source/ControlRig/Private/Rigs/RigHierarchyElements.cpp
@@ -389,7 +389,7 @@ FRotator FRigPreferredEulerAngles::SetRotator(const FRotator& InValue, bool bIni
 			return FixedValue;
 		}
 	}
-	SetAngles(InValue.Euler(), bInitial, DefaultRotationOrder);
+	SetAngles(InValue.Euler(), bInitial, DefaultRotationOrder, false);
 	return InValue;
 }
 
@@ -402,16 +402,51 @@ FVector FRigPreferredEulerAngles::GetAngles(bool bInitial, EEulerRotationOrder I
 	return AnimationCore::ChangeEulerRotationOrder(Get(bInitial), RotationOrder, InRotationOrder);
 }
 
-void FRigPreferredEulerAngles::SetAngles(const FVector& InValue, bool bInitial, EEulerRotationOrder InRotationOrder)
+void FRigPreferredEulerAngles::SetAngles(const FVector& InValue, bool bInitial, EEulerRotationOrder InRotationOrder, bool bFixEulerFlips)
 {
 	FVector Value = InValue;
 	if(RotationOrder != InRotationOrder)
 	{
 		Value = AnimationCore::ChangeEulerRotationOrder(Value, InRotationOrder, RotationOrder);
 	}
+
+	if(bFixEulerFlips)
+	{
+		// @MikeZ we need something here as well for euler values
+		// Do you think we can rely on the same mechanism?
+		/*
+		const FRotator CurrentValue = GetRotator(bInitial);
+			
+		//Find Diff of the rotation from current and just add that instead of setting so we can go over/under -180
+		FRotator CurrentWinding;
+		FRotator CurrentRotRemainder;
+		CurrentValue.GetWindingAndRemainder(CurrentWinding, CurrentRotRemainder);
+
+		FRotator DeltaRot = InValue - CurrentRotRemainder;
+		DeltaRot.Normalize();
+		const FRotator FixedValue = CurrentValue + DeltaRot;
+
+		SetAngles(FixedValue.Euler(), bInitial, DefaultRotationOrder);
+		return FixedValue;
+		*/
+	}
+	
 	Get(bInitial) = Value;
 }
 
+void FRigPreferredEulerAngles::SetRotationOrder(EEulerRotationOrder InRotationOrder)
+{
+	if(RotationOrder != InRotationOrder)
+	{
+		const EEulerRotationOrder PreviousRotationOrder = RotationOrder;
+		const FVector PreviousAnglesCurrent = GetAngles(false, RotationOrder);
+		const FVector PreviousAnglesInitial = GetAngles(true, RotationOrder);
+		RotationOrder = InRotationOrder;
+		SetAngles(PreviousAnglesCurrent, false, PreviousRotationOrder);
+		SetAngles(PreviousAnglesInitial, true, PreviousRotationOrder);
+	}
+}
+
 ////////////////////////////////////////////////////////////////////////////////
 // FRigElementHandle
 ////////////////////////////////////////////////////////////////////////////////
@@ -744,6 +779,7 @@ FRigControlSettings::FRigControlSettings()
 , Customization()
 , bGroupWithParentControl(false)
 , bRestrictSpaceSwitching(false)
+, PreferredRotationOrder(FRigPreferredEulerAngles::DefaultRotationOrder)
 {
 	// rely on the default provided by the shape definition
 	ShapeName = FControlRigShapeDefinition().ShapeName; 
@@ -789,6 +825,7 @@ void FRigControlSettings::Save(FArchive& Ar)
 	Ar << bGroupWithParentControl;
 	Ar << bRestrictSpaceSwitching;
 	Ar << FilteredChannels;
+	Ar << PreferredRotationOrder;
 }
 
 void FRigControlSettings::Load(FArchive& Ar)
@@ -957,6 +994,15 @@ void FRigControlSettings::Load(FArchive& Ar)
 	{
 		FilteredChannels.Reset();
 	}
+
+	if (Ar.CustomVer(FControlRigObjectVersion::GUID) >= FControlRigObjectVersion::RigHierarchyControlPreferredRotationOrder)
+	{
+		Ar << PreferredRotationOrder;
+	}
+	else
+	{
+		PreferredRotationOrder = FRigPreferredEulerAngles::DefaultRotationOrder;
+	}
 }
 
 uint32 GetTypeHash(const FRigControlSettings& Settings)
@@ -980,6 +1026,7 @@ uint32 GetTypeHash(const FRigControlSettings& Settings)
 	{
 		Hash = HashCombine(Hash, GetTypeHash(Channel));
 	}
+	Hash = HashCombine(Hash, GetTypeHash(Settings.PreferredRotationOrder));
 	return Hash;
 }
 
@@ -1061,6 +1108,10 @@ bool FRigControlSettings::operator==(const FRigControlSettings& InOther) const
 	{
 		return false;
 	}
+	if(PreferredRotationOrder != InOther.PreferredRotationOrder)
+	{
+		return false;
+	}
 
 	const FTransform MinimumTransform = MinimumValue.GetAsTransform(ControlType, PrimaryAxis);
 	const FTransform OtherMinimumTransform = InOther.MinimumValue.GetAsTransform(ControlType, PrimaryAxis);
@@ -1170,6 +1221,7 @@ void FRigControlElement::Load(FArchive& Ar, URigHierarchy* Hierarchy, ESerializa
 		{
 			PreferredEulerAngles.Reset();
 		}
+		PreferredEulerAngles.SetRotationOrder(Settings.PreferredRotationOrder);
 	}
 }
 
diff --git a/Engine/Plugins/Animation/ControlRig/Source/ControlRig/Private/Sequencer/MovieSceneControlRigParameterTemplate.cpp b/Engine/Plugins/Animation/ControlRig/Source/ControlRig/Private/Sequencer/MovieSceneControlRigParameterTemplate.cpp
--- a/Engine/Plugins/Animation/ControlRig/Source/ControlRig/Private/Sequencer/MovieSceneControlRigParameterTemplate.cpp
+++ b/Engine/Plugins/Animation/ControlRig/Source/ControlRig/Private/Sequencer/MovieSceneControlRigParameterTemplate.cpp
@@ -805,6 +805,7 @@ struct FControlRigParameterPreAnimatedTokenProducer : IMovieScenePreAnimatedToke
 										case ERigControlType::Transform:
 										{
 											ControlRig->SetControlValue<FRigControlValue::FTransform_Float>(Value.Name, Value.Value.ToFTransform(), true, FRigControlModifiedContext(EControlRigSetKey::Never), bSetupUndo);
+											// @MikeZ here we want to call SetControlPreferredEulerAngles instead
 											ControlRig->GetHierarchy()->SetControlPreferredRotator(ControlElement, Value.Value.Rotation);
 											break;
 										}
@@ -812,6 +813,7 @@ struct FControlRigParameterPreAnimatedTokenProducer : IMovieScenePreAnimatedToke
 										{
 											FTransformNoScale NoScale = Value.Value.ToFTransform();
 											ControlRig->SetControlValue<FRigControlValue::FTransformNoScale_Float>(Value.Name, NoScale, true, FRigControlModifiedContext(EControlRigSetKey::Never), bSetupUndo);
+											// @MikeZ here we want to call SetControlPreferredEulerAngles instead
 											ControlRig->GetHierarchy()->SetControlPreferredRotator(ControlElement, Value.Value.Rotation);
 											break;
 										}
@@ -1659,12 +1661,14 @@ struct TControlRigParameterActuatorTransform : TMovieSceneBlendingActuator<FCont
 					if (ControlElement && ControlElement->Settings.ControlType == ERigControlType::Transform)
 					{
 						ControlRig->SetControlValue<FRigControlValue::FTransform_Float>(ParameterName, InFinalValue.Value.ToFTransform(), true, EControlRigSetKey::Never, bSetupUndo);
+						// @MikeZ here we want to call SetControlPreferredEulerAngles instead
 						ControlRig->GetHierarchy()->SetControlPreferredRotator(ControlElement, InFinalValue.Value.Rotator());
 					}
 					else if (ControlElement && ControlElement->Settings.ControlType == ERigControlType::TransformNoScale)
 					{
 						const FTransformNoScale NoScale = InFinalValue.Value.ToFTransform();
 						ControlRig->SetControlValue<FRigControlValue::FTransformNoScale_Float>(ParameterName, NoScale, true, EControlRigSetKey::Never, bSetupUndo);
+						// @MikeZ here we want to call SetControlPreferredEulerAngles instead
 						ControlRig->GetHierarchy()->SetControlPreferredRotator(ControlElement, InFinalValue.Value.Rotator());
 					}
 					else if (ControlElement && ControlElement->Settings.ControlType == ERigControlType::EulerTransform)
diff --git a/Engine/Plugins/Animation/ControlRig/Source/ControlRig/Public/ControlRigObjectVersion.h b/Engine/Plugins/Animation/ControlRig/Source/ControlRig/Public/ControlRigObjectVersion.h
--- a/Engine/Plugins/Animation/ControlRig/Source/ControlRig/Public/ControlRigObjectVersion.h
+++ b/Engine/Plugins/Animation/ControlRig/Source/ControlRig/Public/ControlRigObjectVersion.h
@@ -87,6 +87,9 @@ struct CONTROLRIG_API FControlRigObjectVersion
 		// Hierarchy storing previous names
 		RigHierarchyStoringPreviousNames,
 
+		// Control supporting preferred rotation order
+		RigHierarchyControlPreferredRotationOrder,
+
 		// -----<new versions can be added above this line>-------------------------------------------------
 		VersionPlusOne,
 		LatestVersion = VersionPlusOne - 1
diff --git a/Engine/Plugins/Animation/ControlRig/Source/ControlRig/Public/Rigs/RigHierarchy.h b/Engine/Plugins/Animation/ControlRig/Source/ControlRig/Public/Rigs/RigHierarchy.h
--- a/Engine/Plugins/Animation/ControlRig/Source/ControlRig/Public/Rigs/RigHierarchy.h
+++ b/Engine/Plugins/Animation/ControlRig/Source/ControlRig/Public/Rigs/RigHierarchy.h
@@ -25,6 +25,8 @@ DECLARE_EVENT_FiveParams(URigHierarchy, FRigHierarchyUndoRedoTransformEvent, URi
 DECLARE_MULTICAST_DELEGATE_TwoParams(FRigHierarchyMetadataChangedDelegate, const FRigElementKey& /* Key */, const FName& /* Name */);
 DECLARE_MULTICAST_DELEGATE_ThreeParams(FRigHierarchyMetadataTagChangedDelegate, const FRigElementKey& /* Key */, const FName& /* Tag */, bool /* AddedOrRemoved */);
 
+extern CONTROLRIG_API TAutoConsoleVariable<bool> CVarControlRigHierarchyEnableRotationOrder;
+
 UENUM()
 enum ERigTransformStackEntryType : int
 {
@@ -1986,7 +1988,7 @@ public:
 	 * @param bInitial If true we'll return the preferred rotator for the initial - otherwise current transform
 	 * @return Returns the current preferred rotator
 	 */
-	FRotator GetControlPreferredRotator(FRigControlElement* InControlElement, bool bInitial = false) const
+	FRotator GetControlPreferredRotator(const FRigControlElement* InControlElement, bool bInitial = false) const
 	{
 		if(InControlElement)
 		{
@@ -1998,32 +2000,32 @@ public:
 	/**
 	 * Sets a control's preferred rotator (local transform rotation)
 	 * @param InKey The key of the element to retrieve the current value for
-	 * @param InValue The new preferred rotator to set
+	 * @param InRotator The new preferred rotator to set
 	 * @param bInitial If true we'll return the preferred rotator for the initial - otherwise current transform
 	 * @param bFixEulerFlips If true the new rotator value will use the shortest path
 	 */
 	UFUNCTION(BlueprintCallable, Category = URigHierarchy)
-	void SetControlPreferredRotator(FRigElementKey InKey, const FRotator& InValue, bool bInitial = false, bool bFixEulerFlips = false)
+	void SetControlPreferredRotator(FRigElementKey InKey, const FRotator& InRotator, bool bInitial = false, bool bFixEulerFlips = false)
 	{
-		SetControlPreferredRotatorByIndex(GetIndex(InKey), InValue, bInitial, bFixEulerFlips);
+		SetControlPreferredRotatorByIndex(GetIndex(InKey), InRotator, bInitial, bFixEulerFlips);
 	}
 	
 
 	/**
 	 * Sets a control's preferred rotator (local transform rotation)
 	 * @param InElementIndex The element index to look up
-	 * @param InValue The new preferred rotator to set
+	 * @param InRotator The new preferred rotator to set
 	 * @param bInitial If true we'll return the preferred rotator for the initial - otherwise current transform
 	 * @param bFixEulerFlips If true the new rotator value will use the shortest path
 	 */
 	UFUNCTION(BlueprintCallable, Category = URigHierarchy)
-	void SetControlPreferredRotatorByIndex(int32 InElementIndex, const FRotator& InValue, bool bInitial = false, bool bFixEulerFlips = false)
+	void SetControlPreferredRotatorByIndex(int32 InElementIndex, const FRotator& InRotator, bool bInitial = false, bool bFixEulerFlips = false)
 	{
 		if(Elements.IsValidIndex(InElementIndex))
 		{
 			if(FRigControlElement* ControlElement = Cast<FRigControlElement>(Elements[InElementIndex]))
 			{
-				SetControlPreferredRotator(ControlElement, InValue, bInitial, bFixEulerFlips);
+				SetControlPreferredRotator(ControlElement, InRotator, bInitial, bFixEulerFlips);
 			}
 		}
 	}
@@ -2031,18 +2033,211 @@ public:
 	/**
 	 * Sets a control's preferred rotator (local transform rotation)
 	 * @param InControlElement The element to look up
-	 * @param InValue The new preferred rotator to set
+	 * @param InRotator The new preferred rotator to set
 	 * @param bInitial If true we'll return the preferred rotator for the initial - otherwise current transform
 	 * @param bFixEulerFlips If true the new rotator value will use the shortest path
 	 */
-	void SetControlPreferredRotator(FRigControlElement* InControlElement, const FRotator& InValue, bool bInitial = false, bool bFixEulerFlips = false)
+	void SetControlPreferredRotator(FRigControlElement* InControlElement, const FRotator& InRotator, bool bInitial = false, bool bFixEulerFlips = false)
 	{
 		if(InControlElement)
 		{
-			InControlElement->PreferredEulerAngles.SetRotator(InValue, bInitial, bFixEulerFlips);
+			InControlElement->PreferredEulerAngles.SetRotator(InRotator, bInitial, bFixEulerFlips);
 		}
 	}
+
+	/**
+	 * Returns a control's preferred euler angles (local transform rotation)
+	 * @param InKey The key of the element to retrieve the current value for
+	 * @param InRotationOrder The rotation order to use when retrieving the euler angles
+	 * @param bInitial If true we'll return the preferred euler angles for the initial - otherwise current transform
+	 * @return Returns the current preferred euler angles
+	 */
+	UFUNCTION(BlueprintCallable, Category = URigHierarchy)
+	FVector GetControlPreferredEulerAngles(FRigElementKey InKey, EEulerRotationOrder InRotationOrder, bool bInitial = false) const
+	{
+		return GetControlPreferredEulerAnglesByIndex(GetIndex(InKey), InRotationOrder, bInitial);
+	}
+
+	/**
+	 * Returns a control's preferred euler angles (local transform rotation)
+	 * @param InElementIndex The element index to look up
+	 * @param InRotationOrder The rotation order to use when retrieving the euler angles
+	 * @param bInitial If true we'll return the preferred euler angles for the initial - otherwise current transform
+	 * @return Returns the current preferred euler angles
+	 */
+	UFUNCTION(BlueprintCallable, Category = URigHierarchy)
+	FVector GetControlPreferredEulerAnglesByIndex(int32 InElementIndex, EEulerRotationOrder InRotationOrder, bool bInitial = false) const
+	{
+		if(Elements.IsValidIndex(InElementIndex))
+		{
+			if(FRigControlElement* ControlElement = Cast<FRigControlElement>(Elements[InElementIndex]))
+			{
+				return GetControlPreferredEulerAngles(ControlElement, InRotationOrder, bInitial);
+			}
+		}
+		return FVector::ZeroVector;
+	}
+
+	/**
+	 * Returns a control's preferred euler angles (local transform rotation)
+	 * @param InControlElement The element to look up
+	 * @param InRotationOrder The rotation order to use when retrieving the euler angles
+	 * @param bInitial If true we'll return the preferred euler angles for the initial - otherwise current transform
+	 * @return Returns the current preferred euler angles
+	 */
+	FVector GetControlPreferredEulerAngles(const FRigControlElement* InControlElement, EEulerRotationOrder InRotationOrder, bool bInitial = false) const
+	{
+		if(InControlElement)
+		{
+			return InControlElement->PreferredEulerAngles.GetAngles(bInitial, InRotationOrder);
+		}
+		return FVector::ZeroVector;
+	}
+
+	/**
+	 * Sets a control's preferred euler angles (local transform rotation)
+	 * @param InKey The key of the element to retrieve the current value for
+	 * @param InEulerAngles The new preferred euler angles to set
+	 * @param InRotationOrder The rotation order to use when setting the euler angles
+	 * @param bInitial If true we'll return the preferred euler angles for the initial - otherwise current transform
+	 * @param bFixEulerFlips If true the new euler angles value will use the shortest path
+	 */
+	UFUNCTION(BlueprintCallable, Category = URigHierarchy)
+	void SetControlPreferredEulerAngles(FRigElementKey InKey, const FVector& InEulerAngles, EEulerRotationOrder InRotationOrder, bool bInitial = false, bool bFixEulerFlips = false)
+	{
+		SetControlPreferredEulerAnglesByIndex(GetIndex(InKey), InEulerAngles, InRotationOrder, bInitial, bFixEulerFlips);
+	}
 	
+
+	/**
+	 * Sets a control's preferred euler angles (local transform rotation)
+	 * @param InElementIndex The element index to look up
+	 * @param InEulerAngles The new preferred euler angles to set
+	 * @param InRotationOrder The rotation order to use when setting the euler angles
+	 * @param bInitial If true we'll return the preferred euler angles for the initial - otherwise current transform
+	 * @param bFixEulerFlips If true the new euler angles value will use the shortest path
+	 */
+	UFUNCTION(BlueprintCallable, Category = URigHierarchy)
+	void SetControlPreferredEulerAnglesByIndex(int32 InElementIndex, const FVector& InEulerAngles, EEulerRotationOrder InRotationOrder, bool bInitial = false, bool bFixEulerFlips = false)
+	{
+		if(Elements.IsValidIndex(InElementIndex))
+		{
+			if(FRigControlElement* ControlElement = Cast<FRigControlElement>(Elements[InElementIndex]))
+			{
+				SetControlPreferredEulerAngles(ControlElement, InEulerAngles, InRotationOrder, bInitial, bFixEulerFlips);
+			}
+		}
+	}
+
+	/**
+	 * Sets a control's preferred euler angles (local transform rotation)
+	 * @param InControlElement The element to look up
+	 * @param InEulerAngles The new preferred euler angles to set
+	 * @param InRotationOrder The rotation order to use when setting the euler angles
+	 * @param bInitial If true we'll return the preferred euler angles for the initial - otherwise current transform
+	 * @param bFixEulerFlips If true the new euler angles value will use the shortest path
+	 */
+	void SetControlPreferredEulerAngles(FRigControlElement* InControlElement, const FVector& InEulerAngles, EEulerRotationOrder InRotationOrder, bool bInitial = false, bool bFixEulerFlips = false)
+	{
+		if(InControlElement)
+		{
+			InControlElement->PreferredEulerAngles.SetRotationOrder(InRotationOrder);
+			InControlElement->PreferredEulerAngles.SetAngles(InEulerAngles, bInitial, InRotationOrder, bFixEulerFlips);
+		}
+	}
+
+	/**
+	 * Returns a control's preferred euler rotation order
+	 * @param InKey The key of the element to retrieve the current value for
+	 * @param bFromSettings If true the rotation order will be looked up from the control's settings, otherwise from the currently set animation value
+	 * @return Returns the current preferred euler rotation order
+	 */
+	UFUNCTION(BlueprintCallable, Category = URigHierarchy)
+	EEulerRotationOrder GetControlPreferredEulerRotationOrder(FRigElementKey InKey, bool bFromSettings = true) const
+	{
+		return GetControlPreferredEulerRotationOrderByIndex(GetIndex(InKey), bFromSettings);
+	}
+
+	/**
+	 * Returns a control's preferred euler rotation order
+	 * @param InElementIndex The element index to look up
+	 * @param bFromSettings If true the rotation order will be looked up from the control's settings, otherwise from the currently set animation value
+	 * @return Returns the current preferred euler rotation order
+	 */
+	UFUNCTION(BlueprintCallable, Category = URigHierarchy)
+	EEulerRotationOrder GetControlPreferredEulerRotationOrderByIndex(int32 InElementIndex, bool bFromSettings = true) const
+	{
+		if(Elements.IsValidIndex(InElementIndex))
+		{
+			if(FRigControlElement* ControlElement = Cast<FRigControlElement>(Elements[InElementIndex]))
+			{
+				return GetControlPreferredEulerRotationOrder(ControlElement, bFromSettings);
+			}
+		}
+		return FRigPreferredEulerAngles::DefaultRotationOrder;
+	}
+
+	/**
+	 * Returns a control's preferred euler rotation order
+	 * @param InControlElement The element to look up
+	 * @param bFromSettings If true the rotation order will be looked up from the control's settings, otherwise from the currently set animation value
+	 * @return Returns the current preferred euler rotation order
+	 */
+	EEulerRotationOrder GetControlPreferredEulerRotationOrder(const FRigControlElement* InControlElement, bool bFromSettings = true) const
+	{
+		if(InControlElement)
+		{
+			if(bFromSettings)
+			{
+				return InControlElement->Settings.PreferredRotationOrder;
+			}
+			return InControlElement->PreferredEulerAngles.RotationOrder;
+		}
+		return FRigPreferredEulerAngles::DefaultRotationOrder;
+	}
+	
+	/**
+	 * Sets a control's preferred euler rotation order
+	 * @param InKey The key of the element to retrieve the current value for
+	 * @param InRotationOrder The rotation order to use when setting the euler angles
+	 */
+	UFUNCTION(BlueprintCallable, Category = URigHierarchy)
+	void SetControlPreferredRotationOrder(FRigElementKey InKey, EEulerRotationOrder InRotationOrder)
+	{
+		SetControlPreferredRotationOrderByIndex(GetIndex(InKey), InRotationOrder);
+	}
+	
+
+	/**
+	 * Sets a control's preferred euler rotation order
+	 * @param InElementIndex The element index to look up
+	 * @param InRotationOrder The rotation order to use when setting the euler angles
+	 */
+	UFUNCTION(BlueprintCallable, Category = URigHierarchy)
+	void SetControlPreferredRotationOrderByIndex(int32 InElementIndex, EEulerRotationOrder InRotationOrder)
+	{
+		if(Elements.IsValidIndex(InElementIndex))
+		{
+			if(FRigControlElement* ControlElement = Cast<FRigControlElement>(Elements[InElementIndex]))
+			{
+				SetControlPreferredRotationOrder(ControlElement, InRotationOrder);
+			}
+		}
+	}
+
+	/**
+	 * Sets a control's preferred euler rotation order
+	 * @param InControlElement The element to look up
+	 * @param InRotationOrder The rotation order to use when setting the euler angles
+	 */
+	void SetControlPreferredRotationOrder(FRigControlElement* InControlElement, EEulerRotationOrder InRotationOrder)
+	{
+		if(InControlElement)
+		{
+			InControlElement->PreferredEulerAngles.SetRotationOrder(InRotationOrder);
+		}
+	}
+
 	/**
 	 * Returns the pin type to use for a control
 	 * @param InControlElement The control to return the pin type for
diff --git a/Engine/Plugins/Animation/ControlRig/Source/ControlRig/Public/Rigs/RigHierarchyElements.h b/Engine/Plugins/Animation/ControlRig/Source/ControlRig/Public/Rigs/RigHierarchyElements.h
--- a/Engine/Plugins/Animation/ControlRig/Source/ControlRig/Public/Rigs/RigHierarchyElements.h
+++ b/Engine/Plugins/Animation/ControlRig/Source/ControlRig/Public/Rigs/RigHierarchyElements.h
@@ -470,7 +470,8 @@ struct CONTROLRIG_API FRigPreferredEulerAngles
 	FRotator GetRotator(bool bInitial = false) const;
 	FRotator SetRotator(const FRotator& InValue, bool bInitial = false, bool bFixEulerFlips = false);
 	FVector GetAngles(bool bInitial = false, EEulerRotationOrder InRotationOrder = DefaultRotationOrder) const;
-	void SetAngles(const FVector& InValue, bool bInitial = false, EEulerRotationOrder InRotationOrder = DefaultRotationOrder);
+	void SetAngles(const FVector& InValue, bool bInitial = false, EEulerRotationOrder InRotationOrder = DefaultRotationOrder, bool bFixEulerFlips = false);
+	void SetRotationOrder(EEulerRotationOrder InRotationOrder);
 
 	UPROPERTY(BlueprintReadOnly, EditAnywhere, Category = "Pose")
 	EEulerRotationOrder RotationOrder;
@@ -1187,6 +1188,12 @@ struct CONTROLRIG_API FRigControlSettings
 	UPROPERTY(BlueprintReadOnly, EditAnywhere, Category = Animation)
 	TArray<ERigControlTransformChannel> FilteredChannels;
 
+	/**
+	 * The euler rotation order this control prefers for animation
+	 */
+	UPROPERTY(BlueprintReadOnly, EditAnywhere, Category = Animation)
+	EEulerRotationOrder PreferredRotationOrder;
+
 #if WITH_EDITORONLY_DATA
 	/**
 	 * Deprecated properties.
diff --git a/Engine/Plugins/Animation/ControlRig/Source/ControlRigEditor/Private/ControlRigElementDetails.cpp b/Engine/Plugins/Animation/ControlRig/Source/ControlRigEditor/Private/ControlRigElementDetails.cpp
--- a/Engine/Plugins/Animation/ControlRig/Source/ControlRigEditor/Private/ControlRigElementDetails.cpp
+++ b/Engine/Plugins/Animation/ControlRig/Source/ControlRigEditor/Private/ControlRigElementDetails.cpp
@@ -2970,6 +2970,16 @@ void FRigControlElementDetails::CustomizeControl(IDetailLayoutBuilder& DetailBui
 		.IsEnabled(bIsEnabled);
 	}
 
+	if(CVarControlRigHierarchyEnableRotationOrder.GetValueOnAnyThread())
+	{
+		if(IsAnyControlOfValueType(ERigControlType::EulerTransform))
+		{
+			const TSharedPtr<IPropertyHandle> PreferredRotationOrderHandle = SettingsHandle->GetChildHandle(TEXT("PreferredRotationOrder"));
+			ControlCategory.AddProperty(PreferredRotationOrderHandle.ToSharedRef()).DisplayName(FText::FromString(TEXT("Preferred Rotation Order")))
+			.IsEnabled(bIsEnabled);
+		}
+	}
+
 	if(IsAnyControlOfValueType(ERigControlType::Integer))
 	{
 		const TSharedPtr<IPropertyHandle> ControlEnumHandle = SettingsHandle->GetChildHandle(TEXT("ControlEnum"));
diff --git a/Engine/Plugins/Animation/ControlRig/Source/ControlRigEditor/Private/Sequencer/ControlRigParameterTrackEditor.cpp b/Engine/Plugins/Animation/ControlRig/Source/ControlRigEditor/Private/Sequencer/ControlRigParameterTrackEditor.cpp
--- a/Engine/Plugins/Animation/ControlRig/Source/ControlRigEditor/Private/Sequencer/ControlRigParameterTrackEditor.cpp
+++ b/Engine/Plugins/Animation/ControlRig/Source/ControlRigEditor/Private/Sequencer/ControlRigParameterTrackEditor.cpp
@@ -1670,6 +1670,7 @@ static void EvaluateThisControl(UMovieSceneControlRigParameterSection* Section,
 	if(FRigControlElement* ControlElement = ControlRig->FindControl(ControlName))
 	{ 
 		FControlRigInteractionScope InteractionScope(ControlRig, ControlElement->GetKey());
+		URigHierarchy* RigHierarchy = ControlRig->GetHierarchy();
 		
 		//eval any space for this channel, if not additive section
 		if (Section->GetBlendType().Get() != EMovieSceneBlendType::Additive)
@@ -1677,7 +1678,6 @@ static void EvaluateThisControl(UMovieSceneControlRigParameterSection* Section,
 			TOptional<FMovieSceneControlRigSpaceBaseKey> SpaceKey = Section->EvaluateSpaceChannel(FrameTime, ControlName);
 			if (SpaceKey.IsSet())
 			{
-				URigHierarchy* RigHierarchy = ControlRig->GetHierarchy();
 				switch (SpaceKey.GetValue().SpaceType)
 				{
 				case EMovieSceneControlRigSpaceType::Parent:
@@ -1768,24 +1768,32 @@ static void EvaluateThisControl(UMovieSceneControlRigParameterSection* Section,
 			case ERigControlType::TransformNoScale:
 			case ERigControlType::EulerTransform:
 			{
+				// @MikeZ here I suppose we want to retrieve the rotation order and then also extract the Euler angles
+				// instead of an assumed FRotator coming from the section?
+				// EEulerRotationOrder RotationOrder = SomehowGetRotationOrder();
+					
 				TOptional <FEulerTransform> Value = Section->EvaluateTransformParameter(FrameTime, ControlName);
 				if (Value.IsSet())
 				{
 					if (ControlElement->Settings.ControlType == ERigControlType::Transform)
 					{
 						ControlRig->SetControlValue<FRigControlValue::FTransform_Float>(ControlName, Value.GetValue().ToFTransform(), true, EControlRigSetKey::Never, bSetupUndo);
-						ControlRig->GetHierarchy()->SetControlPreferredRotator(ControlElement, Value.GetValue().Rotation);
+						// @MikeZ here we want to call SetControlPreferredEulerAngles instead
+						RigHierarchy->SetControlPreferredRotator(ControlElement, Value.GetValue().Rotation);
 					}
 					else if (ControlElement->Settings.ControlType == ERigControlType::TransformNoScale)
 					{
 						FTransformNoScale NoScale = Value.GetValue().ToFTransform();
 						ControlRig->SetControlValue<FRigControlValue::FTransformNoScale_Float>(ControlName, NoScale, true, EControlRigSetKey::Never, bSetupUndo);
-						ControlRig->GetHierarchy()->SetControlPreferredRotator(ControlElement, Value.GetValue().Rotation);
+						// @MikeZ here we want to call SetControlPreferredEulerAngles instead
+						RigHierarchy->SetControlPreferredRotator(ControlElement, Value.GetValue().Rotation);
 					}
 					else if (ControlElement->Settings.ControlType == ERigControlType::EulerTransform)
 					{
 						const FEulerTransform& Euler = Value.GetValue();
 						ControlRig->SetControlValue<FRigControlValue::FEulerTransform_Float>(ControlName, Euler, true, EControlRigSetKey::Never, bSetupUndo);
+						// @MikeZ here we want to call SetControlPreferredEulerAngles instead
+						RigHierarchy->SetControlPreferredRotator(ControlElement, Euler.Rotation);
 					}
 				}
 				break;
