Control Rig: Interactive sliders working in details panels

diff --git a/Engine/Plugins/Animation/ControlRig/Source/ControlRig/Public/ControlRig.h b/Engine/Plugins/Animation/ControlRig/Source/ControlRig/Public/ControlRig.h
--- a/Engine/Plugins/Animation/ControlRig/Source/ControlRig/Public/ControlRig.h
+++ b/Engine/Plugins/Animation/ControlRig/Source/ControlRig/Public/ControlRig.h
@@ -713,6 +713,7 @@ private:
 	friend class UControlRigGraph;
 	friend class AControlRigControlActor;
 	friend class AControlRigShapeActor;
+	friend class FRigTransformElementDetails;
 };
 
 class CONTROLRIG_API FControlRigBracketScope
diff --git a/Engine/Plugins/Animation/ControlRig/Source/ControlRigEditor/Private/ControlRigElementDetails.cpp b/Engine/Plugins/Animation/ControlRig/Source/ControlRigEditor/Private/ControlRigElementDetails.cpp
--- a/Engine/Plugins/Animation/ControlRig/Source/ControlRigEditor/Private/ControlRigElementDetails.cpp
+++ b/Engine/Plugins/Animation/ControlRig/Source/ControlRigEditor/Private/ControlRigElementDetails.cpp
@@ -2267,6 +2267,59 @@ FDetailWidgetRow& FRigTransformElementDetails::CreateEulerTransformValueWidgetRo
 		SliderTransaction.Reset();
 	});
 
+	TransformWidgetArgs.OnBeginSliderMovement_Lambda(
+		[
+			this
+		](
+			ESlateTransformComponent::Type Component,
+			ESlateRotationRepresentation::Type Representation,
+			ESlateTransformSubComponent::Type SubComponent)
+		{
+			if (UControlRig* DebuggedRig = Cast<UControlRig>(PerElementInfos[0].GetBlueprint()->GetObjectBeingDebugged()))
+			{
+				EControlRigInteractionType Type = EControlRigInteractionType::None;
+				switch (Component)
+				{
+					case ESlateTransformComponent::Location: Type = EControlRigInteractionType::Translate; break;
+					case ESlateTransformComponent::Rotation: Type = EControlRigInteractionType::Rotate; break;
+					case ESlateTransformComponent::Scale: Type = EControlRigInteractionType::Scale; break;
+					default: Type = EControlRigInteractionType::All;
+				}
+				DebuggedRig->InteractionType = (uint8)Type;
+				DebuggedRig->ElementsBeingInteracted.Reset();
+				for (FPerElementInfo& ElementInfo : PerElementInfos)
+				{
+					DebuggedRig->ElementsBeingInteracted.AddUnique(ElementInfo.Element.GetKey());
+				}
+				
+				FControlRigInteractionScope* InteractionScope = new FControlRigInteractionScope(DebuggedRig);
+				InteractionScopes.Add(InteractionScope);
+			}
+		});
+	TransformWidgetArgs.OnEndSliderMovement_Lambda(
+		[
+			this
+		](
+			ESlateTransformComponent::Type Component,
+			ESlateRotationRepresentation::Type Representation,
+			ESlateTransformSubComponent::Type SubComponent,
+			FVector::FReal InNumericValue)
+		{
+			if (UControlRig* DebuggedRig = Cast<UControlRig>(PerElementInfos[0].GetBlueprint()->GetObjectBeingDebugged()))
+			{
+				DebuggedRig->InteractionType = (uint8)EControlRigInteractionType::None;
+				DebuggedRig->ElementsBeingInteracted.Reset();
+			}
+			for (FControlRigInteractionScope* InteractionScope : InteractionScopes)
+			{
+				if (InteractionScope)
+				{
+					delete InteractionScope; 
+				}
+			}
+			InteractionScopes.Reset();
+		});
+
 	TransformWidgetArgs.OnCopyToClipboard_Lambda([Keys, IsComponentRelative, ConformComponentRelative, GetSingleTransform](
 		ESlateTransformComponent::Type InComponent
 		)
diff --git a/Engine/Plugins/Animation/ControlRig/Source/ControlRigEditor/Private/ControlRigElementDetails.h b/Engine/Plugins/Animation/ControlRig/Source/ControlRigEditor/Private/ControlRigElementDetails.h
--- a/Engine/Plugins/Animation/ControlRig/Source/ControlRigEditor/Private/ControlRigElementDetails.h
+++ b/Engine/Plugins/Animation/ControlRig/Source/ControlRigEditor/Private/ControlRigElementDetails.h
@@ -250,6 +250,7 @@ protected:
 
 	TArray<FPerElementInfo> PerElementInfos;
 	FDelegateHandle MetadataHandle;
+	TArray<FControlRigInteractionScope*> InteractionScopes;
 	
 	TSharedPtr<SButton> SelectParentElementButton;
 };
diff --git a/Engine/Plugins/Animation/ControlRig/Source/ControlRigEditor/Private/EditMode/ControlRigControlsProxy.cpp b/Engine/Plugins/Animation/ControlRig/Source/ControlRigEditor/Private/EditMode/ControlRigControlsProxy.cpp
--- a/Engine/Plugins/Animation/ControlRig/Source/ControlRigEditor/Private/EditMode/ControlRigControlsProxy.cpp
+++ b/Engine/Plugins/Animation/ControlRig/Source/ControlRigEditor/Private/EditMode/ControlRigControlsProxy.cpp
@@ -19,6 +19,7 @@
 #include "DetailCategoryBuilder.h"
 #include "DetailWidgetRow.h"
 #include "SEnumCombo.h"
+#include "ControlRig.h"
 
 
 void UControlRigControlsProxy::SetIsMultiple(bool bIsVal)
@@ -102,9 +103,9 @@ void UControlRigControlsProxy::SelectionChanged(bool bInSelected)
 	}
 }
 
-void UControlRigControlsProxy::PostEditChangeProperty(struct FPropertyChangedEvent& PropertyChangedEvent)
+void UControlRigControlsProxy::PostEditChangeChainProperty(struct FPropertyChangedChainEvent& PropertyChangedEvent)
 {
-	Super::PostEditChangeProperty(PropertyChangedEvent);
+	Super::PostEditChangeChainProperty(PropertyChangedEvent);
 	if (PropertyChangedEvent.Property && PropertyChangedEvent.Property->GetFName() == GET_MEMBER_NAME_CHECKED(UControlRigControlsProxy, bSelected))
 	{
 		FRigControlElement* ControlElement = GetControlElement();
@@ -115,6 +116,57 @@ void UControlRigControlsProxy::PostEditChangeProperty(struct FPropertyChangedEve
 			ControlRig->Evaluate_AnyThread();
 		}
 	}
+#if WITH_EDITOR
+	if (PropertyChangedEvent.Property)
+	{
+		FRigControlElement* ControlElement = GetControlElement();
+		if (PropertyChangedEvent.ChangeType == EPropertyChangeType::Interactive)
+		{
+			EControlRigInteractionType Type = EControlRigInteractionType::None;
+			FProperty* Owner = PropertyChangedEvent.Property->GetOwnerProperty();
+			if (FEditPropertyChain::TDoubleLinkedListNode* MemberNode = PropertyChangedEvent.PropertyChain.GetActiveMemberNode())
+			{
+				if (FProperty* MemberProperty = MemberNode->GetValue())
+				{
+					if (MemberProperty->GetFName() == GET_MEMBER_NAME_CHECKED(UControlRigTransformControlProxy, Transform))
+					{
+						if (FEditPropertyChain::TDoubleLinkedListNode* NextNode = MemberNode->GetNextNode())
+						{
+							if (FProperty* SubMemberProperty = NextNode->GetValue())
+							{
+								if (SubMemberProperty->GetFName() == GET_MEMBER_NAME_CHECKED(FEulerTransform, Location))
+								{
+									Type = EControlRigInteractionType::Translate;
+								}
+								else if (SubMemberProperty->GetFName() == GET_MEMBER_NAME_CHECKED(FEulerTransform, Rotation))
+								{
+									Type = EControlRigInteractionType::Rotate;
+								}
+								else if (SubMemberProperty->GetFName() == GET_MEMBER_NAME_CHECKED(FEulerTransform, Scale))
+								{
+									Type = EControlRigInteractionType::Scale;
+								}
+							}
+						}
+					}
+				}
+			}
+			FControlRigInteractionScope* InteractionScope = new FControlRigInteractionScope(ControlRig.Get(), ControlElement->GetKey(), Type);
+			InteractionScopes.Add(InteractionScope);
+		}
+		else
+		{
+			for (FControlRigInteractionScope* Scope : InteractionScopes)
+			{
+				if (Scope)
+				{
+					delete Scope; 
+				}
+			}
+			InteractionScopes.Reset();
+		}
+	}
+#endif
 }
 
 #if WITH_EDITOR
@@ -137,9 +189,9 @@ FRigControlElement* UControlRigControlsProxy::GetControlElement() const
 	return nullptr;
 }
 
-void UControlRigTransformControlProxy::PostEditChangeProperty(struct FPropertyChangedEvent& PropertyChangedEvent)
+void UControlRigTransformControlProxy::PostEditChangeChainProperty(struct FPropertyChangedChainEvent& PropertyChangedEvent)
 {
-	Super::PostEditChangeProperty(PropertyChangedEvent);
+	Super::PostEditChangeChainProperty(PropertyChangedEvent);
 	if ((PropertyChangedEvent.Property && PropertyChangedEvent.Property->GetFName() == GET_MEMBER_NAME_CHECKED(UControlRigTransformControlProxy, Transform))
 		|| (PropertyChangedEvent.MemberProperty && PropertyChangedEvent.MemberProperty->GetFName() == GET_MEMBER_NAME_CHECKED(UControlRigTransformControlProxy, Transform)))
 	{
@@ -147,7 +199,6 @@ void UControlRigTransformControlProxy::PostEditChangeProperty(struct FPropertyCh
 		if (ControlElement && ControlRig.IsValid())
 		{
 			//MUST set through ControlRig
-			FControlRigInteractionScope InteractionScope(ControlRig.Get(), ControlElement->GetKey());
 			const FTransform RealTransform = Transform.ToFTransform(); //Transform is FEulerTransform
 			ControlRig->SetControlValue<FRigControlValue::FTransform_Float>(ControlName, RealTransform, true, EControlRigSetKey::DoNotCare,false);
 			ControlRig->GetHierarchy()->SetControlPreferredRotator(ControlElement, Transform.Rotation);
@@ -211,9 +262,9 @@ void UControlRigTransformControlProxy::SetKey(const IPropertyHandle& KeyedProper
 	}
 }
 
-void UControlRigTransformNoScaleControlProxy::PostEditChangeProperty(struct FPropertyChangedEvent& PropertyChangedEvent)
+void UControlRigTransformNoScaleControlProxy::PostEditChangeChainProperty(struct FPropertyChangedChainEvent& PropertyChangedEvent)
 {
-	Super::PostEditChangeProperty(PropertyChangedEvent);
+	Super::PostEditChangeChainProperty(PropertyChangedEvent);
 	if ((PropertyChangedEvent.Property && PropertyChangedEvent.Property->GetFName() == GET_MEMBER_NAME_CHECKED(UControlRigTransformNoScaleControlProxy, Transform))
 		|| (PropertyChangedEvent.MemberProperty && PropertyChangedEvent.MemberProperty->GetFName() == GET_MEMBER_NAME_CHECKED(UControlRigTransformNoScaleControlProxy, Transform)))
 	{
@@ -221,7 +272,6 @@ void UControlRigTransformNoScaleControlProxy::PostEditChangeProperty(struct FPro
 		if (ControlElement && ControlRig.IsValid())
 		{
 			//MUST set through ControlRig
-			FControlRigInteractionScope InteractionScope(ControlRig.Get(), ControlElement->GetKey());
 			ControlRig->SetControlValue<FRigControlValue::FTransformNoScale_Float>(ControlName, Transform, true, EControlRigSetKey::DoNotCare,false);
 			ControlRig->Evaluate_AnyThread();
 
@@ -274,9 +324,9 @@ void UControlRigTransformNoScaleControlProxy::SetKey(const IPropertyHandle& Keye
 }
 
 
-void UControlRigEulerTransformControlProxy::PostEditChangeProperty(struct FPropertyChangedEvent& PropertyChangedEvent)
+void UControlRigEulerTransformControlProxy::PostEditChangeChainProperty(struct FPropertyChangedChainEvent& PropertyChangedEvent)
 {
-	Super::PostEditChangeProperty(PropertyChangedEvent);
+	Super::PostEditChangeChainProperty(PropertyChangedEvent);
 	if ((PropertyChangedEvent.Property && PropertyChangedEvent.Property->GetFName() == GET_MEMBER_NAME_CHECKED(UControlRigEulerTransformControlProxy, Transform))
 		|| (PropertyChangedEvent.MemberProperty && PropertyChangedEvent.MemberProperty->GetFName() == GET_MEMBER_NAME_CHECKED(UControlRigEulerTransformControlProxy, Transform)))
 	{
@@ -284,7 +334,6 @@ void UControlRigEulerTransformControlProxy::PostEditChangeProperty(struct FPrope
 		if (ControlElement && ControlRig.IsValid())
 		{
 			//MUST set through ControlRig
-			FControlRigInteractionScope InteractionScope(ControlRig.Get(), ControlElement->GetKey());
 			ControlRig->SetControlValue<FRigControlValue::FEulerTransform_Float>(ControlName, Transform, true, EControlRigSetKey::DoNotCare,false);
 			ControlRig->Evaluate_AnyThread();
 
@@ -342,16 +391,15 @@ void UControlRigEulerTransformControlProxy::SetKey(const IPropertyHandle& KeyedP
 	}
 }
 
-void UControlRigFloatControlProxy::PostEditChangeProperty(struct FPropertyChangedEvent& PropertyChangedEvent)
+void UControlRigFloatControlProxy::PostEditChangeChainProperty(struct FPropertyChangedChainEvent& PropertyChangedEvent)
 {
-	Super::PostEditChangeProperty(PropertyChangedEvent);
+	Super::PostEditChangeChainProperty(PropertyChangedEvent);
 	if (PropertyChangedEvent.Property && PropertyChangedEvent.Property->GetFName() == GET_MEMBER_NAME_CHECKED(UControlRigFloatControlProxy, Float))
 	{
 		FRigControlElement* ControlElement = GetControlElement();
 		if (ControlElement && ControlRig.IsValid())
 		{
 			//MUST set through ControlRig
-			FControlRigInteractionScope InteractionScope(ControlRig.Get(), ControlElement->GetKey());
 			ControlRig->SetControlValue<float>(ControlName, Float, true, EControlRigSetKey::DoNotCare,false);
 			ControlRig->Evaluate_AnyThread();
 
@@ -394,16 +442,15 @@ void UControlRigFloatControlProxy::SetKey(const IPropertyHandle& KeyedPropertyHa
 	}
 }
 
-void UControlRigIntegerControlProxy::PostEditChangeProperty(struct FPropertyChangedEvent& PropertyChangedEvent)
+void UControlRigIntegerControlProxy::PostEditChangeChainProperty(struct FPropertyChangedChainEvent& PropertyChangedEvent)
 {
-	Super::PostEditChangeProperty(PropertyChangedEvent);
+	Super::PostEditChangeChainProperty(PropertyChangedEvent);
 	if (PropertyChangedEvent.Property && PropertyChangedEvent.Property->GetFName() == GET_MEMBER_NAME_CHECKED(UControlRigIntegerControlProxy, Integer))
 	{
 		FRigControlElement* ControlElement = GetControlElement();
 		if (ControlElement && ControlRig.IsValid())
 		{
 			//MUST set through ControlRig
-			FControlRigInteractionScope InteractionScope(ControlRig.Get(), ControlElement->GetKey());
 			ControlRig->SetControlValue<int32>(ControlName, Integer, true, EControlRigSetKey::DoNotCare,false);
 			ControlRig->Evaluate_AnyThread();
 		}
@@ -444,16 +491,15 @@ void UControlRigIntegerControlProxy::SetKey(const IPropertyHandle& KeyedProperty
 	}
 }
 
-void UControlRigEnumControlProxy::PostEditChangeProperty(struct FPropertyChangedEvent& PropertyChangedEvent)
+void UControlRigEnumControlProxy::PostEditChangeChainProperty(struct FPropertyChangedChainEvent& PropertyChangedEvent)
 {
-	Super::PostEditChangeProperty(PropertyChangedEvent);
+	Super::PostEditChangeChainProperty(PropertyChangedEvent);
 	if (PropertyChangedEvent.Property && PropertyChangedEvent.Property->GetFName() == GET_MEMBER_NAME_CHECKED(UControlRigEnumControlProxy, Enum))
 	{
 		FRigControlElement* ControlElement = GetControlElement();
 		if (ControlElement && ControlRig.IsValid())
 		{
 			//MUST set through ControlRig
-			FControlRigInteractionScope InteractionScope(ControlRig.Get(), ControlElement->GetKey());
 			ControlRig->SetControlValue<int32>(ControlName, Enum.EnumIndex, true, EControlRigSetKey::DoNotCare,false);
 			ControlRig->Evaluate_AnyThread();
 
@@ -499,9 +545,9 @@ void UControlRigEnumControlProxy::SetKey(const IPropertyHandle& KeyedPropertyHan
 	}
 }
 
-void UControlRigVectorControlProxy::PostEditChangeProperty(struct FPropertyChangedEvent& PropertyChangedEvent)
+void UControlRigVectorControlProxy::PostEditChangeChainProperty(struct FPropertyChangedChainEvent& PropertyChangedEvent)
 {
-	Super::PostEditChangeProperty(PropertyChangedEvent);
+	Super::PostEditChangeChainProperty(PropertyChangedEvent);
 	if ((PropertyChangedEvent.Property && PropertyChangedEvent.Property->GetFName() == GET_MEMBER_NAME_CHECKED(UControlRigVectorControlProxy, Vector))
 		|| (PropertyChangedEvent.MemberProperty && PropertyChangedEvent.MemberProperty->GetFName() == GET_MEMBER_NAME_CHECKED(UControlRigVectorControlProxy, Vector)))
 	{
@@ -509,7 +555,6 @@ void UControlRigVectorControlProxy::PostEditChangeProperty(struct FPropertyChang
 		if (ControlElement && ControlRig.IsValid())
 		{
 			//MUST set through ControlRig
-			FControlRigInteractionScope InteractionScope(ControlRig.Get(), ControlElement->GetKey());
 			ControlRig->SetControlValue<FVector3f>(ControlName, Vector, true, EControlRigSetKey::DoNotCare,false);
 			ControlRig->Evaluate_AnyThread();
 
@@ -599,9 +644,9 @@ void UControlRigVectorControlProxy::SetKey(const IPropertyHandle& KeyedPropertyH
 	}
 }
 
-void UControlRigVector2DControlProxy::PostEditChangeProperty(struct FPropertyChangedEvent& PropertyChangedEvent)
+void UControlRigVector2DControlProxy::PostEditChangeChainProperty(struct FPropertyChangedChainEvent& PropertyChangedEvent)
 {
-	Super::PostEditChangeProperty(PropertyChangedEvent);
+	Super::PostEditChangeChainProperty(PropertyChangedEvent);
 	if ((PropertyChangedEvent.Property && PropertyChangedEvent.Property->GetFName() == GET_MEMBER_NAME_CHECKED(UControlRigVector2DControlProxy, Vector2D))
 		|| ((PropertyChangedEvent.MemberProperty && PropertyChangedEvent.MemberProperty->GetFName() == GET_MEMBER_NAME_CHECKED(UControlRigVector2DControlProxy, Vector2D))))
 	{
@@ -609,7 +654,6 @@ void UControlRigVector2DControlProxy::PostEditChangeProperty(struct FPropertyCha
 		if (ControlElement && ControlRig.IsValid())
 		{
 			//MUST set through ControlRig
-			FControlRigInteractionScope InteractionScope(ControlRig.Get(), ControlElement->GetKey());
 			ControlRig->SetControlValue<FVector3f>(ControlName, FVector3f(Vector2D.X, Vector2D.Y, 0.f), true, EControlRigSetKey::DoNotCare,false);
 			ControlRig->Evaluate_AnyThread();
 
@@ -664,16 +708,15 @@ void UControlRigVector2DControlProxy::SetKey(const IPropertyHandle& KeyedPropert
 }
 
 
-void UControlRigBoolControlProxy::PostEditChangeProperty(struct FPropertyChangedEvent& PropertyChangedEvent)
+void UControlRigBoolControlProxy::PostEditChangeChainProperty(struct FPropertyChangedChainEvent& PropertyChangedEvent)
 {
-	Super::PostEditChangeProperty(PropertyChangedEvent);
+	Super::PostEditChangeChainProperty(PropertyChangedEvent);
 	if (PropertyChangedEvent.Property && PropertyChangedEvent.Property->GetFName() == GET_MEMBER_NAME_CHECKED(UControlRigBoolControlProxy, Bool))
 	{
 		FRigControlElement* ControlElement = GetControlElement();
 		if (ControlElement && ControlRig.IsValid())
 		{
 			//MUST set through ControlRig
-			FControlRigInteractionScope InteractionScope(ControlRig.Get(), ControlElement->GetKey());
 			ControlRig->SetControlValue<bool>(ControlName, Bool, true, EControlRigSetKey::DoNotCare,false);
 			ControlRig->Evaluate_AnyThread();
 		}
diff --git a/Engine/Plugins/Animation/ControlRig/Source/ControlRigEditor/Private/EditMode/ControlRigControlsProxy.h b/Engine/Plugins/Animation/ControlRig/Source/ControlRigEditor/Private/EditMode/ControlRigControlsProxy.h
--- a/Engine/Plugins/Animation/ControlRig/Source/ControlRigEditor/Private/EditMode/ControlRigControlsProxy.h
+++ b/Engine/Plugins/Animation/ControlRig/Source/ControlRigEditor/Private/EditMode/ControlRigControlsProxy.h
@@ -16,6 +16,7 @@
 struct FRigControlElement;
 class UControlRig;
 class IPropertyHandle;
+class FControlRigInteractionScope;
 
 UCLASS(Abstract)
 class UControlRigControlsProxy : public UObject
@@ -31,7 +32,7 @@ public:
 	virtual void SetKey(const IPropertyHandle& KeyedPropertyHandle) {};
 
 	// UObject interface
-	virtual void PostEditChangeProperty(struct FPropertyChangedEvent& PropertyChangedEvent) override;
+	virtual void PostEditChangeChainProperty( struct FPropertyChangedChainEvent& PropertyChangedEvent ) override;
 #if WITH_EDITOR
 	virtual void PostEditUndo() override;
 #endif
@@ -53,6 +54,9 @@ public:
 protected:
 	bool bIsMultiple = 0;
 	FName Name;
+#if WITH_EDITOR
+	TArray<FControlRigInteractionScope*> InteractionScopes;
+#endif
 };
 
 UCLASS()
@@ -62,7 +66,7 @@ class UControlRigTransformControlProxy : public UControlRigControlsProxy
 	UControlRigTransformControlProxy() {}
 
 	// UObject interface
-	virtual void PostEditChangeProperty(struct FPropertyChangedEvent& PropertyChangedEvent) override;
+	virtual void PostEditChangeChainProperty(struct FPropertyChangedChainEvent& PropertyChangedEvent) override;
 #if WITH_EDITOR
 	virtual void PostEditUndo() override;
 #endif
@@ -85,7 +89,7 @@ class UControlRigEulerTransformControlProxy : public UControlRigControlsProxy
 	UControlRigEulerTransformControlProxy() {}
 
 	// UObject interface
-	virtual void PostEditChangeProperty(struct FPropertyChangedEvent& PropertyChangedEvent) override;
+	virtual void PostEditChangeChainProperty(struct FPropertyChangedChainEvent& PropertyChangedEvent) override;
 #if WITH_EDITOR
 	virtual void PostEditUndo() override;
 #endif
@@ -108,7 +112,7 @@ class UControlRigTransformNoScaleControlProxy : public UControlRigControlsProxy
 	UControlRigTransformNoScaleControlProxy() {}
 
 	// UObject interface
-	virtual void PostEditChangeProperty(struct FPropertyChangedEvent& PropertyChangedEvent) override;
+	virtual void PostEditChangeChainProperty(struct FPropertyChangedChainEvent& PropertyChangedEvent) override;
 #if WITH_EDITOR
 	virtual void PostEditUndo() override;
 #endif
@@ -130,7 +134,7 @@ class UControlRigFloatControlProxy : public UControlRigControlsProxy
 	UControlRigFloatControlProxy()  {}
 
 	// UObject interface
-	virtual void PostEditChangeProperty(struct FPropertyChangedEvent& PropertyChangedEvent) override;
+	virtual void PostEditChangeChainProperty(struct FPropertyChangedChainEvent& PropertyChangedEvent) override;
 #if WITH_EDITOR
 	virtual void PostEditUndo() override;
 #endif
@@ -152,7 +156,7 @@ class UControlRigIntegerControlProxy : public UControlRigControlsProxy
 	UControlRigIntegerControlProxy() {}
 
 	// UObject interface
-	virtual void PostEditChangeProperty(struct FPropertyChangedEvent& PropertyChangedEvent) override;
+	virtual void PostEditChangeChainProperty(struct FPropertyChangedChainEvent& PropertyChangedEvent) override;
 #if WITH_EDITOR
 	virtual void PostEditUndo() override;
 #endif
@@ -216,7 +220,7 @@ class UControlRigEnumControlProxy : public UControlRigControlsProxy
 	UControlRigEnumControlProxy() {}
 
 	// UObject interface
-	virtual void PostEditChangeProperty(struct FPropertyChangedEvent& PropertyChangedEvent) override;
+	virtual void PostEditChangeChainProperty(struct FPropertyChangedChainEvent& PropertyChangedEvent) override;
 #if WITH_EDITOR
 	virtual void PostEditUndo() override;
 #endif
@@ -238,7 +242,7 @@ class UControlRigVectorControlProxy : public UControlRigControlsProxy
 	UControlRigVectorControlProxy() {}
 
 	// UObject interface
-	virtual void PostEditChangeProperty(struct FPropertyChangedEvent& PropertyChangedEvent) override;
+	virtual void PostEditChangeChainProperty(struct FPropertyChangedChainEvent& PropertyChangedEvent) override;
 #if WITH_EDITOR
 	virtual void PostEditUndo() override;
 #endif
@@ -260,7 +264,7 @@ class UControlRigVector2DControlProxy : public UControlRigControlsProxy
 	UControlRigVector2DControlProxy() {}
 
 	// UObject interface
-	virtual void PostEditChangeProperty(struct FPropertyChangedEvent& PropertyChangedEvent) override;
+	virtual void PostEditChangeChainProperty(struct FPropertyChangedChainEvent& PropertyChangedEvent) override;
 #if WITH_EDITOR
 	virtual void PostEditUndo() override;
 #endif
@@ -282,7 +286,7 @@ class UControlRigBoolControlProxy : public UControlRigControlsProxy
 	UControlRigBoolControlProxy() {}
 
 	// UObject interface
-	virtual void PostEditChangeProperty(struct FPropertyChangedEvent& PropertyChangedEvent) override;
+	virtual void PostEditChangeChainProperty(struct FPropertyChangedChainEvent& PropertyChangedEvent) override;
 #if WITH_EDITOR
 	virtual void PostEditUndo() override;
 #endif
diff --git a/Engine/Source/Developer/AnimationWidgets/AnimationWidgets.Build.cs b/Engine/Source/Developer/AnimationWidgets/AnimationWidgets.Build.cs
--- a/Engine/Source/Developer/AnimationWidgets/AnimationWidgets.Build.cs
+++ b/Engine/Source/Developer/AnimationWidgets/AnimationWidgets.Build.cs
@@ -19,5 +19,15 @@ public class AnimationWidgets : ModuleRules
 				"AnimationCore",
 			}
 		);
+
+		if (Target.Type == TargetType.Editor)
+		{
+			PrivateDependencyModuleNames.AddRange(
+				new string[]
+				{
+					"PropertyEditor"
+				}
+			);
+		}
 	}
 }
diff --git a/Engine/Source/Developer/AnimationWidgets/Public/SAdvancedRotationInputBox.h b/Engine/Source/Developer/AnimationWidgets/Public/SAdvancedRotationInputBox.h
--- a/Engine/Source/Developer/AnimationWidgets/Public/SAdvancedRotationInputBox.h
+++ b/Engine/Source/Developer/AnimationWidgets/Public/SAdvancedRotationInputBox.h
@@ -81,6 +81,12 @@ public:
 	/** Notification for numerictype value committed */
 	DECLARE_DELEGATE_FourParams(FOnNumericValueCommitted, ESlateRotationRepresentation::Type, ESlateTransformSubComponent::Type, NumericType, ETextCommit::Type);
 
+	/** Notification for begin slider movement */
+	DECLARE_DELEGATE_TwoParams(FOnBeginSliderMovement, ESlateRotationRepresentation::Type, ESlateTransformSubComponent::Type);
+
+	/** Notification for end slider movement */
+	DECLARE_DELEGATE_ThreeParams(FOnEndSliderMovement, ESlateRotationRepresentation::Type, ESlateTransformSubComponent::Type, NumericType);
+
 	/** Delegate to retrieve toggle checkstate */
 	DECLARE_DELEGATE_RetVal_TwoParams(ECheckBoxState, FOnGetToggleChecked, ESlateRotationRepresentation::Type, ESlateTransformSubComponent::Type);
 
@@ -126,6 +132,12 @@ public:
 		/** Optional delegate to notify a committed value */
 		SLATE_EVENT( FOnNumericValueCommitted, OnNumericValueCommitted )
 
+		/** Optional delegate to notify a begin slider movement */
+		SLATE_EVENT(FOnBeginSliderMovement, OnBeginSliderMovement)
+
+		/** Optional delegate to notify a end slider movement */
+		SLATE_EVENT(FOnEndSliderMovement, OnEndSliderMovement)
+
 		/** Should the axis labels be colored */
 		SLATE_ARGUMENT( bool, bColorAxisLabels )		
 
@@ -166,6 +178,8 @@ public:
 		OnGetNumericValue = InArgs._OnGetNumericValue; 
 		OnNumericValueChanged = InArgs._OnNumericValueChanged; 
 		OnNumericValueCommitted = InArgs._OnNumericValueCommitted; 
+		OnBeginSliderMovement = InArgs._OnBeginSliderMovement;
+		OnEndSliderMovement = InArgs._OnEndSliderMovement;
 		OnGetToggleChecked = InArgs._OnGetToggleChecked; 
 		OnToggleChanged = InArgs._OnToggleChanged;
 
@@ -208,6 +222,12 @@ public:
 				.OnToggleRollChanged(this, &SAdvancedRotationInputBox::HandleOnToggleChanged, ESlateTransformSubComponent::Roll)
 				.OnToggleYawChanged(this, &SAdvancedRotationInputBox::HandleOnToggleChanged, ESlateTransformSubComponent::Yaw)
 				.OnTogglePitchChanged(this, &SAdvancedRotationInputBox::HandleOnToggleChanged, ESlateTransformSubComponent::Pitch)
+				.OnPitchBeginSliderMovement(this, &SAdvancedRotationInputBox::HandleOnBeginSliderMovement, ESlateTransformSubComponent::Pitch)
+				.OnYawBeginSliderMovement(this, &SAdvancedRotationInputBox::HandleOnBeginSliderMovement, ESlateTransformSubComponent::Yaw)
+				.OnRollBeginSliderMovement(this, &SAdvancedRotationInputBox::HandleOnBeginSliderMovement, ESlateTransformSubComponent::Roll)
+				.OnPitchEndSliderMovement(this, &SAdvancedRotationInputBox::HandleOnEndSliderMovement, ESlateTransformSubComponent::Pitch)
+				.OnYawEndSliderMovement(this, &SAdvancedRotationInputBox::HandleOnEndSliderMovement, ESlateTransformSubComponent::Yaw)
+				.OnRollEndSliderMovement(this, &SAdvancedRotationInputBox::HandleOnEndSliderMovement, ESlateTransformSubComponent::Roll)
 				.TogglePadding(InArgs._TogglePadding)
 				.Visibility(this, &SAdvancedRotationInputBox::IsRotatorInputBoxVisible)
 			]
@@ -240,6 +260,12 @@ public:
 				.OnToggleXChanged(this, &SAdvancedRotationInputBox::HandleOnToggleChanged, ESlateTransformSubComponent::X)
 				.OnToggleYChanged(this, &SAdvancedRotationInputBox::HandleOnToggleChanged, ESlateTransformSubComponent::Y)
 				.OnToggleZChanged(this, &SAdvancedRotationInputBox::HandleOnToggleChanged, ESlateTransformSubComponent::Z)
+				.OnXBeginSliderMovement(this, &SAdvancedRotationInputBox::HandleOnBeginSliderMovement, ESlateTransformSubComponent::X)
+				.OnXEndSliderMovement(this, &SAdvancedRotationInputBox::HandleOnEndSliderMovement, ESlateTransformSubComponent::X)
+				.OnYBeginSliderMovement(this, &SAdvancedRotationInputBox::HandleOnBeginSliderMovement, ESlateTransformSubComponent::Y)
+				.OnYEndSliderMovement(this, &SAdvancedRotationInputBox::HandleOnEndSliderMovement, ESlateTransformSubComponent::Y)
+				.OnZBeginSliderMovement(this, &SAdvancedRotationInputBox::HandleOnBeginSliderMovement, ESlateTransformSubComponent::Z)
+				.OnZEndSliderMovement(this, &SAdvancedRotationInputBox::HandleOnEndSliderMovement, ESlateTransformSubComponent::Z)
 				.TogglePadding(InArgs._TogglePadding)
 				.Visibility(this, &SAdvancedRotationInputBox::IsEulerInputBoxVisible)
 			]
@@ -277,6 +303,14 @@ public:
 				.OnToggleYChanged(this, &SAdvancedRotationInputBox::HandleOnToggleChanged, ESlateTransformSubComponent::Y)
 				.OnToggleZChanged(this, &SAdvancedRotationInputBox::HandleOnToggleChanged, ESlateTransformSubComponent::Z)
 				.OnToggleWChanged(this, &SAdvancedRotationInputBox::HandleOnToggleChanged, ESlateTransformSubComponent::W)
+				.OnXBeginSliderMovement(this, &SAdvancedRotationInputBox::HandleOnBeginSliderMovement, ESlateTransformSubComponent::X)
+				.OnXEndSliderMovement(this, &SAdvancedRotationInputBox::HandleOnEndSliderMovement, ESlateTransformSubComponent::X)
+				.OnYBeginSliderMovement(this, &SAdvancedRotationInputBox::HandleOnBeginSliderMovement, ESlateTransformSubComponent::Y)
+				.OnYEndSliderMovement(this, &SAdvancedRotationInputBox::HandleOnEndSliderMovement, ESlateTransformSubComponent::Y)
+				.OnZBeginSliderMovement(this, &SAdvancedRotationInputBox::HandleOnBeginSliderMovement, ESlateTransformSubComponent::Z)
+				.OnZEndSliderMovement(this, &SAdvancedRotationInputBox::HandleOnEndSliderMovement, ESlateTransformSubComponent::Z)
+				.OnWBeginSliderMovement(this, &SAdvancedRotationInputBox::HandleOnBeginSliderMovement, ESlateTransformSubComponent::W)
+				.OnWEndSliderMovement(this, &SAdvancedRotationInputBox::HandleOnEndSliderMovement, ESlateTransformSubComponent::W)
 				.TogglePadding(InArgs._TogglePadding)
 				.Visibility(this, &SAdvancedRotationInputBox::IsQuaternionInputBoxVisible)
 			]
@@ -324,6 +358,12 @@ public:
 				.OnToggleXChanged(this, &SAdvancedRotationInputBox::HandleOnToggleChanged, ESlateTransformSubComponent::X)
 				.OnToggleYChanged(this, &SAdvancedRotationInputBox::HandleOnToggleChanged, ESlateTransformSubComponent::Y)
 				.OnToggleZChanged(this, &SAdvancedRotationInputBox::HandleOnToggleChanged, ESlateTransformSubComponent::Z)
+				.OnXBeginSliderMovement(this, &SAdvancedRotationInputBox::HandleOnBeginSliderMovement, ESlateTransformSubComponent::X)
+				.OnXEndSliderMovement(this, &SAdvancedRotationInputBox::HandleOnEndSliderMovement, ESlateTransformSubComponent::X)
+				.OnYBeginSliderMovement(this, &SAdvancedRotationInputBox::HandleOnBeginSliderMovement, ESlateTransformSubComponent::Y)
+				.OnYEndSliderMovement(this, &SAdvancedRotationInputBox::HandleOnEndSliderMovement, ESlateTransformSubComponent::Y)
+				.OnZBeginSliderMovement(this, &SAdvancedRotationInputBox::HandleOnBeginSliderMovement, ESlateTransformSubComponent::Z)
+				.OnZEndSliderMovement(this, &SAdvancedRotationInputBox::HandleOnEndSliderMovement, ESlateTransformSubComponent::Z)
 				.TogglePadding(InArgs._TogglePadding)
 				.Visibility(this, &SAdvancedRotationInputBox::IsAxisAndAngleInputBoxVisible)
 			]
@@ -363,6 +403,8 @@ public:
 				.TogglePadding(InArgs._TogglePadding)
 				.ToggleChecked(this, &SAdvancedRotationInputBox::HandleGetToggleChecked, ESlateTransformSubComponent::Angle)
 				.OnToggleChanged(this, &SAdvancedRotationInputBox::HandleOnToggleChanged, ESlateTransformSubComponent::Angle)
+				.OnBeginSliderMovement(this, &SAdvancedRotationInputBox::HandleOnBeginSliderMovement, ESlateTransformSubComponent::Angle)
+				.OnEndSliderMovement(this, &SAdvancedRotationInputBox::HandleOnEndSliderMovement, ESlateTransformSubComponent::Angle)
 				.Visibility(this, &SAdvancedRotationInputBox::IsAxisAndAngleInputBoxVisible)
 			]
 		];
@@ -449,6 +491,8 @@ private:
 	FOnGetNumericValue OnGetNumericValue;
 	FOnNumericValueChanged OnNumericValueChanged;
 	FOnNumericValueCommitted OnNumericValueCommitted;
+	FOnBeginSliderMovement OnBeginSliderMovement;
+	FOnEndSliderMovement OnEndSliderMovement;
 	FOnGetToggleChecked OnGetToggleChecked;
 	FOnToggleChanged OnToggleChanged;
 
@@ -852,6 +896,30 @@ private:
 		}
 	}
 
+	void HandleOnBeginSliderMovement(ESlateTransformSubComponent::Type InSubComponent)
+	{
+		const ESlateRotationRepresentation::Type& Repr = Representation.Get();
+		if(IsValidComponent(Repr, InSubComponent))
+		{
+			if (OnBeginSliderMovement.IsBound())
+			{
+				OnBeginSliderMovement.Execute(Repr, InSubComponent);
+			}
+		}
+	}
+	
+	void HandleOnEndSliderMovement(NumericType InValue, ESlateTransformSubComponent::Type InSubComponent)
+	{
+		const ESlateRotationRepresentation::Type& Repr = Representation.Get();
+		if(IsValidComponent(Repr, InSubComponent))
+		{
+			if (OnEndSliderMovement.IsBound())
+			{
+				OnEndSliderMovement.Execute(Repr, InSubComponent, InValue);
+			}
+		}
+	}
+
 	EVisibility IsRotatorInputBoxVisible() const
 	{
 		const ESlateRotationRepresentation::Type& Repr = Representation.Get();
diff --git a/Engine/Source/Developer/AnimationWidgets/Public/SAdvancedTransformInputBox.h b/Engine/Source/Developer/AnimationWidgets/Public/SAdvancedTransformInputBox.h
--- a/Engine/Source/Developer/AnimationWidgets/Public/SAdvancedTransformInputBox.h
+++ b/Engine/Source/Developer/AnimationWidgets/Public/SAdvancedTransformInputBox.h
@@ -60,6 +60,12 @@ public:
 	/** Notification for numerictype value committed */
 	DECLARE_DELEGATE_FiveParams(FOnNumericValueCommitted, ESlateTransformComponent::Type, ESlateRotationRepresentation::Type, ESlateTransformSubComponent::Type, NumericType, ETextCommit::Type);
 
+	/** Notification for begin slider movement */
+	DECLARE_DELEGATE_ThreeParams(FOnBeginSliderMovement, ESlateTransformComponent::Type, ESlateRotationRepresentation::Type, ESlateTransformSubComponent::Type);
+
+	/** Notification for end slider movement */
+	DECLARE_DELEGATE_FourParams(FOnEndSliderMovement, ESlateTransformComponent::Type, ESlateRotationRepresentation::Type, ESlateTransformSubComponent::Type, NumericType);
+
 	/** Delegate to retrieve relative / world state */
 	DECLARE_DELEGATE_RetVal_OneParam(bool, FOnGetIsComponentRelative, ESlateTransformComponent::Type);
 
@@ -150,6 +156,12 @@ public:
 		/** Delegate to notify a committed value */
 		SLATE_EVENT( FOnNumericValueCommitted, OnNumericValueCommitted )
 
+		/** Delegate to notify a begin slider movement */
+		SLATE_EVENT(FOnBeginSliderMovement, OnBeginSliderMovement)
+
+		/** Delegate to notify a end slider movement */
+		SLATE_EVENT(FOnEndSliderMovement, OnEndSliderMovement)
+
 		/** Should the axis labels be colored */
 		SLATE_ARGUMENT( bool, bColorAxisLabels )		
 
@@ -292,6 +304,8 @@ public:
 		FOnGetNumericValue OnGetNumericValue = InArgs._OnGetNumericValue;
 		FOnNumericValueChanged OnNumericValueChanged = InArgs._OnNumericValueChanged;
 		FOnNumericValueCommitted OnNumericValueCommitted = InArgs._OnNumericValueCommitted;
+		FOnBeginSliderMovement OnBeginSliderMovement = InArgs._OnBeginSliderMovement;
+		FOnEndSliderMovement OnEndSliderMovement = InArgs._OnEndSliderMovement;
 		FOnGetToggleChecked OnGetToggleChecked = InArgs._OnGetToggleChecked;
 		FOnToggleChanged OnToggleChanged = InArgs._OnToggleChanged;
 		const bool UseQuaternionForRotation = InArgs._UseQuaternionForRotation;
@@ -464,6 +478,52 @@ public:
 					); 
 				}
 
+				typename SNumericVectorInputBox3::FOnBeginSliderMovement XBeginSlide, YBeginSlide, ZBeginSlide;
+				if(OnBeginSliderMovement.IsBound())
+				{
+					XBeginSlide = SNumericVectorInputBox3::FOnBeginSliderMovement::CreateLambda(
+						[OnBeginSliderMovement, InComponent]()
+						{
+							OnBeginSliderMovement.Execute(InComponent, ESlateRotationRepresentation::EulerXYZ, ESlateTransformSubComponent::X);
+						}
+					); 
+					YBeginSlide = SNumericVectorInputBox3::FOnBeginSliderMovement::CreateLambda(
+						[OnBeginSliderMovement, InComponent]()
+						{
+							OnBeginSliderMovement.Execute(InComponent, ESlateRotationRepresentation::EulerXYZ, ESlateTransformSubComponent::Y);
+						}
+					); 
+					ZBeginSlide = SNumericVectorInputBox3::FOnBeginSliderMovement::CreateLambda(
+						[OnBeginSliderMovement, InComponent]()
+						{
+							OnBeginSliderMovement.Execute(InComponent, ESlateRotationRepresentation::EulerXYZ, ESlateTransformSubComponent::Z);
+						}
+					); 
+				}
+
+				typename SNumericVectorInputBox3::FOnEndSliderMovement XEndSlide, YEndSlide, ZEndSlide;
+				if(OnEndSliderMovement.IsBound())
+				{
+					XEndSlide = SNumericVectorInputBox3::FOnEndSliderMovement::CreateLambda(
+						[OnEndSliderMovement, InComponent](NumericType InValue)
+						{
+							OnEndSliderMovement.Execute(InComponent, ESlateRotationRepresentation::EulerXYZ, ESlateTransformSubComponent::X, InValue);
+						}
+					); 
+					YEndSlide = SNumericVectorInputBox3::FOnEndSliderMovement::CreateLambda(
+						[OnEndSliderMovement, InComponent](NumericType InValue)
+						{
+							OnEndSliderMovement.Execute(InComponent, ESlateRotationRepresentation::EulerXYZ, ESlateTransformSubComponent::Y, InValue);
+						}
+					); 
+					ZEndSlide = SNumericVectorInputBox3::FOnEndSliderMovement::CreateLambda(
+						[OnEndSliderMovement, InComponent](NumericType InValue)
+						{
+							OnEndSliderMovement.Execute(InComponent, ESlateRotationRepresentation::EulerXYZ, ESlateTransformSubComponent::Z, InValue);
+						}
+					); 
+				}
+
 				typename SNumericVectorInputBox3::FOnVectorValueChanged VectorChanged;
 				if(OnTransformChanged.IsBound())
 				{
@@ -569,6 +629,12 @@ public:
 					.ConstrainVector(ConstrainComponents)
 					.DisplayToggle(InArgs._DisplayToggle)
 					.TogglePadding(InArgs._TogglePadding)
+					.OnXBeginSliderMovement(XBeginSlide)
+					.OnYBeginSliderMovement(YBeginSlide)
+					.OnZBeginSliderMovement(ZBeginSlide)
+					.OnXEndSliderMovement(XEndSlide)
+					.OnYEndSliderMovement(YEndSlide)
+					.OnZEndSliderMovement(ZEndSlide)
 					.ToggleXChecked_Lambda([OnGetToggleChecked, InComponent]() -> ECheckBoxState
 					{
 						if(OnGetToggleChecked.IsBound())
@@ -649,6 +715,26 @@ public:
 					});
 				}
 
+				typename SAdvancedRotationInputBox<NumericType>::FOnBeginSliderMovement RotationBeginSliderMovement;
+				if(OnBeginSliderMovement.IsBound())
+				{
+					RotationBeginSliderMovement = SAdvancedRotationInputBox<NumericType>::FOnBeginSliderMovement::CreateLambda(
+						[OnBeginSliderMovement, InComponent](ESlateRotationRepresentation::Type InRepr, ESlateTransformSubComponent::Type InSubComponent)
+					{
+						OnBeginSliderMovement.Execute(InComponent, InRepr, InSubComponent);
+					});
+				}
+
+				typename SAdvancedRotationInputBox<NumericType>::FOnEndSliderMovement RotationEndSliderMovement;
+				if(OnEndSliderMovement.IsBound())
+				{
+					RotationEndSliderMovement = SAdvancedRotationInputBox<NumericType>::FOnEndSliderMovement::CreateLambda(
+						[OnEndSliderMovement, InComponent](ESlateRotationRepresentation::Type InRepr, ESlateTransformSubComponent::Type InSubComponent, NumericType InValue)
+					{
+						OnEndSliderMovement.Execute(InComponent, InRepr, InSubComponent, InValue);
+					});
+				}
+
 				TAttribute<TOptional<FQuat>> QuaternionAttribute;
 				TAttribute<TOptional<FRotator>> RotatorAttribute;
 				if(InArgs._Transform.IsBound() || InArgs._Transform.IsSet())
@@ -775,6 +861,8 @@ public:
 					.OnQuaternionCommitted(QuaternionCommitted)
 					.DisplayToggle(InArgs._DisplayToggle)
 					.TogglePadding(InArgs._TogglePadding)
+					.OnBeginSliderMovement(RotationBeginSliderMovement)
+					.OnEndSliderMovement(RotationEndSliderMovement)
 					.OnGetToggleChecked(SAdvancedRotationInputBox<NumericType>::FOnGetToggleChecked::CreateLambda(
 						[OnGetToggleChecked, InComponent](ESlateRotationRepresentation::Type InRepr, ESlateTransformSubComponent::Type InSubComponent) -> ECheckBoxState
 					{
diff --git a/Engine/Source/Editor/DetailCustomizations/Private/ComponentTransformDetails.cpp b/Engine/Source/Editor/DetailCustomizations/Private/ComponentTransformDetails.cpp
--- a/Engine/Source/Editor/DetailCustomizations/Private/ComponentTransformDetails.cpp
+++ b/Engine/Source/Editor/DetailCustomizations/Private/ComponentTransformDetails.cpp
@@ -474,8 +474,12 @@ void FComponentTransformDetails::GenerateChildContent( IDetailChildrenBuilder& C
 			.Yaw( this, &FComponentTransformDetails::GetRotationZ )
 			.bColorAxisLabels( true )
 			.IsEnabled( this, &FComponentTransformDetails::GetIsEnabled )
-			.OnBeginSliderMovement( this, &FComponentTransformDetails::OnBeginRotationSlider )
-			.OnEndSliderMovement( this, &FComponentTransformDetails::OnEndRotationSlider )
+			.OnPitchBeginSliderMovement( this, &FComponentTransformDetails::OnBeginRotationSlider )
+			.OnYawBeginSliderMovement( this, &FComponentTransformDetails::OnBeginRotationSlider )
+			.OnRollBeginSliderMovement( this, &FComponentTransformDetails::OnBeginRotationSlider )
+			.OnPitchEndSliderMovement( this, &FComponentTransformDetails::OnEndRotationSlider )
+			.OnYawEndSliderMovement( this, &FComponentTransformDetails::OnEndRotationSlider )
+			.OnRollEndSliderMovement( this, &FComponentTransformDetails::OnEndRotationSlider )
 			.OnRollChanged( this, &FComponentTransformDetails::OnSetTransformAxis, ETextCommit::Default, ETransformField::Rotation, EAxisList::X, false )
 			.OnPitchChanged( this, &FComponentTransformDetails::OnSetTransformAxis, ETextCommit::Default, ETransformField::Rotation, EAxisList::Y, false )
 			.OnYawChanged( this, &FComponentTransformDetails::OnSetTransformAxis, ETextCommit::Default, ETransformField::Rotation, EAxisList::Z, false )
diff --git a/Engine/Source/Runtime/Slate/Public/Widgets/Input/SRotatorInputBox.h b/Engine/Source/Runtime/Slate/Public/Widgets/Input/SRotatorInputBox.h
--- a/Engine/Source/Runtime/Slate/Public/Widgets/Input/SRotatorInputBox.h
+++ b/Engine/Source/Runtime/Slate/Public/Widgets/Input/SRotatorInputBox.h
@@ -89,11 +89,43 @@ public:
 		/** Called when the roll value is committed */
 		SLATE_EVENT( FOnNumericValueCommitted, OnRollCommitted )
 
+		FSimpleDelegate ConvertOnBeginSliderMovement(const FSimpleDelegate& LegacyDelegate)
+		{
+			return FSimpleDelegate::CreateLambda([LegacyDelegate]()
+			{
+				return LegacyDelegate.Execute();
+			});
+		}
 		/** Called when the slider begins to move on any axis */
-		SLATE_EVENT( FSimpleDelegate, OnBeginSliderMovement )
+		SLATE_EVENT_DEPRECATED(5.3, "Use the new OnPitchBeginSliderMovement()/OnYawBeginSliderMovement()/OnRollBeginSliderMovement() event instead", FSimpleDelegate, OnBeginSliderMovement, OnPitchBeginSliderMovement, ConvertOnBeginSliderMovement)
 
+		FOnNumericValueChanged ConvertOnEndSliderMovement(const FOnNumericValueChanged& LegacyDelegate)
+		{
+			return FOnNumericValueChanged::CreateLambda([LegacyDelegate](const NumericType& Value)
+			{
+				return LegacyDelegate.Execute(Value);
+			});
+		}
 		/** Called when the slider for any axis is released */
-		SLATE_EVENT( FOnNumericValueChanged, OnEndSliderMovement )
+		SLATE_EVENT_DEPRECATED(5.3, "Use the new OnPitchEndSliderMovement()/OnYawEndSliderMovement()/OnRollEndSliderMovement() event instead", FOnNumericValueChanged, OnEndSliderMovement, OnPitchEndSliderMovement, ConvertOnEndSliderMovement)
+	
+		/** Called when the slider begins to move on pitch */
+		SLATE_EVENT( FSimpleDelegate, OnPitchBeginSliderMovement )
+
+		/** Called when the slider begins to move on yaw */
+		SLATE_EVENT( FSimpleDelegate, OnYawBeginSliderMovement )
+
+		/** Called when the slider begins to move on roll */
+		SLATE_EVENT( FSimpleDelegate, OnRollBeginSliderMovement )
+
+		/** Called when the slider for pitch is released */
+		SLATE_EVENT( FOnNumericValueChanged, OnPitchEndSliderMovement )
+
+		/** Called when the slider for yaw is released */
+		SLATE_EVENT( FOnNumericValueChanged, OnYawEndSliderMovement )
+
+		/** Called when the slider for roll is released */
+		SLATE_EVENT( FOnNumericValueChanged, OnRollEndSliderMovement )
 
 		/** Provide custom type functionality for the rotator */
 		SLATE_ARGUMENT( TSharedPtr< INumericTypeInterface<NumericType> >, TypeInterface )
@@ -150,8 +182,8 @@ public:
 				.Value( InArgs._Roll )
 				.OnValueChanged( InArgs._OnRollChanged )
 				.OnValueCommitted( InArgs._OnRollCommitted )
-				.OnBeginSliderMovement( InArgs._OnBeginSliderMovement )
-				.OnEndSliderMovement( InArgs._OnEndSliderMovement )
+				.OnBeginSliderMovement( InArgs._OnRollBeginSliderMovement )
+				.OnEndSliderMovement( InArgs._OnRollEndSliderMovement )
 				.UndeterminedString( NSLOCTEXT("SRotatorInputBox", "MultipleValues", "Multiple Values") )
 				.ToolTipText_Lambda([RollAttr = InArgs._Roll]
 				{
@@ -182,8 +214,8 @@ public:
 				.Value( InArgs._Pitch )
 				.OnValueChanged( InArgs._OnPitchChanged )
 				.OnValueCommitted( InArgs._OnPitchCommitted )
-				.OnBeginSliderMovement( InArgs._OnBeginSliderMovement )
-				.OnEndSliderMovement( InArgs._OnEndSliderMovement )
+				.OnBeginSliderMovement( InArgs._OnPitchBeginSliderMovement )
+				.OnEndSliderMovement( InArgs._OnPitchEndSliderMovement )
 				.UndeterminedString( NSLOCTEXT("SRotatorInputBox", "MultipleValues", "Multiple Values") )
 				.ToolTipText_Lambda([PitchAttr = InArgs._Pitch]
 				{
@@ -214,8 +246,8 @@ public:
 				.Value( InArgs._Yaw )
 				.OnValueChanged( InArgs._OnYawChanged )
 				.OnValueCommitted( InArgs._OnYawCommitted )
-				.OnBeginSliderMovement( InArgs._OnBeginSliderMovement )
-				.OnEndSliderMovement( InArgs._OnEndSliderMovement )
+				.OnBeginSliderMovement( InArgs._OnYawBeginSliderMovement )
+				.OnEndSliderMovement( InArgs._OnYawEndSliderMovement )
 				.UndeterminedString( NSLOCTEXT("SRotatorInputBox", "MultipleValues", "Multiple Values") )
 				.ToolTipText_Lambda([YawAttr = InArgs._Yaw]
 				{
diff --git a/Engine/Source/Runtime/Slate/Public/Widgets/Input/SVectorInputBox.h b/Engine/Source/Runtime/Slate/Public/Widgets/Input/SVectorInputBox.h
--- a/Engine/Source/Runtime/Slate/Public/Widgets/Input/SVectorInputBox.h
+++ b/Engine/Source/Runtime/Slate/Public/Widgets/Input/SVectorInputBox.h
@@ -40,6 +40,12 @@ public:
 	/** Delegate to constrain the vector during a change */
 	DECLARE_DELEGATE_ThreeParams(FOnConstrainVector, int32 /* Component */, VectorType /* old */ , VectorType& /* new */);
 
+	/** Notification for vector value slider began movement */
+	DECLARE_DELEGATE(FOnBeginSliderMovement);
+
+	/** Notification for vector value slider ended movement */
+	DECLARE_DELEGATE_OneParam(FOnEndSliderMovement, NumericType);
+
 	struct FArguments;
 
 private:
@@ -72,6 +78,12 @@ private:
 
 		/** Menu extender delegate for the X value */
 		SLATE_EVENT(FMenuExtensionDelegate, ContextMenuExtenderX)
+
+		/** Called when the x value of the vector slider began movement */
+		SLATE_EVENT(FOnBeginSliderMovement, OnXBeginSliderMovement)
+
+		/** Called when the x value of the vector slider ended movement */
+		SLATE_EVENT(FOnEndSliderMovement, OnXEndSliderMovement)
 	};
 
 	struct FVectorYArgumentsEmpty {};
@@ -101,6 +113,12 @@ private:
 
 		/** Menu extender delegate for the Y value */
 		SLATE_EVENT(FMenuExtensionDelegate, ContextMenuExtenderY)
+
+		/** Called when the y value of the vector slider began movement */
+		SLATE_EVENT(FOnBeginSliderMovement, OnYBeginSliderMovement)
+
+		/** Called when the y value of the vector slider ended movement */
+		SLATE_EVENT(FOnEndSliderMovement, OnYEndSliderMovement)
 	};
 
 	struct FVectorZArgumentsEmpty {};
@@ -130,6 +148,12 @@ private:
 
 		/** Menu extender delegate for the Z value */
 		SLATE_EVENT(FMenuExtensionDelegate, ContextMenuExtenderZ)
+
+		/** Called when the z value of the vector slider began movement */
+		SLATE_EVENT(FOnBeginSliderMovement, OnZBeginSliderMovement)
+
+		/** Called when the z value of the vector slider ended movement */
+		SLATE_EVENT(FOnEndSliderMovement, OnZEndSliderMovement)
 	};
 
 	struct FVectorWArgumentsEmpty {};
@@ -159,6 +183,12 @@ private:
 
 		/** Menu extender delegate for the W value */
 		SLATE_EVENT(FMenuExtensionDelegate, ContextMenuExtenderW)
+
+		/** Called when the w value of the vector slider began movement */
+		SLATE_EVENT(FOnBeginSliderMovement, OnWBeginSliderMovement)
+
+		/** Called when the w value of the vector slider ended movement */
+		SLATE_EVENT(FOnEndSliderMovement, OnWEndSliderMovement)
 	};
 
 public:
@@ -390,7 +420,9 @@ private:
 		const FOnNumericValueCommitted& OnComponentCommitted,
 		const TAttribute<ECheckBoxState> ToggleChecked,
 		const FOnCheckStateChanged& OnToggleChanged,
-		const FMenuExtensionDelegate& OnContextMenuExtenderComponent)
+		const FMenuExtensionDelegate& OnContextMenuExtenderComponent,
+		const FSimpleDelegate& OnBeginSliderMovement,
+		const FOnNumericValueChanged& OnEndSliderMovement)
 	{
 		TSharedRef<SWidget> LabelWidget = SNullWidget::NullWidget;
 		if (InArgs._bColorAxisLabels)
@@ -430,8 +462,8 @@ private:
 			.MaxSliderValue(CreatePerComponentGetter(ComponentIndex, TOptional<NumericType>(), InArgs._MaxSliderVector))
 			.LinearDeltaSensitivity(InArgs._LinearDeltaSensitivity)
 			.Delta(InArgs._SpinDelta)
-			.OnBeginSliderMovement(InArgs._OnBeginSliderMovement)
-			.OnEndSliderMovement(InArgs._OnEndSliderMovement)
+			.OnBeginSliderMovement(OnBeginSliderMovement)
+			.OnEndSliderMovement(OnEndSliderMovement)
 			.LabelPadding(FMargin(3.f))
 			.LabelLocation(SNumericEntryBox<NumericType>::ELabelLocation::Inside)
 			.Label()
@@ -460,7 +492,9 @@ private:
 			InArgs._OnXCommitted,
 			InArgs._ToggleXChecked,
 			InArgs._OnToggleXChanged,
-			InArgs._ContextMenuExtenderX
+			InArgs._ContextMenuExtenderX,
+			InArgs._OnXBeginSliderMovement,
+			InArgs._OnXEndSliderMovement
 		);
 	}
 
@@ -479,7 +513,9 @@ private:
 			InArgs._OnYCommitted,
 			InArgs._ToggleYChecked,
 			InArgs._OnToggleYChanged,
-			InArgs._ContextMenuExtenderY
+			InArgs._ContextMenuExtenderY,
+			InArgs._OnYBeginSliderMovement,
+			InArgs._OnYEndSliderMovement
 		);
 	}
 
@@ -498,7 +534,9 @@ private:
 			InArgs._OnZCommitted,
 			InArgs._ToggleZChecked,
 			InArgs._OnToggleZChanged,
-			InArgs._ContextMenuExtenderZ
+			InArgs._ContextMenuExtenderZ,
+			InArgs._OnZBeginSliderMovement,
+			InArgs._OnZEndSliderMovement
 		);
 	}
 
@@ -517,7 +555,9 @@ private:
 			InArgs._OnWCommitted,
 			InArgs._ToggleWChecked,
 			InArgs._OnToggleWChanged,
-			InArgs._ContextMenuExtenderW
+			InArgs._ContextMenuExtenderW,
+			InArgs._OnWBeginSliderMovement,
+			InArgs._OnWEndSliderMovement
 		);
 	}
 
