Fix 'A = B != C' if condition reported by PVS. This was causing RemoveIndex to be set to either 0 or 1, likely removing the incorrect bake frames/transforms.

diff --git a/Engine/Plugins/Animation/ControlRig/Source/ControlRigEditor/Private/Tools/ConstraintBaker.cpp b/Engine/Plugins/Animation/ControlRig/Source/ControlRigEditor/Private/Tools/ConstraintBaker.cpp
--- a/Engine/Plugins/Animation/ControlRig/Source/ControlRigEditor/Private/Tools/ConstraintBaker.cpp
+++ b/Engine/Plugins/Animation/ControlRig/Source/ControlRigEditor/Private/Tools/ConstraintBaker.cpp
@@ -312,7 +312,8 @@ void FConstraintBaker::Bake(UWorld* InWorld,
 				//we also set the tangent at the break to the default type
 				FMovieSceneConstraintChannelHelper::ChangeKeyInterpolation(DoubleTransformChannels, Frame, KeyType);
 			}
-			if (int32 RemoveIndex = Algo::BinarySearch(FramesToBake, FrameMinusOne) != INDEX_NONE)
+			if (const int32 RemoveIndex = Algo::BinarySearch(FramesToBake, FrameMinusOne);
+				RemoveIndex != INDEX_NONE)
 			{
 				FramesToBake.RemoveAt(RemoveIndex);
 				Transforms.RemoveAt(RemoveIndex);
