ControlRig: Don add data source registry to the CDO, which is not a allowed class and thus will generate error during content validation.

diff --git a/Engine/Plugins/Animation/ControlRig/Source/ControlRig/Private/ControlRig.cpp b/Engine/Plugins/Animation/ControlRig/Source/ControlRig/Private/ControlRig.cpp
--- a/Engine/Plugins/Animation/ControlRig/Source/ControlRig/Private/ControlRig.cpp
+++ b/Engine/Plugins/Animation/ControlRig/Source/ControlRig/Private/ControlRig.cpp
@@ -1278,16 +1278,23 @@ void UControlRig::GetMappableNodeData(TArray<FName>& OutNames, TArray<FNodeItem>
 
 UAnimationDataSourceRegistry* UControlRig::GetDataSourceRegistry()
 {
-	if (DataSourceRegistry)
+	if (!HasAnyFlags(RF_ClassDefaultObject))
 	{
-		if (DataSourceRegistry->GetOuter() != this)
+		if (DataSourceRegistry)
+		{
+			if (DataSourceRegistry->GetOuter() != this)
+			{
+				DataSourceRegistry = nullptr;
+			}
+		}
+		if (DataSourceRegistry == nullptr)
 		{
-			DataSourceRegistry = nullptr;
+			DataSourceRegistry = NewObject<UAnimationDataSourceRegistry>(this);
 		}
 	}
-	if (DataSourceRegistry == nullptr)
+	else
 	{
-		DataSourceRegistry = NewObject<UAnimationDataSourceRegistry>(this);
+		DataSourceRegistry = nullptr;
 	}
 	return DataSourceRegistry;
 }
diff --git a/Engine/Plugins/Animation/ControlRig/Source/ControlRig/Public/ControlRig.h b/Engine/Plugins/Animation/ControlRig/Source/ControlRig/Public/ControlRig.h
--- a/Engine/Plugins/Animation/ControlRig/Source/ControlRig/Public/ControlRig.h
+++ b/Engine/Plugins/Animation/ControlRig/Source/ControlRig/Public/ControlRig.h
@@ -396,7 +396,7 @@ private:
 #endif
 	
 	/** The registry to access data source */
-	UPROPERTY(transient)
+	UPROPERTY(Transient)
 	TObjectPtr<UAnimationDataSourceRegistry> DataSourceRegistry;
 
 	/** Broadcasts a notification when launching the construction event */
