Control Rig: Mark DataSourceRegistry object in control rig as RF_Transient, such that it does not get serialized with the CDO

diff --git a/Engine/Plugins/Animation/ControlRig/Source/ControlRig/Private/ControlRig.cpp b/Engine/Plugins/Animation/ControlRig/Source/ControlRig/Private/ControlRig.cpp
--- a/Engine/Plugins/Animation/ControlRig/Source/ControlRig/Private/ControlRig.cpp
+++ b/Engine/Plugins/Animation/ControlRig/Source/ControlRig/Private/ControlRig.cpp
@@ -1278,24 +1278,18 @@ void UControlRig::GetMappableNodeData(TArray<FName>& OutNames, TArray<FNodeItem>
 
 UAnimationDataSourceRegistry* UControlRig::GetDataSourceRegistry()
 {
-	if (!HasAnyFlags(RF_ClassDefaultObject))
+	if (DataSourceRegistry)
 	{
-		if (DataSourceRegistry)
+		if (DataSourceRegistry->GetOuter() != this)
 		{
-			if (DataSourceRegistry->GetOuter() != this)
-			{
-				DataSourceRegistry = nullptr;
-			}
-		}
-		if (DataSourceRegistry == nullptr)
-		{
-			DataSourceRegistry = NewObject<UAnimationDataSourceRegistry>(this);
+			DataSourceRegistry = nullptr;
 		}
 	}
-	else
+	if (DataSourceRegistry == nullptr)
 	{
-		DataSourceRegistry = nullptr;
+		DataSourceRegistry = NewObject<UAnimationDataSourceRegistry>(this, NAME_None, RF_Transient);
 	}
+
 	return DataSourceRegistry;
 }
 
