Control Rig: Mark bp as structurally modified for function header based changes

diff --git a/Engine/Plugins/Animation/ControlRig/Source/ControlRigDeveloper/Private/ControlRigBlueprint.cpp b/Engine/Plugins/Animation/ControlRig/Source/ControlRigDeveloper/Private/ControlRigBlueprint.cpp
--- a/Engine/Plugins/Animation/ControlRig/Source/ControlRigDeveloper/Private/ControlRigBlueprint.cpp
+++ b/Engine/Plugins/Animation/ControlRig/Source/ControlRigDeveloper/Private/ControlRigBlueprint.cpp
@@ -127,6 +127,7 @@ UControlRigBlueprint::UControlRigBlueprint(const FObjectInitializer& ObjectIniti
 	bSuspendModelNotificationsForOthers = false;
 	bSuspendAllNotifications = false;
 	bSuspendPythonMessagesForRigVMClient = true;
+	bMarkBlueprintAsStructurallyModifiedPending = false;
 
 #if WITH_EDITORONLY_DATA
 	ReferencedObjectPathsStored = false;
@@ -3531,6 +3532,22 @@ void UControlRigBlueprint::HandleModifiedEvent(ERigVMGraphNotifType InNotifType,
 	// the notify for still pending (or already sent)
 	bool bNotifForOthersPending = true;
 
+	auto MarkBlueprintAsStructurallyModified = [this]()
+	{
+		if(VMRecompilationBracket == 0)
+		{
+			if(bMarkBlueprintAsStructurallyModifiedPending)
+			{
+				bMarkBlueprintAsStructurallyModifiedPending = false;
+				FBlueprintEditorUtils::MarkBlueprintAsStructurallyModified(this);
+			}
+		}
+		else
+		{
+			bMarkBlueprintAsStructurallyModifiedPending = true;
+		}
+	};
+
 	if (!bSuspendModelNotificationsForSelf)
 	{
 		switch (InNotifType)
@@ -3544,6 +3561,7 @@ void UControlRigBlueprint::HandleModifiedEvent(ERigVMGraphNotifType InNotifType,
 			case ERigVMGraphNotifType::InteractionBracketCanceled:
 			{
 				DecrementVMRecompileBracket();
+				MarkBlueprintAsStructurallyModified();
 				break;
 			}
 			case ERigVMGraphNotifType::PinDefaultValueChanged:
@@ -3719,7 +3737,7 @@ void UControlRigBlueprint::HandleModifiedEvent(ERigVMGraphNotifType InNotifType,
 					RequestAutoVMRecompilation();
 
 					MarkPackageDirty();
-					FBlueprintEditorUtils::MarkBlueprintAsStructurallyModified(this);
+					MarkBlueprintAsStructurallyModified();
 					break;
 				}
 
@@ -3756,7 +3774,7 @@ void UControlRigBlueprint::HandleModifiedEvent(ERigVMGraphNotifType InNotifType,
 				// need to recompile the VM here - unless we don't auto recompile.
 				if(!bAutoRecompileVM)
 				{
-					FBlueprintEditorUtils::MarkBlueprintAsModified(this);
+					MarkBlueprintAsStructurallyModified();
 				}
 				break;
 			}
@@ -3834,6 +3852,25 @@ void UControlRigBlueprint::HandleModifiedEvent(ERigVMGraphNotifType InNotifType,
 						}
 					}
 				}
+				// fall through another time
+			}
+			case ERigVMGraphNotifType::PinAdded:
+			case ERigVMGraphNotifType::PinRemoved:
+			case ERigVMGraphNotifType::PinRenamed:
+			{
+				// exposed pin changes like this (as well as type change etc)
+				// require to mark the blueprint as structurally modified,
+				// so that the instance actions work out.
+				if (URigVMPin* ModelPin = Cast<URigVMPin>(InSubject))
+				{
+					if(URigVMCollapseNode* CollapseNode = Cast<URigVMCollapseNode>(ModelPin->GetNode()))
+					{
+						if(Cast<URigVMFunctionLibrary>(CollapseNode->GetOuter()))
+						{
+							MarkBlueprintAsStructurallyModified();
+						}
+					}
+				}
 				break;
 			}
 			case ERigVMGraphNotifType::PinBoundVariableChanged:
@@ -3863,7 +3900,7 @@ void UControlRigBlueprint::HandleModifiedEvent(ERigVMGraphNotifType InNotifType,
 						ContainedEdGraph->Rename(*CollapseNode->GetEditorSubGraphName(), nullptr);
 					}
 
-					FBlueprintEditorUtils::MarkBlueprintAsStructurallyModified(this);
+					MarkBlueprintAsStructurallyModified();
 				}
 				break;
 			}
@@ -3871,7 +3908,7 @@ void UControlRigBlueprint::HandleModifiedEvent(ERigVMGraphNotifType InNotifType,
 			case ERigVMGraphNotifType::NodeKeywordsChanged:
 			case ERigVMGraphNotifType::NodeDescriptionChanged:
 			{
-				FBlueprintEditorUtils::MarkBlueprintAsStructurallyModified(this);
+				MarkBlueprintAsStructurallyModified();
 				break;
 			}
 			default:
diff --git a/Engine/Plugins/Animation/ControlRig/Source/ControlRigDeveloper/Public/ControlRigBlueprint.h b/Engine/Plugins/Animation/ControlRig/Source/ControlRigDeveloper/Public/ControlRigBlueprint.h
--- a/Engine/Plugins/Animation/ControlRig/Source/ControlRigDeveloper/Public/ControlRigBlueprint.h
+++ b/Engine/Plugins/Animation/ControlRig/Source/ControlRigDeveloper/Public/ControlRigBlueprint.h
@@ -814,6 +814,7 @@ private:
 	bool bDirtyDuringLoad;
 	bool bErrorsDuringCompilation;
 	bool bSuspendPythonMessagesForRigVMClient;
+	bool bMarkBlueprintAsStructurallyModifiedPending;
 
 	friend class FControlRigBlueprintCompilerContext;
 	friend class SRigHierarchy;
diff --git a/Engine/Plugins/Animation/ControlRig/Source/ControlRigEditor/Private/Graph/NodeSpawners/ControlRigVariableNodeSpawner.cpp b/Engine/Plugins/Animation/ControlRig/Source/ControlRigEditor/Private/Graph/NodeSpawners/ControlRigVariableNodeSpawner.cpp
--- a/Engine/Plugins/Animation/ControlRig/Source/ControlRigEditor/Private/Graph/NodeSpawners/ControlRigVariableNodeSpawner.cpp
+++ b/Engine/Plugins/Animation/ControlRig/Source/ControlRigEditor/Private/Graph/NodeSpawners/ControlRigVariableNodeSpawner.cpp
@@ -151,12 +151,28 @@ UEdGraphNode* UControlRigVariableNodeSpawner::Invoke(UEdGraph* ParentGraph, FBin
 	FString NodeName;
 	if(bIsTemplateNode)
 	{
-		// since we are removing the node at the end of this function
-		// we need to create a unique here.
+		// for template controllers let's rely on locally defined nodes
+		// without a backing model node. access to local variables or
+		// input arguments doesn't work on the template controller.
+		
 		static constexpr TCHAR VariableNodeNameFormat[] = TEXT("VariableNode_%s_%s");
 		static const FString GetterPrefix = TEXT("Getter");
 		static const FString SetterPrefix = TEXT("Setter");
 		NodeName = FString::Printf(VariableNodeNameFormat, bIsGetter ? *GetterPrefix : *SetterPrefix, *ExternalVariable.Name.ToString());
+
+		NewNode = NewObject<UControlRigGraphNode>(ParentGraph, *NodeName);
+		ParentGraph->AddNode(NewNode, false);
+
+		NewNode->CreateNewGuid();
+		NewNode->PostPlacedNewNode();
+
+		UEdGraphPin* ValuePin = UEdGraphPin::CreatePin(NewNode);
+		NewNode->Pins.Add(ValuePin);
+
+		ValuePin->PinType = RigVMTypeUtils::PinTypeFromExternalVariable(ExternalVariable);
+		ValuePin->Direction = bIsGetter ? EGPD_Output : EGPD_Input;
+		NewNode->SetFlags(RF_Transactional);
+		return NewNode;
 	}
 
 	if (URigVMNode* ModelNode = Controller->AddVariableNodeFromObjectPath(ExternalVariable.Name, TypeName, ObjectPath, bIsGetter, FString(), Location, NodeName, !bIsTemplateNode, !bIsTemplateNode))
diff --git a/Engine/Plugins/Runtime/RigVM/Source/RigVMDeveloper/Private/RigVMModel/RigVMGraph.cpp b/Engine/Plugins/Runtime/RigVM/Source/RigVMDeveloper/Private/RigVMModel/RigVMGraph.cpp
--- a/Engine/Plugins/Runtime/RigVM/Source/RigVMDeveloper/Private/RigVMModel/RigVMGraph.cpp
+++ b/Engine/Plugins/Runtime/RigVM/Source/RigVMDeveloper/Private/RigVMModel/RigVMGraph.cpp
@@ -338,7 +338,11 @@ TArray<FRigVMGraphVariableDescription> URigVMGraph::GetInputArguments() const
 	if (URigVMFunctionEntryNode* EntryNode = GetEntryNode())
 	{
 		for (URigVMPin* Pin : EntryNode->GetPins())
-		{			
+		{
+			if(Pin->IsExecuteContext())
+			{
+				continue;
+			}
 			FRigVMGraphVariableDescription Description;
 			Description.Name = Pin->GetFName();
 			Description.CPPType = Pin->GetCPPType();
@@ -355,7 +359,11 @@ TArray<FRigVMGraphVariableDescription> URigVMGraph::GetOutputArguments() const
 	if (URigVMFunctionReturnNode* ReturnNode = GetReturnNode())
 	{
 		for (URigVMPin* Pin : ReturnNode->GetPins())
-		{			
+		{
+			if(Pin->IsExecuteContext())
+			{
+				continue;
+			}
 			FRigVMGraphVariableDescription Description;
 			Description.Name = Pin->GetFName();
 			Description.CPPType = Pin->GetCPPType();
