Control Rig: Fix crash when undoing duplication of ThirdPersonCharacter

diff --git a/Engine/Plugins/Animation/ControlRig/Source/ControlRig/Private/Rigs/RigHierarchyController.cpp b/Engine/Plugins/Animation/ControlRig/Source/ControlRig/Private/Rigs/RigHierarchyController.cpp
--- a/Engine/Plugins/Animation/ControlRig/Source/ControlRig/Private/Rigs/RigHierarchyController.cpp
+++ b/Engine/Plugins/Animation/ControlRig/Source/ControlRig/Private/Rigs/RigHierarchyController.cpp
@@ -41,7 +41,7 @@ void URigHierarchyController::SetHierarchy(URigHierarchy* InHierarchy)
 	// controlling a different hierarchy is no longer allowed
 	if (ensure(InHierarchy == GetOuter()))
 	{
-		// make sure making multiple valid SetHieararchy() calls won't lead to accumulated delegates
+		// make sure making multiple valid SetHierarchy() calls won't lead to accumulated delegates
 		// though it should not happen in the first place
 		if (Hierarchy.IsValid())
 		{
@@ -49,10 +49,11 @@ void URigHierarchyController::SetHierarchy(URigHierarchy* InHierarchy)
 			{
 				  Hierarchy->OnModified().RemoveAll(this);
 			}
+			Hierarchy = nullptr;
 		}
 		
 		URigHierarchy* OuterHierarchy = Cast<URigHierarchy>(GetOuter());
-		if (ensure(OuterHierarchy))
+		if (ensure(OuterHierarchy) && ::IsValid(OuterHierarchy))
 		{
 			Hierarchy = OuterHierarchy;
 			Hierarchy->OnModified().AddUObject(this, &URigHierarchyController::HandleHierarchyModified);
