Fixed import hierarchy crash + anim instance initialization

diff --git a/Engine/Plugins/Animation/ControlRig/Source/ControlRigEditor/Private/Editor/ControlRigEditor.cpp b/Engine/Plugins/Animation/ControlRig/Source/ControlRigEditor/Private/Editor/ControlRigEditor.cpp
--- a/Engine/Plugins/Animation/ControlRig/Source/ControlRigEditor/Private/Editor/ControlRigEditor.cpp
+++ b/Engine/Plugins/Animation/ControlRig/Source/ControlRigEditor/Private/Editor/ControlRigEditor.cpp
@@ -4550,15 +4550,6 @@ void FControlRigEditor::RebindToSkeletalMeshComponent()
 	}
 }
 
-void FControlRigEditor::UpdateMeshInAnimInstance(USkeletalMesh* InNewSkeletalMesh)
-{
-	UControlRigSkeletalMeshComponent* EditorSkelComp = Cast<UControlRigSkeletalMeshComponent>(GetPersonaToolkit()->GetPreviewScene()->GetPreviewMeshComponent());
-	if (EditorSkelComp)
-	{
-		EditorSkelComp->GetAnimInstance()->CurrentSkeleton = InNewSkeletalMesh->GetSkeleton();
-	}
-}
-
 void FControlRigEditor::UpdateStaleWatchedPins()
 {
 	UControlRigBlueprint* ControlRigBP = GetControlRigBlueprint();
diff --git a/Engine/Plugins/Animation/ControlRig/Source/ControlRigEditor/Private/Editor/ControlRigEditor.h b/Engine/Plugins/Animation/ControlRig/Source/ControlRigEditor/Private/Editor/ControlRigEditor.h
--- a/Engine/Plugins/Animation/ControlRig/Source/ControlRigEditor/Private/Editor/ControlRigEditor.h
+++ b/Engine/Plugins/Animation/ControlRig/Source/ControlRigEditor/Private/Editor/ControlRigEditor.h
@@ -328,9 +328,6 @@ private:
 	/** Rebind our anim instance to the preview's skeletal mesh component */
 	void RebindToSkeletalMeshComponent();
 
-	/** Update the skeletal mesh componens so that the anim instance has the correct skeletal mesh */
-	void UpdateMeshInAnimInstance(USkeletalMesh* InNewSkeletalMesh);
-
 	/** Update stale watch pins */
 	void UpdateStaleWatchedPins();
 
diff --git a/Engine/Plugins/Animation/ControlRig/Source/ControlRigEditor/Private/Editor/SRigHierarchy.cpp b/Engine/Plugins/Animation/ControlRig/Source/ControlRigEditor/Private/Editor/SRigHierarchy.cpp
--- a/Engine/Plugins/Animation/ControlRig/Source/ControlRigEditor/Private/Editor/SRigHierarchy.cpp
+++ b/Engine/Plugins/Animation/ControlRig/Source/ControlRigEditor/Private/Editor/SRigHierarchy.cpp
@@ -1338,6 +1338,41 @@ void SRigHierarchy::CreateRefreshMenu(FMenuBuilder& MenuBuilder)
 	);
 }
 
+void SRigHierarchy::UpdateMesh(USkeletalMesh* InMesh, const bool bImport) const
+{
+	if (!InMesh || !ControlRigBlueprint.IsValid() || !ControlRigEditor.IsValid())
+	{
+		return;
+	}
+
+	const bool bUpdateMesh = bImport ? ControlRigBlueprint->GetPreviewMesh() == nullptr : true;
+	if (!bUpdateMesh)
+	{
+		return;
+	}
+
+	const TSharedPtr<FControlRigEditor> EditorSharedPtr = ControlRigEditor.Pin();
+	EditorSharedPtr->GetPersonaToolkit()->SetPreviewMesh(InMesh, true);
+
+	UControlRigSkeletalMeshComponent* Component = Cast<UControlRigSkeletalMeshComponent>(EditorSharedPtr->GetPersonaToolkit()->GetPreviewMeshComponent());
+	if (bImport)
+	{
+		Component->InitAnim(true);
+	}
+
+	UAnimInstance* AnimInstance = Component->GetAnimInstance();
+	if (UControlRigLayerInstance* ControlRigLayerInstance = Cast<UControlRigLayerInstance>(AnimInstance))
+	{
+		EditorSharedPtr->PreviewInstance = Cast<UAnimPreviewInstance>(ControlRigLayerInstance->GetSourceAnimInstance());
+	}
+	else
+	{
+		EditorSharedPtr->PreviewInstance = Cast<UAnimPreviewInstance>(AnimInstance);
+	}
+
+	EditorSharedPtr->Compile();
+}
+
 void SRigHierarchy::RefreshHierarchy(const FAssetData& InAssetData)
 {
 	if (bIsChangingRigHierarchy)
@@ -1386,24 +1421,8 @@ void SRigHierarchy::RefreshHierarchy(const FAssetData& InAssetData)
 	RefreshTreeView();
 	FSlateApplication::Get().DismissAllMenus();
 
-	if (Mesh != nullptr)
-	{
-		StrongEditor->GetPersonaToolkit()->SetPreviewMesh(Mesh, true);
-		const UControlRigSkeletalMeshComponent* EditorSkelComp = Cast<UControlRigSkeletalMeshComponent>(StrongEditor->GetPersonaToolkit()->GetPreviewMeshComponent());
-		if (UControlRigLayerInstance* ControlRigLayerInstance = Cast<UControlRigLayerInstance>(EditorSkelComp->GetAnimInstance()))
-		{
-			StrongEditor->PreviewInstance = Cast<UAnimPreviewInstance>(ControlRigLayerInstance->GetSourceAnimInstance());
-		}
-		else
-		{
-			StrongEditor->PreviewInstance = Cast<UAnimPreviewInstance>(EditorSkelComp->GetAnimInstance());
-		}
-	}
-
-	if (ControlRigEditor.IsValid())
-	{
-		StrongEditor->Compile();
-	}
+	static constexpr bool bImport = false;
+	UpdateMesh(Mesh, bImport);
 }
 
 void SRigHierarchy::CreateImportMenu(FMenuBuilder& MenuBuilder)
@@ -1439,16 +1458,22 @@ void SRigHierarchy::ImportHierarchy(const FAssetData& InAssetData)
 	{
 		return;
 	}
+
+	USkeletalMesh* Mesh = Cast<USkeletalMesh>(InAssetData.GetAsset());
+	if (!Mesh)
+	{
+		return;
+	}
+	
 	TGuardValue<bool> GuardRigHierarchyChanges(bIsChangingRigHierarchy, true);
 
 	if (!ControlRigEditor.IsValid())
 	{
 		return;
 	}
-
-	URigHierarchy* Hierarchy = GetDefaultHierarchy();
-	USkeletalMesh* Mesh = Cast<USkeletalMesh>(InAssetData.GetAsset());
-	if (Mesh && Hierarchy)
+	
+	const TSharedPtr<FControlRigEditor> EditorSharedPtr = ControlRigEditor.Pin();
+	if (URigHierarchy* Hierarchy = GetDefaultHierarchy())
 	{
 		// filter out meshes that don't contain a skeleton
 		if(Mesh->GetSkeleton() == nullptr)
@@ -1472,11 +1497,12 @@ void SRigHierarchy::ImportHierarchy(const FAssetData& InAssetData)
 
 		FScopedTransaction Transaction(LOCTEXT("HierarchyImport", "Import Hierarchy"));
 
+
 		// don't select bone if we are in construction mode.
 		// we do this to avoid the editmode / viewport shapes to refresh recursively,
 		// which can add an extreme slowdown depending on the number of bones (n^(n-1))
 		bool bSelectBones = true;
-		if (UControlRig* CurrentRig = ControlRigEditor.Pin()->ControlRig)
+		if (const UControlRig* CurrentRig = EditorSharedPtr->ControlRig)
 		{
 			bSelectBones = !CurrentRig->IsConstructionModeEnabled();
 		}
@@ -1484,36 +1510,26 @@ void SRigHierarchy::ImportHierarchy(const FAssetData& InAssetData)
 		URigHierarchyController* Controller = Hierarchy->GetController(true);
 		check(Controller);
 
-		TArray<FRigElementKey> ImportedBones = Controller->ImportBones(Mesh->GetSkeleton(), NAME_None, false, false, bSelectBones, true, true);
+		const TArray<FRigElementKey> ImportedBones = Controller->ImportBones(Mesh->GetSkeleton(), NAME_None, false, false, bSelectBones, true, true);
 		Controller->ImportCurves(Mesh->GetSkeleton(), NAME_None, false, true);
 
 		ControlRigBlueprint->SourceHierarchyImport = Mesh->GetSkeleton();
 		ControlRigBlueprint->SourceCurveImport = Mesh->GetSkeleton();
 
-		if(ImportedBones.Num() > 0)
+		if (ImportedBones.Num() > 0)
 		{
-			ControlRigEditor.Pin()->GetEditMode()->FrameItems(ImportedBones);
+			EditorSharedPtr->GetEditMode()->FrameItems(ImportedBones);
 		}
 	}
 
 	ControlRigBlueprint->PropagateHierarchyFromBPToInstances();
-	ControlRigEditor.Pin()->OnHierarchyChanged();
+	EditorSharedPtr->OnHierarchyChanged();
 	ControlRigBlueprint->BroadcastRefreshEditor();
 	RefreshTreeView();
 	FSlateApplication::Get().DismissAllMenus();
 
-	if (ControlRigBlueprint->GetPreviewMesh() == nullptr &&
-		ControlRigEditor.IsValid() && 
-		Mesh != nullptr)
-	{
-		ControlRigEditor.Pin()->UpdateMeshInAnimInstance(Mesh);
-		ControlRigEditor.Pin()->GetPersonaToolkit()->SetPreviewMesh(Mesh, true);
-	}
-
-	if (ControlRigEditor.IsValid())
-	{
-		ControlRigEditor.Pin()->Compile();
-	}
+	static constexpr bool bImport = true;
+	UpdateMesh(Mesh, bImport);
 }
 
 bool SRigHierarchy::IsMultiSelected(bool bIncludeProcedural) const
diff --git a/Engine/Plugins/Animation/ControlRig/Source/ControlRigEditor/Private/Editor/SRigHierarchy.h b/Engine/Plugins/Animation/ControlRig/Source/ControlRigEditor/Private/Editor/SRigHierarchy.h
--- a/Engine/Plugins/Animation/ControlRig/Source/ControlRigEditor/Private/Editor/SRigHierarchy.h
+++ b/Engine/Plugins/Animation/ControlRig/Source/ControlRigEditor/Private/Editor/SRigHierarchy.h
@@ -199,6 +199,7 @@ private:
 	void CreateRefreshMenu(FMenuBuilder& MenuBuilder);
 	bool ShouldFilterOnImport(const FAssetData& AssetData) const;
 	void RefreshHierarchy(const FAssetData& InAssetData);
+	void UpdateMesh(USkeletalMesh* InMesh, const bool bImport) const;
 
 	void HandleResetTransform(bool bSelectionOnly);
 	void HandleResetInitialTransform();
