Control Rig: Fix transaction buffer bug with user authored variables

diff --git a/Engine/Plugins/Animation/ControlRig/Source/ControlRig/Private/ControlRig.cpp b/Engine/Plugins/Animation/ControlRig/Source/ControlRig/Private/ControlRig.cpp
--- a/Engine/Plugins/Animation/ControlRig/Source/ControlRig/Private/ControlRig.cpp
+++ b/Engine/Plugins/Animation/ControlRig/Source/ControlRig/Private/ControlRig.cpp
@@ -2295,6 +2295,26 @@ void UControlRig::SetInteractionRigClass(TSubclassOf<UControlRig> InInteractionR
 
 void UControlRig::PreEditChange(FProperty* PropertyAboutToChange)
 {
+	// for BP user authored properties let's ignore changes since they
+	// will be distributed from the BP anyway to all archetype instances.
+	if(!PropertyAboutToChange->IsNative())
+	{
+		if(const UBlueprint* Blueprint = Cast<UBlueprint>(GetClass()->ClassGeneratedBy))
+		{
+			const struct FBPVariableDescription* FoundVariable = Blueprint->NewVariables.FindByPredicate(
+				[PropertyAboutToChange](const struct FBPVariableDescription& NewVariable)
+				{
+					return NewVariable.VarName == PropertyAboutToChange->GetFName();
+				}
+			);
+
+			if(FoundVariable)
+			{
+				return;
+			}
+		}
+	}
+	
 	Super::PreEditChange(PropertyAboutToChange);
 
 	if (PropertyAboutToChange && PropertyAboutToChange->GetFName() == GET_MEMBER_NAME_CHECKED(UControlRig, InteractionRig))
