Sequencer: Make sure we are on game thread when handling initialize events...
#preflght na
#trivial

diff --git a/Engine/Plugins/Animation/ControlRig/Source/ControlRig/Private/Sequencer/MovieSceneControlRigParameterTrack.cpp b/Engine/Plugins/Animation/ControlRig/Source/ControlRig/Private/Sequencer/MovieSceneControlRigParameterTrack.cpp
--- a/Engine/Plugins/Animation/ControlRig/Source/ControlRig/Private/Sequencer/MovieSceneControlRigParameterTrack.cpp
+++ b/Engine/Plugins/Animation/ControlRig/Source/ControlRig/Private/Sequencer/MovieSceneControlRigParameterTrack.cpp
@@ -458,17 +458,20 @@ void UMovieSceneControlRigParameterTrack::PostLoad()
 
 void UMovieSceneControlRigParameterTrack::HandleOnInitialized(URigVMHost* Subject, const FName& InEventName)
 {
-	if (IsValid(ControlRig))
+	if (IsInGameThread())
 	{
-		TArray<FRigControlElement*> SortedControls;
-		ControlRig->GetControlsInOrder(SortedControls);
-		for (UMovieSceneSection* BaseSection : GetAllSections())
+		if (IsValid(ControlRig))
 		{
-			if (UMovieSceneControlRigParameterSection* Section = Cast<UMovieSceneControlRigParameterSection>(BaseSection))
+			TArray<FRigControlElement*> SortedControls;
+			ControlRig->GetControlsInOrder(SortedControls);
+			for (UMovieSceneSection* BaseSection : GetAllSections())
 			{
-				if (Section->IsDifferentThanLastControlsUsedToReconstruct(SortedControls))
+				if (UMovieSceneControlRigParameterSection* Section = Cast<UMovieSceneControlRigParameterSection>(BaseSection))
 				{
-					Section->RecreateWithThisControlRig(ControlRig, Section->GetBlendType() == EMovieSceneBlendType::Absolute);
+					if (Section->IsDifferentThanLastControlsUsedToReconstruct(SortedControls))
+					{
+						Section->RecreateWithThisControlRig(ControlRig, Section->GetBlendType() == EMovieSceneBlendType::Absolute);
+					}
 				}
 			}
 		}
diff --git a/Engine/Plugins/Animation/ControlRig/Source/ControlRigEditor/Private/Sequencer/ControlRigParameterTrackEditor.cpp b/Engine/Plugins/Animation/ControlRig/Source/ControlRigEditor/Private/Sequencer/ControlRigParameterTrackEditor.cpp
--- a/Engine/Plugins/Animation/ControlRig/Source/ControlRigEditor/Private/Sequencer/ControlRigParameterTrackEditor.cpp
+++ b/Engine/Plugins/Animation/ControlRig/Source/ControlRigEditor/Private/Sequencer/ControlRigParameterTrackEditor.cpp
@@ -3061,14 +3061,17 @@ void FControlRigParameterTrackEditor::HandleControlSelected(UControlRig* Subject
 
 void FControlRigParameterTrackEditor::HandleOnInitialized(URigVMHost* Subject, const FName& InEventName)
 {
-	UControlRig* ControlRig = CastChecked<UControlRig>(Subject);
-	
-	if (GetSequencer().IsValid())
+	if (IsInGameThread())
 	{
-		//If FK control rig on next tick we refresh the tree
-		if (ControlRig->IsA<UFKControlRig>())
+		UControlRig* ControlRig = CastChecked<UControlRig>(Subject);
+
+		if (GetSequencer().IsValid())
 		{
-			GetSequencer()->NotifyMovieSceneDataChanged(EMovieSceneDataChangeType::RefreshTree);
+			//If FK control rig on next tick we refresh the tree
+			if (ControlRig->IsA<UFKControlRig>())
+			{
+				GetSequencer()->NotifyMovieSceneDataChanged(EMovieSceneDataChangeType::RefreshTree);
+			}
 		}
 	}
 }
