Sequencer: Snapper: Do not include controls without shapes when snapping.
#trivial

diff --git a/Engine/Plugins/Animation/ControlRig/Source/ControlRigEditor/Private/Tools/ControlRigSnapper.cpp b/Engine/Plugins/Animation/ControlRig/Source/ControlRigEditor/Private/Tools/ControlRigSnapper.cpp
--- a/Engine/Plugins/Animation/ControlRig/Source/ControlRigEditor/Private/Tools/ControlRigSnapper.cpp
+++ b/Engine/Plugins/Animation/ControlRig/Source/ControlRigEditor/Private/Tools/ControlRigSnapper.cpp
@@ -647,6 +647,25 @@ bool FControlRigSnapper::SnapIt(FFrameNumber StartFrame, FFrameNumber EndFrame,c
 							}
 						}
 
+						//only set keys on those specified
+						if (SnapSettings)
+						{
+							Context.KeyMask = 0;
+
+							if (SnapSettings->bSnapPosition == true)
+							{
+								Context.KeyMask |= (uint32)EControlRigContextChannelToKey::Translation;
+							}
+							if (SnapSettings->bSnapRotation == true)
+							{
+								Context.KeyMask |= (uint32)EControlRigContextChannelToKey::Rotation;
+							}
+							if (SnapSettings->bSnapScale == true)
+							{
+								Context.KeyMask |= (uint32)EControlRigContextChannelToKey::Scale;
+							}
+						}
+
 						for (int32 Index = 0; Index < WorldTransformToSnap.Num(); ++Index)
 						{
 							OneFrame[0] = Frames[Index];
