Sequencer: Control Rig: Handle control rig initialization in the game thread.

diff --git a/Engine/Plugins/Animation/ControlRig/Source/ControlRig/Private/Sequencer/MovieSceneControlRigParameterTrack.cpp b/Engine/Plugins/Animation/ControlRig/Source/ControlRig/Private/Sequencer/MovieSceneControlRigParameterTrack.cpp
--- a/Engine/Plugins/Animation/ControlRig/Source/ControlRig/Private/Sequencer/MovieSceneControlRigParameterTrack.cpp
+++ b/Engine/Plugins/Animation/ControlRig/Source/ControlRig/Private/Sequencer/MovieSceneControlRigParameterTrack.cpp
@@ -10,6 +10,8 @@
 #include "Channels/MovieSceneChannelProxy.h"
 #include "Rigs/RigHierarchyController.h"
 #include "UObject/Package.h"
+#include "Async/Async.h"
+
 
 #include UE_INLINE_GENERATED_CPP_BY_NAME(MovieSceneControlRigParameterTrack)
 
@@ -455,29 +457,97 @@ void UMovieSceneControlRigParameterTrack::PostLoad()
 	}
 }
 
+//mz todo this is coming from BuildPatchServices in Runtime/Online/BuildPatchServices/Online/Core/AsyncHelper.h
+//looking at moving this over
+//this is very useful since it properly handles passing in a this pointer to the async task.
+namespace MovieSceneControlRigTrack
+{
+	/**
+	 * Helper functions for wrapping async functionality.
+	 */
+	namespace AsyncHelpers
+	{
+		template<typename ResultType, typename... Args>
+		static TFunction<void()> MakePromiseKeeper(const TSharedRef<TPromise<ResultType>, ESPMode::ThreadSafe>& Promise, const TFunction<ResultType(Args...)>& Function, Args... FuncArgs)
+		{
+			return [Promise, Function, FuncArgs...]()
+			{
+				Promise->SetValue(Function(FuncArgs...));
+			};
+		}
 
-void UMovieSceneControlRigParameterTrack::HandleOnInitialized(URigVMHost* Subject, const FName& InEventName)
+		template<typename... Args>
+		static TFunction<void()> MakePromiseKeeper(const TSharedRef<TPromise<void>, ESPMode::ThreadSafe>& Promise, const TFunction<void(Args...)>& Function, Args... FuncArgs)
+		{
+			return [Promise, Function, FuncArgs...]()
+			{
+				Function(FuncArgs...);
+				Promise->SetValue();
+			};
+		}
+
+		template<typename ResultType, typename... Args>
+		static TFuture<ResultType> ExecuteOnGameThread(const TFunction<ResultType(Args...)>& Function, Args... FuncArgs)
+		{
+			TSharedRef<TPromise<ResultType>, ESPMode::ThreadSafe> Promise = MakeShareable(new TPromise<ResultType>());
+			TFunction<void()> PromiseKeeper = MakePromiseKeeper(Promise, Function, FuncArgs...);
+			if (!IsInGameThread())
+			{
+				AsyncTask(ENamedThreads::GameThread, MoveTemp(PromiseKeeper));
+			}
+			else
+			{
+				PromiseKeeper();
+			}
+			return Promise->GetFuture();
+		}
+
+		template<typename ResultType>
+		static TFuture<ResultType> ExecuteOnGameThread(const TFunction<ResultType()>& Function)
+		{
+			TSharedRef<TPromise<ResultType>, ESPMode::ThreadSafe> Promise = MakeShareable(new TPromise<ResultType>());
+			TFunction<void()> PromiseKeeper = MakePromiseKeeper(Promise, Function);
+			if (!IsInGameThread())
+			{
+				AsyncTask(ENamedThreads::GameThread, MoveTemp(PromiseKeeper));
+			}
+			else
+			{
+				PromiseKeeper();
+			}
+			return Promise->GetFuture();
+		}
+	}
+}
+
+void UMovieSceneControlRigParameterTrack::HandleOnInitialized_GameThread()
 {
-	if (IsInGameThread())
+	if (IsValid(ControlRig))
 	{
-		if (IsValid(ControlRig))
+		TArray<FRigControlElement*> SortedControls;
+		ControlRig->GetControlsInOrder(SortedControls);
+		for (UMovieSceneSection* BaseSection : GetAllSections())
 		{
-			TArray<FRigControlElement*> SortedControls;
-			ControlRig->GetControlsInOrder(SortedControls);
-			for (UMovieSceneSection* BaseSection : GetAllSections())
+			if (UMovieSceneControlRigParameterSection* Section = Cast<UMovieSceneControlRigParameterSection>(BaseSection))
 			{
-				if (UMovieSceneControlRigParameterSection* Section = Cast<UMovieSceneControlRigParameterSection>(BaseSection))
+				if (Section->IsDifferentThanLastControlsUsedToReconstruct(SortedControls))
 				{
-					if (Section->IsDifferentThanLastControlsUsedToReconstruct(SortedControls))
-					{
-						Section->RecreateWithThisControlRig(ControlRig, Section->GetBlendType() == EMovieSceneBlendType::Absolute);
-					}
+					Section->RecreateWithThisControlRig(ControlRig, Section->GetBlendType() == EMovieSceneBlendType::Absolute);
 				}
 			}
 		}
 	}
 }
 
+void UMovieSceneControlRigParameterTrack::HandleOnInitialized(URigVMHost* Subject, const FName& InEventName)
+{
+	MovieSceneControlRigTrack::AsyncHelpers::ExecuteOnGameThread<void>([this]()
+	{
+		HandleOnInitialized_GameThread();
+
+	}).Wait();
+}
+
 #if WITH_EDITORONLY_DATA
 void UMovieSceneControlRigParameterTrack::DeclareConstructClasses(TArray<FTopLevelAssetPath>& OutConstructClasses, const UClass* SpecificSubclass)
 {
diff --git a/Engine/Plugins/Animation/ControlRig/Source/ControlRig/Public/Sequencer/MovieSceneControlRigParameterTrack.h b/Engine/Plugins/Animation/ControlRig/Source/ControlRig/Public/Sequencer/MovieSceneControlRigParameterTrack.h
--- a/Engine/Plugins/Animation/ControlRig/Source/ControlRig/Public/Sequencer/MovieSceneControlRigParameterTrack.h
+++ b/Engine/Plugins/Animation/ControlRig/Source/ControlRig/Public/Sequencer/MovieSceneControlRigParameterTrack.h
@@ -55,6 +55,7 @@ public:
 	virtual void GetSelectedNodes(TArray<FName>& OutSelectedNodes) override;
 
 	//UControlRig Delegates
+	void HandleOnInitialized_GameThread();
 	void HandleOnInitialized(URigVMHost* Subject, const FName& InEventName);
 
 #if WITH_EDITOR
