Control Rig: Offer output bone transfer list for anim node

diff --git a/Engine/Plugins/Animation/ControlRig/Source/ControlRig/Private/AnimNode_ControlRigBase.cpp b/Engine/Plugins/Animation/ControlRig/Source/ControlRig/Private/AnimNode_ControlRigBase.cpp
--- a/Engine/Plugins/Animation/ControlRig/Source/ControlRig/Private/AnimNode_ControlRigBase.cpp
+++ b/Engine/Plugins/Animation/ControlRig/Source/ControlRig/Private/AnimNode_ControlRigBase.cpp
@@ -633,28 +633,33 @@ void FAnimNode_ControlRigBase::CacheBones_AnyThread(const FAnimationCacheBonesCo
 				}
 			}
 
-			if(!InputBonesToTransfer.IsEmpty())
-			{
-				ControlRigBoneOutputMappingByName = ControlRigBoneInputMappingByName;
-				ControlRigBoneInputMappingByName.Reset();
+			auto UpdatingMappingFromSpecificTransferList = [] (
+				TArray<FBoneReference>& InTransferList,
+				const TWeakObjectPtr<UNodeMappingContainer>& InMappingContainer,
+				FBoneContainer& InRequiredBones,
+				const FReferenceSkeleton& InRefSkeleton,
+				const TArray<FBoneIndexType>& InRequiredBonesArray,
+				const UControlRig* InControlRig,
+				TMap<FName, uint16>& OutMapping
+			) {
+				OutMapping.Reset();
 				
-				if (NodeMappingContainer.IsValid())
+				if (InMappingContainer.IsValid())
 				{
 					// get target to source mapping table - this is reversed mapping table
 					TMap<FName, FName> TargetToSourceMappingTable;
-					NodeMappingContainer->GetTargetToSourceMappingTable(TargetToSourceMappingTable);
+					InMappingContainer->GetTargetToSourceMappingTable(TargetToSourceMappingTable);
 
-					for(FBoneReference& InputBoneToTransfer : InputBonesToTransfer)
+					for(FBoneReference& InputBoneToTransfer : InTransferList)
 					{
-						if(!InputBoneToTransfer.Initialize(RequiredBones))
+						if(!InputBoneToTransfer.Initialize(InRequiredBones))
 						{
 							continue;
 						}
-						FName TargetNodeName = RefSkeleton.GetBoneName(InputBoneToTransfer.BoneIndex);
-						FName* SourceName = TargetToSourceMappingTable.Find(TargetNodeName);
-						if (SourceName)
+						const FName TargetNodeName = InRefSkeleton.GetBoneName(InputBoneToTransfer.BoneIndex);
+						if (const FName* SourceName = TargetToSourceMappingTable.Find(TargetNodeName))
 						{
-							ControlRigBoneInputMappingByName.Add(*SourceName, InputBoneToTransfer.BoneIndex);
+							OutMapping.Add(*SourceName, InputBoneToTransfer.BoneIndex);
 						}
 					}
 				}
@@ -662,21 +667,47 @@ void FAnimNode_ControlRigBase::CacheBones_AnyThread(const FAnimationCacheBonesCo
 				{
 					TArray<FName> NodeNames;
 					TArray<FNodeItem> NodeItems;
-					ControlRig->GetMappableNodeData(NodeNames, NodeItems);
+					InControlRig->GetMappableNodeData(NodeNames, NodeItems);
 
-					for(FBoneReference& InputBoneToTransfer : InputBonesToTransfer)
+					for(FBoneReference& InputBoneToTransfer : InTransferList)
 					{
-						if(!InputBoneToTransfer.Initialize(RequiredBones))
+						if(!InputBoneToTransfer.Initialize(InRequiredBones))
 						{
 							continue;
 						}
-						const FName& BoneName = RefSkeleton.GetBoneName(RequiredBonesArray[InputBoneToTransfer.BoneIndex]);
+						const FName& BoneName = InRefSkeleton.GetBoneName(InRequiredBonesArray[InputBoneToTransfer.BoneIndex]);
 						if (NodeNames.Contains(BoneName))
 						{
-							ControlRigBoneInputMappingByName.Add(BoneName, InputBoneToTransfer.BoneIndex);
+							OutMapping.Add(BoneName, InputBoneToTransfer.BoneIndex);
 						}
 					}
 				}
+			};
+			
+			if(!InputBonesToTransfer.IsEmpty())
+			{
+				ControlRigBoneOutputMappingByName = ControlRigBoneInputMappingByName;
+
+				UpdatingMappingFromSpecificTransferList(
+					InputBonesToTransfer,
+					NodeMappingContainer,
+					RequiredBones,
+					RefSkeleton,
+					RequiredBonesArray,
+					ControlRig,
+					ControlRigBoneInputMappingByName);
+			}
+
+			if(!OutputBonesToTransfer.IsEmpty())
+			{
+				UpdatingMappingFromSpecificTransferList(
+					OutputBonesToTransfer,
+					NodeMappingContainer,
+					RequiredBones,
+					RefSkeleton,
+					RequiredBonesArray,
+					ControlRig,
+					ControlRigBoneOutputMappingByName);
 			}
 
 			// check if we can switch the bones to an index based mapping.
diff --git a/Engine/Plugins/Animation/ControlRig/Source/ControlRig/Public/AnimNode_ControlRigBase.h b/Engine/Plugins/Animation/ControlRig/Source/ControlRig/Public/AnimNode_ControlRigBase.h
--- a/Engine/Plugins/Animation/ControlRig/Source/ControlRig/Public/AnimNode_ControlRigBase.h
+++ b/Engine/Plugins/Animation/ControlRig/Source/ControlRig/Public/AnimNode_ControlRigBase.h
@@ -119,6 +119,14 @@ protected:
 	UPROPERTY(EditAnywhere, AdvancedDisplay, Category = Settings)
 	TArray<FBoneReference> InputBonesToTransfer;
 
+	/**
+	 * An inclusive list of bones to transfer as part
+	 * of the output pose transfer phase.
+	 * If this list is empty all bones will be transferred.
+	 */
+	UPROPERTY(EditAnywhere, AdvancedDisplay, Category = Settings)
+	TArray<FBoneReference> OutputBonesToTransfer;
+
 	/** Complete mapping from skeleton to control rig bone index */
 	TArray<TPair<uint16, uint16>> ControlRigBoneInputMappingByIndex;
 	TArray<TPair<uint16, uint16>> ControlRigBoneOutputMappingByIndex;
