Control Rig: Add nullptr check.

diff --git a/Engine/Plugins/Animation/ControlRig/Source/ControlRig/Private/ControlRig.cpp b/Engine/Plugins/Animation/ControlRig/Source/ControlRig/Private/ControlRig.cpp
--- a/Engine/Plugins/Animation/ControlRig/Source/ControlRig/Private/ControlRig.cpp
+++ b/Engine/Plugins/Animation/ControlRig/Source/ControlRig/Private/ControlRig.cpp
@@ -2299,7 +2299,7 @@ void UControlRig::PreEditChange(FProperty* PropertyAboutToChange)
 {
 	// for BP user authored properties let's ignore changes since they
 	// will be distributed from the BP anyway to all archetype instances.
-	if(!PropertyAboutToChange->IsNative())
+	if(PropertyAboutToChange && !PropertyAboutToChange->IsNative())
 	{
 		if(const UBlueprint* Blueprint = Cast<UBlueprint>(GetClass()->ClassGeneratedBy))
 		{
