Control Rig: Fix crash when right clicking a set shape settings node
#rnx

diff --git a/Engine/Plugins/Animation/ControlRig/Source/ControlRig/Private/Units/Execution/RigUnit_DynamicHierarchy.cpp b/Engine/Plugins/Animation/ControlRig/Source/ControlRig/Private/Units/Execution/RigUnit_DynamicHierarchy.cpp
--- a/Engine/Plugins/Animation/ControlRig/Source/ControlRig/Private/Units/Execution/RigUnit_DynamicHierarchy.cpp
+++ b/Engine/Plugins/Animation/ControlRig/Source/ControlRig/Private/Units/Execution/RigUnit_DynamicHierarchy.cpp
@@ -990,18 +990,20 @@ FRigUnit_HierarchyGetShapeSettings_Execute()
 
 FRigUnit_HierarchySetShapeSettings_Execute()
 {
-	if(FRigControlElement* ControlElement = ExecuteContext.Hierarchy->Find<FRigControlElement>(Item))
-	{
-		ControlElement->Settings.bShapeVisible = Settings.bVisible;
-		ControlElement->Settings.ShapeColor = Settings.Color;
-		ControlElement->Settings.ShapeName = Settings.Name;
-		ExecuteContext.Hierarchy->Notify(ERigHierarchyNotification::ControlSettingChanged, ControlElement);
-		ExecuteContext.Hierarchy->SetControlShapeTransformByIndex(ControlElement->GetIndex(), Settings.Transform, true, false);
-		ExecuteContext.Hierarchy->SetControlShapeTransformByIndex(ControlElement->GetIndex(), Settings.Transform, false, false);
-	}
-	else
+	if(URigHierarchy* Hierarchy = ExecuteContext.Hierarchy)
 	{
-		UE_CONTROLRIG_RIGUNIT_REPORT_ERROR(TEXT("Control '%s' does not exist."), *Item.ToString());
+		if(FRigControlElement* ControlElement = Hierarchy->Find<FRigControlElement>(Item))
+		{
+			ControlElement->Settings.bShapeVisible = Settings.bVisible;
+			ControlElement->Settings.ShapeColor = Settings.Color;
+			ControlElement->Settings.ShapeName = Settings.Name;
+			Hierarchy->Notify(ERigHierarchyNotification::ControlSettingChanged, ControlElement);
+			Hierarchy->SetControlShapeTransformByIndex(ControlElement->GetIndex(), Settings.Transform, true, false);
+			Hierarchy->SetControlShapeTransformByIndex(ControlElement->GetIndex(), Settings.Transform, false, false);
+		}
+		else
+		{
+			UE_CONTROLRIG_RIGUNIT_REPORT_ERROR(TEXT("Control '%s' does not exist."), *Item.ToString());
+		}
 	}
-
 }
\ No newline at end of file
