clear detail object when importing a new skeletal mesh

diff --git a/Engine/Plugins/Animation/ControlRig/Source/ControlRigEditor/Private/Editor/SRigHierarchy.cpp b/Engine/Plugins/Animation/ControlRig/Source/ControlRigEditor/Private/Editor/SRigHierarchy.cpp
--- a/Engine/Plugins/Animation/ControlRig/Source/ControlRigEditor/Private/Editor/SRigHierarchy.cpp
+++ b/Engine/Plugins/Animation/ControlRig/Source/ControlRigEditor/Private/Editor/SRigHierarchy.cpp
@@ -1493,11 +1493,12 @@ void SRigHierarchy::ImportHierarchy(const FAssetData& InAssetData)
 			return;
 		}
 		
+		EditorSharedPtr->ClearDetailObject();
+		
 		TGuardValue<bool> SuspendBlueprintNotifs(ControlRigBlueprint->bSuspendAllNotifications, true);
 
 		FScopedTransaction Transaction(LOCTEXT("HierarchyImport", "Import Hierarchy"));
 
-
 		// don't select bone if we are in construction mode.
 		// we do this to avoid the editmode / viewport shapes to refresh recursively,
 		// which can add an extreme slowdown depending on the number of bones (n^(n-1))
