RigVM: Avoid crash when creating function with same name as deleted function

diff --git a/Engine/Plugins/Animation/ControlRig/Source/ControlRigDeveloper/Private/ControlRigBlueprint.cpp b/Engine/Plugins/Animation/ControlRig/Source/ControlRigDeveloper/Private/ControlRigBlueprint.cpp
--- a/Engine/Plugins/Animation/ControlRig/Source/ControlRigDeveloper/Private/ControlRigBlueprint.cpp
+++ b/Engine/Plugins/Animation/ControlRig/Source/ControlRigDeveloper/Private/ControlRigBlueprint.cpp
@@ -3822,7 +3822,7 @@ void UControlRigBlueprint::HandleModifiedEvent(ERigVMGraphNotifType InNotifType,
 				if (URigVMPin* ModelPin = Cast<URigVMPin>(InSubject))
 				{
 					if (UEdGraph* EdGraph = GetEdGraph(InGraph))
-					{							
+					{
 						if (UControlRigGraph* Graph = Cast<UControlRigGraph>(EdGraph))
 						{
 							if (UEdGraphNode* EdNode = Graph->FindNodeForModelNodeName(ModelPin->GetNode()->GetFName()))
@@ -5366,7 +5366,7 @@ void UControlRigBlueprint::CreateEdGraphForCollapseNodeIfNeeded(URigVMCollapseNo
 			if (!bFunctionGraphExists)
 			{
 				// create a sub graph
-				UControlRigGraph* RigFunctionGraph = NewObject<UControlRigGraph>(this, *InNode->GetName(), RF_Transactional);
+				UControlRigGraph* RigFunctionGraph = NewObject<UControlRigGraph>(this, *InNode->GetEditorSubGraphName(), RF_Transactional);
 				RigFunctionGraph->Schema = UControlRigGraphSchema::StaticClass();
 				RigFunctionGraph->bAllowRenaming = 1;
 				RigFunctionGraph->bEditable = 1;
