RigVm: Fix User defined event nodes with wrong title
#rnx

diff --git a/Engine/Plugins/Animation/ControlRig/Source/ControlRigDeveloper/Private/Graph/ControlRigGraphNode.cpp b/Engine/Plugins/Animation/ControlRig/Source/ControlRigDeveloper/Private/Graph/ControlRigGraphNode.cpp
--- a/Engine/Plugins/Animation/ControlRig/Source/ControlRigDeveloper/Private/Graph/ControlRigGraphNode.cpp
+++ b/Engine/Plugins/Animation/ControlRig/Source/ControlRigDeveloper/Private/Graph/ControlRigGraphNode.cpp
@@ -70,11 +70,11 @@ FText UControlRigGraphNode::GetNodeTitle(ENodeTitleType::Type TitleType) const
 			if (URigVMUnitNode* UnitNode = Cast<URigVMUnitNode>(ModelNode))
 			{
 				const UScriptStruct* ScriptStruct = UnitNode->GetScriptStruct();
-				if (ScriptStruct && ScriptStruct->IsChildOf(FRigUnit::StaticStruct()))
+				if (ScriptStruct && ScriptStruct->IsChildOf(FRigVMStruct::StaticStruct()))
 				{
 					if (TSharedPtr<FStructOnScope> StructOnScope = UnitNode->ConstructStructInstance())
 					{
-						FRigUnit* RigUnit = (FRigUnit*)StructOnScope->GetStructMemory();
+						FRigVMStruct* RigUnit = (FRigVMStruct*)StructOnScope->GetStructMemory();
 						NodeTitle = FText::FromString(RigUnit->GetUnitLabel());
 					}
 				}
