Control Rig: Fix nulls not resetting to initial pose when Ctrl+Shift+G is triggered

diff --git a/Engine/Plugins/Animation/ControlRig/Source/ControlRigEditor/Private/EditMode/ControlRigEditMode.cpp b/Engine/Plugins/Animation/ControlRig/Source/ControlRigEditor/Private/EditMode/ControlRigEditMode.cpp
--- a/Engine/Plugins/Animation/ControlRig/Source/ControlRigEditor/Private/EditMode/ControlRigEditMode.cpp
+++ b/Engine/Plugins/Animation/ControlRig/Source/ControlRigEditor/Private/EditMode/ControlRigEditMode.cpp
@@ -2846,7 +2846,7 @@ void FControlRigEditMode::ResetTransforms(bool bSelectionOnly)
 	{
 		if (UControlRig* ControlRig = RuntimeRigPtr.Get())
 		{
-			if (ControlRig->CurrentControlSelection().Num() > 0)
+			if (!bSelectionOnly || ControlRig->CurrentControlSelection().Num() > 0)
 			{
 				ControlRigs.Add(ControlRig);
 			}
@@ -2864,8 +2864,17 @@ void FControlRigEditMode::ResetTransforms(bool bSelectionOnly)
 		TArray<FRigElementKey> SelectedRigElements = GetSelectedRigElements(ControlRig);
 		TArray<FRigElementKey> ControlsToReset = SelectedRigElements;
 		TArray<FRigElementKey> ControlsInteracting = SelectedRigElements;
+		TArray<FRigElementKey> TransformElementsToReset = SelectedRigElements;
 		if (!bSelectionOnly)
 		{
+			TArray<FRigBaseElement*> Elements = ControlRig->GetHierarchy()->GetElementsOfType<FRigBaseElement>(true);
+			TransformElementsToReset.Reset();
+			TransformElementsToReset.Reserve(Elements.Num());
+			for (const FRigBaseElement* Element : Elements)
+			{
+				TransformElementsToReset.Add(Element->GetKey());
+			}
+			
 			TArray<FRigControlElement*> Controls;
 			ControlRig->GetControlsInOrder(Controls);
 			ControlsToReset.SetNum(0);
@@ -2881,46 +2890,55 @@ void FControlRigEditMode::ResetTransforms(bool bSelectionOnly)
 			}
 		}
 		bool bHasNonDefaultParent = false;
-		TArray<FRigElementKey> Parents;
-		for (const FRigElementKey& ControlKey : ControlsToReset)
+		TMap<FRigElementKey, FRigElementKey> Parents;
+		for (const FRigElementKey& Key : TransformElementsToReset)
 		{
-			FRigElementKey SpaceKey = ControlRig->GetHierarchy()->GetActiveParent(ControlKey);
-			Parents.Add(SpaceKey);
-			if (SpaceKey != ControlRig->GetHierarchy()->GetDefaultParentKey())
+			FRigElementKey SpaceKey = ControlRig->GetHierarchy()->GetActiveParent(Key);
+			Parents.Add(Key, SpaceKey);
+			if (!bHasNonDefaultParent && SpaceKey != ControlRig->GetHierarchy()->GetDefaultParentKey())
 			{
 				bHasNonDefaultParent = true;
 			}
 		}
 
 		FControlRigInteractionScope InteractionScope(ControlRig, ControlsInteracting);
-
-		for (const FRigElementKey& ControlToReset : ControlsToReset)
+		for (const FRigElementKey& ElementToReset : TransformElementsToReset)
 		{
-			if (ControlToReset.Type == ERigElementType::Control)
+			FRigControlElement* ControlElement = nullptr;
+			if (ElementToReset.Type == ERigElementType::Control)
 			{
-				FRigControlElement* ControlElement = ControlRig->FindControl(ControlToReset.Name);
-				if (ControlElement && !ControlElement->Settings.bIsTransientControl)
+				ControlElement = ControlRig->FindControl(ElementToReset.Name);
+				if (ControlElement->Settings.bIsTransientControl)
 				{
-					const FTransform InitialLocalTransform = ControlRig->GetHierarchy()->GetInitialLocalTransform(ControlToReset);
-					ControlRig->Modify();
-					if (bHasNonDefaultParent == true) //possibly not at default parent so switch to it
-					{
-						ControlRig->GetHierarchy()->SwitchToDefaultParent(ControlElement->GetKey());
-					}
-					ControlRig->SetControlLocalTransform(ControlToReset.Name, InitialLocalTransform, true);
-					NotifyDrivenControls(ControlRig, ControlToReset);
-					if (bHasNonDefaultParent == false)
-					{
-						ControlRig->ControlModified().Broadcast(ControlRig, ControlElement, EControlRigSetKey::DoNotCare);
-					}
-
-					//@helge not sure what to do if the non-default parent
-					if (UControlRigBlueprint* Blueprint = Cast<UControlRigBlueprint>(ControlRig->GetClass()->ClassGeneratedBy))
-					{
-						Blueprint->Hierarchy->SetLocalTransform(ControlToReset, InitialLocalTransform);
-					}
+					ControlElement = nullptr;
 				}
 			}
+			
+			const FTransform InitialLocalTransform = ControlRig->GetHierarchy()->GetInitialLocalTransform(ElementToReset);
+			ControlRig->Modify();
+			if (bHasNonDefaultParent == true) //possibly not at default parent so switch to it
+			{
+				ControlRig->GetHierarchy()->SwitchToDefaultParent(ElementToReset);
+			}
+			if (ControlElement)
+			{
+				ControlRig->SetControlLocalTransform(ElementToReset.Name, InitialLocalTransform, true);
+				NotifyDrivenControls(ControlRig, ElementToReset);
+				if (bHasNonDefaultParent == false)
+				{
+					ControlRig->ControlModified().Broadcast(ControlRig, ControlElement, EControlRigSetKey::DoNotCare);
+				}
+			}
+			else
+			{
+				ControlRig->GetHierarchy()->SetLocalTransform(ElementToReset, InitialLocalTransform, false, true, true);
+			}
+
+			//@helge not sure what to do if the non-default parent
+			if (UControlRigBlueprint* Blueprint = Cast<UControlRigBlueprint>(ControlRig->GetClass()->ClassGeneratedBy))
+			{
+				Blueprint->Hierarchy->SetLocalTransform(ElementToReset, InitialLocalTransform);
+			}
 		}
 
 		if (bHasNonDefaultParent == true) //now we have the initial pose setup we need to get the global transforms as specified now then set them in the current parent space
@@ -2928,43 +2946,66 @@ void FControlRigEditMode::ResetTransforms(bool bSelectionOnly)
 			ControlRig->Evaluate_AnyThread();
 
 			//get global transforms
-			TArray<FTransform> GlobalTransforms;
-			for (const FRigElementKey& ControlToReset : ControlsToReset)
+			TMap<FRigElementKey, FTransform> GlobalTransforms;
+			for (const FRigElementKey& ElementToReset : TransformElementsToReset)
 			{
-				FRigControlElement* ControlElement = ControlRig->FindControl(ControlToReset.Name);
-				if (ControlElement && !ControlElement->Settings.bIsTransientControl)
+				if (ElementToReset.IsTypeOf(ERigElementType::Control))
 				{
-					FTransform GlobalTransform = ControlRig->GetHierarchy()->GetGlobalTransform(ControlToReset);
-					GlobalTransforms.Add(GlobalTransform);
+					FRigControlElement* ControlElement = ControlRig->FindControl(ElementToReset.Name);
+					if (ControlElement && !ControlElement->Settings.bIsTransientControl)
+					{
+						FTransform GlobalTransform = ControlRig->GetHierarchy()->GetGlobalTransform(ElementToReset);
+						GlobalTransforms.Add(ElementToReset, GlobalTransform);
+					}
+					NotifyDrivenControls(ControlRig, ElementToReset);
+				}
+				else
+				{
+					FTransform GlobalTransform = ControlRig->GetHierarchy()->GetGlobalTransform(ElementToReset);
+					GlobalTransforms.Add(ElementToReset, GlobalTransform);
 				}
-				NotifyDrivenControls(ControlRig, ControlToReset);
 			}
 			//switch back to original parent space
-			int32 Index = 0;
-			for (const FRigElementKey& ControlToReset : ControlsToReset)
+			for (const FRigElementKey& ElementToReset : TransformElementsToReset)
 			{
-				FRigControlElement* ControlElement = ControlRig->FindControl(ControlToReset.Name);
-				if (ControlElement && !ControlElement->Settings.bIsTransientControl)
+				if (const FRigElementKey* SpaceKey = Parents.Find(ElementToReset))
 				{
-					ControlRig->GetHierarchy()->SwitchToParent(ControlToReset, Parents[Index]);
-					++Index;
+					if (ElementToReset.IsTypeOf(ERigElementType::Control))
+					{
+						FRigControlElement* ControlElement = ControlRig->FindControl(ElementToReset.Name);
+						if (ControlElement && !ControlElement->Settings.bIsTransientControl)
+						{
+							ControlRig->GetHierarchy()->SwitchToParent(ElementToReset, *SpaceKey);
+						}
+					}
+					else
+					{
+						ControlRig->GetHierarchy()->SwitchToParent(ElementToReset, *SpaceKey);
+					}
 				}
 			}
 			//set global transforms in this space // do it twice since ControlsInOrder is not really always in order
 			for (int32 SetHack = 0; SetHack < 2; ++SetHack)
 			{
 				ControlRig->Evaluate_AnyThread();
-				Index = 0;
-				for (const FRigElementKey& ControlToReset : ControlsToReset)
-
+				for (const FRigElementKey& ElementToReset : TransformElementsToReset)
 				{
-					FRigControlElement* ControlElement = ControlRig->FindControl(ControlToReset.Name);
-					if (ControlElement && !ControlElement->Settings.bIsTransientControl)
+					if (const FTransform* GlobalTransform = GlobalTransforms.Find(ElementToReset))
 					{
-						ControlRig->SetControlGlobalTransform(ControlToReset.Name, GlobalTransforms[Index], true);
-						ControlRig->Evaluate_AnyThread();
-						NotifyDrivenControls(ControlRig, ControlToReset);
-						++Index;
+						if (ElementToReset.IsTypeOf(ERigElementType::Control))
+						{
+							FRigControlElement* ControlElement = ControlRig->FindControl(ElementToReset.Name);
+							if (ControlElement && !ControlElement->Settings.bIsTransientControl)
+							{
+								ControlRig->SetControlGlobalTransform(ElementToReset.Name, *GlobalTransform, true);
+								ControlRig->Evaluate_AnyThread();
+								NotifyDrivenControls(ControlRig, ElementToReset);
+							}
+						}
+						else
+						{
+							ControlRig->GetHierarchy()->SetGlobalTransform(ElementToReset, *GlobalTransform, false, true, true);
+						}
 					}
 				}
 			}
