ControlRig: Catch transaction ensure when changing comment box text

diff --git a/Engine/Plugins/Animation/ControlRig/Source/ControlRigDeveloper/Private/Graph/ControlRigGraph.cpp b/Engine/Plugins/Animation/ControlRig/Source/ControlRigDeveloper/Private/Graph/ControlRigGraph.cpp
--- a/Engine/Plugins/Animation/ControlRig/Source/ControlRigDeveloper/Private/Graph/ControlRigGraph.cpp
+++ b/Engine/Plugins/Animation/ControlRig/Source/ControlRigDeveloper/Private/Graph/ControlRigGraph.cpp
@@ -636,7 +636,7 @@ void UControlRigGraph::HandleModifiedEvent(ERigVMGraphNotifType InNotifType, URi
 				UEdGraphNode_Comment* EdNode = Cast<UEdGraphNode_Comment>(FindNodeForModelNodeName(ModelNode->GetFName()));
 				if (EdNode)
 				{
-					EdNode->OnUpdateCommentText(ModelNode->GetCommentText());
+					EdNode->NodeComment = ModelNode->GetCommentText();
 					EdNode->FontSize = ModelNode->GetCommentFontSize();
 					EdNode->bCommentBubbleVisible = ModelNode->GetCommentBubbleVisible();
 					EdNode->bCommentBubbleVisible_InDetailsPanel = ModelNode->GetCommentBubbleVisible();
