Control Rig: Fix SwitchParent in construction graph not working properly

diff --git a/Engine/Plugins/Animation/ControlRig/Source/ControlRig/Private/ControlRig.cpp b/Engine/Plugins/Animation/ControlRig/Source/ControlRig/Private/ControlRig.cpp
--- a/Engine/Plugins/Animation/ControlRig/Source/ControlRig/Private/ControlRig.cpp
+++ b/Engine/Plugins/Animation/ControlRig/Source/ControlRig/Private/ControlRig.cpp
@@ -655,11 +655,11 @@ bool UControlRig::Execute(const FName& InEventName)
 				// run in between forward events
 				// the saved pose is reapplied to the rig after construction event as the pose scope goes out of scope
 				TUniquePtr<FPoseScope> PoseScope;
-				if (!bConstructionModeEnabled)
+				if (!bConstructionModeEnabled && !bJustRanInit)
 				{
 					// only do this in non-construction mode because 
 					// when construction mode is enabled, the control values are cleared before reaching here (too late to save them)
-					PoseScope = MakeUnique<FPoseScope>(this, ERigElementType::ToResetAfterConstructionEvent);
+					PoseScope = MakeUnique<FPoseScope>(this, ERigElementType::ToResetAfterConstructionEvent, TArray<FRigElementKey>(), ERigTransformType::CurrentGlobal);
 				}
 				
 				{
@@ -2898,20 +2898,21 @@ void UControlRig::OnHierarchyTransformUndoRedo(URigHierarchy* InHierarchy, const
 	}
 }
 
-UControlRig::FPoseScope::FPoseScope(UControlRig* InControlRig, ERigElementType InFilter, const TArray<FRigElementKey>& InElements)
+UControlRig::FPoseScope::FPoseScope(UControlRig* InControlRig, ERigElementType InFilter, const TArray<FRigElementKey>& InElements, const ERigTransformType::Type InTransformType)
 : ControlRig(InControlRig)
 , Filter(InFilter)
+, TransformType(InTransformType)
 {
 	check(InControlRig);
 	const TArrayView<const FRigElementKey> ElementView(InElements.GetData(), InElements.Num());
-	CachedPose = InControlRig->GetHierarchy()->GetPose(false, InFilter, ElementView);
+	CachedPose = InControlRig->GetHierarchy()->GetPose(IsInitial(InTransformType), InFilter, ElementView);
 }
 
 UControlRig::FPoseScope::~FPoseScope()
 {
 	check(ControlRig);
 
-	ControlRig->GetHierarchy()->SetPose(CachedPose);
+	ControlRig->GetHierarchy()->SetPose(CachedPose, TransformType);
 }
 
 #if WITH_EDITOR
diff --git a/Engine/Plugins/Animation/ControlRig/Source/ControlRig/Public/ControlRig.h b/Engine/Plugins/Animation/ControlRig/Source/ControlRig/Public/ControlRig.h
--- a/Engine/Plugins/Animation/ControlRig/Source/ControlRig/Public/ControlRig.h
+++ b/Engine/Plugins/Animation/ControlRig/Source/ControlRig/Public/ControlRig.h
@@ -610,7 +610,7 @@ private:
 	{
 	public:
 		FPoseScope(UControlRig* InControlRig, ERigElementType InFilter = ERigElementType::All,
-			const TArray<FRigElementKey>& InElements = TArray<FRigElementKey>());
+			const TArray<FRigElementKey>& InElements = TArray<FRigElementKey>(), const ERigTransformType::Type InTransformType = ERigTransformType::CurrentLocal);
 		~FPoseScope();
 
 	private:
@@ -618,6 +618,7 @@ private:
 		UControlRig* ControlRig;
 		ERigElementType Filter;
 		FRigPose CachedPose;
+		ERigTransformType::Type TransformType;
 	};
 
 #if WITH_EDITOR
diff --git a/Engine/Plugins/Animation/ControlRig/Source/ControlRigEditor/Private/EditMode/ControlRigEditMode.cpp b/Engine/Plugins/Animation/ControlRig/Source/ControlRigEditor/Private/EditMode/ControlRigEditMode.cpp
--- a/Engine/Plugins/Animation/ControlRig/Source/ControlRigEditor/Private/EditMode/ControlRigEditMode.cpp
+++ b/Engine/Plugins/Animation/ControlRig/Source/ControlRigEditor/Private/EditMode/ControlRigEditMode.cpp
@@ -4119,24 +4119,6 @@ void FControlRigEditMode::MoveControlShape(AControlRigShapeActor* ShapeActor, co
 			}
 		}
 	}
-#if WITH_EDITOR
-	if (bTransformChanged)
-	{
-		if (UControlRig* ControlRig = ShapeActor->ControlRig.Get())
-		{
-			if (UWorld* World = ControlRig->GetWorld())
-			{
-				if (World->IsPreviewWorld())
-				{
-					if (UControlRigBlueprint* Blueprint = Cast<UControlRigBlueprint>(ControlRig->GetClass()->ClassGeneratedBy))
-					{
-						Blueprint->PropagatePoseFromInstanceToBP(ControlRig);
-					}
-				}
-			}
-		}
-	}
-#endif
 }
 
 void FControlRigEditMode::ChangeControlShapeTransform(AControlRigShapeActor* ShapeActor, const bool bTranslation, FVector& InDrag,
