Animation Mode: Fix Anim Details not changing Controls due to PostEditChangeChainedProperty changes upstream.
#trivial

diff --git a/Engine/Plugins/Animation/ControlRig/Source/ControlRigEditor/Private/EditMode/ControlRigControlsProxy.cpp b/Engine/Plugins/Animation/ControlRig/Source/ControlRigEditor/Private/EditMode/ControlRigControlsProxy.cpp
--- a/Engine/Plugins/Animation/ControlRig/Source/ControlRigEditor/Private/EditMode/ControlRigControlsProxy.cpp
+++ b/Engine/Plugins/Animation/ControlRig/Source/ControlRigEditor/Private/EditMode/ControlRigControlsProxy.cpp
@@ -192,8 +192,15 @@ FRigControlElement* UControlRigControlsProxy::GetControlElement() const
 void UControlRigTransformControlProxy::PostEditChangeChainProperty(struct FPropertyChangedChainEvent& PropertyChangedEvent)
 {
 	Super::PostEditChangeChainProperty(PropertyChangedEvent);
-	if ((PropertyChangedEvent.Property && PropertyChangedEvent.Property->GetFName() == GET_MEMBER_NAME_CHECKED(UControlRigTransformControlProxy, Transform))
-		|| (PropertyChangedEvent.MemberProperty && PropertyChangedEvent.MemberProperty->GetFName() == GET_MEMBER_NAME_CHECKED(UControlRigTransformControlProxy, Transform)))
+	FProperty* Property = PropertyChangedEvent.Property;
+	FProperty* MemberProperty = nullptr;
+	if (PropertyChangedEvent.PropertyChain.GetActiveMemberNode())
+	{
+		MemberProperty = PropertyChangedEvent.PropertyChain.GetActiveMemberNode()->GetValue();
+	}
+
+	if ((Property && Property->GetFName() == GET_MEMBER_NAME_CHECKED(UControlRigTransformControlProxy, Transform))
+		|| (MemberProperty && MemberProperty->GetFName() == GET_MEMBER_NAME_CHECKED(UControlRigTransformControlProxy, Transform)))
 	{
 		FRigControlElement* ControlElement = GetControlElement();
 		if (ControlElement && ControlRig.IsValid())
@@ -265,8 +272,15 @@ void UControlRigTransformControlProxy::SetKey(const IPropertyHandle& KeyedProper
 void UControlRigTransformNoScaleControlProxy::PostEditChangeChainProperty(struct FPropertyChangedChainEvent& PropertyChangedEvent)
 {
 	Super::PostEditChangeChainProperty(PropertyChangedEvent);
-	if ((PropertyChangedEvent.Property && PropertyChangedEvent.Property->GetFName() == GET_MEMBER_NAME_CHECKED(UControlRigTransformNoScaleControlProxy, Transform))
-		|| (PropertyChangedEvent.MemberProperty && PropertyChangedEvent.MemberProperty->GetFName() == GET_MEMBER_NAME_CHECKED(UControlRigTransformNoScaleControlProxy, Transform)))
+	FProperty* Property = PropertyChangedEvent.Property;
+	FProperty* MemberProperty = nullptr;
+	if (PropertyChangedEvent.PropertyChain.GetActiveMemberNode())
+	{
+		MemberProperty = PropertyChangedEvent.PropertyChain.GetActiveMemberNode()->GetValue();
+	}
+
+	if ((Property && Property->GetFName() == GET_MEMBER_NAME_CHECKED(UControlRigTransformNoScaleControlProxy, Transform))
+		|| (MemberProperty && MemberProperty->GetFName() == GET_MEMBER_NAME_CHECKED(UControlRigTransformNoScaleControlProxy, Transform)))
 	{
 		FRigControlElement* ControlElement = GetControlElement();
 		if (ControlElement && ControlRig.IsValid())
@@ -327,8 +341,16 @@ void UControlRigTransformNoScaleControlProxy::SetKey(const IPropertyHandle& Keye
 void UControlRigEulerTransformControlProxy::PostEditChangeChainProperty(struct FPropertyChangedChainEvent& PropertyChangedEvent)
 {
 	Super::PostEditChangeChainProperty(PropertyChangedEvent);
-	if ((PropertyChangedEvent.Property && PropertyChangedEvent.Property->GetFName() == GET_MEMBER_NAME_CHECKED(UControlRigEulerTransformControlProxy, Transform))
-		|| (PropertyChangedEvent.MemberProperty && PropertyChangedEvent.MemberProperty->GetFName() == GET_MEMBER_NAME_CHECKED(UControlRigEulerTransformControlProxy, Transform)))
+
+	FProperty* Property = PropertyChangedEvent.Property;
+	FProperty* MemberProperty = nullptr;
+	if (PropertyChangedEvent.PropertyChain.GetActiveMemberNode())
+	{
+		MemberProperty = PropertyChangedEvent.PropertyChain.GetActiveMemberNode()->GetValue();
+	}
+
+	if ((Property && Property->GetFName() == GET_MEMBER_NAME_CHECKED(UControlRigEulerTransformControlProxy, Transform))
+		|| (MemberProperty && MemberProperty->GetFName() == GET_MEMBER_NAME_CHECKED(UControlRigEulerTransformControlProxy, Transform)))
 	{
 		FRigControlElement* ControlElement = GetControlElement();
 		if (ControlElement && ControlRig.IsValid())
@@ -548,8 +570,15 @@ void UControlRigEnumControlProxy::SetKey(const IPropertyHandle& KeyedPropertyHan
 void UControlRigVectorControlProxy::PostEditChangeChainProperty(struct FPropertyChangedChainEvent& PropertyChangedEvent)
 {
 	Super::PostEditChangeChainProperty(PropertyChangedEvent);
-	if ((PropertyChangedEvent.Property && PropertyChangedEvent.Property->GetFName() == GET_MEMBER_NAME_CHECKED(UControlRigVectorControlProxy, Vector))
-		|| (PropertyChangedEvent.MemberProperty && PropertyChangedEvent.MemberProperty->GetFName() == GET_MEMBER_NAME_CHECKED(UControlRigVectorControlProxy, Vector)))
+	FProperty* Property = PropertyChangedEvent.Property;
+	FProperty* MemberProperty = nullptr;
+	if (PropertyChangedEvent.PropertyChain.GetActiveMemberNode())
+	{
+		MemberProperty = PropertyChangedEvent.PropertyChain.GetActiveMemberNode()->GetValue();
+	}
+
+	if ((Property && Property->GetFName() == GET_MEMBER_NAME_CHECKED(UControlRigVectorControlProxy, Vector))
+		|| (MemberProperty && MemberProperty->GetFName() == GET_MEMBER_NAME_CHECKED(UControlRigVectorControlProxy, Vector)))
 	{
 		FRigControlElement* ControlElement = GetControlElement();
 		if (ControlElement && ControlRig.IsValid())
@@ -647,8 +676,15 @@ void UControlRigVectorControlProxy::SetKey(const IPropertyHandle& KeyedPropertyH
 void UControlRigVector2DControlProxy::PostEditChangeChainProperty(struct FPropertyChangedChainEvent& PropertyChangedEvent)
 {
 	Super::PostEditChangeChainProperty(PropertyChangedEvent);
-	if ((PropertyChangedEvent.Property && PropertyChangedEvent.Property->GetFName() == GET_MEMBER_NAME_CHECKED(UControlRigVector2DControlProxy, Vector2D))
-		|| ((PropertyChangedEvent.MemberProperty && PropertyChangedEvent.MemberProperty->GetFName() == GET_MEMBER_NAME_CHECKED(UControlRigVector2DControlProxy, Vector2D))))
+	FProperty* Property = PropertyChangedEvent.Property;
+	FProperty* MemberProperty = nullptr;
+	if (PropertyChangedEvent.PropertyChain.GetActiveMemberNode())
+	{
+		MemberProperty = PropertyChangedEvent.PropertyChain.GetActiveMemberNode()->GetValue();
+	}
+
+	if ((Property && Property->GetFName() == GET_MEMBER_NAME_CHECKED(UControlRigVector2DControlProxy, Vector2D))
+		|| (MemberProperty && MemberProperty->GetFName() == GET_MEMBER_NAME_CHECKED(UControlRigVector2DControlProxy, Vector2D)))
 	{
 		FRigControlElement* ControlElement = GetControlElement();
 		if (ControlElement && ControlRig.IsValid())
