Fix crash reloading a control rig Anim BP

Control rig's custom proxy wasnt passing through the pre-evaluate call to its wrapped instance, so cached object ptrs were not getting updated on first evaluation

diff --git a/Engine/Plugins/Animation/ControlRig/Source/ControlRig/Private/Sequencer/ControlRigLayerInstanceProxy.cpp b/Engine/Plugins/Animation/ControlRig/Source/ControlRig/Private/Sequencer/ControlRigLayerInstanceProxy.cpp
--- a/Engine/Plugins/Animation/ControlRig/Source/ControlRig/Private/Sequencer/ControlRigLayerInstanceProxy.cpp
+++ b/Engine/Plugins/Animation/ControlRig/Source/ControlRig/Private/Sequencer/ControlRigLayerInstanceProxy.cpp
@@ -28,6 +28,16 @@ void FControlRigLayerInstanceProxy::CacheBones()
 	}
 }
 
+void FControlRigLayerInstanceProxy::PreEvaluateAnimation(UAnimInstance* InAnimInstance)
+{
+	Super::PreEvaluateAnimation(InAnimInstance);
+
+	if (CurrentSourceAnimInstance)
+	{
+		CurrentSourceAnimInstance->PreEvaluateAnimation();
+	}
+}
+
 bool FControlRigLayerInstanceProxy::Evaluate(FPoseContext& Output)
 {
 	check(CurrentRoot);
diff --git a/Engine/Plugins/Animation/ControlRig/Source/ControlRig/Private/Sequencer/ControlRigLayerInstanceProxy.h b/Engine/Plugins/Animation/ControlRig/Source/ControlRig/Private/Sequencer/ControlRigLayerInstanceProxy.h
--- a/Engine/Plugins/Animation/ControlRig/Source/ControlRig/Private/Sequencer/ControlRigLayerInstanceProxy.h
+++ b/Engine/Plugins/Animation/ControlRig/Source/ControlRig/Private/Sequencer/ControlRigLayerInstanceProxy.h
@@ -72,6 +72,7 @@ public:
 	virtual bool Evaluate(FPoseContext& Output) override;
 	virtual void CacheBones() override;
 	virtual void UpdateAnimationNode(const FAnimationUpdateContext& InContext) override;
+	virtual void PreEvaluateAnimation(UAnimInstance* InAnimInstance) override;
 
 	/** Anim Instance Source info - created externally and used here */
 	void SetSourceAnimInstance(UAnimInstance* SourceAnimInstance, FAnimInstanceProxy* SourceAnimInputProxy);
