RigVMEditor: EdGraph, EdGraphNode, EdGraphSchema and Blueprint

diff --git a/Engine/Plugins/Animation/ControlRig/Config/BaseControlRig.ini b/Engine/Plugins/Animation/ControlRig/Config/BaseControlRig.ini
--- a/Engine/Plugins/Animation/ControlRig/Config/BaseControlRig.ini
+++ b/Engine/Plugins/Animation/ControlRig/Config/BaseControlRig.ini
@@ -430,3 +430,6 @@
 +StructRedirects=(OldName="/Script/ControlRig.RigUnit_GetWorldTime",NewName="/Script/RigVM.RigVMFunction_GetWorldTime")
 +StructRedirects=(OldName="/Script/ControlRig.RigUnit_FramesToSeconds",NewName="/Script/RigVM.RigVMFunction_FramesToSeconds")
 +StructRedirects=(OldName="/Script/ControlRig.RigUnit_SecondsToFrames",NewName="/Script/RigVM.RigVMFunction_SecondsToFrames")
+
+; moving editor side types to RigVMEditor module
++StructRedirects=(OldName="/Script/ControlRig.RigGraphDisplaySettings",NewName="/Script/RigVMDeveloper.RigVMEdGraphDisplaySettings")
diff --git a/Engine/Plugins/Animation/ControlRig/Source/ControlRig/Private/Sequencer/MovieSceneControlRigParameterTrack.cpp b/Engine/Plugins/Animation/ControlRig/Source/ControlRig/Private/Sequencer/MovieSceneControlRigParameterTrack.cpp
--- a/Engine/Plugins/Animation/ControlRig/Source/ControlRig/Private/Sequencer/MovieSceneControlRigParameterTrack.cpp
+++ b/Engine/Plugins/Animation/ControlRig/Source/ControlRig/Private/Sequencer/MovieSceneControlRigParameterTrack.cpp
@@ -583,7 +583,7 @@ void UMovieSceneControlRigParameterTrack::HandlePackageDone(const FEndLoadPackag
 	FCoreUObjectDelegates::OnEndLoadPackage.RemoveAll(this);
 }
 
-void UMovieSceneControlRigParameterTrack::HandleControlRigPackageDone(UControlRig* InControlRig)
+void UMovieSceneControlRigParameterTrack::HandleControlRigPackageDone(URigVMHost* InControlRig)
 {
 	if (ensure(ControlRig == InControlRig))
 	{
diff --git a/Engine/Plugins/Animation/ControlRig/Source/ControlRig/Public/ControlRig.h b/Engine/Plugins/Animation/ControlRig/Source/ControlRig/Public/ControlRig.h
--- a/Engine/Plugins/Animation/ControlRig/Source/ControlRig/Public/ControlRig.h
+++ b/Engine/Plugins/Animation/ControlRig/Source/ControlRig/Public/ControlRig.h
@@ -1,10 +1,7 @@
 // Copyright Epic Games, Inc. All Rights Reserved.
 
 #pragma once
-#include "UObject/Object.h"
 #include "CoreMinimal.h"
-#include "UObject/ObjectMacros.h"
-#include "UObject/ScriptMacros.h"
 #include "UObject/SparseDelegate.h"
 #include "Engine/EngineBaseTypes.h"
 #include "Templates/SubclassOf.h"
@@ -72,11 +69,6 @@ public:
 	// To support Blueprints/scripting, we need a different delegate type (a 'Dynamic' delegate) which supports looser style UFunction binding (using names).
 	DECLARE_DYNAMIC_MULTICAST_SPARSE_DELEGATE_ThreeParams(FOnControlSelectedBP, UControlRig, OnControlSelected_BP, UControlRig*, Rig, const FRigControlElement&, Control, bool, bSelected);
 
-#if WITH_EDITOR
-	/** Bindable event for external objects to be notified that a control rig is fully end-loaded*/
-	DECLARE_EVENT_OneParam(UControlRig, FOnEndLoadPackage, UControlRig*);
-#endif
-
 	/** Bindable event to notify object binding change. */
 	DECLARE_EVENT_OneParam(UControlRig, FControlRigBoundEvent, UControlRig*);
 
@@ -91,19 +83,6 @@ public:
 	virtual void PostLoad() override;
 	virtual UScriptStruct* GetPublicContextStruct() const override { return FControlRigExecuteContext::StaticStruct(); }
 
-#if WITH_EDITOR
-
-private:
-	FOnEndLoadPackage EndLoadPackageEvent;
-
-public:
-	// these are needed so that sequencer can have a chance to update its 
-	// ControlRig instances after the package is fully end-loaded
-	void BroadCastEndLoadPackage() { EndLoadPackageEvent.Broadcast(this); }
-	FOnEndLoadPackage& OnEndLoadPackage() { return EndLoadPackageEvent; }
-
-#endif
-
 	/** Creates a transformable control handle for the specified control to be used by the constraints system. Should use the UObject from 
 	ConstraintsScriptingLibrary::GetManager(UWorld* InWorld)*/
 	UFUNCTION(BlueprintCallable, Category = "Control Rig | Constraints")
diff --git a/Engine/Plugins/Animation/ControlRig/Source/ControlRig/Public/Sequencer/MovieSceneControlRigParameterTrack.h b/Engine/Plugins/Animation/ControlRig/Source/ControlRig/Public/Sequencer/MovieSceneControlRigParameterTrack.h
--- a/Engine/Plugins/Animation/ControlRig/Source/ControlRig/Public/Sequencer/MovieSceneControlRigParameterTrack.h
+++ b/Engine/Plugins/Animation/ControlRig/Source/ControlRig/Public/Sequencer/MovieSceneControlRigParameterTrack.h
@@ -64,7 +64,7 @@ public:
 #if WITH_EDITOR
 	void HandlePackageDone(const FEndLoadPackageContext& Context);
 	// control Rigs are ready only after its package is fully end-loaded
-	void HandleControlRigPackageDone(UControlRig* InControlRig);
+	void HandleControlRigPackageDone(URigVMHost* InControlRig);
 #endif
 
 public:
diff --git a/Engine/Plugins/Animation/ControlRig/Source/ControlRigDeveloper/Private/ControlRigBlueprint.cpp b/Engine/Plugins/Animation/ControlRig/Source/ControlRigDeveloper/Private/ControlRigBlueprint.cpp
--- a/Engine/Plugins/Animation/ControlRig/Source/ControlRigDeveloper/Private/ControlRigBlueprint.cpp
+++ b/Engine/Plugins/Animation/ControlRig/Source/ControlRigDeveloper/Private/ControlRigBlueprint.cpp
@@ -43,7 +43,7 @@
 #include "Kismet2/KismetDebugUtilities.h"
 #include "Kismet2/WatchedPin.h"
 #include "Kismet2/BlueprintEditorUtils.h"
-#include "ControlRigBlueprintUtils.h"
+#include "RigVMBlueprintUtils.h"
 #include "Settings/ControlRigSettings.h"
 #include "UnrealEdGlobals.h"
 #include "Editor/UnrealEdEngine.h"
@@ -1555,7 +1555,7 @@ void UControlRigBlueprint::OnRigVMRegistryChanged()
 	RefreshAllModels();
 	RebuildGraphFromModel();
 	// avoids slate crash
-	FControlRigBlueprintUtils::HandleRefreshAllNodes(this);
+	FRigVMBlueprintUtils::HandleRefreshAllNodes(this);
 }
 
 void UControlRigBlueprint::HandleReportFromCompiler(EMessageSeverity::Type InSeverity, UObject* InSubject, const FString& InMessage)
diff --git a/Engine/Plugins/Animation/ControlRig/Source/ControlRigDeveloper/Private/Graph/ControlRigGraph.cpp b/Engine/Plugins/Animation/ControlRig/Source/ControlRigDeveloper/Private/Graph/ControlRigGraph.cpp
--- a/Engine/Plugins/Animation/ControlRig/Source/ControlRigDeveloper/Private/Graph/ControlRigGraph.cpp
+++ b/Engine/Plugins/Animation/ControlRig/Source/ControlRigDeveloper/Private/Graph/ControlRigGraph.cpp
@@ -20,7 +20,7 @@
 
 #if WITH_EDITOR
 #include "Kismet2/BlueprintEditorUtils.h"
-#include "ControlRigBlueprintUtils.h"
+#include "RigVMBlueprintUtils.h"
 #include "BlueprintCompilationManager.h"
 #include "Widgets/Notifications/SNotificationList.h"
 #include "Framework/Notifications/NotificationManager.h"
diff --git a/Engine/Plugins/Animation/ControlRig/Source/ControlRigDeveloper/Private/Graph/ControlRigGraphNode.cpp b/Engine/Plugins/Animation/ControlRig/Source/ControlRigDeveloper/Private/Graph/ControlRigGraphNode.cpp
--- a/Engine/Plugins/Animation/ControlRig/Source/ControlRigDeveloper/Private/Graph/ControlRigGraphNode.cpp
+++ b/Engine/Plugins/Animation/ControlRig/Source/ControlRigDeveloper/Private/Graph/ControlRigGraphNode.cpp
@@ -20,7 +20,7 @@
 #include "ScopedTransaction.h"
 #include "StructReference.h"
 #include "UObject/PropertyPortFlags.h"
-#include "ControlRigBlueprintUtils.h"
+#include "RigVMBlueprintUtils.h"
 #include "Curves/CurveFloat.h"
 #include "RigVMCore/RigVMExecuteContext.h"
 #include "ControlRigDeveloper.h"
@@ -1444,7 +1444,7 @@ void UControlRigGraphNode::DestroyNode()
 		{
 			if(PropertyName_DEPRECATED.IsValid())
 			{
-				FControlRigBlueprintUtils::RemoveMemberVariableIfNotUsed(ControlRigBlueprint, PropertyName_DEPRECATED, this);
+				FRigVMBlueprintUtils::RemoveMemberVariableIfNotUsed(ControlRigBlueprint, PropertyName_DEPRECATED);
 			}
 		}
 	}
diff --git a/Engine/Plugins/Animation/ControlRig/Source/ControlRigDeveloper/Public/Graph/ControlRigGraphNode.h b/Engine/Plugins/Animation/ControlRig/Source/ControlRigDeveloper/Public/Graph/ControlRigGraphNode.h
--- a/Engine/Plugins/Animation/ControlRig/Source/ControlRigDeveloper/Public/Graph/ControlRigGraphNode.h
+++ b/Engine/Plugins/Animation/ControlRig/Source/ControlRigDeveloper/Public/Graph/ControlRigGraphNode.h
@@ -32,7 +32,6 @@ class CONTROLRIGDEVELOPER_API UControlRigGraphNode : public UEdGraphNode
 	friend class FControlRigGraphTraverser;
 	friend class FControlRigGraphPanelPinFactory;
 	friend class FControlRigEditor;
-	friend struct FControlRigBlueprintUtils;
 	friend class SControlRigGraphPinCurveFloat;
 
 private:
diff --git a/Engine/Plugins/Animation/ControlRig/Source/ControlRigEditor/Private/ControlRigEditorModule.cpp b/Engine/Plugins/Animation/ControlRig/Source/ControlRigEditor/Private/ControlRigEditorModule.cpp
--- a/Engine/Plugins/Animation/ControlRig/Source/ControlRigEditor/Private/ControlRigEditorModule.cpp
+++ b/Engine/Plugins/Animation/ControlRig/Source/ControlRigEditor/Private/ControlRigEditorModule.cpp
@@ -58,7 +58,7 @@
 #include "EdGraphUtilities.h"
 #include "Graph/ControlRigGraphPanelNodeFactory.h"
 #include "Graph/ControlRigGraphPanelPinFactory.h"
-#include "ControlRigBlueprintUtils.h"
+#include "RigVMBlueprintUtils.h"
 #include "ControlRigBlueprintCommands.h"
 #include "ControlRigHierarchyCommands.h"
 #include "ControlRigStackCommands.h"
@@ -240,8 +240,8 @@ void FControlRigEditorModule::StartupModule()
 	ControlRigGraphPanelPinFactory = MakeShared<FControlRigGraphPanelPinFactory>();
 	FEdGraphUtilities::RegisterVisualPinFactory(ControlRigGraphPanelPinFactory);
 
-	ReconstructAllNodesDelegateHandle = FBlueprintEditorUtils::OnReconstructAllNodesEvent.AddStatic(&FControlRigBlueprintUtils::HandleReconstructAllNodes);
-	RefreshAllNodesDelegateHandle = FBlueprintEditorUtils::OnRefreshAllNodesEvent.AddStatic(&FControlRigBlueprintUtils::HandleRefreshAllNodes);
+	ReconstructAllNodesDelegateHandle = FBlueprintEditorUtils::OnReconstructAllNodesEvent.AddStatic(&FRigVMBlueprintUtils::HandleReconstructAllNodes);
+	RefreshAllNodesDelegateHandle = FBlueprintEditorUtils::OnRefreshAllNodesEvent.AddStatic(&FRigVMBlueprintUtils::HandleRefreshAllNodes);
 
 	ICurveEditorModule& CurveEditorModule = FModuleManager::LoadModuleChecked<ICurveEditorModule>("CurveEditor");
 	FControlRigSpaceChannelCurveModel::ViewID = CurveEditorModule.RegisterView(FOnCreateCurveEditorView::CreateStatic(
diff --git a/Engine/Plugins/Animation/ControlRig/Source/ControlRigEditor/Private/Editor/ControlRigEditor.cpp b/Engine/Plugins/Animation/ControlRig/Source/ControlRigEditor/Private/Editor/ControlRigEditor.cpp
--- a/Engine/Plugins/Animation/ControlRig/Source/ControlRigEditor/Private/Editor/ControlRigEditor.cpp
+++ b/Engine/Plugins/Animation/ControlRig/Source/ControlRigEditor/Private/Editor/ControlRigEditor.cpp
@@ -30,7 +30,7 @@
 #include "ControlRig.h"
 #include "Editor/ControlRigSkeletalMeshComponent.h"
 #include "ControlRigObjectBinding.h"
-#include "ControlRigBlueprintUtils.h"
+#include "RigVMBlueprintUtils.h"
 #include "EditorViewportClient.h"
 #include "AnimationEditorPreviewActor.h"
 #include "Misc/MessageDialog.h"
@@ -565,7 +565,7 @@ void FControlRigEditor::InitControlRigEditor(const EToolkitMode::Type Mode, cons
 		OpenGraphAndBringToFront(ActiveGraph, true);
 	}
 
-	FControlRigBlueprintUtils::HandleRefreshAllNodes(InControlRigBlueprint);
+	FRigVMBlueprintUtils::HandleRefreshAllNodes(InControlRigBlueprint);
 
 	bControlRigEditorInitialized = true;
 
@@ -6629,7 +6629,7 @@ void FControlRigEditor::HandleMakeElementGetterSetter(ERigElementGetterSetterTyp
 			NodePositionIncrement = FVector2D(380.f, 0.f);
 		}
 
-		FName Name = FControlRigBlueprintUtils::ValidateName(GetControlRigBlueprint(), StructTemplate->GetName());
+		FName Name = FRigVMBlueprintUtils::ValidateName(GetControlRigBlueprint(), StructTemplate->GetName());
 		if (URigVMUnitNode* ModelNode = GetFocusedController()->AddUnitNode(StructTemplate, FRigUnit::GetMethodName(), NodePosition, FString(), true, true))
 		{
 			FString ItemTypeStr = StaticEnum<ERigElementType>()->GetDisplayNameTextByValue((int64)Key.Type).ToString();
diff --git a/Engine/Plugins/Animation/ControlRig/Source/ControlRigEditor/Private/Editor/SControlRigUnitCombo.cpp b/Engine/Plugins/Animation/ControlRig/Source/ControlRigEditor/Private/Editor/SControlRigUnitCombo.cpp
--- a/Engine/Plugins/Animation/ControlRig/Source/ControlRigEditor/Private/Editor/SControlRigUnitCombo.cpp
+++ b/Engine/Plugins/Animation/ControlRig/Source/ControlRigEditor/Private/Editor/SControlRigUnitCombo.cpp
@@ -13,7 +13,7 @@
 #include "Units/RigUnit.h"
 #include "UObject/UObjectIterator.h"
 #include "Styling/AppStyle.h"
-#include "ControlRigBlueprintUtils.h"
+#include "RigVMBlueprintUtils.h"
 #include "SGraphEditorActionMenu.h"
 #include "Graph/ControlRigGraph.h"
 
@@ -125,7 +125,7 @@ void SControlRigUnitCombo::OnSelectionChanged(TSharedPtr<FRigUnitTypeItem> InIte
 
 void SControlRigUnitCombo::BuildUnitTypesList()
 {
-	FControlRigBlueprintUtils::ForAllRigUnits([this](UStruct* InStruct)
+	FRigVMBlueprintUtils::ForAllRigVMStructs([this](UStruct* InStruct)
 	{
 		UnitTypeList.Add(MakeShared<FRigUnitTypeItem>(InStruct));
 	});
diff --git a/Engine/Plugins/Animation/ControlRig/Source/ControlRigEditor/Private/Editor/SRigHierarchy.cpp b/Engine/Plugins/Animation/ControlRig/Source/ControlRigEditor/Private/Editor/SRigHierarchy.cpp
--- a/Engine/Plugins/Animation/ControlRig/Source/ControlRigEditor/Private/Editor/SRigHierarchy.cpp
+++ b/Engine/Plugins/Animation/ControlRig/Source/ControlRigEditor/Private/Editor/SRigHierarchy.cpp
@@ -14,7 +14,7 @@
 #include "Widgets/Images/SImage.h"
 #include "Widgets/Text/STextBlock.h"
 #include "K2Node_VariableGet.h"
-#include "ControlRigBlueprintUtils.h"
+#include "RigVMBlueprintUtils.h"
 #include "ControlRigHierarchyCommands.h"
 #include "ControlRigBlueprint.h"
 #include "Graph/ControlRigGraph.h"
diff --git a/Engine/Plugins/Animation/ControlRig/Source/ControlRigEditor/Private/Graph/NodeSpawners/ControlRigEnumNodeSpawner.cpp b/Engine/Plugins/Animation/ControlRig/Source/ControlRigEditor/Private/Graph/NodeSpawners/ControlRigEnumNodeSpawner.cpp
--- a/Engine/Plugins/Animation/ControlRig/Source/ControlRigEditor/Private/Graph/NodeSpawners/ControlRigEnumNodeSpawner.cpp
+++ b/Engine/Plugins/Animation/ControlRig/Source/ControlRigEditor/Private/Graph/NodeSpawners/ControlRigEnumNodeSpawner.cpp
@@ -8,7 +8,6 @@
 #include "BlueprintNodeTemplateCache.h"
 #include "RigVMModel/RigVMController.h"
 #include "ControlRigBlueprint.h"
-#include "ControlRigBlueprintUtils.h"
 
 #include "RigVMModel/Nodes/RigVMEnumNode.h"
 
diff --git a/Engine/Plugins/Animation/ControlRig/Source/ControlRigEditor/Private/Graph/NodeSpawners/ControlRigFunctionRefNodeSpawner.cpp b/Engine/Plugins/Animation/ControlRig/Source/ControlRigEditor/Private/Graph/NodeSpawners/ControlRigFunctionRefNodeSpawner.cpp
--- a/Engine/Plugins/Animation/ControlRig/Source/ControlRigEditor/Private/Graph/NodeSpawners/ControlRigFunctionRefNodeSpawner.cpp
+++ b/Engine/Plugins/Animation/ControlRig/Source/ControlRigEditor/Private/Graph/NodeSpawners/ControlRigFunctionRefNodeSpawner.cpp
@@ -13,7 +13,7 @@
 #include "EditorCategoryUtils.h"
 #include "K2Node_Variable.h"
 #include "BlueprintNodeTemplateCache.h"
-#include "ControlRigBlueprintUtils.h"
+#include "RigVMBlueprintUtils.h"
 #include "ScopedTransaction.h"
 #include "Units/Execution/RigUnit_BeginExecution.h"
 #include "ControlRig.h"
@@ -331,7 +331,7 @@ UControlRigGraphNode* UControlRigFunctionRefNodeSpawner::SpawnNode(UEdGraph* Par
 		bool const bIsTemplateNode = FBlueprintNodeTemplateCache::IsTemplateOuter(ParentGraph);
 		bool const bIsUserFacingNode = !bIsTemplateNode;
 
-		FName Name = bIsTemplateNode ? InFunction.Name : FControlRigBlueprintUtils::ValidateName(RigBlueprint, InFunction.Name.ToString());
+		FName Name = bIsTemplateNode ? InFunction.Name : FRigVMBlueprintUtils::ValidateName(RigBlueprint, InFunction.Name.ToString());
 		URigVMController* Controller = bIsTemplateNode ? RigGraph->GetTemplateController() : RigBlueprint->GetController(ParentGraph);
 
 		if (!bIsTemplateNode)
diff --git a/Engine/Plugins/Animation/ControlRig/Source/ControlRigEditor/Private/Graph/NodeSpawners/ControlRigInvokeEntryNodeSpawner.cpp b/Engine/Plugins/Animation/ControlRig/Source/ControlRigEditor/Private/Graph/NodeSpawners/ControlRigInvokeEntryNodeSpawner.cpp
--- a/Engine/Plugins/Animation/ControlRig/Source/ControlRigEditor/Private/Graph/NodeSpawners/ControlRigInvokeEntryNodeSpawner.cpp
+++ b/Engine/Plugins/Animation/ControlRig/Source/ControlRigEditor/Private/Graph/NodeSpawners/ControlRigInvokeEntryNodeSpawner.cpp
@@ -12,7 +12,6 @@
 #include "BlueprintNodeTemplateCache.h"
 #include "RigVMModel/RigVMController.h"
 #include "ControlRigBlueprint.h"
-#include "ControlRigBlueprintUtils.h"
 #include "RigVMModel/Nodes/RigVMInvokeEntryNode.h"
 
 #include UE_INLINE_GENERATED_CPP_BY_NAME(ControlRigInvokeEntryNodeSpawner)
diff --git a/Engine/Plugins/Animation/ControlRig/Source/ControlRigEditor/Private/Graph/NodeSpawners/ControlRigTemplateNodeSpawner.cpp b/Engine/Plugins/Animation/ControlRig/Source/ControlRigEditor/Private/Graph/NodeSpawners/ControlRigTemplateNodeSpawner.cpp
--- a/Engine/Plugins/Animation/ControlRig/Source/ControlRigEditor/Private/Graph/NodeSpawners/ControlRigTemplateNodeSpawner.cpp
+++ b/Engine/Plugins/Animation/ControlRig/Source/ControlRigEditor/Private/Graph/NodeSpawners/ControlRigTemplateNodeSpawner.cpp
@@ -13,7 +13,7 @@
 #include "EditorCategoryUtils.h"
 #include "K2Node_Variable.h"
 #include "BlueprintNodeTemplateCache.h"
-#include "ControlRigBlueprintUtils.h"
+#include "RigVMBlueprintUtils.h"
 #include "ScopedTransaction.h"
 #include "ControlRig.h"
 #include "RigVMModel/Nodes/RigVMDispatchNode.h"
@@ -187,7 +187,7 @@ UControlRigGraphNode* UControlRigTemplateNodeSpawner::SpawnNode(UEdGraph* Parent
 		const int32 NotationHash = (int32)GetTypeHash(InNotation);
 		const FString TemplateName = TEXT("RigVMTemplate_") + FString::FromInt(NotationHash);
 
-		FName Name = bIsTemplateNode ? *TemplateName : FControlRigBlueprintUtils::ValidateName(RigBlueprint, Template->GetName().ToString());
+		FName Name = bIsTemplateNode ? *TemplateName : FRigVMBlueprintUtils::ValidateName(RigBlueprint, Template->GetName().ToString());
 		URigVMController* Controller = bIsTemplateNode ? RigGraph->GetTemplateController() : RigBlueprint->GetController(ParentGraph);
 
 		if (!bIsTemplateNode)
diff --git a/Engine/Plugins/Animation/ControlRig/Source/ControlRigEditor/Private/Graph/NodeSpawners/ControlRigUnitNodeSpawner.cpp b/Engine/Plugins/Animation/ControlRig/Source/ControlRigEditor/Private/Graph/NodeSpawners/ControlRigUnitNodeSpawner.cpp
--- a/Engine/Plugins/Animation/ControlRig/Source/ControlRigEditor/Private/Graph/NodeSpawners/ControlRigUnitNodeSpawner.cpp
+++ b/Engine/Plugins/Animation/ControlRig/Source/ControlRigEditor/Private/Graph/NodeSpawners/ControlRigUnitNodeSpawner.cpp
@@ -12,7 +12,7 @@
 #include "EditorCategoryUtils.h"
 #include "K2Node_Variable.h"
 #include "BlueprintNodeTemplateCache.h"
-#include "ControlRigBlueprintUtils.h"
+#include "RigVMBlueprintUtils.h"
 #include "ScopedTransaction.h"
 #include "Units/Execution/RigUnit_BeginExecution.h"
 #include "ControlRig.h"
@@ -172,7 +172,7 @@ UControlRigGraphNode* UControlRigUnitNodeSpawner::SpawnNode(UEdGraph* ParentGrap
 		bool const bIsTemplateNode = FBlueprintNodeTemplateCache::IsTemplateOuter(ParentGraph);
 		bool const bIsUserFacingNode = !bIsTemplateNode;
 
-		FName Name = bIsTemplateNode ? *StructTemplate->GetStructCPPName() : FControlRigBlueprintUtils::ValidateName(RigBlueprint, StructTemplate->GetFName().ToString());
+		FName Name = bIsTemplateNode ? *StructTemplate->GetStructCPPName() : FRigVMBlueprintUtils::ValidateName(RigBlueprint, StructTemplate->GetFName().ToString());
 
 		FStructOnScope StructScope(StructTemplate);
 		const FRigVMStruct* StructInstance = (const FRigVMStruct*)StructScope.GetStructMemory();
diff --git a/Engine/Plugins/Runtime/RigVM/Source/RigVM/Public/RigVMHost.h b/Engine/Plugins/Runtime/RigVM/Source/RigVM/Public/RigVMHost.h
--- a/Engine/Plugins/Runtime/RigVM/Source/RigVM/Public/RigVMHost.h
+++ b/Engine/Plugins/Runtime/RigVM/Source/RigVM/Public/RigVMHost.h
@@ -141,6 +141,20 @@ public:
 	UFUNCTION(BlueprintCallable, Category = "RigVM")
 	virtual bool Execute(const FName& InEventName);
 
+#if WITH_EDITOR
+
+	/** Bindable event for external objects to be notified that a control rig is fully end-loaded*/
+	DECLARE_EVENT_OneParam(URigVMHost, FOnEndLoadPackage, URigVMHost*);
+
+	// these are needed so that sequencer can have a chance to update its 
+	// ControlRig instances after the package is fully end-loaded
+	void BroadCastEndLoadPackage() { EndLoadPackageEvent.Broadcast(this); }
+	FOnEndLoadPackage& OnEndLoadPackage() { return EndLoadPackageEvent; }
+
+private:
+	FOnEndLoadPackage EndLoadPackageEvent;
+
+#endif
 
 protected:
 
@@ -402,6 +416,7 @@ protected:
 	uint64 AccumulatedCycles;
 #endif
 
+	friend class URigVMBlueprint;
 	friend class URigVMBlueprintGeneratedClass;
 };
 
