RigVM: Missing include, non-unity breakage

diff --git a/Engine/Plugins/Animation/ControlRig/Config/BaseControlRig.ini b/Engine/Plugins/Animation/ControlRig/Config/BaseControlRig.ini
--- a/Engine/Plugins/Animation/ControlRig/Config/BaseControlRig.ini
+++ b/Engine/Plugins/Animation/ControlRig/Config/BaseControlRig.ini
@@ -432,4 +432,4 @@
 +StructRedirects=(OldName="/Script/ControlRig.RigUnit_SecondsToFrames",NewName="/Script/RigVM.RigVMFunction_SecondsToFrames")
 
 ; moving editor side types to RigVMEditor module
-+StructRedirects=(OldName="/Script/ControlRig.RigGraphDisplaySettings",NewName="/Script/RigVMDeveloper.RigVMEdGraphDisplaySettings")
+; +StructRedirects=(OldName="/Script/ControlRig.RigGraphDisplaySettings",NewName="/Script/RigVMDeveloper.RigVMEdGraphDisplaySettings")
diff --git a/Engine/Plugins/Runtime/RigVM/Source/RigVMDeveloper/Private/EdGraph/RigVMEdGraph.cpp b/Engine/Plugins/Runtime/RigVM/Source/RigVMDeveloper/Private/EdGraph/RigVMEdGraph.cpp
--- a/Engine/Plugins/Runtime/RigVM/Source/RigVMDeveloper/Private/EdGraph/RigVMEdGraph.cpp
+++ b/Engine/Plugins/Runtime/RigVM/Source/RigVMDeveloper/Private/EdGraph/RigVMEdGraph.cpp
@@ -11,6 +11,7 @@
 #include "RigVMModel/Nodes/RigVMFunctionEntryNode.h"
 #include "RigVMModel/Nodes/RigVMFunctionReturnNode.h"
 #include "RigVMModel/Nodes/RigVMRerouteNode.h"
+#include "Stats/StatsHierarchical.h"
 
 #include UE_INLINE_GENERATED_CPP_BY_NAME(RigVMEdGraph)
 
diff --git a/Engine/Plugins/Runtime/RigVM/Source/RigVMDeveloper/Private/EdGraph/RigVMEdGraphNode.cpp b/Engine/Plugins/Runtime/RigVM/Source/RigVMDeveloper/Private/EdGraph/RigVMEdGraphNode.cpp
--- a/Engine/Plugins/Runtime/RigVM/Source/RigVMDeveloper/Private/EdGraph/RigVMEdGraphNode.cpp
+++ b/Engine/Plugins/Runtime/RigVM/Source/RigVMDeveloper/Private/EdGraph/RigVMEdGraphNode.cpp
@@ -23,6 +23,7 @@
 #include "RigVMModel/Nodes/RigVMInvokeEntryNode.h"
 #include "RigVMModel/Nodes/RigVMDispatchNode.h"
 #include "Stats/StatsHierarchical.h"
+#include "Styling/AppStyle.h"
 
 #include UE_INLINE_GENERATED_CPP_BY_NAME(RigVMEdGraphNode)
 
diff --git a/Engine/Plugins/Runtime/RigVM/Source/RigVMDeveloper/Private/RigVMBlueprint.cpp b/Engine/Plugins/Runtime/RigVM/Source/RigVMDeveloper/Private/RigVMBlueprint.cpp
--- a/Engine/Plugins/Runtime/RigVM/Source/RigVMDeveloper/Private/RigVMBlueprint.cpp
+++ b/Engine/Plugins/Runtime/RigVM/Source/RigVMDeveloper/Private/RigVMBlueprint.cpp
@@ -22,6 +22,7 @@
 #include "RigVMModel/Nodes/RigVMAggregateNode.h"
 #include "RigVMModel/Nodes/RigVMDispatchNode.h"
 #include "UObject/UE5MainStreamObjectVersion.h"
+#include "Stats/StatsHierarchical.h"
 
 #include UE_INLINE_GENERATED_CPP_BY_NAME(RigVMBlueprint)
 
