Control Rig: Delay level sequence asset creation until next tick after the control rig actor is dropped into the viewport. This prevents the level sequence asset creation to be in the same transaction as the create actors transaction. When it was in the same transaction, it caused the movie scene to be garbage collected on an undo because it's creation was in a transaction.

diff --git a/Engine/Plugins/Animation/ControlRig/Source/ControlRigEditor/Private/ControlRigBlueprintActions.cpp b/Engine/Plugins/Animation/ControlRig/Source/ControlRigEditor/Private/ControlRigBlueprintActions.cpp
--- a/Engine/Plugins/Animation/ControlRig/Source/ControlRigEditor/Private/ControlRigBlueprintActions.cpp
+++ b/Engine/Plugins/Animation/ControlRig/Source/ControlRigEditor/Private/ControlRigBlueprintActions.cpp
@@ -12,6 +12,7 @@
 #include "Styling/AppStyle.h"
 #include "Subsystems/AssetEditorSubsystem.h"
 
+#include "AssetRegistry/AssetRegistryModule.h"
 #include "ToolMenus.h"
 #include "ContentBrowserMenuContexts.h"
 #include "Engine/SkeletalMesh.h"
@@ -295,147 +296,156 @@ void FControlRigBlueprintActions::OnSpawnedSkeletalMeshActorChanged(UObject* InO
 
 	FCoreUObjectDelegates::OnObjectPropertyChanged.Remove(OnSpawnedSkeletalMeshActorChangedHandle);
 
-	UControlRigBlueprint* RigBlueprint = Cast<UControlRigBlueprint>(InAsset);
-	if (RigBlueprint == nullptr)
+	// Create a level sequence but delay until next tick so that the creation of the asset is not in the existing transaction
+	GEditor->GetTimerManager()->SetTimerForNextTick([MeshActor, InAsset]()
 	{
-		return;
-	}
-	UClass* ControlRigClass = RigBlueprint->GeneratedClass;
+		UControlRigBlueprint* RigBlueprint = Cast<UControlRigBlueprint>(InAsset);
+		if (RigBlueprint == nullptr)
+		{
+			return;
+		}
 
-	TGuardValue<bool> DisableTrackCreation(FControlRigParameterTrackEditor::bAutoGenerateControlRigTrack, false);
+		UClass* ControlRigClass = RigBlueprint->GeneratedClass;
 
-	// find a level sequence in the world, if can't find that, create one
-	ULevelSequence* Sequence = ULevelSequenceEditorBlueprintLibrary::GetFocusedLevelSequence();
-	if (Sequence == nullptr)
-	{
-		FString SequenceName = FString::Printf(TEXT("%s_Take1"), *InAsset->GetName());
-		
-		FString PackagePath;
-		const FString DefaultDirectory = FEditorDirectories::Get().GetLastDirectory(ELastDirectory::NEW_ASSET);
-		FPackageName::TryConvertFilenameToLongPackageName(DefaultDirectory, PackagePath);
-		if (PackagePath.IsEmpty())
+		TGuardValue<bool> DisableTrackCreation(FControlRigParameterTrackEditor::bAutoGenerateControlRigTrack, false);
+
+		// find a level sequence in the world, if can't find that, create one
+		ULevelSequence* Sequence = ULevelSequenceEditorBlueprintLibrary::GetFocusedLevelSequence();
+		if (Sequence == nullptr)
 		{
-			PackagePath = TEXT("/Game");
-		}
+			FString SequenceName = FString::Printf(TEXT("%s_Take1"), *InAsset->GetName());
 
-		FAssetToolsModule& AssetToolsModule = FModuleManager::LoadModuleChecked<FAssetToolsModule>("AssetTools");
-		FString UniquePackageName;
-		FString UniqueAssetName;
-		AssetToolsModule.Get().CreateUniqueAssetName(PackagePath / SequenceName, TEXT(""), UniquePackageName, UniqueAssetName);
+			FString PackagePath;
+			const FString DefaultDirectory = FEditorDirectories::Get().GetLastDirectory(ELastDirectory::NEW_ASSET);
+			FPackageName::TryConvertFilenameToLongPackageName(DefaultDirectory, PackagePath);
+			if (PackagePath.IsEmpty())
+			{
+				PackagePath = TEXT("/Game");
+			}
 
-		UPackage* Package = CreatePackage(*UniquePackageName);
-		Sequence = NewObject<ULevelSequence>(Package, *UniqueAssetName, RF_Public | RF_Standalone);
-		Sequence->Initialize(); //creates movie scene
-		Sequence->MarkPackageDirty();
+			FAssetToolsModule& AssetToolsModule = FModuleManager::LoadModuleChecked<FAssetToolsModule>("AssetTools");
+			FString UniquePackageName;
+			FString UniqueAssetName;
+			AssetToolsModule.Get().CreateUniqueAssetName(PackagePath / SequenceName, TEXT(""), UniquePackageName, UniqueAssetName);
 
-		// Set up some sensible defaults
-		const UMovieSceneToolsProjectSettings* ProjectSettings = GetDefault<UMovieSceneToolsProjectSettings>();
-		FFrameRate TickResolution = Sequence->GetMovieScene()->GetTickResolution();
-		Sequence->GetMovieScene()->SetPlaybackRange((ProjectSettings->DefaultStartTime*TickResolution).FloorToFrame(), (ProjectSettings->DefaultDuration*TickResolution).FloorToFrame().Value);
+			UPackage* Package = CreatePackage(*UniquePackageName);
+			Sequence = NewObject<ULevelSequence>(Package, *UniqueAssetName, RF_Public | RF_Standalone);
+			Sequence->Initialize(); //creates movie scene
+			Sequence->MarkPackageDirty();
 
-		if (UActorFactory* ActorFactory = GEditor->FindActorFactoryForActorClass(ALevelSequenceActor::StaticClass()))
-		{
-			if (ALevelSequenceActor* LevelSequenceActor = Cast<ALevelSequenceActor>(GEditor->UseActorFactory(ActorFactory, FAssetData(Sequence), &FTransform::Identity)))
+			// Notify the asset registry
+			FAssetRegistryModule::AssetCreated(Sequence);
+
+			// Set up some sensible defaults
+			const UMovieSceneToolsProjectSettings* ProjectSettings = GetDefault<UMovieSceneToolsProjectSettings>();
+			FFrameRate TickResolution = Sequence->GetMovieScene()->GetTickResolution();
+			Sequence->GetMovieScene()->SetPlaybackRange((ProjectSettings->DefaultStartTime * TickResolution).FloorToFrame(), (ProjectSettings->DefaultDuration * TickResolution).FloorToFrame().Value);
+
+			if (UActorFactory* ActorFactory = GEditor->FindActorFactoryForActorClass(ALevelSequenceActor::StaticClass()))
 			{
-				LevelSequenceActor->SetSequence(Sequence);
+				if (ALevelSequenceActor* LevelSequenceActor = Cast<ALevelSequenceActor>(GEditor->UseActorFactory(ActorFactory, FAssetData(Sequence), &FTransform::Identity)))
+				{
+					LevelSequenceActor->SetSequence(Sequence);
+				}
 			}
 		}
-	}
- 
-	if (Sequence == nullptr)
-	{
-		return;
-	}
-	UMovieScene* MovieScene = Sequence->GetMovieScene();
 
-	GEditor->GetEditorSubsystem<UAssetEditorSubsystem>()->OpenEditorForAsset(Sequence);
+		if (Sequence == nullptr)
+		{
+			return;
+		}
+
+		UMovieScene* MovieScene = Sequence->GetMovieScene();
 
-	IAssetEditorInstance* AssetEditor = GEditor->GetEditorSubsystem<UAssetEditorSubsystem>()->FindEditorForAsset(Sequence, false);
-	ILevelSequenceEditorToolkit* LevelSequenceEditor = static_cast<ILevelSequenceEditorToolkit*>(AssetEditor);
-	TWeakPtr<ISequencer> WeakSequencer = LevelSequenceEditor ? LevelSequenceEditor->GetSequencer() : nullptr;
+		GEditor->GetEditorSubsystem<UAssetEditorSubsystem>()->OpenEditorForAsset(Sequence);
 
-	if (WeakSequencer.IsValid())
-	{
-		TArray<TWeakObjectPtr<AActor> > ActorsToAdd;
-		ActorsToAdd.Add(MeshActor);
-		TArray<FGuid> ActorTracks = WeakSequencer.Pin()->AddActors(ActorsToAdd, false);
-		FControlRigEditMode* ControlRigEditMode = static_cast<FControlRigEditMode*>(GLevelEditorModeTools().GetActiveMode(FControlRigEditMode::ModeName));
+		IAssetEditorInstance* AssetEditor = GEditor->GetEditorSubsystem<UAssetEditorSubsystem>()->FindEditorForAsset(Sequence, false);
+		ILevelSequenceEditorToolkit* LevelSequenceEditor = static_cast<ILevelSequenceEditorToolkit*>(AssetEditor);
+		TWeakPtr<ISequencer> WeakSequencer = LevelSequenceEditor ? LevelSequenceEditor->GetSequencer() : nullptr;
 
-		for (FGuid ActorTrackGuid : ActorTracks)
+		if (WeakSequencer.IsValid())
 		{
-			//Delete binding from default animating rig
-			FGuid CompGuid = WeakSequencer.Pin()->FindObjectId(*(MeshActor->GetSkeletalMeshComponent()), WeakSequencer.Pin()->GetFocusedTemplateID());
-			if (CompGuid.IsValid())
+			TArray<TWeakObjectPtr<AActor> > ActorsToAdd;
+			ActorsToAdd.Add(MeshActor);
+			TArray<FGuid> ActorTracks = WeakSequencer.Pin()->AddActors(ActorsToAdd, false);
+			FControlRigEditMode* ControlRigEditMode = static_cast<FControlRigEditMode*>(GLevelEditorModeTools().GetActiveMode(FControlRigEditMode::ModeName));
+
+			for (FGuid ActorTrackGuid : ActorTracks)
 			{
-				if (ControlRigEditMode)
+				//Delete binding from default animating rig
+				FGuid CompGuid = WeakSequencer.Pin()->FindObjectId(*(MeshActor->GetSkeletalMeshComponent()), WeakSequencer.Pin()->GetFocusedTemplateID());
+				if (CompGuid.IsValid())
 				{
-					UMovieSceneControlRigParameterTrack* Track = Cast<UMovieSceneControlRigParameterTrack>(MovieScene->FindTrack(UMovieSceneControlRigParameterTrack::StaticClass(), CompGuid, NAME_None));
-					if (Track && Track->GetControlRig())
+					if (ControlRigEditMode)
 					{
-						ControlRigEditMode->RemoveControlRig(Track->GetControlRig());
+						UMovieSceneControlRigParameterTrack* Track = Cast<UMovieSceneControlRigParameterTrack>(MovieScene->FindTrack(UMovieSceneControlRigParameterTrack::StaticClass(), CompGuid, NAME_None));
+						if (Track && Track->GetControlRig())
+						{
+							ControlRigEditMode->RemoveControlRig(Track->GetControlRig());
+						}
+					}
+					if (!MovieScene->RemovePossessable(CompGuid))
+					{
+						MovieScene->RemoveSpawnable(CompGuid);
 					}
 				}
-				if (!MovieScene->RemovePossessable(CompGuid))
+
+				UMovieSceneControlRigParameterTrack* Track = Cast<UMovieSceneControlRigParameterTrack>(MovieScene->FindTrack(UMovieSceneControlRigParameterTrack::StaticClass(), ActorTrackGuid));
+				if (!Track)
 				{
-					MovieScene->RemoveSpawnable(CompGuid);
+					Track = MovieScene->AddTrack<UMovieSceneControlRigParameterTrack>(ActorTrackGuid);
 				}
-			}
 
-			UMovieSceneControlRigParameterTrack* Track = Cast<UMovieSceneControlRigParameterTrack>(MovieScene->FindTrack(UMovieSceneControlRigParameterTrack::StaticClass(), ActorTrackGuid));
-			if (!Track)
-			{
-				Track = MovieScene->AddTrack<UMovieSceneControlRigParameterTrack>(ActorTrackGuid);
-			}
+				UControlRig* ControlRig = Track->GetControlRig();
 
-			UControlRig* ControlRig = Track->GetControlRig();
+				FString ObjectName = (ControlRigClass->GetName());
 
-			FString ObjectName = (ControlRigClass->GetName());
+				if (!ControlRig || ControlRig->GetClass() != ControlRigClass)
+				{
+					USkeletalMesh* SkeletalMesh = MeshActor->GetSkeletalMeshComponent()->GetSkeletalMeshAsset();
+					USkeleton* Skeleton = SkeletalMesh->GetSkeleton();
 
-			if (!ControlRig || ControlRig->GetClass() != ControlRigClass)
-			{		
-				USkeletalMesh* SkeletalMesh = MeshActor->GetSkeletalMeshComponent()->GetSkeletalMeshAsset();
-				USkeleton* Skeleton = SkeletalMesh->GetSkeleton();
-				
-				ObjectName.RemoveFromEnd(TEXT("_C"));
+					ObjectName.RemoveFromEnd(TEXT("_C"));
 
-				ControlRig = NewObject<UControlRig>(Track, ControlRigClass, FName(*ObjectName), RF_Transactional);
-				ControlRig->SetObjectBinding(MakeShared<FControlRigObjectBinding>());
-				ControlRig->GetObjectBinding()->BindToObject(MeshActor->GetSkeletalMeshComponent());
-				ControlRig->GetDataSourceRegistry()->RegisterDataSource(UControlRig::OwnerComponent, ControlRig->GetObjectBinding()->GetBoundObject());
-				ControlRig->Initialize();
-				ControlRig->Evaluate_AnyThread();
-				ControlRig->CreateRigControlsForCurveContainer();
+					ControlRig = NewObject<UControlRig>(Track, ControlRigClass, FName(*ObjectName), RF_Transactional);
+					ControlRig->SetObjectBinding(MakeShared<FControlRigObjectBinding>());
+					ControlRig->GetObjectBinding()->BindToObject(MeshActor->GetSkeletalMeshComponent());
+					ControlRig->GetDataSourceRegistry()->RegisterDataSource(UControlRig::OwnerComponent, ControlRig->GetObjectBinding()->GetBoundObject());
+					ControlRig->Initialize();
+					ControlRig->Evaluate_AnyThread();
+					ControlRig->CreateRigControlsForCurveContainer();
 
-				WeakSequencer.Pin()->NotifyMovieSceneDataChanged(EMovieSceneDataChangeType::MovieSceneStructureItemsChanged);
-			}
+					WeakSequencer.Pin()->NotifyMovieSceneDataChanged(EMovieSceneDataChangeType::MovieSceneStructureItemsChanged);
+				}
 
-			UMovieSceneSection* Section = Track->GetAllSections().Num() ? Track->GetAllSections()[0] : nullptr;
-			if (!Section)
-			{
-				Track->Modify();
-				Section = Track->CreateControlRigSection(0, ControlRig, true);
-				//mz todo need to have multiple rigs with same class
-				Track->SetTrackName(FName(*ObjectName));
-				Track->SetDisplayName(FText::FromString(ObjectName));
-
-				WeakSequencer.Pin()->EmptySelection();
-				WeakSequencer.Pin()->SelectSection(Section);
-				WeakSequencer.Pin()->ThrobSectionSelection();
-				WeakSequencer.Pin()->ObjectImplicitlyAdded(ControlRig);
-			}
-			
-			WeakSequencer.Pin()->NotifyMovieSceneDataChanged(EMovieSceneDataChangeType::MovieSceneStructureItemAdded);
-			if (!ControlRigEditMode)
-			{
-				GLevelEditorModeTools().ActivateMode(FControlRigEditMode::ModeName);
-				ControlRigEditMode = static_cast<FControlRigEditMode*>(GLevelEditorModeTools().GetActiveMode(FControlRigEditMode::ModeName));
-			}
-			if (ControlRigEditMode)
-			{
-				ControlRigEditMode->AddControlRigObject(ControlRig, WeakSequencer.Pin());
+				UMovieSceneSection* Section = Track->GetAllSections().Num() ? Track->GetAllSections()[0] : nullptr;
+				if (!Section)
+				{
+					Track->Modify();
+					Section = Track->CreateControlRigSection(0, ControlRig, true);
+					//mz todo need to have multiple rigs with same class
+					Track->SetTrackName(FName(*ObjectName));
+					Track->SetDisplayName(FText::FromString(ObjectName));
+
+					WeakSequencer.Pin()->EmptySelection();
+					WeakSequencer.Pin()->SelectSection(Section);
+					WeakSequencer.Pin()->ThrobSectionSelection();
+					WeakSequencer.Pin()->ObjectImplicitlyAdded(ControlRig);
+				}
+
+				WeakSequencer.Pin()->NotifyMovieSceneDataChanged(EMovieSceneDataChangeType::MovieSceneStructureItemAdded);
+				if (!ControlRigEditMode)
+				{
+					GLevelEditorModeTools().ActivateMode(FControlRigEditMode::ModeName);
+					ControlRigEditMode = static_cast<FControlRigEditMode*>(GLevelEditorModeTools().GetActiveMode(FControlRigEditMode::ModeName));
+				}
+				if (ControlRigEditMode)
+				{
+					ControlRigEditMode->AddControlRigObject(ControlRig, WeakSequencer.Pin());
+				}
 			}
 		}
-	}
+	});
 }
 
 const TArray<FText>& FControlRigBlueprintActions::GetSubMenus() const
