Sequencer:ControlRig fix drag n drop of control rigs, remove handle and capture parameter
#trivial

diff --git a/Engine/Plugins/Animation/ControlRig/Source/ControlRigEditor/Private/ControlRigBlueprintActions.cpp b/Engine/Plugins/Animation/ControlRig/Source/ControlRigEditor/Private/ControlRigBlueprintActions.cpp
--- a/Engine/Plugins/Animation/ControlRig/Source/ControlRigEditor/Private/ControlRigBlueprintActions.cpp
+++ b/Engine/Plugins/Animation/ControlRig/Source/ControlRigEditor/Private/ControlRigBlueprintActions.cpp
@@ -291,14 +291,14 @@ void FControlRigBlueprintActions::OnSpawnedSkeletalMeshActorChanged(UObject* InO
 		return;
 	}
 
-	ASkeletalMeshActor* MeshActor = Cast<ASkeletalMeshActor>(InObject);
-	check(MeshActor);
-
 	FCoreUObjectDelegates::OnObjectPropertyChanged.Remove(OnSpawnedSkeletalMeshActorChangedHandle);
+	OnSpawnedSkeletalMeshActorChangedHandle.Reset();
 
 	// Create a level sequence but delay until next tick so that the creation of the asset is not in the existing transaction
-	GEditor->GetTimerManager()->SetTimerForNextTick([MeshActor, InAsset]()
+	GEditor->GetTimerManager()->SetTimerForNextTick([InObject, InAsset]()
 	{
+		ASkeletalMeshActor* MeshActor = Cast<ASkeletalMeshActor>(InObject);
+		check(MeshActor);
 		UControlRigBlueprint* RigBlueprint = Cast<UControlRigBlueprint>(InAsset);
 		if (RigBlueprint == nullptr)
 		{
