Fix for ControlRig editor issues when adding comments and using undo / redo
- Crash when attempting to save control rig asset after deleting undone and redone comment boxes
- Comment boxes not being visibly deleted after undoing and redoing their creation

diff --git a/Engine/Plugins/Animation/ControlRig/Source/ControlRigDeveloper/Private/Graph/ControlRigGraph.cpp b/Engine/Plugins/Animation/ControlRig/Source/ControlRigDeveloper/Private/Graph/ControlRigGraph.cpp
--- a/Engine/Plugins/Animation/ControlRig/Source/ControlRigDeveloper/Private/Graph/ControlRigGraph.cpp
+++ b/Engine/Plugins/Animation/ControlRig/Source/ControlRigDeveloper/Private/Graph/ControlRigGraph.cpp
@@ -475,8 +475,12 @@ void UControlRigGraph::HandleModifiedEvent(ERigVMGraphNotifType InNotifType, URi
 				else
 				{
 					// check if the node is already part of the graph
-					if(FindNodeForModelNodeName(ModelNode->GetFName()) != nullptr)
+					if(UEdGraphNode* NodeToAdd = FindNodeForModelNodeName(ModelNode->GetFName()))
 					{
+						// Notify the GraphPanel a new node has been added
+						EEdGraphActionType AddNodeAction = GRAPHACTION_AddNode;
+						FEdGraphEditAction Action(AddNodeAction, this, NodeToAdd, false);
+						NotifyGraphChanged(Action);
 						break;
 					}
 				}
@@ -562,12 +566,12 @@ void UControlRigGraph::HandleModifiedEvent(ERigVMGraphNotifType InNotifType, URi
 					break;
 				}
 
-				ModelNodePathToEdNode.Remove(ModelNode->GetFName());
-
 				if (UEdGraphNode* EdNode = FindNodeForModelNodeName(ModelNode->GetFName(), false))
 				{
 					RemoveNode(EdNode);
 				}
+
+				ModelNodePathToEdNode.Remove(ModelNode->GetFName());
 			}
 			break;
 		}
@@ -844,7 +848,6 @@ void UControlRigGraph::HandleModifiedEvent(ERigVMGraphNotifType InNotifType, URi
 		{
 			if (URigVMNode* ModelNode = Cast<URigVMNode>(InSubject))
 			{
-				ModelNodePathToEdNode.Remove(ModelNode->GetPreviousFName());
 				UEdGraphNode* EdNode = FindNodeForModelNodeName(ModelNode->GetPreviousFName());
 				ModelNodePathToEdNode.Remove(ModelNode->GetPreviousFName());
 				if (UControlRigGraphNode* RigNode = Cast<UControlRigGraphNode>(EdNode))
@@ -1140,6 +1143,27 @@ const UControlRigGraph* UControlRigGraph::GetRootGraph() const
 	return this;
 }
 
+void UControlRigGraph::AddNode(UEdGraphNode* NodeToAdd, bool bUserAction, bool bSelectNewNode)
+{
+	// Comments are added outside of the ControlRigEditor, so we add here the node to the model
+	if (const UEdGraphNode_Comment* CommentNode = Cast<const UEdGraphNode_Comment>(NodeToAdd))
+	{
+		if (URigVMController* Controller = GetBlueprint()->GetOrCreateController(GetModel()))
+		{
+			if (GetModel()->FindNodeByName(NodeToAdd->GetFName()) == nullptr) // When recreating nodes at RebuildGraphFromModel, the model node already exists
+			{
+				TGuardValue<bool> BlueprintNotifGuard(GetBlueprint()->bSuspendModelNotificationsForOthers, true);
+				FVector2D NodePos(CommentNode->NodePosX, CommentNode->NodePosY);
+				FVector2D NodeSize(CommentNode->NodeWidth, CommentNode->NodeHeight);
+				FLinearColor NodeColor = CommentNode->CommentColor;
+				Controller->AddCommentNode(CommentNode->NodeComment, NodePos, NodeSize, NodeColor, CommentNode->GetName(), GIsTransacting == false, true);
+			}
+		}
+	}
+
+	Super::AddNode(NodeToAdd, bUserAction, bSelectNewNode);
+}
+
 void UControlRigGraph::RemoveNode(UEdGraphNode* InNode)
 {
 	// Make sure EdGraph is not part of the transaction
@@ -1157,21 +1181,6 @@ void UControlRigGraph::RemoveNode(UEdGraphNode* InNode)
 	}
 	InNode->Pins.Reset();
 					
-	// Rename the soon to be deleted object to a unique name, so that other objects can use
-	// the old name
-	FString DeletedName;
-	{
-		UObject* ExistingObject = nullptr;
-		static int32 DeletedIndex = FMath::Rand();
-		do
-		{
-			DeletedName = FString::Printf(TEXT("ControlRigGraph_%s_Deleted_%d"), *InNode->GetName(), DeletedIndex++); 
-			ExistingObject = StaticFindObject(/*Class=*/ NULL, this, *DeletedName, true);						
-		}
-		while (ExistingObject);
-	}
-	InNode->Rename(*DeletedName, GetTransientPackage(), REN_ForceNoResetLoaders | REN_DontCreateRedirectors);	
-
 	// this also subsequently calls NotifyGraphChanged
 	Super::RemoveNode(InNode);
 }
diff --git a/Engine/Plugins/Animation/ControlRig/Source/ControlRigDeveloper/Public/Graph/ControlRigGraph.h b/Engine/Plugins/Animation/ControlRig/Source/ControlRigDeveloper/Public/Graph/ControlRigGraph.h
--- a/Engine/Plugins/Animation/ControlRig/Source/ControlRigDeveloper/Public/Graph/ControlRigGraph.h
+++ b/Engine/Plugins/Animation/ControlRig/Source/ControlRigDeveloper/Public/Graph/ControlRigGraph.h
@@ -95,6 +95,10 @@ public:
 	UPROPERTY()
 	bool bIsFunctionDefinition;
 
+protected:
+	using Super::AddNode;
+	virtual void AddNode(UEdGraphNode* NodeToAdd, bool bUserAction = false, bool bSelectNewNode = true) override;
+
 private:
 
 	template<class T>
diff --git a/Engine/Plugins/Animation/ControlRig/Source/ControlRigEditor/Private/Editor/ControlRigEditor.cpp b/Engine/Plugins/Animation/ControlRig/Source/ControlRigEditor/Private/Editor/ControlRigEditor.cpp
--- a/Engine/Plugins/Animation/ControlRig/Source/ControlRigEditor/Private/Editor/ControlRigEditor.cpp
+++ b/Engine/Plugins/Animation/ControlRig/Source/ControlRigEditor/Private/Editor/ControlRigEditor.cpp
@@ -3658,18 +3658,6 @@ void FControlRigEditor::OnSelectedNodesChangedImpl(const TSet<class UObject*>& N
 			}
 			else if(UEdGraphNode* Node = Cast<UEdGraphNode>(Object))
 			{
-				if (UEdGraphNode_Comment* CommentNode = Cast<UEdGraphNode_Comment>(Node))
-				{
-					URigVMNode* ModelNode = GetFocusedModel()->FindNodeByName(Node->GetFName());
-					if (ModelNode == nullptr)
-					{
-						TGuardValue<bool> BlueprintNotifGuard(ControlRigBlueprint->bSuspendModelNotificationsForOthers, true);
-						FVector2D NodePos(CommentNode->NodePosX, CommentNode->NodePosY);
-						FVector2D NodeSize(CommentNode->NodeWidth, CommentNode->NodeHeight);
-						FLinearColor NodeColor = CommentNode->CommentColor;
-						GetFocusedController()->AddCommentNode(CommentNode->NodeComment, NodePos, NodeSize, NodeColor, CommentNode->GetName(), true, true);
-					}
-				}
 				NodeNamesToSelect.Add(Node->GetFName());
 			}
 		}
diff --git a/Engine/Source/Runtime/Engine/Classes/EdGraph/EdGraph.h b/Engine/Source/Runtime/Engine/Classes/EdGraph/EdGraph.h
--- a/Engine/Source/Runtime/Engine/Classes/EdGraph/EdGraph.h
+++ b/Engine/Source/Runtime/Engine/Classes/EdGraph/EdGraph.h
@@ -145,7 +145,7 @@ public:
 	 * @param bUserAction	true if the node was added as the result of a direct user action
 	 * @param bSelectNewNode	Whether or not to select the new node being created
 	 */
-	ENGINE_API void AddNode( UEdGraphNode* NodeToAdd, bool bUserAction = false, bool bSelectNewNode = true );
+	ENGINE_API virtual void AddNode( UEdGraphNode* NodeToAdd, bool bUserAction = false, bool bSelectNewNode = true );
 
 	/**
 	 * Queues up a select operation for a series of nodes in this graph.
diff --git a/Engine/Source/Runtime/Engine/Private/EdGraph/EdGraphNode.cpp b/Engine/Source/Runtime/Engine/Private/EdGraph/EdGraphNode.cpp
--- a/Engine/Source/Runtime/Engine/Private/EdGraph/EdGraphNode.cpp
+++ b/Engine/Source/Runtime/Engine/Private/EdGraph/EdGraphNode.cpp
@@ -855,6 +855,12 @@ UEdGraphPin* UEdGraphNode::GetPinWithDirectionAt(int32 PinIndex, EEdGraphPinDire
 
 void UEdGraphNode::AddSearchMetaDataInfo(TArray<struct FSearchTagDataPair>& OutTaggedMetaData) const
 {
+	if (GetSchema() == nullptr)
+	{
+		ensure(false);
+		return;
+	}
+
 	// Searchable - Primary label for the item in the search results
 	OutTaggedMetaData.Add(FSearchTagDataPair(FFindInBlueprintSearchTags::FiB_Name, GetNodeTitle(ENodeTitleType::ListView)));
 
