Sequencer: Control Rig: When collapsing all sections make sure to take into account per frame setting.
#trivial

diff --git a/Engine/Plugins/Animation/ControlRig/Source/ControlRigEditor/Private/Sequencer/ControlRigParameterTrackEditor.cpp b/Engine/Plugins/Animation/ControlRig/Source/ControlRigEditor/Private/Sequencer/ControlRigParameterTrackEditor.cpp
--- a/Engine/Plugins/Animation/ControlRig/Source/ControlRigEditor/Private/Sequencer/ControlRigParameterTrackEditor.cpp
+++ b/Engine/Plugins/Animation/ControlRig/Source/ControlRigEditor/Private/Sequencer/ControlRigParameterTrackEditor.cpp
@@ -4498,149 +4498,156 @@ bool CollapseAllLayersPerKey(TSharedPtr<ISequencer>& SequencerPtr, UMovieSceneTr
 
 bool FControlRigParameterTrackEditor::CollapseAllLayers(TSharedPtr<ISequencer>&SequencerPtr, UMovieSceneTrack * OwnerTrack, const FBakingAnimationKeySettings &InSettings)
 {
-	return CollapseAllLayersPerKey(SequencerPtr, OwnerTrack, InSettings);
-	if (SequencerPtr.IsValid() && OwnerTrack)
+	if (InSettings.BakingKeySettings == EBakingKeySettings::KeysOnly)
 	{
-		TArray<UMovieSceneSection*> Sections = OwnerTrack->GetAllSections();
-		//make sure right type
-		if (Sections.Num() < 1)
-		{
-			UE_LOG(LogControlRigEditor, Log, TEXT("CollapseAllSections::No sections on track"));
-			return false;
-		}
-		if (UMovieSceneControlRigParameterSection* ParameterSection = Cast<UMovieSceneControlRigParameterSection>(Sections[0]))
+		return CollapseAllLayersPerKey(SequencerPtr, OwnerTrack, InSettings);
+	}
+	else
+	{
+		if (SequencerPtr.IsValid() && OwnerTrack)
 		{
-			if (ParameterSection->GetBlendType().Get() == EMovieSceneBlendType::Absolute)
+			TArray<UMovieSceneSection*> Sections = OwnerTrack->GetAllSections();
+			//make sure right type
+			if (Sections.Num() < 1)
 			{
-				FScopedTransaction Transaction(LOCTEXT("CollapseAllSections", "Collapse All Sections"));
-				ParameterSection->Modify();
-				UControlRig* ControlRig = ParameterSection->GetControlRig();
-				FMovieSceneSequenceTransform RootToLocalTransform = SequencerPtr->GetFocusedMovieSceneSequenceTransform();
-
-				FFrameNumber StartFrame = InSettings.StartFrame;
-				FFrameNumber EndFrame = InSettings.EndFrame;
-				TRange<FFrameNumber> Range(StartFrame, EndFrame);
-				const FFrameRate& FrameRate = SequencerPtr->GetFocusedDisplayRate();
-				const FFrameRate& TickResolution = SequencerPtr->GetFocusedTickResolution();
-				
-				//frames and (optional) tangents
-				TArray<TPair<FFrameNumber, TArray<FMovieSceneTangentData>>> StoredTangents; //store tangents so we can reset them
-				TArray<FFrameNumber> Frames;
-				FFrameNumber FrameRateInFrameNumber = TickResolution.AsFrameNumber(FrameRate.AsInterval());
-				for (FFrameNumber& Frame = StartFrame; Frame <= EndFrame; Frame += FrameRateInFrameNumber)
-				{
-					Frames.Add(Frame);
-				}
-				
-				//Store transforms
-				TArray<TPair<FName, TArray<FTransform>>> ControlLocalTransforms;
-				TArray<FRigControlElement*> Controls;
-				ControlRig->GetControlsInOrder(Controls);
-
-				for (FRigControlElement* ControlElement : Controls)
+				UE_LOG(LogControlRigEditor, Log, TEXT("CollapseAllSections::No sections on track"));
+				return false;
+			}
+			if (UMovieSceneControlRigParameterSection* ParameterSection = Cast<UMovieSceneControlRigParameterSection>(Sections[0]))
+			{
+				if (ParameterSection->GetBlendType().Get() == EMovieSceneBlendType::Absolute)
 				{
-					if (!ControlRig->GetHierarchy()->IsAnimatable(ControlElement))
+					FScopedTransaction Transaction(LOCTEXT("CollapseAllSections", "Collapse All Sections"));
+					ParameterSection->Modify();
+					UControlRig* ControlRig = ParameterSection->GetControlRig();
+					FMovieSceneSequenceTransform RootToLocalTransform = SequencerPtr->GetFocusedMovieSceneSequenceTransform();
+
+					FFrameNumber StartFrame = InSettings.StartFrame;
+					FFrameNumber EndFrame = InSettings.EndFrame;
+					TRange<FFrameNumber> Range(StartFrame, EndFrame);
+					const FFrameRate& FrameRate = SequencerPtr->GetFocusedDisplayRate();
+					const FFrameRate& TickResolution = SequencerPtr->GetFocusedTickResolution();
+
+					//frames and (optional) tangents
+					TArray<TPair<FFrameNumber, TArray<FMovieSceneTangentData>>> StoredTangents; //store tangents so we can reset them
+					TArray<FFrameNumber> Frames;
+					FFrameNumber FrameRateInFrameNumber = TickResolution.AsFrameNumber(FrameRate.AsInterval());
+					FrameRateInFrameNumber.Value *= InSettings.FrameIncrement;
+					for (FFrameNumber& Frame = StartFrame; Frame <= EndFrame; Frame += FrameRateInFrameNumber)
 					{
-						continue;
+						Frames.Add(Frame);
 					}
-					TPair<FName, TArray<FTransform>> NameTransforms;
-					NameTransforms.Key = ControlElement->GetName();
-					NameTransforms.Value.SetNum(Frames.Num());
-					ControlLocalTransforms.Add(NameTransforms);
-				}
 
-				//get all of the local 
-				int32 Index = 0;
-				for (Index = 0; Index < Frames.Num(); ++Index)
-				{
-					const FFrameNumber& FrameNumber = Frames[Index];
-					FFrameTime GlobalTime(FrameNumber);
-					GlobalTime = GlobalTime * RootToLocalTransform.InverseLinearOnly();
+					//Store transforms
+					TArray<TPair<FName, TArray<FTransform>>> ControlLocalTransforms;
+					TArray<FRigControlElement*> Controls;
+					ControlRig->GetControlsInOrder(Controls);
 
-					FMovieSceneContext Context = FMovieSceneContext(FMovieSceneEvaluationRange(GlobalTime, TickResolution), SequencerPtr->GetPlaybackStatus()).SetHasJumped(true);
+					for (FRigControlElement* ControlElement : Controls)
+					{
+						if (!ControlRig->GetHierarchy()->IsAnimatable(ControlElement))
+						{
+							continue;
+						}
+						TPair<FName, TArray<FTransform>> NameTransforms;
+						NameTransforms.Key = ControlElement->GetName();
+						NameTransforms.Value.SetNum(Frames.Num());
+						ControlLocalTransforms.Add(NameTransforms);
+					}
 
-					SequencerPtr->GetEvaluationTemplate().EvaluateSynchronousBlocking(Context, *SequencerPtr);
-					ControlRig->Evaluate_AnyThread();
-					for (TPair<FName, TArray<FTransform>>& TrailControlTransform : ControlLocalTransforms)
+					//get all of the local 
+					int32 Index = 0;
+					for (Index = 0; Index < Frames.Num(); ++Index)
 					{
-						TrailControlTransform.Value[Index] = ControlRig->GetControlLocalTransform(TrailControlTransform.Key);
+						const FFrameNumber& FrameNumber = Frames[Index];
+						FFrameTime GlobalTime(FrameNumber);
+						GlobalTime = GlobalTime * RootToLocalTransform.InverseLinearOnly();
+
+						FMovieSceneContext Context = FMovieSceneContext(FMovieSceneEvaluationRange(GlobalTime, TickResolution), SequencerPtr->GetPlaybackStatus()).SetHasJumped(true);
+
+						SequencerPtr->GetEvaluationTemplate().EvaluateSynchronousBlocking(Context, *SequencerPtr);
+						ControlRig->Evaluate_AnyThread();
+						for (TPair<FName, TArray<FTransform>>& TrailControlTransform : ControlLocalTransforms)
+						{
+							TrailControlTransform.Value[Index] = ControlRig->GetControlLocalTransform(TrailControlTransform.Key);
+						}
 					}
-				}
-				//delete other sections
-				OwnerTrack->Modify();
-				for (Index = Sections.Num() - 1; Index >= 0; --Index)
-				{
-					if (Sections[Index] != ParameterSection)
+					//delete other sections
+					OwnerTrack->Modify();
+					for (Index = Sections.Num() - 1; Index >= 0; --Index)
 					{
-						OwnerTrack->RemoveSectionAt(Index);
+						if (Sections[Index] != ParameterSection)
+						{
+							OwnerTrack->RemoveSectionAt(Index);
+						}
 					}
-				}
 
-				//remove all keys, except Space Channels, from the Section.
-				ParameterSection->RemoveAllKeys(false /*bIncludedSpaceKeys*/);
+					//remove all keys, except Space Channels, from the Section.
+					ParameterSection->RemoveAllKeys(false /*bIncludedSpaceKeys*/);
 
-				FRigControlModifiedContext Context;
-				Context.SetKey = EControlRigSetKey::Always;
+					FRigControlModifiedContext Context;
+					Context.SetKey = EControlRigSetKey::Always;
 
-				FScopedSlowTask Feedback(Frames.Num(), LOCTEXT("CollapsingSections", "Collapsing Sections"));
-				Feedback.MakeDialog(true);
+					FScopedSlowTask Feedback(Frames.Num(), LOCTEXT("CollapsingSections", "Collapsing Sections"));
+					Feedback.MakeDialog(true);
 
-				const ERichCurveInterpMode InterpMode = InSettings.bReduceKeys ? RCIM_Cubic : RCIM_Linear;
+					const ERichCurveInterpMode InterpMode = InSettings.bReduceKeys ? RCIM_Cubic : RCIM_Linear;
 
-				Index = 0;
-				for (Index = 0; Index < Frames.Num(); ++Index)
-				{
-					Feedback.EnterProgressFrame(1, LOCTEXT("CollapsingSections", "Collapsing Sections"));
-					const FFrameNumber& FrameNumber = Frames[Index];
-					Context.LocalTime = TickResolution.AsSeconds(FFrameTime(FrameNumber));
-					//need to do the twice hack since controls aren't really in order
-					for (int32 TwiceHack = 0; TwiceHack < 2; ++TwiceHack)
+					Index = 0;
+					for (Index = 0; Index < Frames.Num(); ++Index)
 					{
-						for (TPair<FName, TArray<FTransform>>& TrailControlTransform : ControlLocalTransforms)
+						Feedback.EnterProgressFrame(1, LOCTEXT("CollapsingSections", "Collapsing Sections"));
+						const FFrameNumber& FrameNumber = Frames[Index];
+						Context.LocalTime = TickResolution.AsSeconds(FFrameTime(FrameNumber));
+						//need to do the twice hack since controls aren't really in order
+						for (int32 TwiceHack = 0; TwiceHack < 2; ++TwiceHack)
 						{
-							ControlRig->SetControlLocalTransform(TrailControlTransform.Key, TrailControlTransform.Value[Index], false, Context, false);
+							for (TPair<FName, TArray<FTransform>>& TrailControlTransform : ControlLocalTransforms)
+							{
+								ControlRig->SetControlLocalTransform(TrailControlTransform.Key, TrailControlTransform.Value[Index], false, Context, false);
+							}
 						}
-					}
-					ControlRig->Evaluate_AnyThread();
-					ParameterSection->RecordControlRigKey(FrameNumber, true, InterpMode);
+						ControlRig->Evaluate_AnyThread();
+						ParameterSection->RecordControlRigKey(FrameNumber, true, InterpMode);
 
-					if (Feedback.ShouldCancel())
+						if (Feedback.ShouldCancel())
+						{
+							Transaction.Cancel();
+							SequencerPtr->NotifyMovieSceneDataChanged(EMovieSceneDataChangeType::MovieSceneStructureItemAdded);
+							return false;
+						}
+					}
+					if (InSettings.bReduceKeys)
 					{
-						Transaction.Cancel();
-						SequencerPtr->NotifyMovieSceneDataChanged(EMovieSceneDataChangeType::MovieSceneStructureItemAdded);
-						return false;
+						FKeyDataOptimizationParams Params;
+						Params.bAutoSetInterpolation = true;
+						Params.Tolerance = InSettings.Tolerance;
+						FMovieSceneChannelProxy& ChannelProxy = ParameterSection->GetChannelProxy();
+						TArrayView<FMovieSceneFloatChannel*> FloatChannels = ChannelProxy.GetChannels<FMovieSceneFloatChannel>();
+
+						for (FMovieSceneFloatChannel* Channel : FloatChannels)
+						{
+							Channel->Optimize(Params); //should also auto tangent
+						}
 					}
+					//reset everything back
+					SequencerPtr->NotifyMovieSceneDataChanged(EMovieSceneDataChangeType::MovieSceneStructureItemAdded);
+					return true;
 				}
-				if (InSettings.bReduceKeys)
+				else
 				{
-					FKeyDataOptimizationParams Params;
-					Params.bAutoSetInterpolation = true;
-					Params.Tolerance = InSettings.Tolerance;
-					FMovieSceneChannelProxy& ChannelProxy = ParameterSection->GetChannelProxy();
-					TArrayView<FMovieSceneFloatChannel*> FloatChannels = ChannelProxy.GetChannels<FMovieSceneFloatChannel>();
-
-					for (FMovieSceneFloatChannel* Channel : FloatChannels)
-					{
-						Channel->Optimize(Params); //should also auto tangent
-					}
+					UE_LOG(LogControlRigEditor, Log, TEXT("CollapseAllSections:: First section is not additive"));
+					return false;
 				}
-				//reset everything back
-				SequencerPtr->NotifyMovieSceneDataChanged(EMovieSceneDataChangeType::MovieSceneStructureItemAdded);
-				return true;
 			}
 			else
 			{
-				UE_LOG(LogControlRigEditor, Log, TEXT("CollapseAllSections:: First section is not additive"));
+				UE_LOG(LogControlRigEditor, Log, TEXT("CollapseAllSections:: No Control Rig section"));
 				return false;
 			}
 		}
-		else
-		{
-			UE_LOG(LogControlRigEditor, Log, TEXT("CollapseAllSections:: No Control Rig section"));
-			return false;
-		}
+		UE_LOG(LogControlRigEditor, Log, TEXT("CollapseAllSections:: Sequencer or track is invalid"));
 	}
-	UE_LOG(LogControlRigEditor, Log, TEXT("CollapseAllSections:: Sequencer or track is invalid"));
 	return false;
 }
 
diff --git a/Engine/Source/Editor/MovieSceneTools/Private/MovieSceneToolHelpers.cpp b/Engine/Source/Editor/MovieSceneTools/Private/MovieSceneToolHelpers.cpp
--- a/Engine/Source/Editor/MovieSceneTools/Private/MovieSceneToolHelpers.cpp
+++ b/Engine/Source/Editor/MovieSceneTools/Private/MovieSceneToolHelpers.cpp
@@ -5130,6 +5130,19 @@ bool MovieSceneToolHelpers::CollapseSection(TSharedPtr<ISequencer>& SequencerPtr
 				}
 			}
 		}
+		if (InSettings.bReduceKeys)
+		{
+			FKeyDataOptimizationParams Params;
+			Params.bAutoSetInterpolation = true;
+			Params.Tolerance = InSettings.Tolerance;
+			FMovieSceneChannelProxy& ChannelProxy = BaseSection->GetChannelProxy();
+			TArrayView<FMovieSceneFloatChannel*> FloatChannels = ChannelProxy.GetChannels<FMovieSceneFloatChannel>();
+
+			for (FMovieSceneFloatChannel* Channel : FloatChannels)
+			{
+				Channel->Optimize(Params); //should also auto tangent
+			}
+		}
 		//reset everything back
 		SequencerPtr->NotifyMovieSceneDataChanged(EMovieSceneDataChangeType::MovieSceneStructureItemAdded);
 	}
