Rotation Order: Sequencer: Control Rig: New Rotation Orders Add flag to Control RigelElements to specify whether not to use the preferred angles if we do we set them when setting control rig transforms if we are flipping eulers. We also directly set the preferred angles in particularly when coming from Sequencer.  We currently use use the Setting preferrend angle values but have set it up so we may be able to override this later from the anmator.  There was a bug with SetControlLocalTransform not flipping eulers due to the fact it has a different signature(missing the bPythonPrint flag) then SetControlGlobalTransform.  Since ue has crazy right/left hand mix added some flags to the converion functions in the animation core library to specify we should keep the rotations in UE handyness when converting.

diff --git a/Engine/Plugins/Animation/ControlRig/Source/ControlRig/Private/Constraints/ControlRigTransformableHandle.cpp b/Engine/Plugins/Animation/ControlRig/Source/ControlRig/Private/Constraints/ControlRigTransformableHandle.cpp
--- a/Engine/Plugins/Animation/ControlRig/Source/ControlRig/Private/Constraints/ControlRigTransformableHandle.cpp
+++ b/Engine/Plugins/Animation/ControlRig/Source/ControlRig/Private/Constraints/ControlRigTransformableHandle.cpp
@@ -95,12 +95,12 @@ void UTransformableControlHandle::SetGlobalTransform(const FTransform& InGlobal)
 		return;
 	}
 	
-	URigHierarchy* Hierarchy = ControlRig->GetHierarchy();
 	const FRigElementKey& ControlKey = ControlElement->GetKey();
-	const int32 CtrlIndex = Hierarchy->GetIndex(ControlKey);
 	
 	const FTransform& ComponentTransform = SkeletalMeshComponent->GetComponentTransform();
-	Hierarchy->SetGlobalTransform(CtrlIndex, InGlobal.GetRelativeTransform(ComponentTransform));
+	//use this function so we don't set the preferred angles
+	ControlRig->SetControlGlobalTransform(ControlKey.Name, InGlobal.GetRelativeTransform(ComponentTransform), false/*bNotify*/, FRigControlModifiedContext(), false/*bSetupUndo*/, false /*bPrintPython*/, false/* bFixEulerFlips*/);
+
 }
 
 void UTransformableControlHandle::SetLocalTransform(const FTransform& InLocal) const
diff --git a/Engine/Plugins/Animation/ControlRig/Source/ControlRig/Private/ControlRig.cpp b/Engine/Plugins/Animation/ControlRig/Source/ControlRig/Private/ControlRig.cpp
--- a/Engine/Plugins/Animation/ControlRig/Source/ControlRig/Private/ControlRig.cpp
+++ b/Engine/Plugins/Animation/ControlRig/Source/ControlRig/Private/ControlRig.cpp
@@ -1574,7 +1574,7 @@ void UControlRig::SetControlLocalTransform(const FName& InControlName, const FTr
 			OnFilterControl.Broadcast(this, ControlElement, Value);
 			
 		}
-		SetControlValue(InControlName, Value, bNotify, Context, bSetupUndo, bFixEulerFlips);
+		SetControlValue(InControlName, Value, bNotify, Context, bSetupUndo, false /*bPrintPython not defined!*/, bFixEulerFlips);
 	}
 }
 
diff --git a/Engine/Plugins/Animation/ControlRig/Source/ControlRig/Private/Rigs/RigHierarchy.cpp b/Engine/Plugins/Animation/ControlRig/Source/ControlRig/Private/Rigs/RigHierarchy.cpp
--- a/Engine/Plugins/Animation/ControlRig/Source/ControlRig/Private/Rigs/RigHierarchy.cpp
+++ b/Engine/Plugins/Animation/ControlRig/Source/ControlRig/Private/Rigs/RigHierarchy.cpp
@@ -3546,17 +3546,26 @@ void URigHierarchy::SetControlValue(FRigControlElement* InControlElement, const
 				}
 				case ERigControlType::EulerTransform:
 				{
-					InControlElement->PreferredEulerAngles.SetRotator(GetEulerTransformFromControlValue(InValue).Rotation, bInitial, bFixEulerFlips);
+					FEulerTransform EulerTransform = GetEulerTransformFromControlValue(InValue);
+					FQuat Quat = EulerTransform.GetRotation();
+					const FVector Angle = GetControlAnglesFromQuat(InControlElement, Quat);
+					InControlElement->PreferredEulerAngles.SetAngles(Angle, bInitial, InControlElement->PreferredEulerAngles.RotationOrder, bFixEulerFlips);
 					break;
 				}
 				case ERigControlType::Transform:
 				{
-					InControlElement->PreferredEulerAngles.SetRotator(GetTransformFromControlValue(InValue).Rotator(), bInitial, bFixEulerFlips);
+					FTransform Transform = GetTransformFromControlValue(InValue);
+					FQuat Quat = Transform.GetRotation();
+					const FVector Angle = GetControlAnglesFromQuat(InControlElement, Quat);
+					InControlElement->PreferredEulerAngles.SetAngles(Angle, bInitial, InControlElement->PreferredEulerAngles.RotationOrder, bFixEulerFlips);
 					break;
 				}
 				case ERigControlType::TransformNoScale:
 				{
-					InControlElement->PreferredEulerAngles.SetRotator(GetTransformNoScaleFromControlValue(InValue).Rotation.Rotator(), bInitial, bFixEulerFlips);
+					FTransform Transform = GetTransformNoScaleFromControlValue(InValue);
+					FQuat Quat = Transform.GetRotation();
+					const FVector Angle = GetControlAnglesFromQuat(InControlElement, Quat);
+					InControlElement->PreferredEulerAngles.SetAngles(Angle, bInitial, InControlElement->PreferredEulerAngles.RotationOrder, bFixEulerFlips);
 					break;
 				}
 				default:
@@ -3586,8 +3595,10 @@ void URigHierarchy::SetControlValue(FRigControlElement* InControlElement, const
 					bForce,
 					bPrintPythonCommands
 				);
-
-				SetPreferredEulerAnglesFromValue();
+				if (bFixEulerFlips)
+				{
+					SetPreferredEulerAnglesFromValue();
+				}
 				break;
 			}
 			case ERigControlValueType::Initial:
@@ -3609,7 +3620,10 @@ void URigHierarchy::SetControlValue(FRigControlElement* InControlElement, const
 					bPrintPythonCommands
 				);
 
-				SetPreferredEulerAnglesFromValue();
+				if (bFixEulerFlips)
+				{
+					SetPreferredEulerAnglesFromValue();
+				}
 				break;
 			}
 			case ERigControlValueType::Minimum:
diff --git a/Engine/Plugins/Animation/ControlRig/Source/ControlRig/Private/Rigs/RigHierarchyElements.cpp b/Engine/Plugins/Animation/ControlRig/Source/ControlRig/Private/Rigs/RigHierarchyElements.cpp
--- a/Engine/Plugins/Animation/ControlRig/Source/ControlRig/Private/Rigs/RigHierarchyElements.cpp
+++ b/Engine/Plugins/Animation/ControlRig/Source/ControlRig/Private/Rigs/RigHierarchyElements.cpp
@@ -365,31 +365,12 @@ void FRigPreferredEulerAngles::Reset()
 
 FRotator FRigPreferredEulerAngles::GetRotator(bool bInitial) const
 {
-	return FRotator::MakeFromEuler(GetAngles(bInitial, DefaultRotationOrder));
+	return FRotator::MakeFromEuler(GetAngles(bInitial, RotationOrder));
 }
 
 FRotator FRigPreferredEulerAngles::SetRotator(const FRotator& InValue, bool bInitial, bool bFixEulerFlips)
 {
-	if(RotationOrder == DefaultRotationOrder)
-	{
-		if(bFixEulerFlips)
-		{
-			const FRotator CurrentValue = GetRotator(bInitial);
-			
-			//Find Diff of the rotation from current and just add that instead of setting so we can go over/under -180
-			FRotator CurrentWinding;
-			FRotator CurrentRotRemainder;
-			CurrentValue.GetWindingAndRemainder(CurrentWinding, CurrentRotRemainder);
-
-			FRotator DeltaRot = InValue - CurrentRotRemainder;
-			DeltaRot.Normalize();
-			const FRotator FixedValue = CurrentValue + DeltaRot;
-
-			SetAngles(FixedValue.Euler(), bInitial, DefaultRotationOrder);
-			return FixedValue;
-		}
-	}
-	SetAngles(InValue.Euler(), bInitial, DefaultRotationOrder, false);
+	SetAngles(InValue.Euler(), bInitial, RotationOrder, bFixEulerFlips);
 	return InValue;
 }
 
@@ -412,23 +393,21 @@ void FRigPreferredEulerAngles::SetAngles(const FVector& InValue, bool bInitial,
 
 	if(bFixEulerFlips)
 	{
-		// @MikeZ we need something here as well for euler values
-		// Do you think we can rely on the same mechanism?
-		/*
-		const FRotator CurrentValue = GetRotator(bInitial);
-			
+		const FRotator CurrentRotator = FRotator::MakeFromEuler(GetAngles(bInitial, RotationOrder));
+		const FRotator InRotator = FRotator::MakeFromEuler(Value);
+
 		//Find Diff of the rotation from current and just add that instead of setting so we can go over/under -180
 		FRotator CurrentWinding;
 		FRotator CurrentRotRemainder;
-		CurrentValue.GetWindingAndRemainder(CurrentWinding, CurrentRotRemainder);
+		CurrentRotator.GetWindingAndRemainder(CurrentWinding, CurrentRotRemainder);
 
-		FRotator DeltaRot = InValue - CurrentRotRemainder;
+		FRotator DeltaRot = InRotator - CurrentRotRemainder;
 		DeltaRot.Normalize();
-		const FRotator FixedValue = CurrentValue + DeltaRot;
+		const FRotator FixedValue = CurrentRotator + DeltaRot;
+
+		Get(bInitial) = FixedValue.Euler();
+		return;
 
-		SetAngles(FixedValue.Euler(), bInitial, DefaultRotationOrder);
-		return FixedValue;
-		*/
 	}
 	
 	Get(bInitial) = Value;
@@ -447,6 +426,19 @@ void FRigPreferredEulerAngles::SetRotationOrder(EEulerRotationOrder InRotationOr
 	}
 }
 
+FRotator FRigPreferredEulerAngles::GetRotatorFromQuat(const FQuat& InQuat) const
+{
+	FVector Vector = AnimationCore::EulerFromQuat(InQuat, RotationOrder, true);
+	return FRotator::MakeFromEuler(Vector);
+}
+
+FQuat FRigPreferredEulerAngles::GetQuatFromRotator(const FRotator& InRotator) const
+{
+	FVector Vector = InRotator.Euler();
+	return AnimationCore::QuatFromEuler(Vector, RotationOrder, true);
+}
+
+
 ////////////////////////////////////////////////////////////////////////////////
 // FRigElementHandle
 ////////////////////////////////////////////////////////////////////////////////
@@ -778,8 +770,9 @@ FRigControlSettings::FRigControlSettings()
 , ControlEnum(nullptr)
 , Customization()
 , bGroupWithParentControl(false)
-, bRestrictSpaceSwitching(false)
+, bRestrictSpaceSwitching(false) 
 , PreferredRotationOrder(FRigPreferredEulerAngles::DefaultRotationOrder)
+, bUsePreferredRotationOrder(false) 
 {
 	// rely on the default provided by the shape definition
 	ShapeName = FControlRigShapeDefinition().ShapeName; 
@@ -826,6 +819,8 @@ void FRigControlSettings::Save(FArchive& Ar)
 	Ar << bRestrictSpaceSwitching;
 	Ar << FilteredChannels;
 	Ar << PreferredRotationOrder;
+	Ar << bUsePreferredRotationOrder;
+
 }
 
 void FRigControlSettings::Load(FArchive& Ar)
@@ -1003,6 +998,16 @@ void FRigControlSettings::Load(FArchive& Ar)
 	{
 		PreferredRotationOrder = FRigPreferredEulerAngles::DefaultRotationOrder;
 	}
+
+	if (Ar.CustomVer(FControlRigObjectVersion::GUID) >= FControlRigObjectVersion::RigHierarchyControlPreferredRotationOrderFlag)
+	{
+		Ar << bUsePreferredRotationOrder;
+	}
+	else
+	{
+		bUsePreferredRotationOrder = false;
+	}
+
 }
 
 uint32 GetTypeHash(const FRigControlSettings& Settings)
@@ -1112,6 +1117,11 @@ bool FRigControlSettings::operator==(const FRigControlSettings& InOther) const
 	{
 		return false;
 	}
+	if (bUsePreferredRotationOrder != InOther.bUsePreferredRotationOrder)
+	{
+		return false;
+	}
+
 
 	const FTransform MinimumTransform = MinimumValue.GetAsTransform(ControlType, PrimaryAxis);
 	const FTransform OtherMinimumTransform = InOther.MinimumValue.GetAsTransform(ControlType, PrimaryAxis);
diff --git a/Engine/Plugins/Animation/ControlRig/Source/ControlRig/Private/Sequencer/MovieSceneControlRigParameterSection.cpp b/Engine/Plugins/Animation/ControlRig/Source/ControlRig/Private/Sequencer/MovieSceneControlRigParameterSection.cpp
--- a/Engine/Plugins/Animation/ControlRig/Source/ControlRig/Private/Sequencer/MovieSceneControlRigParameterSection.cpp
+++ b/Engine/Plugins/Animation/ControlRig/Source/ControlRig/Private/Sequencer/MovieSceneControlRigParameterSection.cpp
@@ -22,6 +22,7 @@
 #include "UObject/UE5MainStreamObjectVersion.h"
 #include "TransformConstraint.h"
 #include "TransformableHandle.h"
+#include "AnimationCoreLibrary.h"
 #include "UObject/ObjectSaveContext.h"
 
 #include UE_INLINE_GENERATED_CPP_BY_NAME(MovieSceneControlRigParameterSection)
@@ -608,6 +609,7 @@ struct FParameterTransformChannelEditorData
 	static TOptional<FRotator> GetRotator(UControlRig* ControlRig, FName ParameterName, UObject& InObject, FTrackInstancePropertyBindings* Bindings)
 	{
 		if (ControlRig)
+
 		{
 			FRigControlElement* ControlElement = ControlRig->FindControl(ParameterName);
 			if (ControlElement)
@@ -623,7 +625,9 @@ struct FParameterTransformChannelEditorData
 					}
 				}
 				
-				return ControlRig->GetHierarchy()->GetControlPreferredRotator(ControlElement);
+				FVector Vector = ControlRig->GetHierarchy()->GetControlSpecifiedEulerAngle(ControlElement);
+				FRotator Rotation = FRotator(Vector.Y, Vector.Z, Vector.X);
+				return Rotation;
 			}
 		}
 		return TOptional<FRotator>();
@@ -2583,6 +2587,94 @@ void UMovieSceneControlRigParameterSection::SetControlRig(UControlRig* InControl
 	ControlRigClass = ControlRig ? ControlRig->GetClass() : nullptr;
 }
 
+void UMovieSceneControlRigParameterSection::ChangeControlRotationOrder(const FName& InControlName, const  TOptional<EEulerRotationOrder>& OldOrder,
+	const  TOptional<EEulerRotationOrder>& NewOrder, EMovieSceneKeyInterpolation Interpolation)
+{
+	FChannelMapInfo* pChannelIndex = ControlChannelMap.Find(InControlName);
+	if (pChannelIndex == nullptr || GetControlRig() == nullptr)
+	{
+		return;
+	}
+	int32 ChannelIndex = pChannelIndex->ChannelIndex;
+	TArrayView<FMovieSceneFloatChannel*> FloatChannels = GetChannelProxy().GetChannels<FMovieSceneFloatChannel>();
+	if (FRigControlElement* ControlElement = GetControlRig()->FindControl(InControlName))
+	{
+		if (ControlElement->Settings.ControlType == ERigControlType::Rotator ||
+			ControlElement->Settings.ControlType == ERigControlType::EulerTransform ||
+			ControlElement->Settings.ControlType == ERigControlType::Transform ||
+			ControlElement->Settings.ControlType == ERigControlType::TransformNoScale)
+		{
+
+			auto AddArrayToSortedMap = [](const TArray<FFrameNumber>& InFrames, TSortedMap<FFrameNumber, FFrameNumber>& OutFrameMap)
+			{
+				for (const FFrameNumber& Frame : InFrames)
+				{
+					OutFrameMap.Add(Frame, Frame);
+				}
+			};
+
+			const int32 StartIndex = (ControlElement->Settings.ControlType == ERigControlType::Rotator) ? 0 : 3;
+			const int32 XIndex = StartIndex + ChannelIndex;
+			const int32 YIndex = XIndex + 1;
+			const int32 ZIndex = XIndex + 2;
+
+			TSortedMap<FFrameNumber, FFrameNumber> AllKeys;
+			TArray<FFrameNumber> KeyTimes;
+			TArray<FKeyHandle> Handles;
+			for (int32 Index = XIndex; Index < XIndex + 3; ++Index)
+			{
+				KeyTimes.SetNum(0);
+				Handles.SetNum(0);
+				FloatChannels[Index]->GetKeys(TRange<FFrameNumber>(), &KeyTimes, &Handles);
+				AddArrayToSortedMap(KeyTimes, AllKeys);
+			}
+			KeyTimes.SetNum(0);
+			AllKeys.GenerateKeyArray(KeyTimes);
+			if (KeyTimes.Num() <= 0)
+			{
+				//no keys so bail
+				return;
+			}
+			FRotator Rotator(0.0f, 0.0f, 0.0f);
+			const FFrameNumber StartFrame = KeyTimes[0];
+			const FFrameNumber EndFrame = KeyTimes[KeyTimes.Num() - 1];
+			for (const FFrameNumber& Frame : KeyTimes)
+			{
+				float Roll = 0.0f, Pitch = 0.0f, Yaw = 0.0f;
+				FloatChannels[XIndex]->Evaluate(Frame, Roll);
+				FloatChannels[YIndex]->Evaluate(Frame, Pitch);
+				FloatChannels[ZIndex]->Evaluate(Frame, Yaw);
+				Rotator = FRotator(Pitch, Yaw, Roll);
+				FQuat Quat;
+				//if set use animation core conversion, else use rotator conversion
+				if (OldOrder.IsSet())
+				{
+					FVector Vector = Rotator.Euler();
+					Quat = AnimationCore::QuatFromEuler(Vector, OldOrder.GetValue(),true);
+				}
+				else
+				{
+					Quat = FQuat(Rotator);
+				}
+				if (NewOrder.IsSet())
+				{
+					FVector Vector = AnimationCore::EulerFromQuat(Quat, NewOrder.GetValue(), true);
+					Rotator = FRotator::MakeFromEuler(Vector);
+				}
+				else
+				{
+					Rotator = FRotator(Quat);
+				}
+				//this will reuse tangent like we want and only add if new
+				AddKeyToChannel(FloatChannels[XIndex], Frame, Rotator.Roll, Interpolation);
+				AddKeyToChannel(FloatChannels[YIndex], Frame, Rotator.Pitch, Interpolation);
+				AddKeyToChannel(FloatChannels[ZIndex], Frame, Rotator.Yaw, Interpolation);
+			}
+			FixRotationWinding(InControlName, StartFrame, EndFrame);
+		}
+	}
+}
+
 void UMovieSceneControlRigParameterSection::FixRotationWinding(const FName& ControlName, FFrameNumber StartFrame, FFrameNumber EndFrame)
 {
 	FChannelMapInfo* pChannelIndex = ControlChannelMap.Find(ControlName);
diff --git a/Engine/Plugins/Animation/ControlRig/Source/ControlRig/Private/Sequencer/MovieSceneControlRigParameterTemplate.cpp b/Engine/Plugins/Animation/ControlRig/Source/ControlRig/Private/Sequencer/MovieSceneControlRigParameterTemplate.cpp
--- a/Engine/Plugins/Animation/ControlRig/Source/ControlRig/Private/Sequencer/MovieSceneControlRigParameterTemplate.cpp
+++ b/Engine/Plugins/Animation/ControlRig/Source/ControlRig/Private/Sequencer/MovieSceneControlRigParameterTemplate.cpp
@@ -24,6 +24,7 @@
 #include "Rigs/RigHierarchy.h"
 #include "TransformConstraint.h"
 #include "TransformableHandle.h"
+#include "AnimationCoreLibrary.h"
 
 #include UE_INLINE_GENERATED_CPP_BY_NAME(MovieSceneControlRigParameterTemplate)
 
@@ -805,22 +806,24 @@ struct FControlRigParameterPreAnimatedTokenProducer : IMovieScenePreAnimatedToke
 										case ERigControlType::Transform:
 										{
 											ControlRig->SetControlValue<FRigControlValue::FTransform_Float>(Value.Name, Value.Value.ToFTransform(), true, FRigControlModifiedContext(EControlRigSetKey::Never), bSetupUndo);
-											// @MikeZ here we want to call SetControlPreferredEulerAngles instead
-											ControlRig->GetHierarchy()->SetControlPreferredRotator(ControlElement, Value.Value.Rotation);
+											FVector EulerAngle(Value.Value.Rotation.Roll, Value.Value.Rotation.Pitch, Value.Value.Rotation.Yaw);
+											RigHierarchy->SetControlSpecifiedEulerAngle(ControlElement, EulerAngle);
 											break;
 										}
 										case ERigControlType::TransformNoScale:
 										{
 											FTransformNoScale NoScale = Value.Value.ToFTransform();
 											ControlRig->SetControlValue<FRigControlValue::FTransformNoScale_Float>(Value.Name, NoScale, true, FRigControlModifiedContext(EControlRigSetKey::Never), bSetupUndo);
-											// @MikeZ here we want to call SetControlPreferredEulerAngles instead
-											ControlRig->GetHierarchy()->SetControlPreferredRotator(ControlElement, Value.Value.Rotation);
+											FVector EulerAngle(Value.Value.Rotation.Roll, Value.Value.Rotation.Pitch, Value.Value.Rotation.Yaw);
+											RigHierarchy->SetControlSpecifiedEulerAngle(ControlElement, EulerAngle);
 											break;
 										}
 										case ERigControlType::EulerTransform:
 										{
 											FEulerTransform EulerTransform = Value.Value;
 											ControlRig->SetControlValue<FRigControlValue::FEulerTransform_Float>(Value.Name, EulerTransform, true, FRigControlModifiedContext(EControlRigSetKey::Never), bSetupUndo);
+											FVector EulerAngle(Value.Value.Rotation.Roll, Value.Value.Rotation.Pitch, Value.Value.Rotation.Yaw);
+											RigHierarchy->SetControlSpecifiedEulerAngle(ControlElement, EulerAngle);
 											break;
 										}
 										default:
@@ -940,7 +943,8 @@ struct FControlRigParameterPreAnimatedTokenProducer : IMovieScenePreAnimatedToke
 						Token.SpaceValues.Add(FControlSpaceAndValue(ControlElement->GetName(), SpaceValue));
 						const FTransform Val = ControlRig->GetHierarchy()->GetControlValue(ControlElement, ERigControlValueType::Current).Get<FRigControlValue::FTransform_Float>().ToTransform();
 						FEulerTransform EulerTransform(Val);
-						EulerTransform.Rotation = ControlRig->GetHierarchy()->GetControlPreferredRotator(ControlElement);
+						FVector Vector = ControlRig->GetHierarchy()->GetControlSpecifiedEulerAngle(ControlElement);
+						EulerTransform.Rotation = FRotator(Vector.Y, Vector.Z, Vector.X);
 						Token.TransformValues.Add(TNameAndValue<FEulerTransform>{ ControlElement->GetName(), EulerTransform });
 						break;
 					}
@@ -950,15 +954,18 @@ struct FControlRigParameterPreAnimatedTokenProducer : IMovieScenePreAnimatedToke
 							ControlRig->GetHierarchy()
 							->GetControlValue(ControlElement, ERigControlValueType::Current).Get<FRigControlValue::FTransformNoScale_Float>().ToTransform();
 						FEulerTransform EulerTransform(NoScale.ToFTransform());
-						EulerTransform.Rotation = ControlRig->GetHierarchy()->GetControlPreferredRotator(ControlElement);
+						FVector Vector = ControlRig->GetHierarchy()->GetControlSpecifiedEulerAngle(ControlElement);
+						EulerTransform.Rotation = FRotator(Vector.Y, Vector.Z, Vector.X);
 						Token.TransformValues.Add(TNameAndValue<FEulerTransform>{ ControlElement->GetName(), EulerTransform });
 						break;
 					}
 					case ERigControlType::EulerTransform:
 					{
-						const FEulerTransform EulerTransform = 
+						FEulerTransform EulerTransform = 
 							ControlRig->GetHierarchy()
 							->GetControlValue(ControlElement, ERigControlValueType::Current).Get<FRigControlValue::FEulerTransform_Float>().ToTransform();
+						FVector Vector = ControlRig->GetHierarchy()->GetControlSpecifiedEulerAngle(ControlElement);
+						EulerTransform.Rotation = FRotator(Vector.Y, Vector.Z, Vector.X);
 						Token.TransformValues.Add(TNameAndValue<FEulerTransform>{ ControlElement->GetName(), EulerTransform });
 						break;
 					}
@@ -1611,7 +1618,8 @@ struct TControlRigParameterActuatorTransform : TMovieSceneBlendingActuator<FCont
 				{
 					const FTransform Val = ControlRig->GetHierarchy()->GetControlValue(ControlElement, ERigControlValueType::Current).Get<FRigControlValue::FTransform_Float>().ToTransform();
 					FEulerTransform EulerTransform(Val);
-					EulerTransform.Rotation = ControlRig->GetHierarchy()->GetControlPreferredRotator(ControlElement);
+					FVector Vector = ControlRig->GetHierarchy()->GetControlSpecifiedEulerAngle(ControlElement);
+					EulerTransform.Rotation = FRotator(Vector.Y, Vector.Z, Vector.X);
 					return FControlRigTrackTokenTransform(EulerTransform);
 				}
 				else if (ControlElement && ControlElement->Settings.ControlType == ERigControlType::TransformNoScale)
@@ -1619,12 +1627,15 @@ struct TControlRigParameterActuatorTransform : TMovieSceneBlendingActuator<FCont
 					FTransformNoScale ValNoScale = ControlRig->GetHierarchy()->GetControlValue(ControlElement, ERigControlValueType::Current).Get<FRigControlValue::FTransformNoScale_Float>().ToTransform();
 					FTransform Val = ValNoScale;
 					FEulerTransform EulerTransform(Val);
-					EulerTransform.Rotation = ControlRig->GetHierarchy()->GetControlPreferredRotator(ControlElement);
+					FVector Vector = ControlRig->GetHierarchy()->GetControlSpecifiedEulerAngle(ControlElement);
+					EulerTransform.Rotation = FRotator(Vector.Y, Vector.Z, Vector.X);
 					return FControlRigTrackTokenTransform(EulerTransform);
 				}
 				else if (ControlElement && ControlElement->Settings.ControlType == ERigControlType::EulerTransform)
 				{
 					FEulerTransform EulerTransform = ControlRig->GetHierarchy()->GetControlValue(ControlElement, ERigControlValueType::Current).Get<FRigControlValue::FEulerTransform_Float>().ToTransform();
+					FVector Vector = ControlRig->GetHierarchy()->GetControlSpecifiedEulerAngle(ControlElement);
+					EulerTransform.Rotation = FRotator(Vector.Y, Vector.Z, Vector.X);
 					return FControlRigTrackTokenTransform(EulerTransform);
 				}
 			}
@@ -1661,19 +1672,25 @@ struct TControlRigParameterActuatorTransform : TMovieSceneBlendingActuator<FCont
 					if (ControlElement && ControlElement->Settings.ControlType == ERigControlType::Transform)
 					{
 						ControlRig->SetControlValue<FRigControlValue::FTransform_Float>(ParameterName, InFinalValue.Value.ToFTransform(), true, EControlRigSetKey::Never, bSetupUndo);
-						// @MikeZ here we want to call SetControlPreferredEulerAngles instead
-						ControlRig->GetHierarchy()->SetControlPreferredRotator(ControlElement, InFinalValue.Value.Rotator());
+						FVector EulerAngle(InFinalValue.Value.Rotation.Roll, InFinalValue.Value.Rotation.Pitch, InFinalValue.Value.Rotation.Yaw);
+						ControlRig->GetHierarchy()->SetControlSpecifiedEulerAngle(ControlElement, EulerAngle);
 					}
 					else if (ControlElement && ControlElement->Settings.ControlType == ERigControlType::TransformNoScale)
 					{
 						const FTransformNoScale NoScale = InFinalValue.Value.ToFTransform();
 						ControlRig->SetControlValue<FRigControlValue::FTransformNoScale_Float>(ParameterName, NoScale, true, EControlRigSetKey::Never, bSetupUndo);
-						// @MikeZ here we want to call SetControlPreferredEulerAngles instead
-						ControlRig->GetHierarchy()->SetControlPreferredRotator(ControlElement, InFinalValue.Value.Rotator());
+						FVector EulerAngle(InFinalValue.Value.Rotation.Roll, InFinalValue.Value.Rotation.Pitch, InFinalValue.Value.Rotation.Yaw);
+						ControlRig->GetHierarchy()->SetControlSpecifiedEulerAngle(ControlElement, EulerAngle);
 					}
 					else if (ControlElement && ControlElement->Settings.ControlType == ERigControlType::EulerTransform)
 					{
-						ControlRig->SetControlValue<FRigControlValue::FEulerTransform_Float>(ParameterName, InFinalValue.Value, true, EControlRigSetKey::Never, bSetupUndo);
+						FVector EulerAngle(InFinalValue.Value.Rotation.Roll, InFinalValue.Value.Rotation.Pitch, InFinalValue.Value.Rotation.Yaw);
+						FQuat Quat = ControlRig->GetHierarchy()->GetControlQuaternion(ControlElement, EulerAngle);
+						ControlRig->GetHierarchy()->SetControlSpecifiedEulerAngle(ControlElement, EulerAngle);
+						FRotator UERotator(Quat);
+						FEulerTransform Transform = InFinalValue.Value;
+						Transform.Rotation = UERotator;
+						ControlRig->SetControlValue<FRigControlValue::FEulerTransform_Float>(ParameterName, Transform, true, EControlRigSetKey::Never, bSetupUndo);	
 					}
 				}
 			}
diff --git a/Engine/Plugins/Animation/ControlRig/Source/ControlRig/Private/Sequencer/MovieSceneControlRigParameterTrack.cpp b/Engine/Plugins/Animation/ControlRig/Source/ControlRig/Private/Sequencer/MovieSceneControlRigParameterTrack.cpp
--- a/Engine/Plugins/Animation/ControlRig/Source/ControlRig/Private/Sequencer/MovieSceneControlRigParameterTrack.cpp
+++ b/Engine/Plugins/Animation/ControlRig/Source/ControlRig/Private/Sequencer/MovieSceneControlRigParameterTrack.cpp
@@ -17,6 +17,13 @@
 
 #define LOCTEXT_NAMESPACE "MovieSceneParameterControlRigTrack"
 
+FControlRotationOrder::FControlRotationOrder()
+	:RotationOrder(EEulerRotationOrder::YZX),
+	bOverrideSetting(false)
+
+{
+
+}
 
 UMovieSceneControlRigParameterTrack::UMovieSceneControlRigParameterTrack(const FObjectInitializer& ObjectInitializer)
 	: Super(ObjectInitializer)
@@ -536,6 +543,11 @@ void UMovieSceneControlRigParameterTrack::HandleOnInitialized_GameThread()
 				}
 			}
 		}
+		if (SortedControls.Num() > 0) //really set up
+		{
+			TArray<FName> Names = GetControlsWithDifferentRotationOrders();
+			ResetControlsToSettingsRotationOrder(Names);
+		}
 	}
 }
 
@@ -620,6 +632,12 @@ void UMovieSceneControlRigParameterTrack::RenameParameterName(const FName& OldPa
 				CRSection->RenameParameterName(OldParameterName, NewParameterName);
 			}
 		}
+		if (FControlRotationOrder* RotationOrder = ControlsRotationOrder.Find(OldParameterName))
+		{
+			FControlRotationOrder NewRotationOrder = *RotationOrder;
+			ControlsRotationOrder.Remove(OldParameterName);
+			ControlsRotationOrder.Add(NewParameterName, NewRotationOrder);
+		}
 	}
 }
 
@@ -656,6 +674,107 @@ void UMovieSceneControlRigParameterTrack::ReplaceControlRig(UControlRig* NewCont
 		}
 	}
 }
+TOptional<EEulerRotationOrder> UMovieSceneControlRigParameterTrack::GetControlRotationOrder(const FRigControlElement* ControlElement,
+	bool bCurrent) const
+{
+	TOptional<EEulerRotationOrder> Order;
+	if (bCurrent)
+	{
+		const FControlRotationOrder* RotationOrder = ControlsRotationOrder.Find(ControlElement->GetName());
+		if (RotationOrder)
+		{
+			Order = RotationOrder->RotationOrder;
+		}
+	}
+	else //use setting
+	{
+		if (ControlRig->GetHierarchy()->GetUsePreferredRotationOrder(ControlElement))
+		{
+			Order = ControlRig->GetHierarchy()->GetControlPreferredEulerRotationOrder(ControlElement);
+		}
+	}
+	return Order;
+}
+
+TArray<FName> UMovieSceneControlRigParameterTrack::GetControlsWithDifferentRotationOrders() const
+{
+	TArray<FName> Names;
+	if (ControlRig && ControlRig->GetHierarchy())
+	{
+		URigHierarchy* Hierarchy = ControlRig->GetHierarchy();
+		TArray<FRigControlElement*> SortedControls;
+		ControlRig->GetControlsInOrder(SortedControls);
+		for (const FRigControlElement* ControlElement : SortedControls)
+		{
+			if (!Hierarchy->IsAnimatable(ControlElement))
+			{
+				continue;
+			}
+			TOptional<EEulerRotationOrder> Current = GetControlRotationOrder(ControlElement, true);
+			TOptional<EEulerRotationOrder> Setting = GetControlRotationOrder(ControlElement, false);
+			if (Current != Setting)
+			{
+				Names.Add(ControlElement->GetName());
+			}
+
+		}
+	}
+	return Names;
+}
+
+void UMovieSceneControlRigParameterTrack::ResetControlsToSettingsRotationOrder(const TArray<FName>& Names, EMovieSceneKeyInterpolation Interpolation)
+{
+	if (ControlRig && ControlRig->GetHierarchy())
+	{
+		for (const FName& Name : Names)
+		{
+			if (FRigControlElement* ControlElement = GetControlRig()->FindControl(Name))
+			{
+
+				TOptional<EEulerRotationOrder> Current = GetControlRotationOrder(ControlElement, true);
+				TOptional<EEulerRotationOrder> Setting = GetControlRotationOrder(ControlElement, false);
+				if (Current != Setting)
+				{
+					ChangeControlRotationOrder(Name, Setting, Interpolation);
+				}
+			}
+		}
+	}
+}
+void UMovieSceneControlRigParameterTrack::ChangeControlRotationOrder(const FName& InControlName, const TOptional<EEulerRotationOrder>& NewOrder, EMovieSceneKeyInterpolation Interpolation)
+{
+	if (ControlRig && ControlRig->GetHierarchy())
+	{
+		if (FRigControlElement* ControlElement = GetControlRig()->FindControl(InControlName))
+		{
+			TOptional<EEulerRotationOrder> Current = GetControlRotationOrder(ControlElement, true);
+			if (Current != NewOrder)
+			{
+				if (NewOrder.IsSet())
+				{
+					FControlRotationOrder& RotationOrder = ControlsRotationOrder.FindOrAdd(InControlName);
+					RotationOrder.RotationOrder = NewOrder.GetValue();
+					TOptional<EEulerRotationOrder> Setting = GetControlRotationOrder(ControlElement, false);
+					if (Setting != NewOrder)
+					{
+						RotationOrder.bOverrideSetting = true;
+					}
+				}
+				else //no longer set so just remove
+				{
+					ControlsRotationOrder.Remove(InControlName);
+				}
+				for (UMovieSceneSection* Section : Sections)
+				{
+					if (UMovieSceneControlRigParameterSection* CRSection = Cast<UMovieSceneControlRigParameterSection>(Section))
+					{
+						CRSection->ChangeControlRotationOrder(InControlName, Current, NewOrder, Interpolation);
+					}
+				}
+			}
+		}
+	}	
+}
 
 void UMovieSceneControlRigParameterTrack::GetSelectedNodes(TArray<FName>& SelectedControlNames)
 {
diff --git a/Engine/Plugins/Animation/ControlRig/Source/ControlRig/Private/Tools/ControlRigPose.cpp b/Engine/Plugins/Animation/ControlRig/Source/ControlRig/Private/Tools/ControlRigPose.cpp
--- a/Engine/Plugins/Animation/ControlRig/Source/ControlRig/Private/Tools/ControlRigPose.cpp
+++ b/Engine/Plugins/Animation/ControlRig/Source/ControlRig/Private/Tools/ControlRigPose.cpp
@@ -62,12 +62,12 @@ void FControlRigControlPose::SetControlMirrorTransform(bool bDoLocal, UControlRi
 {
 	if (bDoLocal || bIsMatched)
 	{
-		ControlRig->SetControlLocalTransform(Name, LocalTransform, bNotify,Context,bSetupUndo);
+		ControlRig->SetControlLocalTransform(Name, LocalTransform, bNotify,Context,bSetupUndo, true/* bFixEulerFlips*/);
 
 	}
 	else
 	{
-		ControlRig->SetControlGlobalTransform(Name, GlobalTransform, bNotify,Context,bSetupUndo);
+		ControlRig->SetControlGlobalTransform(Name, GlobalTransform, bNotify,Context,bSetupUndo, false /*bPrintPython*/, true/* bFixEulerFlips*/);
 	}	
 }
 
@@ -120,11 +120,11 @@ void FControlRigControlPose::PastePoseInternal(UControlRig* ControlRig, bool bDo
 				{
 					if (bDoLocal) // -V547  
 					{
-						ControlRig->SetControlLocalTransform(ControlElement->GetName(), CopyRigControl->LocalTransform, true,Context,bSetupUndo);
+						ControlRig->SetControlLocalTransform(ControlElement->GetName(), CopyRigControl->LocalTransform, true,Context,bSetupUndo, true/* bFixEulerFlips*/);
 					}
 					else
 					{
-						ControlRig->SetControlGlobalTransform(ControlElement->GetName(), CopyRigControl->GlobalTransform, true,Context,bSetupUndo);
+						ControlRig->SetControlGlobalTransform(ControlElement->GetName(), CopyRigControl->GlobalTransform, true, Context, bSetupUndo, false /*bPrintPython*/, true/* bFixEulerFlips*/);
 					}
 				}
 				else
@@ -170,7 +170,7 @@ void FControlRigControlPose::PastePoseInternal(UControlRig* ControlRig, bool bDo
 			{
 				FTransform GlobalTransform = ControlRig->GetHierarchy()->GetGlobalTransform(ControlElement->GetKey());
 				ControlRig->GetHierarchy()->SwitchToParent(ControlElement->GetKey(), SpaceKey);
-				ControlRig->SetControlGlobalTransform(ControlElement->GetName(), GlobalTransform, true, Context, bSetupUndo);
+				ControlRig->SetControlGlobalTransform(ControlElement->GetName(), GlobalTransform, true, Context, bSetupUndo, false /*bPrintPython*/, true/* bFixEulerFlips*/);
 			}
 		}
 	}
@@ -241,7 +241,7 @@ void FControlRigControlPose::BlendWithInitialPoses(FControlRigControlPose& Initi
 							Rotation = FQuat::Slerp(InitialVal.GetRotation(), Val.GetRotation(), BlendValue); //doing slerp here not fast lerp, can be slow this is for content creation
 							Scale = FMath::Lerp(InitialVal.GetScale3D(), Val.GetScale3D(), BlendValue);
 							Val = FTransform(Rotation, Translation, Scale);
-							ControlRig->SetControlLocalTransform(ControlElement->GetName(), Val, bDoKey,Context,bSetupUndo);
+							ControlRig->SetControlLocalTransform(ControlElement->GetName(), Val, bDoKey,Context,bSetupUndo, true/* bFixEulerFlips*/);
 						}
 						else
 						{
@@ -253,7 +253,7 @@ void FControlRigControlPose::BlendWithInitialPoses(FControlRigControlPose& Initi
 							Rotation = FQuat::Slerp(InitialVal.GetRotation(), Val.GetRotation(), BlendValue); //doing slerp here not fast lerp, can be slow this is for content creation
 							Scale = FMath::Lerp(InitialVal.GetScale3D(), Val.GetScale3D(), BlendValue);
 							Val = FTransform(Rotation, Translation, Scale);
-							ControlRig->SetControlGlobalTransform(ControlElement->GetName(), Val, bDoKey,Context,bSetupUndo);
+							ControlRig->SetControlGlobalTransform(ControlElement->GetName(), Val, bDoKey,Context,bSetupUndo, false /*bPrintPython*/, true/* bFixEulerFlips*/);
 						}
 					}
 					else
@@ -284,7 +284,7 @@ void FControlRigControlPose::BlendWithInitialPoses(FControlRigControlPose& Initi
 					{
 						FTransform GlobalTransform = ControlRig->GetHierarchy()->GetGlobalTransform(ControlElement->GetKey());
 						ControlRig->GetHierarchy()->SwitchToParent(ControlElement->GetKey(), SpaceKey);
-						ControlRig->SetControlGlobalTransform(ControlElement->GetName(), GlobalTransform, true, Context, bSetupUndo);
+						ControlRig->SetControlGlobalTransform(ControlElement->GetName(), GlobalTransform, true, Context, bSetupUndo, false /*bPrintPython*/, true/* bFixEulerFlips*/);
 					}
 				}
 			}
diff --git a/Engine/Plugins/Animation/ControlRig/Source/ControlRig/Public/ControlRigObjectVersion.h b/Engine/Plugins/Animation/ControlRig/Source/ControlRig/Public/ControlRigObjectVersion.h
--- a/Engine/Plugins/Animation/ControlRig/Source/ControlRig/Public/ControlRigObjectVersion.h
+++ b/Engine/Plugins/Animation/ControlRig/Source/ControlRig/Public/ControlRigObjectVersion.h
@@ -90,6 +90,9 @@ struct CONTROLRIG_API FControlRigObjectVersion
 		// Control supporting preferred rotation order
 		RigHierarchyControlPreferredRotationOrder,
 
+		// Last bit required for Control supporting preferred rotation order
+		RigHierarchyControlPreferredRotationOrderFlag,
+
 		// -----<new versions can be added above this line>-------------------------------------------------
 		VersionPlusOne,
 		LatestVersion = VersionPlusOne - 1
diff --git a/Engine/Plugins/Animation/ControlRig/Source/ControlRig/Public/Rigs/RigHierarchy.h b/Engine/Plugins/Animation/ControlRig/Source/ControlRig/Public/Rigs/RigHierarchy.h
--- a/Engine/Plugins/Animation/ControlRig/Source/ControlRig/Public/Rigs/RigHierarchy.h
+++ b/Engine/Plugins/Animation/ControlRig/Source/ControlRig/Public/Rigs/RigHierarchy.h
@@ -2238,6 +2238,100 @@ public:
 		}
 	}
 
+	bool GetUsePreferredRotationOrder(const FRigControlElement* InControlElement) const
+	{
+		if (InControlElement)
+		{
+			{
+				return InControlElement->Settings.bUsePreferredRotationOrder;
+			}
+		}
+		return false;
+	}
+
+	void  SetUsePreferredRotationOrder(FRigControlElement* InControlElement, bool bVal)
+	{
+		if (InControlElement)
+		{
+			{
+				InControlElement->Settings.bUsePreferredRotationOrder = bVal;
+			}
+		}
+	}
+
+	FVector GetControlSpecifiedEulerAngle(const FRigControlElement* InControlElement, bool bIsInitial = false) const
+	{
+		FVector EulerAngle = FVector::ZeroVector;
+		if (InControlElement)
+		{
+			if (bIsInitial == false && GetUsePreferredRotationOrder(InControlElement))
+			{
+				EEulerRotationOrder RotationOrder = GetControlPreferredEulerRotationOrder(InControlElement);
+				EulerAngle = InControlElement->PreferredEulerAngles.GetAngles(false, RotationOrder);
+			}
+			else
+			{
+				FRotator Rotator = GetControlPreferredRotator(InControlElement);
+				EulerAngle = FVector(Rotator.Roll, Rotator.Pitch, Rotator.Yaw);
+			}
+		}
+		return EulerAngle;
+	}
+	
+	void SetControlSpecifiedEulerAngle(FRigControlElement* InControlElement, const FVector& InEulerAngle, bool bIsInitial = false)
+	{
+		if (InControlElement)
+		{
+			if (GetUsePreferredRotationOrder(InControlElement))
+			{
+				EEulerRotationOrder RotationOrder = GetControlPreferredEulerRotationOrder(InControlElement);
+				SetControlPreferredEulerAngles(InControlElement, InEulerAngle, RotationOrder, bIsInitial);
+			}
+			else
+			{
+				FRotator Rotator(InEulerAngle[1], InEulerAngle[2], InEulerAngle[0]);
+				SetControlPreferredRotator(InControlElement, Rotator, bIsInitial, false /* fix euler flips*/); //test fix todo mikez
+			}
+		}
+	}
+
+	FQuat GetControlQuaternion(const FRigControlElement* InControlElement, const FVector& InEulerAngle) const
+	{
+		if (InControlElement)
+		{
+			FRotator Rotator(InEulerAngle[1], InEulerAngle[2], InEulerAngle[0]);
+
+			if (GetUsePreferredRotationOrder(InControlElement))
+			{
+				return InControlElement->PreferredEulerAngles.GetQuatFromRotator(Rotator);
+			}
+			else
+			{
+				return Rotator.Quaternion();
+			}
+		}
+		return FQuat();
+	}
+
+	FVector GetControlAnglesFromQuat(const FRigControlElement* InControlElement, const FQuat& InQuat) const
+	{
+		FVector Angle(0, 0, 0);
+		if (InControlElement)
+		{
+
+			if (GetUsePreferredRotationOrder(InControlElement))
+			{
+				FRotator Rotator = InControlElement->PreferredEulerAngles.GetRotatorFromQuat(InQuat);
+				Angle = Rotator.Euler();
+			}
+			else
+			{
+				FRotator Rotator(InQuat);
+				Angle = Rotator.Euler();
+			}
+		}
+		return Angle;
+	}
 	/**
 	 * Returns the pin type to use for a control
 	 * @param InControlElement The control to return the pin type for
diff --git a/Engine/Plugins/Animation/ControlRig/Source/ControlRig/Public/Rigs/RigHierarchyElements.h b/Engine/Plugins/Animation/ControlRig/Source/ControlRig/Public/Rigs/RigHierarchyElements.h
--- a/Engine/Plugins/Animation/ControlRig/Source/ControlRig/Public/Rigs/RigHierarchyElements.h
+++ b/Engine/Plugins/Animation/ControlRig/Source/ControlRig/Public/Rigs/RigHierarchyElements.h
@@ -473,6 +473,10 @@ struct CONTROLRIG_API FRigPreferredEulerAngles
 	void SetAngles(const FVector& InValue, bool bInitial = false, EEulerRotationOrder InRotationOrder = DefaultRotationOrder, bool bFixEulerFlips = false);
 	void SetRotationOrder(EEulerRotationOrder InRotationOrder);
 
+
+	FRotator GetRotatorFromQuat(const FQuat& InQuat) const;
+	FQuat GetQuatFromRotator(const FRotator& InRotator) const;
+	
 	UPROPERTY(BlueprintReadOnly, EditAnywhere, Category = "Pose")
 	EEulerRotationOrder RotationOrder;
 
@@ -1189,11 +1193,45 @@ struct CONTROLRIG_API FRigControlSettings
 	TArray<ERigControlTransformChannel> FilteredChannels;
 
 	/**
-	 * The euler rotation order this control prefers for animation
+	 * The euler rotation order this control prefers for animation, if we aren't using default UE rotator
 	 */
 	UPROPERTY(BlueprintReadOnly, EditAnywhere, Category = Animation)
 	EEulerRotationOrder PreferredRotationOrder;
 
+	/**
+	* Whether to use a specfied rotation order or just use the default FRotator order and conversion functions
+	*/
+	UPROPERTY(BlueprintReadOnly, EditAnywhere, Category = Animation)
+	bool bUsePreferredRotationOrder;
+
+	/**
+	* The euler rotation order this control prefers for animation if it is active. If not set then we use the default UE rotator.
+	*/
+	TOptional<EEulerRotationOrder> GetRotationOrder() const
+	{
+		TOptional<EEulerRotationOrder> RotationOrder;
+		if (bUsePreferredRotationOrder)
+		{
+			RotationOrder = PreferredRotationOrder;
+		}
+		return RotationOrder;
+	}
+
+	/**
+	*  Set the rotation order if the rotation is set otherwise use default rotator
+	*/
+	void SetRotationOrder(const TOptional<EEulerRotationOrder>& EulerRotation)
+	{
+		if (EulerRotation.IsSet())
+		{
+			bUsePreferredRotationOrder = true;
+			PreferredRotationOrder = EulerRotation.GetValue();
+		}
+		else
+		{
+			bUsePreferredRotationOrder = false;
+		}
+	}
 #if WITH_EDITORONLY_DATA
 	/**
 	 * Deprecated properties.
@@ -1356,11 +1394,11 @@ public:
 	UPROPERTY(BlueprintReadOnly, EditAnywhere, Category = RigElement)
 	FRigCurrentAndInitialTransform Shape;
 
-protected:
-
 	UPROPERTY(BlueprintReadOnly, EditAnywhere, Category = RigElement)
 	FRigPreferredEulerAngles PreferredEulerAngles;
 	
+protected:
+
 	static bool IsClassOf(const FRigBaseElement* InElement)
 	{
 		return InElement->GetType() == ERigElementType::Control;
diff --git a/Engine/Plugins/Animation/ControlRig/Source/ControlRig/Public/Sequencer/MovieSceneControlRigParameterSection.h b/Engine/Plugins/Animation/ControlRig/Source/ControlRig/Public/Sequencer/MovieSceneControlRigParameterSection.h
--- a/Engine/Plugins/Animation/ControlRig/Source/ControlRig/Public/Sequencer/MovieSceneControlRigParameterSection.h
+++ b/Engine/Plugins/Animation/ControlRig/Source/ControlRig/Public/Sequencer/MovieSceneControlRigParameterSection.h
@@ -211,6 +211,10 @@ public:
 	const FName& FindControlNameFromConstraintChannel(const FMovieSceneConstraintChannel* InConstraintChannel) const;
 		
 	bool RenameParameterName(const FName& OldParameterName, const FName& NewParameterName);
+
+	void ChangeControlRotationOrder(const FName& InControlName, const TOptional<EEulerRotationOrder>& CurrentOrder, 
+		const  TOptional<EEulerRotationOrder>& NewOrder, EMovieSceneKeyInterpolation Interpolation);
+
 private:
 
 	FSpaceChannelAddedEvent OnSpaceChannelAdded;
diff --git a/Engine/Plugins/Animation/ControlRig/Source/ControlRig/Public/Sequencer/MovieSceneControlRigParameterTrack.h b/Engine/Plugins/Animation/ControlRig/Source/ControlRig/Public/Sequencer/MovieSceneControlRigParameterTrack.h
--- a/Engine/Plugins/Animation/ControlRig/Source/ControlRig/Public/Sequencer/MovieSceneControlRigParameterTrack.h
+++ b/Engine/Plugins/Animation/ControlRig/Source/ControlRig/Public/Sequencer/MovieSceneControlRigParameterTrack.h
@@ -10,9 +10,27 @@
 #include "Compilation/IMovieSceneTrackTemplateProducer.h"
 #include "INodeAndChannelMappings.h"
 #include "MovieSceneControlRigParameterSection.h"
+#include "EulerTransform.h"
 #include "MovieSceneControlRigParameterTrack.generated.h"
 
 struct FEndLoadPackageContext;
+struct FRigControlElement;
+
+USTRUCT()
+struct CONTROLRIG_API FControlRotationOrder
+{
+	GENERATED_BODY()
+
+	FControlRotationOrder();
+
+	/* Rotation Order*/
+	UPROPERTY()
+	EEulerRotationOrder RotationOrder;
+
+	/** Override the default control setting*/
+	UPROPERTY()
+	bool bOverrideSetting;
+};
 
 /**
  * Handles animation of skeletal mesh actors using animation ControlRigs
@@ -131,6 +149,17 @@ public:
 	UMovieSceneControlRigParameterSection::FSpaceChannelAddedEvent& SpaceChannelAdded() { return OnSpaceChannelAdded; }
 	IMovieSceneConstrainedSection::FConstraintChannelAddedEvent& ConstraintChannelAdded() { return OnConstraintChannelAdded; }
 
+public:
+	CONTROLRIG_API TArray<FName> GetControlsWithDifferentRotationOrders() const;
+	CONTROLRIG_API void ResetControlsToSettingsRotationOrder(const TArray<FName>& Names,
+		EMovieSceneKeyInterpolation Interpolation = EMovieSceneKeyInterpolation::SmartAuto);
+	// if order is not set then it uses the default FRotator conversions
+	CONTROLRIG_API void ChangeControlRotationOrder(const FName& InControlName, const TOptional<EEulerRotationOrder>& NewOrder,
+		EMovieSceneKeyInterpolation Interpolation = EMovieSceneKeyInterpolation::SmartAuto);
+private:
+	//get the rotation orderm will not be set if it's default FRotator order. If bCurrent is true, it uses what's set
+	//if false it uses the default setting from the current control rig.
+	TOptional<EEulerRotationOrder> GetControlRotationOrder(const FRigControlElement* ControlElement, bool bCurrent) const;
 private:
 
 	//we register this with sections.
@@ -165,6 +194,11 @@ private:
 	/** Unique Name*/
 	UPROPERTY()
 	FName TrackName;
+
+	/** Uses Rotation Order*/
+	UPROPERTY()
+	TMap<FName,FControlRotationOrder> ControlsRotationOrder;
+
 };
 
 
diff --git a/Engine/Plugins/Animation/ControlRig/Source/ControlRigEditor/Private/ControlRigElementDetails.cpp b/Engine/Plugins/Animation/ControlRig/Source/ControlRigEditor/Private/ControlRigElementDetails.cpp
--- a/Engine/Plugins/Animation/ControlRig/Source/ControlRigEditor/Private/ControlRigElementDetails.cpp
+++ b/Engine/Plugins/Animation/ControlRig/Source/ControlRigEditor/Private/ControlRigElementDetails.cpp
@@ -1543,7 +1543,7 @@ FDetailWidgetRow& FRigTransformElementDetails::CreateEulerTransformValueWidgetRo
 
 	const TSharedPtr<ESlateRotationRepresentation::Type> RotationRepresentationStorage = Settings.RotationRepresentation;
 	TransformWidgetArgs.RotationRepresentation(RotationRepresentationStorage);
-	
+
 	auto IsComponentRelative = [TransformWidgetArgs](int32 Component) -> bool
 	{
 		if(TransformWidgetArgs._OnGetIsComponentRelative.IsBound())
@@ -1892,7 +1892,8 @@ FDetailWidgetRow& FRigTransformElementDetails::CreateEulerTransformValueWidgetRo
 					case ERigControlType::Transform:
 					case ERigControlType::TransformNoScale:
 					{
-						RelativeTransform.Rotation = Hierarchy->GetControlPreferredRotator(ControlElement, bInitial);
+						FVector Vector = Hierarchy->GetControlSpecifiedEulerAngle(ControlElement, bInitial);
+						RelativeTransform.Rotation =  FRotator(Vector.Y, Vector.Z, Vector.X);
 						break;
 					}
 					default:
@@ -2066,7 +2067,8 @@ FDetailWidgetRow& FRigTransformElementDetails::CreateEulerTransformValueWidgetRo
 							case ERigControlType::Transform:
 							case ERigControlType::TransformNoScale:
 							{
-								HierarchyToUpdate->SetControlPreferredRotator(ControlElement, InTransform.Rotator(), bInitial);
+								FVector EulerAngle(InTransform.Rotator().Roll, InTransform.Rotator().Pitch, InTransform.Rotator().Yaw);
+								HierarchyToUpdate->SetControlSpecifiedEulerAngle(ControlElement, EulerAngle, bInitial);
 								break;
 							}
 							default:
@@ -3027,6 +3029,10 @@ void FRigControlElementDetails::CustomizeControl(IDetailLayoutBuilder& DetailBui
 	{
 		if(IsAnyControlOfValueType(ERigControlType::EulerTransform))
 		{
+			const TSharedPtr<IPropertyHandle> UsePreferredRotationOrderHandle = SettingsHandle->GetChildHandle(TEXT("bUsePreferredRotationOrder"));
+			ControlCategory.AddProperty(UsePreferredRotationOrderHandle.ToSharedRef()).DisplayName(FText::FromString(TEXT("Use Preferred Rotation Order")))
+				.IsEnabled(bIsEnabled);
+
 			const TSharedPtr<IPropertyHandle> PreferredRotationOrderHandle = SettingsHandle->GetChildHandle(TEXT("PreferredRotationOrder"));
 			ControlCategory.AddProperty(PreferredRotationOrderHandle.ToSharedRef()).DisplayName(FText::FromString(TEXT("Preferred Rotation Order")))
 			.IsEnabled(bIsEnabled);
diff --git a/Engine/Plugins/Animation/ControlRig/Source/ControlRigEditor/Private/ControlRigSequencerEditorLibrary.cpp b/Engine/Plugins/Animation/ControlRig/Source/ControlRigEditor/Private/ControlRigSequencerEditorLibrary.cpp
--- a/Engine/Plugins/Animation/ControlRig/Source/ControlRigEditor/Private/ControlRigSequencerEditorLibrary.cpp
+++ b/Engine/Plugins/Animation/ControlRig/Source/ControlRigEditor/Private/ControlRigSequencerEditorLibrary.cpp
@@ -1038,7 +1038,7 @@ static void LocalSetControlRigWorldTransforms(ULevelSequence* LevelSequence,UCon
 						const FFrameNumber& FrameNumber = Frames[Index];
 						Context.LocalTime = TickResolution.AsSeconds(FFrameTime(FrameNumber));
 						FTransform GlobalTransform = WorldTransforms[Index].GetRelativeTransform(ControlRigParentWorldTransforms[Index]);
-						ControlRig->SetControlGlobalTransform(ControlName, GlobalTransform, true, Context);
+						ControlRig->SetControlGlobalTransform(ControlName, GlobalTransform, true, Context, false /*undo*/, false /*bPrintPython*/, true/* bFixEulerFlips*/);
 					}
 				}
 			}
diff --git a/Engine/Plugins/Animation/ControlRig/Source/ControlRigEditor/Private/EditMode/ControlRigControlsProxy.cpp b/Engine/Plugins/Animation/ControlRig/Source/ControlRigEditor/Private/EditMode/ControlRigControlsProxy.cpp
--- a/Engine/Plugins/Animation/ControlRig/Source/ControlRigEditor/Private/EditMode/ControlRigControlsProxy.cpp
+++ b/Engine/Plugins/Animation/ControlRig/Source/ControlRigEditor/Private/EditMode/ControlRigControlsProxy.cpp
@@ -208,7 +208,8 @@ void UControlRigTransformControlProxy::PostEditChangeChainProperty(struct FPrope
 			//MUST set through ControlRig
 			const FTransform RealTransform = Transform.ToFTransform(); //Transform is FEulerTransform
 			ControlRig->SetControlValue<FRigControlValue::FTransform_Float>(ControlName, RealTransform, true, EControlRigSetKey::DoNotCare,false);
-			ControlRig->GetHierarchy()->SetControlPreferredRotator(ControlElement, Transform.Rotation);
+			FVector EulerAngle(Transform.Rotation.Roll, Transform.Rotation.Pitch, Transform.Rotation.Yaw);
+			ControlRig->GetHierarchy()->SetControlSpecifiedEulerAngle(ControlElement, EulerAngle);
 			ControlRig->Evaluate_AnyThread();
 
 		}
@@ -239,7 +240,8 @@ void UControlRigTransformControlProxy::PostEditUndo()
 		ControlRig->SelectControl(ControlName, bSelected);
 		const FTransform RealTransform = Transform.ToFTransform(); //Transform is FEulerTransform
 		ControlRig->SetControlValue<FRigControlValue::FTransform_Float>(ControlName, RealTransform, true, EControlRigSetKey::Never,false);
-		ControlRig->GetHierarchy()->SetControlPreferredRotator(ControlElement, Transform.Rotation);
+		FVector EulerAngle(Transform.Rotation.Roll, Transform.Rotation.Pitch, Transform.Rotation.Yaw);
+		ControlRig->GetHierarchy()->SetControlSpecifiedEulerAngle(ControlElement, EulerAngle);
 	}
 }
 #endif
@@ -265,7 +267,8 @@ void UControlRigTransformControlProxy::SetKey(const IPropertyHandle& KeyedProper
 		}
 		const FTransform RealTransform = Transform.ToFTransform(); //Transform is FEulerTransform
 		ControlRig->SetControlValue<FRigControlValue::FTransform_Float>(ControlName, RealTransform, true, Context, false);
-		ControlRig->GetHierarchy()->SetControlPreferredRotator(ControlElement, Transform.Rotation);
+		FVector EulerAngle(Transform.Rotation.Roll, Transform.Rotation.Pitch, Transform.Rotation.Yaw);
+		ControlRig->GetHierarchy()->SetControlSpecifiedEulerAngle(ControlElement, EulerAngle);
 	}
 }
 
@@ -355,8 +358,14 @@ void UControlRigEulerTransformControlProxy::PostEditChangeChainProperty(struct F
 		FRigControlElement* ControlElement = GetControlElement();
 		if (ControlElement && ControlRig.IsValid())
 		{
-			//MUST set through ControlRig
-			ControlRig->SetControlValue<FRigControlValue::FEulerTransform_Float>(ControlName, Transform, true, EControlRigSetKey::DoNotCare,false);
+			FVector EulerAngle(Transform.Rotation.Roll, Transform.Rotation.Pitch, Transform.Rotation.Yaw);
+			FQuat Quat = ControlRig->GetHierarchy()->GetControlQuaternion(ControlElement, EulerAngle);
+			ControlRig->GetHierarchy()->SetControlSpecifiedEulerAngle(ControlElement, EulerAngle);
+			FRotator UERotator(Quat);
+			FEulerTransform UETransform = Transform;
+			UETransform.Rotation = UERotator;
+
+			ControlRig->SetControlValue<FRigControlValue::FEulerTransform_Float>(ControlName, UETransform, true, EControlRigSetKey::DoNotCare,false);
 			ControlRig->Evaluate_AnyThread();
 
 		}
@@ -371,8 +380,8 @@ void UControlRigEulerTransformControlProxy::ValueChanged()
 		Modify();
 		const FName PropertyName("Transform");
 		FTrackInstancePropertyBindings Binding(PropertyName, PropertyName.ToString());
-		const FEulerTransform NewTransform = ControlRig.Get()->GetHierarchy()->GetControlValue(ControlElement, ERigControlValueType::Current).Get<FRigControlValue::FEulerTransform_Float>().ToTransform();
-
+		FEulerTransform NewTransform = ControlRig.Get()->GetHierarchy()->GetControlValue(ControlElement, ERigControlValueType::Current).Get<FRigControlValue::FEulerTransform_Float>().ToTransform();
+		NewTransform.Rotation = ControlRig->GetHierarchy()->GetControlPreferredRotator(ControlElement);
 		Binding.CallFunction<FEulerTransform>(*this, NewTransform);
 
 	}
@@ -386,6 +395,8 @@ void UControlRigEulerTransformControlProxy::PostEditUndo()
 	{
 		ControlRig->SelectControl(ControlName, bSelected);
 		ControlRig->SetControlValue<FRigControlValue::FEulerTransform_Float>(ControlName, Transform, true, EControlRigSetKey::Never,false);
+		FVector EulerAngle(Transform.Rotation.Roll, Transform.Rotation.Pitch, Transform.Rotation.Yaw);
+		ControlRig->GetHierarchy()->SetControlSpecifiedEulerAngle(ControlElement, EulerAngle);
 	}
 }
 #endif
@@ -410,6 +421,8 @@ void UControlRigEulerTransformControlProxy::SetKey(const IPropertyHandle& KeyedP
 			Context.KeyMask = (uint32)EControlRigContextChannelToKey::Scale;
 		}
 		ControlRig->SetControlValue<FRigControlValue::FEulerTransform_Float>(ControlName, Transform, true, Context,false);
+		FVector EulerAngle(Transform.Rotation.Roll, Transform.Rotation.Pitch, Transform.Rotation.Yaw);
+		ControlRig->GetHierarchy()->SetControlSpecifiedEulerAngle(ControlElement, EulerAngle);
 	}
 }
 
diff --git a/Engine/Plugins/Animation/ControlRig/Source/ControlRigEditor/Private/EditMode/ControlRigEditMode.cpp b/Engine/Plugins/Animation/ControlRig/Source/ControlRigEditor/Private/EditMode/ControlRigEditMode.cpp
--- a/Engine/Plugins/Animation/ControlRig/Source/ControlRigEditor/Private/EditMode/ControlRigEditMode.cpp
+++ b/Engine/Plugins/Animation/ControlRig/Source/ControlRigEditor/Private/EditMode/ControlRigEditMode.cpp
@@ -2922,7 +2922,7 @@ void FControlRigEditMode::ResetTransforms(bool bSelectionOnly)
 			}
 			if (ControlElement)
 			{
-				ControlRig->SetControlLocalTransform(ElementToReset.Name, InitialLocalTransform, true);
+				ControlRig->SetControlLocalTransform(ElementToReset.Name, InitialLocalTransform, true, FRigControlModifiedContext(), true, true);
 				NotifyDrivenControls(ControlRig, ElementToReset);
 				if (bHasNonDefaultParent == false)
 				{
@@ -4012,7 +4012,6 @@ void FControlRigEditMode::MoveControlShape(AControlRigShapeActor* ShapeActor, co
 	bool bUseLocal, bool bCalcLocal, FTransform& InOutLocal)
 {
 	bool bTransformChanged = false;
-
 	//first case is where we do all controls by the local diff.
 	if (bUseLocal)
 	{
@@ -4632,7 +4631,7 @@ void FControlRigEditMode::NotifyDrivenControls(UControlRig* InControlRig, const
 				if(DrivenKey.Type == ERigElementType::Control)
 				{
 					const FTransform DrivenTransform = InControlRig->GetHierarchy()->GetLocalTransform(DrivenKey);
-					InControlRig->SetControlLocalTransform(DrivenKey.Name, DrivenTransform, true, Context, false);
+					InControlRig->SetControlLocalTransform(DrivenKey.Name, DrivenTransform, true, Context, false /*undo*/, true/* bFixEulerFlips*/);
 				}
 			}
 		}
diff --git a/Engine/Plugins/Animation/ControlRig/Source/ControlRigEditor/Private/Sequencer/ControlRigParameterTrackEditor.cpp b/Engine/Plugins/Animation/ControlRig/Source/ControlRigEditor/Private/Sequencer/ControlRigParameterTrackEditor.cpp
--- a/Engine/Plugins/Animation/ControlRig/Source/ControlRigEditor/Private/Sequencer/ControlRigParameterTrackEditor.cpp
+++ b/Engine/Plugins/Animation/ControlRig/Source/ControlRigEditor/Private/Sequencer/ControlRigParameterTrackEditor.cpp
@@ -1779,23 +1779,28 @@ static void EvaluateThisControl(UMovieSceneControlRigParameterSection* Section,
 				{
 					if (ControlElement->Settings.ControlType == ERigControlType::Transform)
 					{
+						FVector EulerAngle(Value.GetValue().Rotation.Roll, Value.GetValue().Rotation.Pitch, Value.GetValue().Rotation.Yaw);
+						RigHierarchy->SetControlSpecifiedEulerAngle(ControlElement, EulerAngle);
 						ControlRig->SetControlValue<FRigControlValue::FTransform_Float>(ControlName, Value.GetValue().ToFTransform(), true, EControlRigSetKey::Never, bSetupUndo);
-						// @MikeZ here we want to call SetControlPreferredEulerAngles instead
-						RigHierarchy->SetControlPreferredRotator(ControlElement, Value.GetValue().Rotation);
 					}
 					else if (ControlElement->Settings.ControlType == ERigControlType::TransformNoScale)
 					{
 						FTransformNoScale NoScale = Value.GetValue().ToFTransform();
+						FVector EulerAngle(Value.GetValue().Rotation.Roll, Value.GetValue().Rotation.Pitch, Value.GetValue().Rotation.Yaw);
+						RigHierarchy->SetControlSpecifiedEulerAngle(ControlElement, EulerAngle);
 						ControlRig->SetControlValue<FRigControlValue::FTransformNoScale_Float>(ControlName, NoScale, true, EControlRigSetKey::Never, bSetupUndo);
-						// @MikeZ here we want to call SetControlPreferredEulerAngles instead
-						RigHierarchy->SetControlPreferredRotator(ControlElement, Value.GetValue().Rotation);
 					}
 					else if (ControlElement->Settings.ControlType == ERigControlType::EulerTransform)
 					{
 						const FEulerTransform& Euler = Value.GetValue();
-						ControlRig->SetControlValue<FRigControlValue::FEulerTransform_Float>(ControlName, Euler, true, EControlRigSetKey::Never, bSetupUndo);
-						// @MikeZ here we want to call SetControlPreferredEulerAngles instead
-						RigHierarchy->SetControlPreferredRotator(ControlElement, Euler.Rotation);
+						FVector EulerAngle(Euler.Rotation.Roll, Euler.Rotation.Pitch, Euler.Rotation.Yaw);
+						FQuat Quat = ControlRig->GetHierarchy()->GetControlQuaternion(ControlElement, EulerAngle);
+						ControlRig->GetHierarchy()->SetControlSpecifiedEulerAngle(ControlElement, EulerAngle);
+						FRotator UERotator(Quat);
+						FEulerTransform Transform = Euler;
+						Transform.Rotation = UERotator;
+						ControlRig->SetControlValue<FRigControlValue::FEulerTransform_Float>(ControlName, Transform, true, EControlRigSetKey::Never, bSetupUndo);
+
 					}
 				}
 				break;
@@ -3359,8 +3364,8 @@ void FControlRigParameterTrackEditor::GetControlRigKeys(
 			case ERigControlType::EulerTransform:
 			{
 				FVector Translation, Scale(1.0f, 1.0f, 1.0f);
-				FRotator Rotation = InControlRig->GetHierarchy()->GetControlPreferredRotator(ControlElement);
-
+				FVector Vector = InControlRig->GetHierarchy()->GetControlSpecifiedEulerAngle(ControlElement);
+				FRotator Rotation = FRotator(Vector.Y, Vector.Z, Vector.X);
 				if (ControlElement->Settings.ControlType == ERigControlType::TransformNoScale)
 				{
 					FTransformNoScale NoScale = ControlValue.Get<FRigControlValue::FTransformNoScale_Float>().ToTransform();
@@ -3379,7 +3384,6 @@ void FControlRigParameterTrackEditor::GetControlRigKeys(
 					Scale = Val.GetScale3D();
 				}
 
-				// switch values to constraint space?
 				if (bInConstraintSpace)
 				{
 					const uint32 ControlHash = UTransformableControlHandle::ComputeHash(InControlRig, ControlElement->GetName());
@@ -3387,10 +3391,23 @@ void FControlRigParameterTrackEditor::GetControlRigKeys(
 					if (Transform)
 					{
 						Translation = Transform->GetTranslation();
-						Rotation = Transform->GetRotation().Rotator();
+						if (InControlRig->GetHierarchy()->GetUsePreferredRotationOrder(ControlElement))
+						{
+							Rotation = ControlElement->PreferredEulerAngles.GetRotatorFromQuat(Transform->GetRotation());
+							FVector Angle = Rotation.Euler();
+							//need to wind rotators still
+							ControlElement->PreferredEulerAngles.SetAngles(Angle, false, ControlElement->PreferredEulerAngles.RotationOrder, true);
+							Angle = InControlRig->GetHierarchy()->GetControlSpecifiedEulerAngle(ControlElement);
+							Rotation = FRotator(Vector.Y, Vector.Z, Vector.X);
+						}
+						else
+						{
+							Rotation = Transform->GetRotation().Rotator();
+						}
 						Scale = Transform->GetScale3D();
 					}
 				}
+				
 					
 				FVector3f CurrentVector = (FVector3f)Translation;
 				bool bKeyX = bSetKey && EnumHasAnyFlags(ChannelsToKey, EControlRigContextChannelToKey::TranslationX);
@@ -4604,7 +4621,7 @@ bool FControlRigParameterTrackEditor::CollapseAllLayers(TSharedPtr<ISequencer>&S
 						{
 							for (TPair<FName, TArray<FTransform>>& TrailControlTransform : ControlLocalTransforms)
 							{
-								ControlRig->SetControlLocalTransform(TrailControlTransform.Key, TrailControlTransform.Value[Index], false, Context, false);
+								ControlRig->SetControlLocalTransform(TrailControlTransform.Key, TrailControlTransform.Value[Index], false, Context, false /*undo*/, true/* bFixEulerFlips*/);
 							}
 						}
 						ControlRig->Evaluate_AnyThread();
diff --git a/Engine/Plugins/Animation/ControlRig/Source/ControlRigEditor/Private/Sequencer/ControlRigSpaceChannelEditors.cpp b/Engine/Plugins/Animation/ControlRig/Source/ControlRigEditor/Private/Sequencer/ControlRigSpaceChannelEditors.cpp
--- a/Engine/Plugins/Animation/ControlRig/Source/ControlRigEditor/Private/Sequencer/ControlRigSpaceChannelEditors.cpp
+++ b/Engine/Plugins/Animation/ControlRig/Source/ControlRigEditor/Private/Sequencer/ControlRigSpaceChannelEditors.cpp
@@ -453,7 +453,7 @@ FKeyHandle FControlRigSpaceChannelHelpers::SequencerKeyControlRigSpaceChannel(UC
 			Sequencer->GetEvaluationTemplate().EvaluateSynchronousBlocking(SceneContext, *Sequencer);
 
 			Context.LocalTime = TickResolution.AsSeconds(FFrameTime(Time - 1));
-			ControlRig->SetControlGlobalTransform(ControlKey.Name, ControlWorldTransforms[0], true, Context);
+			ControlRig->SetControlGlobalTransform(ControlKey.Name, ControlWorldTransforms[0], true, Context, false /*undo*/, false /*bPrintPython*/, true/* bFixEulerFlips*/);
 
 			//need to do this after eval
 			FChannelMapInfo* pChannelIndex = nullptr;
@@ -484,7 +484,7 @@ FKeyHandle FControlRigSpaceChannelHelpers::SequencerKeyControlRigSpaceChannel(UC
 
 			ControlRig->Evaluate_AnyThread();
 			Context.LocalTime = TickResolution.AsSeconds(FFrameTime(Frame));
-			ControlRig->SetControlGlobalTransform(ControlKey.Name, ControlWorldTransforms[FramesIndex], true, Context);
+			ControlRig->SetControlGlobalTransform(ControlKey.Name, ControlWorldTransforms[FramesIndex], true, Context, false /*undo*/, false /*bPrintPython*/, true/* bFixEulerFlips*/);
 
 			//need to do this after eval
 			FChannelMapInfo* pChannelIndex = nullptr;
@@ -709,7 +709,7 @@ void  FControlRigSpaceChannelHelpers::SequencerSpaceChannelKeyDeleted(UControlRi
 			Context.LocalTime = TickResolution.AsSeconds(FFrameTime(Frame));
 			//make sure we only add keys to those that exist since they may be getting deleted also.
 			Context.KeyMask = (uint32)GetCurrentTransformKeysAtThisTime(ControlRig, ControlKey.Name, SectionToKey, Frame);
-			ControlRig->SetControlGlobalTransform(ControlKey.Name, ControlWorldTransforms[FramesIndex++], true, Context);
+			ControlRig->SetControlGlobalTransform(ControlKey.Name, ControlWorldTransforms[FramesIndex++], true, Context, false /*undo*/, false /*bPrintPython*/, true/* bFixEulerFlips*/);
 		}
 		//now delete any extra TimeOfDelete -1
 		FControlRigSpaceChannelHelpers::DeleteTransformKeysAtThisTime(ControlRig, SectionToKey, ControlName, TimeOfDeletion - 1);
@@ -1105,7 +1105,7 @@ void FControlRigSpaceChannelHelpers::SequencerBakeControlInSpace(UControlRig* Co
 				{
 					Context.KeyMask = (uint32)EControlRigContextChannelToKey::AllTransform;
 				}
-				ControlRig->SetControlGlobalTransform(ControlKey.Name, GlobalTransform, true, Context);
+				ControlRig->SetControlGlobalTransform(ControlKey.Name, GlobalTransform, true, Context, false /*undo*/, false /*bPrintPython*/, true/* bFixEulerFlips*/);
 			}
 
 			//if end compensated set the space that was active previously and set the compensated global value
@@ -1140,7 +1140,7 @@ void FControlRigSpaceChannelHelpers::SequencerBakeControlInSpace(UControlRig* Co
 				const FTransform GlobalTransform = ControlWorldTransforms[Frames.Num() - 1];
 				Context.LocalTime = TickResolution.AsSeconds(FFrameTime(EndFrame));
 				Context.KeyMask = (uint32)EControlRigContextChannelToKey::AllTransform;
-				ControlRig->SetControlGlobalTransform(ControlKey.Name, GlobalTransform, true, Context);
+				ControlRig->SetControlGlobalTransform(ControlKey.Name, GlobalTransform, true, Context, false /*undo*/, false /*bPrintPython*/, true/* bFixEulerFlips*/);
 			}
 			//Fix tangents at removed space switches if needed
 			for (TPair<FFrameNumber, TArray<FMovieSceneTangentData>>& StoredTangent : StoredTangents)
@@ -1346,7 +1346,7 @@ void FControlRigSpaceChannelHelpers::CompensateIfNeeded(
 							//now set time -1 frame value
 							ControlRig->Evaluate_AnyThread();
 							KeyframeContext.LocalTime = TickResolution.AsSeconds(FFrameTime(Time - 1));
-							ControlRig->SetControlGlobalTransform(Control->GetName(), ControlWorldTransforms[0], true, KeyframeContext);
+							ControlRig->SetControlGlobalTransform(Control->GetName(), ControlWorldTransforms[0], true, KeyframeContext, false /*undo*/, false /*bPrintPython*/, true/* bFixEulerFlips*/);
 							
 							bDidIt = true;
 						}
diff --git a/Engine/Plugins/Animation/ControlRig/Source/ControlRigEditor/Private/Tools/ControlRigSnapper.cpp b/Engine/Plugins/Animation/ControlRig/Source/ControlRigEditor/Private/Tools/ControlRigSnapper.cpp
--- a/Engine/Plugins/Animation/ControlRig/Source/ControlRigEditor/Private/Tools/ControlRigSnapper.cpp
+++ b/Engine/Plugins/Animation/ControlRig/Source/ControlRigEditor/Private/Tools/ControlRigSnapper.cpp
@@ -669,7 +669,7 @@ bool FControlRigSnapper::SnapIt(FFrameNumber StartFrame, FFrameNumber EndFrame,c
 							const FFrameNumber& FrameNumber = Frames[Index];
 							Context.LocalTime = TickResolution.AsSeconds(FFrameTime(FrameNumber));
 							FTransform GlobalTransform = WorldTransformToSnap[Index].GetRelativeTransform(ControlRigParentWorldTransforms[Index]);
-							ControlRig->SetControlGlobalTransform(Name, GlobalTransform, true, Context);
+							ControlRig->SetControlGlobalTransform(Name, GlobalTransform, true, Context, false /*undo*/, false /*bPrintPython*/, true/* bFixEulerFlips*/);
 						}
 					}
 
diff --git a/Engine/Plugins/Editor/SequencerAnimTools/Source/SequencerAnimTools/Private/Sequencer/MovieSceneTransformTrail.cpp b/Engine/Plugins/Editor/SequencerAnimTools/Source/SequencerAnimTools/Private/Sequencer/MovieSceneTransformTrail.cpp
--- a/Engine/Plugins/Editor/SequencerAnimTools/Source/SequencerAnimTools/Private/Sequencer/MovieSceneTransformTrail.cpp
+++ b/Engine/Plugins/Editor/SequencerAnimTools/Source/SequencerAnimTools/Private/Sequencer/MovieSceneTransformTrail.cpp
@@ -757,7 +757,7 @@ bool FMovieSceneControlRigTransformTrail::ApplyDelta(const FVector& Pos, const F
 				NewTransform = NewTransform.GetRelativeTransform(KeyInfo->ParentTransform);
 				GetSequencer()->GetEvaluationTemplate().EvaluateSynchronousBlocking(MovieSceneContext, *Player);
 				ControlRig->Evaluate_AnyThread();
-				ControlRig->SetControlGlobalTransform(ControlName, NewTransform, true, Context, false);
+				ControlRig->SetControlGlobalTransform(ControlName, NewTransform, true, Context, false /*undo*/, false /*bPrintPython*/, true/* bFixEulerFlips*/);
 			}
 		};
 
@@ -853,7 +853,7 @@ bool FMovieSceneControlRigTransformTrail::HandleAltClick(FEditorViewportClient*
 	GetSequencer()->GetEvaluationTemplate().EvaluateSynchronousBlocking(MovieSceneContext, *Player);
 	ControlRig->Evaluate_AnyThread();
 	FTransform NewTransform(ControlRig->GetControlGlobalTransform(ControlName));
-	ControlRig->SetControlGlobalTransform(ControlName, NewTransform, true, Context, false);
+	ControlRig->SetControlGlobalTransform(ControlName, NewTransform, true, Context, false /*undo*/, false /*bPrintPython*/, true/* bFixEulerFlips*/);
 
 	//eval back at current time
 	FFrameTime StartTime = GetSequencer()->GetLocalTime().Time;
diff --git a/Engine/Source/Developer/AnimationWidgets/Public/SAdvancedRotationInputBox.h b/Engine/Source/Developer/AnimationWidgets/Public/SAdvancedRotationInputBox.h
--- a/Engine/Source/Developer/AnimationWidgets/Public/SAdvancedRotationInputBox.h
+++ b/Engine/Source/Developer/AnimationWidgets/Public/SAdvancedRotationInputBox.h
@@ -523,7 +523,7 @@ private:
 						case ESlateRotationRepresentation::EulerZYX:
 						{
 								const int32 RotationOrder = int32(Repr) - int32(ESlateRotationRepresentation::EulerXYZ);
-								const FVector Euler = AnimationCore::EulerFromQuat(Quat.GetValue(), EEulerRotationOrder(RotationOrder));
+								const FVector Euler = AnimationCore::EulerFromQuat(Quat.GetValue(), EEulerRotationOrder(RotationOrder), true);
 								switch(InSubComponent)
 								{
 									case ESlateTransformSubComponent::X:
@@ -685,7 +685,7 @@ private:
 				case ESlateRotationRepresentation::EulerZYX:
 				{
 						const int32 RotationOrder = int32(Repr) - int32(ESlateRotationRepresentation::EulerXYZ);
-						FVector Euler = AnimationCore::EulerFromQuat(Quat.GetValue(), EEulerRotationOrder(RotationOrder));
+						FVector Euler = AnimationCore::EulerFromQuat(Quat.GetValue(), EEulerRotationOrder(RotationOrder), true);
 						switch(InSubComponent)
 						{
 							case ESlateTransformSubComponent::X:
@@ -708,7 +708,7 @@ private:
 								break;
 							}
 						}
-						Quat = AnimationCore::QuatFromEuler(Euler, EEulerRotationOrder(RotationOrder));
+						Quat = AnimationCore::QuatFromEuler(Euler, EEulerRotationOrder(RotationOrder), true);
 						Rot = Quat->Rotator();
 						break;
 				}
diff --git a/Engine/Source/Developer/AnimationWidgets/Public/SAdvancedTransformInputBox.h b/Engine/Source/Developer/AnimationWidgets/Public/SAdvancedTransformInputBox.h
--- a/Engine/Source/Developer/AnimationWidgets/Public/SAdvancedTransformInputBox.h
+++ b/Engine/Source/Developer/AnimationWidgets/Public/SAdvancedTransformInputBox.h
@@ -308,6 +308,7 @@ public:
 		FOnEndSliderMovement OnEndSliderMovement = InArgs._OnEndSliderMovement;
 		FOnGetToggleChecked OnGetToggleChecked = InArgs._OnGetToggleChecked;
 		FOnToggleChanged OnToggleChanged = InArgs._OnToggleChanged;
+		FOnRotationRepresentationChanged OnRotationRepresentationChanged = InArgs._OnRotationRepresentationChanged;
 		const bool UseQuaternionForRotation = InArgs._UseQuaternionForRotation;
 
 		auto OnGetLocation = [Transform, OnGetNumericValue]() -> TOptional<FVector>
@@ -1380,7 +1381,7 @@ public:
 					case ESlateRotationRepresentation::EulerZYX:
 					{
 						const int32 RotationOrder = int32(Representation) - int32(ESlateRotationRepresentation::EulerXYZ);
-						const FVector Euler = AnimationCore::EulerFromQuat(Transform.GetRotation(), EEulerRotationOrder(RotationOrder));
+						const FVector Euler = AnimationCore::EulerFromQuat(Transform.GetRotation(), EEulerRotationOrder(RotationOrder), true);
 						switch(SubComponent)
 						{
 							case ESlateTransformSubComponent::X:
@@ -1529,7 +1530,7 @@ public:
 								break;
 							}
 						}
-						Transform.SetRotation(AnimationCore::QuatFromEuler(Euler, EEulerRotationOrder(RotationOrder)));
+						Transform.SetRotation(AnimationCore::QuatFromEuler(Euler, EEulerRotationOrder(RotationOrder), true));
 						break;
 					}
 					case ESlateRotationRepresentation::Rotator:
diff --git a/Engine/Source/Runtime/AnimationCore/Private/AnimationCoreLibrary.cpp b/Engine/Source/Runtime/AnimationCore/Private/AnimationCoreLibrary.cpp
--- a/Engine/Source/Runtime/AnimationCore/Private/AnimationCoreLibrary.cpp
+++ b/Engine/Source/Runtime/AnimationCore/Private/AnimationCoreLibrary.cpp
@@ -188,12 +188,18 @@ FTransform SolveConstraints(const FTransform& CurrentTransform, const FTransform
 	return BlendedLocalTransform * CurrentParentTransform;
 }
 
-FQuat QuatFromEuler(const FVector& XYZAnglesInDegrees, EEulerRotationOrder RotationOrder)
+FQuat QuatFromEuler(const FVector& XYZAnglesInDegrees, EEulerRotationOrder RotationOrder, bool bUseUEHandyness)
 {
 	double X = FMath::DegreesToRadians(XYZAnglesInDegrees.X);
 	double Y = FMath::DegreesToRadians(XYZAnglesInDegrees.Y);
 	double Z = FMath::DegreesToRadians(XYZAnglesInDegrees.Z);
 
+	if (bUseUEHandyness)
+	{
+		X = -X;
+		Y = -Y;
+	}
+
 	double CosX = FMath::Cos( X * 0.5 );
 	double CosY = FMath::Cos( Y * 0.5 );
 	double CosZ = FMath::Cos( Z * 0.5 );
@@ -254,7 +260,7 @@ FQuat QuatFromEuler(const FVector& XYZAnglesInDegrees, EEulerRotationOrder Rotat
 	return FQuat::Identity;
 }
 
-FVector EulerFromQuat(const FQuat& Rotation, EEulerRotationOrder RotationOrder)
+FVector EulerFromQuat(const FQuat& Rotation, EEulerRotationOrder RotationOrder, bool bUseUEHandyness)
 {
 	double X = Rotation.X;
 	double Y = Rotation.Y;
@@ -382,19 +388,25 @@ FVector EulerFromQuat(const FQuat& Rotation, EEulerRotationOrder RotationOrder)
 			Result.Z = 0.0;
 		}
 	}
+	if (bUseUEHandyness)
+	{
+		Result.X = -Result.X;
+		Result.Y = -Result.Y;
+	}
 
 	return Result * 180.0 / DOUBLE_PI;
 }
 
-FVector ChangeEulerRotationOrder(const FVector& XYZAnglesInDegrees, EEulerRotationOrder SourceRotationOrder, EEulerRotationOrder TargetRotationOrder)
+FVector ChangeEulerRotationOrder(const FVector& XYZAnglesInDegrees, EEulerRotationOrder SourceRotationOrder, EEulerRotationOrder TargetRotationOrder, bool bUseUEHandyness)
 {
 	if(SourceRotationOrder == TargetRotationOrder)
 	{
 		return XYZAnglesInDegrees;
 	}
 	
-	const FQuat Quaternion = QuatFromEuler(XYZAnglesInDegrees, SourceRotationOrder);
-	return EulerFromQuat(Quaternion, TargetRotationOrder); 
+	const FQuat Quaternion = QuatFromEuler(XYZAnglesInDegrees, SourceRotationOrder, bUseUEHandyness);
+	return EulerFromQuat(Quaternion, TargetRotationOrder, bUseUEHandyness);
+
 }
 
 }
\ No newline at end of file
diff --git a/Engine/Source/Runtime/AnimationCore/Public/AnimationCoreLibrary.h b/Engine/Source/Runtime/AnimationCore/Public/AnimationCoreLibrary.h
--- a/Engine/Source/Runtime/AnimationCore/Public/AnimationCoreLibrary.h
+++ b/Engine/Source/Runtime/AnimationCore/Public/AnimationCoreLibrary.h
@@ -62,16 +62,18 @@ namespace AnimationCore
 
 	/**
 	 * Converts a euler rotation represented by a vector of rotations in degrees and a rotation order to a quaternion
+	 * bUseUEHandyness If True will use UE handyness with right handed X and Y angles left handed Z, if False all are left handed
 	 */
-	ANIMATIONCORE_API FQuat QuatFromEuler(const FVector& XYZAnglesInDegrees, EEulerRotationOrder RotationOrderr = EEulerRotationOrder::ZYX);
+	ANIMATIONCORE_API FQuat QuatFromEuler(const FVector& XYZAnglesInDegrees, EEulerRotationOrder RotationOrderr = EEulerRotationOrder::ZYX, bool bUseUEHandyness = false);
 
 	/**
-	 * Converts a quaternion to a euler rotation represented by a vector of rotations in degrees and a rotation order
-	 */
-	ANIMATIONCORE_API FVector EulerFromQuat(const FQuat& Rotation, EEulerRotationOrder RotationOrder = EEulerRotationOrder::ZYX);
+	* Converts a quaternion to a euler rotation represented by a vector of rotations in degrees and a rotation order
+	* bUseUEHandyness If True will use UE handyness with right handed X and Y angles left handed Z, if False all are left handed
+	*/
+	ANIMATIONCORE_API FVector EulerFromQuat(const FQuat& Rotation, EEulerRotationOrder RotationOrder = EEulerRotationOrder::ZYX, bool bUseUEHandyness = false);
 
 	/**
 	 * Converts a euler rotation from one rotation order to another
 	 */
-	ANIMATIONCORE_API FVector ChangeEulerRotationOrder(const FVector& XYZAnglesInDegrees, EEulerRotationOrder SourceRotationOrder, EEulerRotationOrder TargetRotationOrder);
+	ANIMATIONCORE_API FVector ChangeEulerRotationOrder(const FVector& XYZAnglesInDegrees, EEulerRotationOrder SourceRotationOrder, EEulerRotationOrder TargetRotationOrder, bool bUseUEHandyness = false);
 }
