Editor crashes using "Ctrl + Z" after adding Edit with FK Control Ring in the Level Sequence
#fix ensure that we do not transact the animation recording used to populate the FK curve data (as it will be marked pending kill afterwards)

diff --git a/Engine/Plugins/Animation/ControlRig/Source/ControlRigEditor/Private/Sequencer/ControlRigParameterTrackEditor.cpp b/Engine/Plugins/Animation/ControlRig/Source/ControlRigEditor/Private/Sequencer/ControlRigParameterTrackEditor.cpp
--- a/Engine/Plugins/Animation/ControlRig/Source/ControlRigEditor/Private/Sequencer/ControlRigParameterTrackEditor.cpp
+++ b/Engine/Plugins/Animation/ControlRig/Source/ControlRigEditor/Private/Sequencer/ControlRigParameterTrackEditor.cpp
@@ -1068,6 +1068,7 @@ void FControlRigParameterTrackEditor::BakeToControlRig(UClass* InClass, FGuid Ob
 			FMovieSceneSequenceTransform RootToLocalTransform = ParentSequencer->GetFocusedMovieSceneSequenceTransform();
 			UAnimSeqExportOption* AnimSeqExportOption = NewObject<UAnimSeqExportOption>(GetTransientPackage(), NAME_None);
 			UBakeToControlRigSettings* BakeSettings = GetMutableDefault<UBakeToControlRigSettings>();
+			AnimSeqExportOption->bTransactRecording = false;
 
 			TSharedPtr<SWindow> ParentWindow;
 			if (FModuleManager::Get().IsModuleLoaded("MainFrame"))
diff --git a/Engine/Source/Editor/MovieSceneTools/Private/MovieSceneToolHelpers.cpp b/Engine/Source/Editor/MovieSceneTools/Private/MovieSceneToolHelpers.cpp
--- a/Engine/Source/Editor/MovieSceneTools/Private/MovieSceneToolHelpers.cpp
+++ b/Engine/Source/Editor/MovieSceneTools/Private/MovieSceneToolHelpers.cpp
@@ -3954,6 +3954,7 @@ bool MovieSceneToolHelpers::ExportToAnimSequence(UAnimSequence* AnimSequence, UA
 		RecordingSettings.bRecordInWorldSpace = ExportOptions->bRecordInWorldSpace;
 		RecordingSettings.IncludeAnimationNames = ExportOptions->IncludeAnimationNames;
 		RecordingSettings.ExcludeAnimationNames = ExportOptions->ExcludeAnimationNames;
+		RecordingSettings.bTransactRecording = ExportOptions->bTransactRecording;
 		AnimationRecorder.Init(SkelMeshComp, AnimSequence, nullptr, RecordingSettings);	
 		});
 
