Rotation Orders: Fix BP/Python API so preferred angels are set when setting local euler transforms
#trivial

diff --git a/Engine/Plugins/Animation/ControlRig/Source/ControlRigEditor/Private/ControlRigSequencerEditorLibrary.cpp b/Engine/Plugins/Animation/ControlRig/Source/ControlRigEditor/Private/ControlRigSequencerEditorLibrary.cpp
--- a/Engine/Plugins/Animation/ControlRig/Source/ControlRigEditor/Private/ControlRigSequencerEditorLibrary.cpp
+++ b/Engine/Plugins/Animation/ControlRig/Source/ControlRigEditor/Private/ControlRigSequencerEditorLibrary.cpp
@@ -2276,6 +2276,11 @@ void UControlRigSequencerEditorLibrary::SetLocalControlRigEulerTransform(ULevelS
 		FRigControlModifiedContext Context;
 		Context.SetKey = bSetKey ? EControlRigSetKey::Always : EControlRigSetKey::DoNotCare;
 		Context.LocalTime = TickResolution.AsSeconds(FFrameTime(Frame));
+		if (FRigControlElement* ControlElement = ControlRig->FindControl(ControlName))
+		{
+			FVector EulerAngle(Value.Rotation.Roll, Value.Rotation.Pitch, Value.Rotation.Yaw);
+			ControlRig->GetHierarchy()->SetControlSpecifiedEulerAngle(ControlElement, EulerAngle);
+		}
 		ControlRig->SetControlValue<FRigControlValue::FEulerTransform_Float>(ControlName, Value, true, Context);
 	}
 }
@@ -2309,6 +2314,12 @@ void UControlRigSequencerEditorLibrary::SetLocalControlRigEulerTransforms(ULevel
 			}
 			FEulerTransform Value = Values[Index];
 			Context.LocalTime = TickResolution.AsSeconds(FFrameTime(Frame));
+
+			if (FRigControlElement* ControlElement = ControlRig->FindControl(ControlName))
+			{
+				FVector EulerAngle(Value.Rotation.Roll, Value.Rotation.Pitch, Value.Rotation.Yaw);
+				ControlRig->GetHierarchy()->SetControlSpecifiedEulerAngle(ControlElement, EulerAngle);
+			}
 			ControlRig->SetControlValue<FRigControlValue::FEulerTransform_Float>(ControlName, Value, true, Context);
 		}
 	}
