Rotation Order: Make sure rig's don't set rotators when getting transforms and when setting in edit mode make sure we flip eulers
#trivial

diff --git a/Engine/Plugins/Animation/ControlRig/Source/ControlRig/Private/Rigs/RigHierarchy.cpp b/Engine/Plugins/Animation/ControlRig/Source/ControlRig/Private/Rigs/RigHierarchy.cpp
--- a/Engine/Plugins/Animation/ControlRig/Source/ControlRig/Private/Rigs/RigHierarchy.cpp
+++ b/Engine/Plugins/Animation/ControlRig/Source/ControlRig/Private/Rigs/RigHierarchy.cpp
@@ -2749,6 +2749,7 @@ FTransform URigHierarchy::GetTransform(FRigTransformElement* InTransformElement,
 			{
 				const FTransform NewTransform = ComputeLocalControlValue(ControlElement, ControlElement->Pose.Get(OpposedType), GlobalType);
 				ControlElement->Pose.Set(InTransformType, NewTransform);
+				/** from mikez we do not want geting a pose to set these preferred angles
 				switch(ControlElement->Settings.ControlType)
 				{
 					case ERigControlType::Rotator:
@@ -2763,7 +2764,8 @@ FTransform URigHierarchy::GetTransform(FRigTransformElement* InTransformElement,
 					{
 						break;
 					}
-				}
+					
+				}*/
 			}
 			else if(FRigMultiParentElement* MultiParentElement = Cast<FRigMultiParentElement>(InTransformElement))
 			{
diff --git a/Engine/Plugins/Animation/ControlRig/Source/ControlRigEditor/Private/EditMode/ControlRigEditMode.cpp b/Engine/Plugins/Animation/ControlRig/Source/ControlRigEditor/Private/EditMode/ControlRigEditMode.cpp
--- a/Engine/Plugins/Animation/ControlRig/Source/ControlRigEditor/Private/EditMode/ControlRigEditMode.cpp
+++ b/Engine/Plugins/Animation/ControlRig/Source/ControlRigEditor/Private/EditMode/ControlRigEditMode.cpp
@@ -4042,7 +4042,7 @@ void FControlRigEditMode::MoveControlShape(AControlRigShapeActor* ShapeActor, co
 				ControlRig->InteractionType = InteractionType;
 				ControlRig->ElementsBeingInteracted.AddUnique(ShapeActor->GetElementKey());
 				
-				ControlRig->SetControlLocalTransform(ShapeActor->ControlName, CurrentLocalTransform);
+				ControlRig->SetControlLocalTransform(ShapeActor->ControlName, CurrentLocalTransform,true, FRigControlModifiedContext(), true, true);
 
 				FTransform CurrentTransform  = ControlRig->GetControlGlobalTransform(ShapeActor->ControlName);			// assumes it's attached to actor
 				CurrentTransform = ToWorldTransform * CurrentTransform;
