Constrains:Multiple Bake Crash
#prefilght 645aa0d4013f375158745ad7
#trivial

diff --git a/Engine/Plugins/Animation/ControlRig/Source/ControlRigEditor/Private/Tools/ConstraintBaker.cpp b/Engine/Plugins/Animation/ControlRig/Source/ControlRigEditor/Private/Tools/ConstraintBaker.cpp
--- a/Engine/Plugins/Animation/ControlRig/Source/ControlRigEditor/Private/Tools/ConstraintBaker.cpp
+++ b/Engine/Plugins/Animation/ControlRig/Source/ControlRigEditor/Private/Tools/ConstraintBaker.cpp
@@ -206,7 +206,11 @@ bool FConstraintBaker::BakeMultiple(
 	TOptional< FBakingAnimationKeySettings> Settings = InSettings;
 	for (UTickableTransformConstraint* Constraint : InConstraints)
 	{
-		FConstraintBaker::Bake(InWorld, Constraint, InSequencer, Settings, Frames);
+		if (IsValid(Constraint->ChildTRSHandle) && Constraint->ChildTRSHandle->IsValid()
+			&& IsValid(Constraint->ParentTRSHandle) && Constraint->ParentTRSHandle->IsValid())
+		{
+			FConstraintBaker::Bake(InWorld, Constraint, InSequencer, Settings, Frames);
+		}
 	}
 	return true;
 }
@@ -268,7 +272,12 @@ void FConstraintBaker::Bake(UWorld* InWorld,
 			}
 		}
 	}
-	
+	if (FramesToBake.Num() == 0 || ConstraintSections.ConstraintSection == nullptr ||
+		ConstraintSections.ActiveChannel == nullptr)
+	{
+		return;
+	}
+
 	FCompensationEvaluator Evaluator(InConstraint);
 	Evaluator.ComputeLocalTransformsForBaking(InWorld, InSequencer, FramesToBake);
 	TArray<FTransform> Transforms = Evaluator.ChildLocals;
