Control Rig: Fix sequencer unit tests failing

diff --git a/Engine/Plugins/Animation/ControlRig/Source/ControlRig/Private/ControlRig.cpp b/Engine/Plugins/Animation/ControlRig/Source/ControlRig/Private/ControlRig.cpp
--- a/Engine/Plugins/Animation/ControlRig/Source/ControlRig/Private/ControlRig.cpp
+++ b/Engine/Plugins/Animation/ControlRig/Source/ControlRig/Private/ControlRig.cpp
@@ -650,16 +650,23 @@ bool UControlRig::Execute(const FName& InEventName)
 
 			// disable selection notifications from the hierarchy
 			TGuardValue<bool> DisableSelectionNotifications(GetHierarchy()->GetController(true)->bSuspendSelectionNotifications, true);
-			{				
+			{
 				// save the current state of all pose elements to preserve user intention, since construction event can
 				// run in between forward events
 				// the saved pose is reapplied to the rig after construction event as the pose scope goes out of scope
 				TUniquePtr<FPoseScope> PoseScope;
-				if (!bConstructionModeEnabled && !bJustRanInit)
+				if (!bConstructionModeEnabled)
 				{
 					// only do this in non-construction mode because 
 					// when construction mode is enabled, the control values are cleared before reaching here (too late to save them)
-					PoseScope = MakeUnique<FPoseScope>(this, ERigElementType::ToResetAfterConstructionEvent, TArray<FRigElementKey>(), ERigTransformType::CurrentGlobal);
+					if (!bJustRanInit)
+					{
+						PoseScope = MakeUnique<FPoseScope>(this, ERigElementType::ToResetAfterConstructionEvent, TArray<FRigElementKey>(), ERigTransformType::CurrentGlobal);
+					}
+					else
+					{
+						PoseScope = MakeUnique<FPoseScope>(this, ERigElementType::ToResetAfterConstructionEvent);
+					}
 				}
 				
 				{
