RigVM: First steps towards rebasing ControlRigBlueprint

diff --git a/Engine/Plugins/Animation/ControlRig/Config/BaseControlRig.ini b/Engine/Plugins/Animation/ControlRig/Config/BaseControlRig.ini
--- a/Engine/Plugins/Animation/ControlRig/Config/BaseControlRig.ini
+++ b/Engine/Plugins/Animation/ControlRig/Config/BaseControlRig.ini
@@ -432,4 +432,5 @@
 +StructRedirects=(OldName="/Script/ControlRig.RigUnit_SecondsToFrames",NewName="/Script/RigVM.RigVMFunction_SecondsToFrames")
 
 ; moving editor side types to RigVMEditor module
-; +StructRedirects=(OldName="/Script/ControlRig.RigGraphDisplaySettings",NewName="/Script/RigVMDeveloper.RigVMEdGraphDisplaySettings")
++StructRedirects=(OldName="/Script/ControlRig.RigGraphDisplaySettings",NewName="/Script/RigVMDeveloper.RigVMEdGraphDisplaySettings")
++StructRedirects=(OldName="/Script/ControlRig.ControlRigPythonSettings",NewName="/Script/RigVMDeveloper.RigVMPythonSettings")
