RigVM: First steps towards rebasing ControlRigBlueprint

diff --git a/Engine/Plugins/Animation/ControlRig/Source/ControlRigDeveloper/Private/ControlRigBlueprint.cpp b/Engine/Plugins/Animation/ControlRig/Source/ControlRigDeveloper/Private/ControlRigBlueprint.cpp
--- a/Engine/Plugins/Animation/ControlRig/Source/ControlRigDeveloper/Private/ControlRigBlueprint.cpp
+++ b/Engine/Plugins/Animation/ControlRig/Source/ControlRigDeveloper/Private/ControlRigBlueprint.cpp
@@ -5,84 +5,39 @@
 #include "RigVMBlueprintGeneratedClass.h"
 #include "EdGraph/EdGraph.h"
 #include "EdGraphNode_Comment.h"
-#include "Modules/ModuleManager.h"
 #include "Engine/SkeletalMesh.h"
 #include "BlueprintActionDatabaseRegistrar.h"
 #include "ControlRig.h"
 #include "Graph/ControlRigGraph.h"
-#include "Graph/ControlRigGraphNode.h"
 #include "Graph/ControlRigGraphSchema.h"
 #include "UObject/ObjectSaveContext.h"
 #include "UObject/UObjectGlobals.h"
 #include "ControlRigObjectVersion.h"
-#include "ControlRigDeveloper.h"
-#include "Curves/CurveFloat.h"
 #include "BlueprintCompilationManager.h"
 #include "RigVMCompiler/RigVMCompiler.h"
 #include "RigVMCore/RigVMRegistry.h"
 #include "Units/Execution/RigUnit_BeginExecution.h"
 #include "Units/Hierarchy/RigUnit_SetBoneTransform.h"
-#include "Async/TaskGraphInterfaces.h"
-#include "Misc/CoreDelegates.h"
 #include "AssetRegistry/AssetRegistryModule.h"
 #include "RigVMPythonUtils.h"
 #include "RigVMTypeUtils.h"
-#include "Algo/Count.h"
-#include "Algo/Transform.h"
 #include "RigVMModel/Nodes/RigVMAggregateNode.h"
-#include "Units/ControlRigNodeWorkflow.h"
 #include "Rigs/RigControlHierarchy.h"
-#include "RigVMModel/Nodes/RigVMDispatchNode.h"
-#include "UObject/UE5MainStreamObjectVersion.h"
-#include "ControlRigSchema.h"
 
 #include UE_INLINE_GENERATED_CPP_BY_NAME(ControlRigBlueprint)
 
 #if WITH_EDITOR
 #include "IControlRigEditorModule.h"
-#include "Kismet2/KismetDebugUtilities.h"
 #include "Kismet2/WatchedPin.h"
 #include "Kismet2/BlueprintEditorUtils.h"
-#include "RigVMBlueprintUtils.h"
 #include "Settings/ControlRigSettings.h"
-#include "UnrealEdGlobals.h"
 #include "Editor/UnrealEdEngine.h"
 #include "Editor/Transactor.h"
 #include "CookOnTheSide/CookOnTheFlyServer.h"
-#include "Widgets/Notifications/SNotificationList.h"
-#include "Framework/Notifications/NotificationManager.h"
-#include "ScopedTransaction.h"
 #endif//WITH_EDITOR
 
 #define LOCTEXT_NAMESPACE "ControlRigBlueprint"
 
-TAutoConsoleVariable<bool> CVarEnablePreLoadFiltering(
-	TEXT("ControlRig.EnablePreLoadFiltering"),
-	true,
-	TEXT("When true the RigVMGraphs will be skipped during preload to speed up load times."));
-
-TAutoConsoleVariable<bool> CVarEnablePostLoadHashing(
-	TEXT("ControlRig.EnablePostLoadHashing"),
-	true,
-	TEXT("When true refreshing the RigVMGraphs will be skipped if the hash matches the serialized hash."));
-
-static TArray<UClass*> GetClassObjectsInPackage(UPackage* InPackage)
-{
-	TArray<UObject*> Objects;
-	GetObjectsWithOuter(InPackage, Objects, false);
-
-	TArray<UClass*> ClassObjects;
-	for (UObject* Object : Objects)
-	{
-		if (UClass* Class = Cast<UClass>(Object))
-		{
-			ClassObjects.Add(Class);
-		}
-	}
-
-	return ClassObjects;
-}
-
 FEdGraphPinType FControlRigPublicFunctionArg::GetPinType() const
 {
 	FRigVMExternalVariable Variable;
@@ -118,209 +73,62 @@ bool FControlRigPublicFunctionData::IsMutable() const
 }
 
 TArray<UControlRigBlueprint*> UControlRigBlueprint::sCurrentlyOpenedRigBlueprints;
-FSoftObjectPath UControlRigBlueprint::PreDuplicateAssetPath;
-FSoftObjectPath UControlRigBlueprint::PreDuplicateHostPath;
 
 UControlRigBlueprint::UControlRigBlueprint(const FObjectInitializer& ObjectInitializer)
+	: URigVMBlueprint(ObjectInitializer)
 {
-	bSuspendModelNotificationsForSelf = false;
-	bSuspendModelNotificationsForOthers = false;
-	bSuspendAllNotifications = false;
-	bSuspendPythonMessagesForRigVMClient = true;
-	bMarkBlueprintAsStructurallyModifiedPending = false;
-
 #if WITH_EDITORONLY_DATA
-	ReferencedObjectPathsStored = false;
 	GizmoLibrary_DEPRECATED = nullptr;
 	ShapeLibraries.Add(UControlRigSettings::Get()->DefaultShapeLibrary);
 #endif
 
-	bRecompileOnLoad = 0;
-	bAutoRecompileVM = true;
-	bVMRecompilationRequired = false;
-	bIsCompiling = false;
-	VMRecompilationBracket = 0;
-
-	RigVMClient.SetOuterClientHost(this, GET_MEMBER_NAME_CHECKED(UControlRigBlueprint, RigVMClient));
-	RigVMClient.SetSchemaClass(UControlRigSchema::StaticClass());
-	{
-		TGuardValue<bool> DisableClientNotifs(RigVMClient.bSuspendNotifications, true);
-		RigVMClient.GetOrCreateFunctionLibrary(false, &ObjectInitializer, false);
-		RigVMClient.AddModel(FRigVMClient::RigVMModelPrefix, false, &ObjectInitializer, false);
-	}
-	RigVMClient.SetExecuteContextStruct(FControlRigExecuteContext::StaticStruct());
-	
-	FunctionLibraryEdGraph = ObjectInitializer.CreateDefaultSubobject<UControlRigGraph>(this, TEXT("RigVMFunctionLibraryEdGraph"));
-	FunctionLibraryEdGraph->Schema = UControlRigGraphSchema::StaticClass();
-	FunctionLibraryEdGraph->bAllowRenaming = 0;
-	FunctionLibraryEdGraph->bEditable = 0;
-	FunctionLibraryEdGraph->bAllowDeletion = 0;
-	FunctionLibraryEdGraph->bIsFunctionDefinition = false;
-	FunctionLibraryEdGraph->ModelNodePath = RigVMClient.GetFunctionLibrary()->GetNodePath();
-	FunctionLibraryEdGraph->Initialize(this);
-
 	Validator = ObjectInitializer.CreateDefaultSubobject<UControlRigValidator>(this, TEXT("ControlRigValidator"));
-
 	DebugBoneRadius = 1.f;
 
-	bUpdatingExternalVariables = false;
-	
-	bDirtyDuringLoad = false;
-	bErrorsDuringCompilation = false;
-
-	SupportedEventNames.Reset();
 	bExposesAnimatableControls = false;
 
-	VMCompileSettings.ASTSettings.ReportDelegate.BindUObject(this, &UControlRigBlueprint::HandleReportFromCompiler);
-
-#if WITH_EDITOR
-	CompileLog.SetSourcePath(GetPathName());
-	CompileLog.bLogDetailedResults = false;
-	CompileLog.EventDisplayThresholdMs = false;
-#endif
-
 	Hierarchy = CreateDefaultSubobject<URigHierarchy>(TEXT("Hierarchy"));
 	URigHierarchyController* Controller = Hierarchy->GetController(true);
 	// give BP a chance to propagate hierarchy changes to available control rig instances
 	Controller->OnModified().AddUObject(this, &UControlRigBlueprint::HandleHierarchyModified);
-}
-
-UControlRigBlueprint::UControlRigBlueprint()
-{
-}
-
-void UControlRigBlueprint::InitializeModelIfRequired(bool bRecompileVM)
-{
-	DECLARE_SCOPE_HIERARCHICAL_COUNTER_FUNC()
 
-	if (RigVMClient.GetController(0) == nullptr)
+	if(GetClass() == UControlRigBlueprint::StaticClass())
 	{
-		check(RigVMClient.Num() == 1);
-		check(RigVMClient.GetFunctionLibrary());
-		
-		RigVMClient.GetOrCreateController(RigVMClient.GetDefaultModel());
-		RigVMClient.GetOrCreateController(RigVMClient.GetFunctionLibrary());
-
-		bool bRecompileRequired = false;
-		for (int32 i = 0; i < UbergraphPages.Num(); ++i)
-		{
-			if (UControlRigGraph* Graph = Cast<UControlRigGraph>(UbergraphPages[i]))
-			{
-				PopulateModelFromGraphForBackwardsCompatibility(Graph);
-				if (bRecompileVM)
-				{
-					bRecompileRequired = true;
-				}
-
-				Graph->Initialize(this);
-			}
-		}
-
-		if(bRecompileRequired)
-		{
-			RecompileVM();
-		}
-
-		FunctionLibraryEdGraph->Initialize(this);
+		CommonInitialization(ObjectInitializer);
 	}
 }
 
-URigVMBlueprintGeneratedClass* UControlRigBlueprint::GetControlRigBlueprintGeneratedClass() const
-{
-	URigVMBlueprintGeneratedClass* Result = Cast<URigVMBlueprintGeneratedClass>(*GeneratedClass);
-	return Result;
-}
-
-URigVMBlueprintGeneratedClass* UControlRigBlueprint::GetControlRigBlueprintSkeletonClass() const
-{
-	URigVMBlueprintGeneratedClass* Result = Cast<URigVMBlueprintGeneratedClass>(*SkeletonGeneratedClass);
-	return Result;
-}
-
-UClass* UControlRigBlueprint::GetBlueprintClass() const
+UControlRigBlueprint::UControlRigBlueprint()
 {
-	return URigVMBlueprintGeneratedClass::StaticClass();
 }
 
 UClass* UControlRigBlueprint::RegenerateClass(UClass* ClassToRegenerate, UObject* PreviousCDO)
 {
-	UClass* Result = nullptr;
-	{
-		TGuardValue<bool> NotificationGuard(bSuspendAllNotifications, true);
-		Result = Super::RegenerateClass(ClassToRegenerate, PreviousCDO);
-	}
+	UClass* Result = Super::RegenerateClass(ClassToRegenerate, PreviousCDO);
 	Hierarchy->CleanupInvalidCaches();
 	PropagateHierarchyFromBPToInstances();
 	return Result;
 }
 
-void UControlRigBlueprint::LoadModulesRequiredForCompilation() 
-{
-}
-
-bool UControlRigBlueprint::ExportGraphToText(UEdGraph* InEdGraph, FString& OutText)
-{
-	OutText.Empty();
-
-	if (URigVMGraph* RigGraph = GetModel(InEdGraph))
-	{
-		if (URigVMCollapseNode* CollapseNode = Cast<URigVMCollapseNode>(RigGraph->GetOuter()))
-		{
-			if (URigVMController* Controller = GetOrCreateController(CollapseNode->GetGraph()))
-			{
-				TArray<FName> NodeNamesToExport;
-				NodeNamesToExport.Add(CollapseNode->GetFName());
-				OutText = Controller->ExportNodesToText(NodeNamesToExport);
-			}
-		}
-	}
-
-	// always return true so that the default mechanism doesn't take over
-	return true;
-}
-
-bool UControlRigBlueprint::CanImportGraphFromText(const FString& InClipboardText)
-{
-	return GetTemplateController(true)->CanImportNodesFromText(InClipboardText);
-}
-
 bool UControlRigBlueprint::RequiresForceLoadMembers(UObject* InObject) const
 {
-	// only filter if the console variable is enabled
-	if(!CVarEnablePreLoadFiltering->GetBool())
-	{
-		return UBlueprint::RequiresForceLoadMembers(InObject);
-	}
-
 	// old assets don't support preload filtering
 	if (GetLinkerCustomVersion(FControlRigObjectVersion::GUID) < FControlRigObjectVersion::RemoveParameters)
 	{
 		return UBlueprint::RequiresForceLoadMembers(InObject);
 	}
 	
-	// we can stop traversing when hitting a URigVMNode
-	// except for collapse nodes - since they contain a graphs again
-	// and variable  nodes - since they are needed during preload by the BP compiler
-	if(InObject->IsA<URigVMNode>())
-	{
-		if(!InObject->IsA<URigVMCollapseNode>() &&
-			!InObject->IsA<URigVMVariableNode>())
-		{
-			return false;
-		}
-	}
-	return UBlueprint::RequiresForceLoadMembers(InObject);
+	return Super::RequiresForceLoadMembers(InObject);
 }
 
 void UControlRigBlueprint::PostEditChangeChainProperty(FPropertyChangedChainEvent& PropertyChangedEvent)
 {
 	Super::PostEditChangeChainProperty(PropertyChangedEvent);
-	PostEditChangeChainPropertyEvent.Broadcast(PropertyChangedEvent);
 
 	// Propagate shape libraries
 	if (PropertyChangedEvent.Property && PropertyChangedEvent.Property->GetFName() == GET_MEMBER_NAME_CHECKED(UControlRigBlueprint, ShapeLibraries))
 	{
-		URigVMBlueprintGeneratedClass* RigClass = GetControlRigBlueprintGeneratedClass();
+		URigVMBlueprintGeneratedClass* RigClass = GetRigVMBlueprintGeneratedClass();
 		UControlRig* CDO = Cast<UControlRig>(RigClass->GetDefaultObject(false /* create if needed */));
 
 		TArray<UObject*> ArchetypeInstances;
@@ -338,5159 +146,1064 @@ void UControlRigBlueprint::PostEditChangeChainProperty(FPropertyChangedChainEven
 	}
 }
 
-void UControlRigBlueprint::PostRename(UObject* OldOuter, const FName OldName)
+USkeletalMesh* UControlRigBlueprint::GetPreviewMesh() const
 {
-	Super::PostRename(OldOuter, OldName);
-
-	// Whenever the asset is renamed/moved, generated classes parented to the old package
-	// are not moved to the new package automatically (see FAssetRenameManager), so we
-	// have to manually perform the move/rename, to avoid invalid reference to the old package
-	
-	// Note: while asset duplication doesn't duplicate the classes either, it is not a problem there
-	// because we always recompile in post duplicate.
-	TArray<UClass*> ClassObjects = GetClassObjectsInPackage(OldOuter->GetPackage());
+	DECLARE_SCOPE_HIERARCHICAL_COUNTER_FUNC()
 
-	for (UClass* ClassObject : ClassObjects)
+#if WITH_EDITORONLY_DATA
+	if (!PreviewSkeletalMesh.IsValid())
 	{
-		if (URigVMMemoryStorageGeneratorClass* MemoryClass = Cast<URigVMMemoryStorageGeneratorClass>(ClassObject))
-		{
-			MemoryClass->Rename(nullptr, GetPackage(), REN_ForceNoResetLoaders | REN_DoNotDirty | REN_DontCreateRedirectors | REN_NonTransactional);
-		}
+		(void)PreviewSkeletalMesh.LoadSynchronous();
 	}
+
+	return PreviewSkeletalMesh.Get();
+#else
+	return nullptr;
+#endif
 }
 
-void UControlRigBlueprint::GetPreloadDependencies(TArray<UObject*>& OutDeps)
+void UControlRigBlueprint::SetPreviewMesh(USkeletalMesh* PreviewMesh, bool bMarkAsDirty/*=true*/)
 {
-	Super::GetPreloadDependencies(OutDeps);
-
-	TArray<UClass*> ClassObjects = GetClassObjectsInPackage(GetPackage());
-
-	for (UClass* ClassObject : ClassObjects)
+#if WITH_EDITORONLY_DATA
+	if(bMarkAsDirty)
 	{
-		if (URigVMMemoryStorageGeneratorClass* MemoryClass = Cast<URigVMMemoryStorageGeneratorClass>(ClassObject))
-		{
-			OutDeps.Add(MemoryClass);
-		}
+		Modify();
 	}
-}
-
-FRigVMClient* UControlRigBlueprint::GetRigVMClient()
-{
-	return &RigVMClient;
-}
 
-const FRigVMClient* UControlRigBlueprint::GetRigVMClient() const
-{
-	return &RigVMClient;
-}
-
-IRigVMGraphFunctionHost* UControlRigBlueprint::GetRigVMGraphFunctionHost() 
-{
-	return GetControlRigBlueprintGeneratedClass();
+	PreviewSkeletalMesh = PreviewMesh;
+#endif
 }
 
-const IRigVMGraphFunctionHost* UControlRigBlueprint::GetRigVMGraphFunctionHost() const 
+void UControlRigBlueprint::Serialize(FArchive& Ar)
 {
-	return GetControlRigBlueprintGeneratedClass();
-}
+	RigVMClient.SetOuterClientHost(this, GET_MEMBER_NAME_CHECKED(UControlRigBlueprint, RigVMClient));
+	
+	Super::Serialize(Ar);
 
-UObject* UControlRigBlueprint::GetEditorObjectForRigVMGraph(URigVMGraph* InVMGraph) const
-{
-	if(InVMGraph)
+	if(Ar.IsObjectReferenceCollector())
 	{
-		if(InVMGraph->GetOutermost() != GetOutermost())
-		{
-			return nullptr;
-		}
-
-		if(InVMGraph->IsA<URigVMFunctionLibrary>())
-		{
-			return FunctionLibraryEdGraph;
-		}
+		Ar.UsingCustomVersion(FControlRigObjectVersion::GUID);
 
-		TArray<UEdGraph*> EdGraphs;
-		GetAllGraphs(EdGraphs);
 
-		bool bIsFunctionDefinition = false;
-		if (URigVMLibraryNode* LibraryNode = Cast<URigVMLibraryNode>(InVMGraph->GetOuter()))
+#if WITH_EDITORONLY_DATA
+		if (Ar.IsCooking() && ReferencedObjectPathsStored)
 		{
-			bIsFunctionDefinition = LibraryNode->GetGraph()->IsA<URigVMFunctionLibrary>();
+			for (FSoftObjectPath ObjectPath : ReferencedObjectPaths)
+			{
+				ObjectPath.Serialize(Ar);
+			}
 		}
-
-		for (UEdGraph* EdGraph : EdGraphs)
+		else
+#endif
 		{
-			if (UControlRigGraph* RigGraph = Cast<UControlRigGraph>(EdGraph))
+			TArray<IRigVMGraphFunctionHost*> ReferencedFunctionHosts = GetReferencedFunctionHosts(false);
+
+			for(IRigVMGraphFunctionHost* ReferencedFunctionHost : ReferencedFunctionHosts)
 			{
-				if (RigGraph->bIsFunctionDefinition != bIsFunctionDefinition)
+				if (URigVMBlueprintGeneratedClass* BPGeneratedClass = Cast<URigVMBlueprintGeneratedClass>(ReferencedFunctionHost))
 				{
-					continue;
+					Ar << BPGeneratedClass;
 				}
+			}
 
-				if ((RigGraph->ModelNodePath == InVMGraph->GetNodePath()) ||
-					(RigGraph->ModelNodePath.IsEmpty() && (RigVMClient.GetDefaultModel() == InVMGraph)))
+			for(const TSoftObjectPtr<UControlRigShapeLibrary>& ShapeLibraryPtr : ShapeLibraries)
+			{
+				if(ShapeLibraryPtr.IsValid())
 				{
-					return RigGraph;
+					UControlRigShapeLibrary* ShapeLibrary = ShapeLibraryPtr.Get();
+					Ar << ShapeLibrary;
 				}
 			}
 		}
 	}
-	
-	return nullptr;
-}
 
-void UControlRigBlueprint::HandleRigVMGraphAdded(const FRigVMClient* InClient, const FString& InNodePath)
-{
-	if(URigVMGraph* Model = InClient->GetModel(InNodePath))
+	if(Ar.IsLoading())
 	{
-		if(!HasAnyFlags(RF_ClassDefaultObject | RF_NeedInitialization | RF_NeedLoad | RF_NeedPostLoad) &&
-			GetOuter() != GetTransientPackage())
-		{
-			CreateEdGraph(Model, true);
-			RecompileVM();
-		}
-
-#if WITH_EDITOR
-		if(!bSuspendPythonMessagesForRigVMClient)
+		if(Model_DEPRECATED || FunctionLibrary_DEPRECATED)
 		{
-			const FString BlueprintName = InClient->GetSchema()->GetSanitizedName(GetName(), true, false);
-			RigVMPythonUtils::Print(BlueprintName, 
-				FString::Printf(TEXT("blueprint.add_model('%s')"),
-					*Model->GetName()));
+			TGuardValue<bool> DisableClientNotifs(RigVMClient.bSuspendNotifications, true);
+			RigVMClient.SetFromDeprecatedData(Model_DEPRECATED, FunctionLibrary_DEPRECATED);
 		}
-#endif
 	}
 }
 
-void UControlRigBlueprint::HandleRigVMGraphRemoved(const FRigVMClient* InClient, const FString& InNodePath)
+void UControlRigBlueprint::PreSave(FObjectPreSaveContext ObjectSaveContext)
 {
-	if(URigVMGraph* Model = InClient->GetModel(InNodePath))
-	{
-		RemoveEdGraph(Model);
-		RecompileVM();
+	Super::PreSave(ObjectSaveContext);
 
-#if WITH_EDITOR
-		if(!bSuspendPythonMessagesForRigVMClient)
+	bExposesAnimatableControls = false;
+	Hierarchy->ForEach<FRigControlElement>([this](FRigControlElement* ControlElement) -> bool
+    {
+		if (Hierarchy->IsAnimatable(ControlElement))
 		{
-			const FString BlueprintName = InClient->GetSchema()->GetSanitizedName(GetName(), true, false);
-			RigVMPythonUtils::Print(BlueprintName, 
-				FString::Printf(TEXT("blueprint.remove_model('%s')"),
-					*Model->GetName()));
+			bExposesAnimatableControls = true;
+			return false;
 		}
-#endif
-	}
+		return true;
+	});
 }
 
-void UControlRigBlueprint::HandleRigVMGraphRenamed(const FRigVMClient* InClient, const FString& InOldNodePath, const FString& InNewNodePath)
+void UControlRigBlueprint::PostLoad()
 {
-	if(URigVMGraph* Model = InClient->GetModel(InNewNodePath))
-	{
-		TArray<UEdGraph*> EdGraphs;
-		GetAllGraphs(EdGraphs);
+	Super::PostLoad();
 
-		for (UEdGraph* EdGraph : EdGraphs)
+	{
+#if WITH_EDITOR
+		
+		// correct the offset transforms
+		if (GetLinkerCustomVersion(FControlRigObjectVersion::GUID) < FControlRigObjectVersion::ControlOffsetTransform)
 		{
-			if (UControlRigGraph* RigGraph = Cast<UControlRigGraph>(EdGraph))
+			HierarchyContainer_DEPRECATED.ControlHierarchy.PostLoad();
+			if (HierarchyContainer_DEPRECATED.ControlHierarchy.Num() > 0)
 			{
-				RigGraph->HandleRigVMGraphRenamed(InOldNodePath, InNewNodePath);
+				MarkDirtyDuringLoad();
 			}
-		}
-
-		FBlueprintEditorUtils::MarkBlueprintAsStructurallyModified(this);
-	}
-}
-
-void UControlRigBlueprint::HandleConfigureRigVMController(const FRigVMClient* InClient,
-	URigVMController* InControllerToConfigure)
-{
-	InControllerToConfigure->OnModified().AddUObject(this, &UControlRigBlueprint::HandleModifiedEvent);
-
-	TWeakObjectPtr<UControlRigBlueprint> WeakThis(this);
 
-	// this delegate is used by the controller to determine variable validity
-	// during a bind process. the controller itself doesn't own the variables,
-	// so we need a delegate to request them from the owning blueprint
-	InControllerToConfigure->GetExternalVariablesDelegate.BindLambda([](URigVMGraph* InGraph) -> TArray<FRigVMExternalVariable> {
-
-		if (InGraph)
-		{
-			if(UControlRigBlueprint* Blueprint = InGraph->GetTypedOuter<UControlRigBlueprint>())
+			for (FRigControl& Control : HierarchyContainer_DEPRECATED.ControlHierarchy)
 			{
-				if (URigVMBlueprintGeneratedClass* RigClass = Blueprint->GetControlRigBlueprintGeneratedClass())
+				const FTransform PreviousOffsetTransform = Control.GetTransformFromValue(ERigControlValueType::Initial);
+				Control.OffsetTransform = PreviousOffsetTransform;
+				Control.InitialValue = Control.Value;
+
+				if (Control.ControlType == ERigControlType::Transform)
+				{
+					Control.InitialValue = FRigControlValue::Make<FTransform>(FTransform::Identity);
+				}
+				else if (Control.ControlType == ERigControlType::TransformNoScale)
+				{
+					Control.InitialValue = FRigControlValue::Make<FTransformNoScale>(FTransformNoScale::Identity);
+				}
+				else if (Control.ControlType == ERigControlType::EulerTransform)
 				{
-                    if (UControlRig* CDO = Cast<UControlRig>(RigClass->GetDefaultObject(true /* create if needed */)))
-                    {
-                        return CDO->GetExternalVariablesImpl(true /* rely on variables within blueprint */);
-                    }
-                }
+					Control.InitialValue = FRigControlValue::Make<FEulerTransform>(FEulerTransform::Identity);
+				}
 			}
 		}
-		return TArray<FRigVMExternalVariable>();
-
-	});
 
+		// convert the hierarchy from V1 to V2
+		if (GetLinkerCustomVersion(FControlRigObjectVersion::GUID) < FControlRigObjectVersion::RigHierarchyV2)
+		{
+			Modify();
+			
+			TGuardValue<bool> SuspendNotifGuard(Hierarchy->GetSuspendNotificationsFlag(), true);
+			
+			Hierarchy->Reset();
+			GetHierarchyController()->ImportFromHierarchyContainer(HierarchyContainer_DEPRECATED, false);
+		}
 
-	// this delegate is used by the controller to retrieve the current bytecode of the VM
-	InControllerToConfigure->GetCurrentByteCodeDelegate.BindLambda([WeakThis]() -> const FRigVMByteCode* {
-
-		if (WeakThis.IsValid())
+		// perform backwards compat value upgrades
+		TArray<URigVMGraph*> GraphsToValidate = GetAllModels();
+		for (int32 GraphIndex = 0; GraphIndex < GraphsToValidate.Num(); GraphIndex++)
 		{
-			if (URigVMBlueprintGeneratedClass* RigClass = WeakThis->GetControlRigBlueprintGeneratedClass())
+			URigVMGraph* GraphToValidate = GraphsToValidate[GraphIndex];
+			if(GraphToValidate == nullptr)
+			{
+				continue;
+			}
+
+			for(URigVMNode* Node : GraphToValidate->GetNodes())
 			{
-				if (UControlRig* CDO = Cast<UControlRig>(RigClass->GetDefaultObject(false)))
+				TArray<URigVMPin*> Pins = Node->GetAllPinsRecursively();
+				for(URigVMPin* Pin : Pins)
 				{
-					if (CDO->VM)
+					if(Pin->GetCPPTypeObject() == StaticEnum<ERigElementType>())
 					{
-						return &CDO->VM->GetByteCode();
+						if(Pin->GetDefaultValue() == TEXT("Space"))
+						{
+							if(URigVMController* Controller = GetController(GraphToValidate))
+							{
+								FRigVMControllerNotifGuard NotifGuard(Controller, true);
+								Controller->SetPinDefaultValue(Pin->GetPinPath(), TEXT("Null"), false, false, false);
+							}
+						}
 					}
 				}
 			}
 		}
-		return nullptr;
-
-	});
-
-#if WITH_EDITOR
 
-	// this sets up three delegates:
-	// a) get external variables (mapped to Controller->GetExternalVariables)
-	// b) bind pin to variable (mapped to Controller->BindPinToVariable)
-	// c) create external variable (mapped to the passed in tfunction)
-	// the last one is defined within the blueprint since the controller
-	// doesn't own the variables and can't create one itself.
-	InControllerToConfigure->SetupDefaultUnitNodeDelegates(TDelegate<FName(FRigVMExternalVariable, FString)>::CreateLambda(
-		[WeakThis](FRigVMExternalVariable InVariableToCreate, FString InDefaultValue) -> FName {
-			if (WeakThis.IsValid())
-			{
-				return WeakThis->AddCRMemberVariableFromExternal(InVariableToCreate, InDefaultValue);
-			}
-			return NAME_None;
-		}
-	));
+#endif
+	}
 
-	TWeakObjectPtr<URigVMController> WeakController = InControllerToConfigure;
-	InControllerToConfigure->RequestBulkEditDialogDelegate.BindLambda([WeakThis, WeakController](URigVMLibraryNode* InFunction, ERigVMControllerBulkEditType InEditType) -> FRigVMController_BulkEditResult 
+	// upgrade the gizmo libraries to shape libraries
+	if(!GizmoLibrary_DEPRECATED.IsNull() || GetLinkerCustomVersion(FControlRigObjectVersion::GUID) < FControlRigObjectVersion::RenameGizmoToShape)
 	{
-		if(WeakThis.IsValid() && WeakController.IsValid())
-		{
-			UControlRigBlueprint* StrongThis = WeakThis.Get();
-            URigVMController* StrongController = WeakController.Get();
-            if(StrongThis->OnRequestBulkEditDialog().IsBound())
-			{
-				return StrongThis->OnRequestBulkEditDialog().Execute(StrongThis, StrongController, InFunction, InEditType);
-			}
-		}
-		return FRigVMController_BulkEditResult();
-	});
+		// if it's an older file and it doesn't have the GizmoLibrary stored,
+		// refer to the previous default.
+		ShapeLibraries.Reset();
 
-	InControllerToConfigure->RequestBreakLinksDialogDelegate.BindLambda([WeakThis, WeakController](TArray<URigVMLink*> InLinks) -> bool 
-	{
-		if(WeakThis.IsValid() && WeakController.IsValid())
+		if(!GizmoLibrary_DEPRECATED.IsNull())
 		{
-			UControlRigBlueprint* StrongThis = WeakThis.Get();
-			if(StrongThis->OnRequestBreakLinksDialog().IsBound())
-			{
-				return StrongThis->OnRequestBreakLinksDialog().Execute(InLinks);
-			}
+			ShapeLibrariesToLoadOnPackageLoaded.Add(GizmoLibrary_DEPRECATED.ToString());
 		}
-		return false;
-	});
-
-	InControllerToConfigure->RequestPinTypeSelectionDelegate.BindLambda([WeakThis](const TArray<TRigVMTypeIndex>& InTypes) -> TRigVMTypeIndex 
-	{
-		if(WeakThis.IsValid())
+		else
 		{
-			UControlRigBlueprint* StrongThis = WeakThis.Get();
-			if(StrongThis->OnRequestPinTypeSelectionDialog().IsBound())
-			{
-				return StrongThis->OnRequestPinTypeSelectionDialog().Execute(InTypes);
-			}
+			static const FString DefaultGizmoLibraryPath = TEXT("/ControlRig/Controls/DefaultGizmoLibrary.DefaultGizmoLibrary");
+			ShapeLibrariesToLoadOnPackageLoaded.Add(DefaultGizmoLibraryPath);
 		}
-		return INDEX_NONE;
-	});
 
-	InControllerToConfigure->RequestNewExternalVariableDelegate.BindLambda([WeakThis](FRigVMGraphVariableDescription InVariable, bool bInIsPublic, bool bInIsReadOnly) -> FName
-	{
-		if (WeakThis.IsValid())
+		URigVMBlueprintGeneratedClass* RigClass = GetRigVMBlueprintGeneratedClass();
+		UControlRig* CDO = Cast<UControlRig>(RigClass->GetDefaultObject(false /* create if needed */));
+
+		TArray<UObject*> ArchetypeInstances;
+		CDO->GetArchetypeInstances(ArchetypeInstances);
+		ArchetypeInstances.Insert(CDO, 0);
+
+		for (UObject* Instance : ArchetypeInstances)
 		{
-			for (FBPVariableDescription& ExistingVariable : WeakThis->NewVariables)
+			if (UControlRig* InstanceRig = Cast<UControlRig>(Instance))
 			{
-				if (ExistingVariable.VarName == InVariable.Name)
-				{
-					return FName();
-				}
+				InstanceRig->ShapeLibraries.Reset();
+				InstanceRig->GizmoLibrary_DEPRECATED.Reset();
 			}
-
-			FRigVMExternalVariable ExternalVariable = InVariable.ToExternalVariable();
-			return WeakThis->AddMemberVariable(InVariable.Name,
-				ExternalVariable.TypeObject ? ExternalVariable.TypeObject->GetPathName() : ExternalVariable.TypeName.ToString(),
-				bInIsPublic,
-				bInIsReadOnly,
-				InVariable.DefaultValue);
 		}
-		
-		return FName();
-	});
+	}
+}
 
+#if WITH_EDITOR
 
-	InControllerToConfigure->RequestJumpToHyperlinkDelegate.BindLambda([WeakThis](const UObject* InSubject)
+void UControlRigBlueprint::HandlePackageDone()
+{
+	if (ShapeLibrariesToLoadOnPackageLoaded.Num() > 0)
 	{
-		if (WeakThis.IsValid())
+		for(const FString& ShapeLibraryToLoadOnPackageLoaded : ShapeLibrariesToLoadOnPackageLoaded)
 		{
-			UControlRigBlueprint* StrongThis = WeakThis.Get();
-			if(StrongThis->OnRequestJumpToHyperlink().IsBound())
-			{
-				StrongThis->OnRequestJumpToHyperlink().Execute(InSubject);
-			}
+			ShapeLibraries.Add(LoadObject<UControlRigShapeLibrary>(nullptr, *ShapeLibraryToLoadOnPackageLoaded));
 		}
-	});
 
-	InControllerToConfigure->ConfigureWorkflowOptionsDelegate.BindLambda([WeakThis](URigVMUserWorkflowOptions* Options)
-	{
-		if(UControlRigWorkflowOptions* ControlRigNodeWorkflowOptions = Cast<UControlRigWorkflowOptions>(Options))
+		URigVMBlueprintGeneratedClass* RigClass = GetRigVMBlueprintGeneratedClass();
+		UControlRig* CDO = Cast<UControlRig>(RigClass->GetDefaultObject(false /* create if needed */));
+
+		TArray<UObject*> ArchetypeInstances;
+		CDO->GetArchetypeInstances(ArchetypeInstances);
+		ArchetypeInstances.Insert(CDO, 0);
+
+		for (UObject* Instance : ArchetypeInstances)
 		{
-			ControlRigNodeWorkflowOptions->Hierarchy = nullptr;
-			ControlRigNodeWorkflowOptions->Selection.Reset();
-			
-			if(const UControlRigBlueprint* StrongThis = WeakThis.Get())
+			if (UControlRig* InstanceRig = Cast<UControlRig>(Instance))
 			{
-				if(UControlRig* ControlRig = Cast<UControlRig>(StrongThis->GetObjectBeingDebugged()))
-				{
-					ControlRigNodeWorkflowOptions->Hierarchy = ControlRig->GetHierarchy();
-				}
-				ControlRigNodeWorkflowOptions->Selection = StrongThis->Hierarchy->GetSelectedKeys();
+				InstanceRig->ShapeLibraries = ShapeLibraries;
 			}
 		}
-	});
+
+		ShapeLibrariesToLoadOnPackageLoaded.Reset();
+	}
+
+	PropagateHierarchyFromBPToInstances();
+
+	Super::HandlePackageDone();
+}
 
 #endif
+
+UClass* UControlRigBlueprint::GetRigVMEdGraphNodeClass() const
+{
+	return UControlRigGraphNode::StaticClass();
 }
 
-bool UControlRigBlueprint::TryImportGraphFromText(const FString& InClipboardText, UEdGraph** OutGraphPtr)
+UClass* UControlRigBlueprint::GetRigVMEdGraphSchemaClass() const
 {
-	if (OutGraphPtr)
-	{
-		*OutGraphPtr = nullptr;
-	}
+	return UControlRigGraphSchema::StaticClass();
+}
 
-	if (URigVMController* FunctionLibraryController = GetOrCreateController(GetLocalFunctionLibrary()))
-	{
-		TGuardValue<FRigVMController_RequestLocalizeFunctionDelegate> RequestLocalizeDelegateGuard(
-            FunctionLibraryController->RequestLocalizeFunctionDelegate,
-            FRigVMController_RequestLocalizeFunctionDelegate::CreateLambda([this](FRigVMGraphFunctionIdentifier& InFunctionToLocalize)
-            {
-            	BroadcastRequestLocalizeFunctionDialog(InFunctionToLocalize);
-				const URigVMLibraryNode* LocalizedFunctionNode = GetLocalFunctionLibrary()->FindPreviouslyLocalizedFunction(InFunctionToLocalize);
-				return LocalizedFunctionNode != nullptr;
-            })
-        );
-		
-		TArray<FName> ImportedNodeNames = FunctionLibraryController->ImportNodesFromText(InClipboardText, true, true);
-		if (ImportedNodeNames.Num() == 0)
-		{
-			return false;
-		}
+UClass* UControlRigBlueprint::GetRigVMEdGraphClass() const
+{
+	return UControlRigGraph::StaticClass();
+}
 
-		URigVMCollapseNode* CollapseNode = Cast<URigVMCollapseNode>(GetLocalFunctionLibrary()->FindFunction(ImportedNodeNames[0]));
-		if (ImportedNodeNames.Num() > 1 || CollapseNode == nullptr || CollapseNode->GetContainedGraph() == nullptr)
-		{
-			FunctionLibraryController->Undo();
-			return false;
-		}
+TArray<FString> UControlRigBlueprint::GeneratePythonCommands(const FString InNewBlueprintName)
+{
+	TArray<FString> InternalCommands;
+	InternalCommands.Add(TEXT("import unreal"));
+	InternalCommands.Add(TEXT("unreal.load_module('ControlRigDeveloper')"));
+	InternalCommands.Add(TEXT("factory = unreal.ControlRigBlueprintFactory"));
+	InternalCommands.Add(FString::Printf(TEXT("blueprint = factory.create_new_control_rig_asset(desired_package_path = '%s')"), *InNewBlueprintName));
+	InternalCommands.Add(TEXT("hierarchy = blueprint.hierarchy"));
+	InternalCommands.Add(TEXT("hierarchy_controller = hierarchy.get_controller()"));
 
-		UEdGraph* EdGraph = GetEdGraph(CollapseNode->GetContainedGraph());
-		if (OutGraphPtr)
-		{
-			*OutGraphPtr = EdGraph;
-		}
+	// Hierarchy
+	InternalCommands.Append(Hierarchy->GetController(true)->GeneratePythonCommands());
 
-		BroadcastGraphImported(EdGraph);
-	}
+#if WITH_EDITORONLY_DATA
+	const FString PreviewMeshPath = GetPreviewMesh()->GetPathName();
+	InternalCommands.Add(FString::Printf(TEXT("blueprint.set_preview_mesh(unreal.load_object(name='%s', outer=None))"),
+		*PreviewMeshPath));
+#endif
 
-	// always return true so that the default mechanism doesn't take over
-	return true;
+	InternalCommands.Append(Super::GeneratePythonCommands(InNewBlueprintName));
+	return InternalCommands;
 }
 
-USkeletalMesh* UControlRigBlueprint::GetPreviewMesh() const
+void UControlRigBlueprint::GetTypeActions(FBlueprintActionDatabaseRegistrar& ActionRegistrar) const
 {
 	DECLARE_SCOPE_HIERARCHICAL_COUNTER_FUNC()
 
-#if WITH_EDITORONLY_DATA
-	if (!PreviewSkeletalMesh.IsValid())
-	{
-		PreviewSkeletalMesh.LoadSynchronous();
-	}
-
-	return PreviewSkeletalMesh.Get();
-#else
-	return nullptr;
-#endif
+	IControlRigEditorModule::Get().GetTypeActions((UControlRigBlueprint*)this, ActionRegistrar);
 }
 
-void UControlRigBlueprint::SetPreviewMesh(USkeletalMesh* PreviewMesh, bool bMarkAsDirty/*=true*/)
+void UControlRigBlueprint::GetInstanceActions(FBlueprintActionDatabaseRegistrar& ActionRegistrar) const
 {
-#if WITH_EDITORONLY_DATA
-	if(bMarkAsDirty)
-	{
-		Modify();
-	}
+	DECLARE_SCOPE_HIERARCHICAL_COUNTER_FUNC()
 
-	PreviewSkeletalMesh = PreviewMesh;
-#endif
+	IControlRigEditorModule::Get().GetInstanceActions((UControlRigBlueprint*)this, ActionRegistrar);
 }
 
-void UControlRigBlueprint::Serialize(FArchive& Ar)
+void UControlRigBlueprint::PostTransacted(const FTransactionObjectEvent& TransactionEvent)
 {
-	RigVMClient.SetOuterClientHost(this, GET_MEMBER_NAME_CHECKED(UControlRigBlueprint, RigVMClient));
-	
-	Super::Serialize(Ar);
+	DECLARE_SCOPE_HIERARCHICAL_COUNTER_FUNC()
+	Super::PostTransacted(TransactionEvent);
 
-	if(Ar.IsObjectReferenceCollector())
+	if (TransactionEvent.GetEventType() == ETransactionObjectEventType::UndoRedo)
 	{
-		Ar.UsingCustomVersion(FControlRigObjectVersion::GUID);
-
-
-#if WITH_EDITORONLY_DATA
-		if (Ar.IsCooking() && ReferencedObjectPathsStored)
-		{
-			for (FSoftObjectPath ObjectPath : ReferencedObjectPaths)
-			{
-				ObjectPath.Serialize(Ar);
-			}
-		}
-		else
-#endif
+		TArray<FName> PropertiesChanged = TransactionEvent.GetChangedProperties();
+		if (PropertiesChanged.Contains(GET_MEMBER_NAME_CHECKED(UControlRigBlueprint, Hierarchy)))
 		{
-			TArray<IRigVMGraphFunctionHost*> ReferencedFunctionHosts = GetReferencedFunctionHosts(false);
+			int32 TransactionIndex = GEditor->Trans->FindTransactionIndex(TransactionEvent.GetTransactionId());
+			const FTransaction* Transaction = GEditor->Trans->GetTransaction(TransactionIndex);
 
-			for(IRigVMGraphFunctionHost* ReferencedFunctionHost : ReferencedFunctionHosts)
+			if (Transaction->GetTitle().BuildSourceString() == TEXT("Transform Gizmo"))
 			{
-				if (URigVMBlueprintGeneratedClass* BPGeneratedClass = Cast<URigVMBlueprintGeneratedClass>(ReferencedFunctionHost))
-				{
-					Ar << BPGeneratedClass;
-				}
+				PropagatePoseFromBPToInstances();
+				return;
 			}
 
-			for(const TSoftObjectPtr<UControlRigShapeLibrary>& ShapeLibraryPtr : ShapeLibraries)
+			PropagateHierarchyFromBPToInstances();
+
+			// make sure the bone name list is up 2 date for the editor graph
+			for (UEdGraph* Graph : UbergraphPages)
 			{
-				if(ShapeLibraryPtr.IsValid())
+				UControlRigGraph* RigGraph = Cast<UControlRigGraph>(Graph);
+				if (RigGraph == nullptr)
 				{
-					UControlRigShapeLibrary* ShapeLibrary = ShapeLibraryPtr.Get();
-					Ar << ShapeLibrary;
+					continue;
 				}
+				RigGraph->CacheNameLists(Hierarchy, &DrawContainer, ShapeLibraries);
 			}
+
+			RequestAutoVMRecompilation();
+			(void)MarkPackageDirty();
 		}
-	}
 
-	if(Ar.IsLoading())
-	{
-		if(Model_DEPRECATED || FunctionLibrary_DEPRECATED)
+		if (PropertiesChanged.Contains(GET_MEMBER_NAME_CHECKED(UControlRigBlueprint, DrawContainer)))
 		{
-			TGuardValue<bool> DisableClientNotifs(RigVMClient.bSuspendNotifications, true);
-			RigVMClient.SetFromDeprecatedData(Model_DEPRECATED, FunctionLibrary_DEPRECATED);
+			PropagateDrawInstructionsFromBPToInstances();
 		}
 	}
 }
 
-void UControlRigBlueprint::PreSave(const class ITargetPlatform* TargetPlatform)
+void UControlRigBlueprint::PostDuplicate(bool bDuplicateForPIE)
 {
-	PRAGMA_DISABLE_DEPRECATION_WARNINGS;
-	Super::PreSave(TargetPlatform);
-	PRAGMA_ENABLE_DEPRECATION_WARNINGS;
+	Super::PostDuplicate(bDuplicateForPIE);
+	
+	if (URigHierarchyController* Controller = Hierarchy->GetController(true))
+	{
+		Controller->OnModified().RemoveAll(this);
+		Controller->OnModified().AddUObject(this, &UControlRigBlueprint::HandleHierarchyModified);
+	}
 }
 
-void UControlRigBlueprint::PreSave(FObjectPreSaveContext ObjectSaveContext)
+TArray<UControlRigBlueprint*> UControlRigBlueprint::GetCurrentlyOpenRigBlueprints()
 {
-	Super::PreSave(ObjectSaveContext);
+	return sCurrentlyOpenedRigBlueprints;
+}
 
-	RigVMClient.PreSave(ObjectSaveContext);
+#if WITH_EDITOR
 
-	SupportedEventNames.Reset();
-	if (URigVMBlueprintGeneratedClass* RigClass = GetControlRigBlueprintGeneratedClass())
+const FControlRigShapeDefinition* UControlRigBlueprint::GetControlShapeByName(const FName& InName) const
+{
+	TMap<FString, FString> LibraryNameMap;
+	if(UControlRig* ControlRig = Cast<UControlRig>(GetObjectBeingDebugged()))
 	{
-		if (UControlRig* CDO = Cast<UControlRig>(RigClass->GetDefaultObject(true /* create if needed */)))
-		{
-			SupportedEventNames = CDO->GetSupportedEvents();
-		}
+		LibraryNameMap = ControlRig->ShapeLibraryNameMap;
+	}
+	return UControlRigShapeLibrary::GetShapeByName(InName, ShapeLibraries, LibraryNameMap);
+}
 
-		PublicGraphFunctions.Reset();
-		PublicGraphFunctions.SetNum(RigClass->GetRigVMGraphFunctionStore()->PublicFunctions.Num());
-		for (int32 i=0; i<PublicGraphFunctions.Num(); ++i)
-		{
-			PublicGraphFunctions[i] = RigClass->GetRigVMGraphFunctionStore()->PublicFunctions[i].Header;
-		}
+FName UControlRigBlueprint::AddTransientControl(URigVMPin* InPin)
+{
+	TUniquePtr<FControlValueScope> ValueScope;
+	if (!UControlRigEditorSettings::Get()->bResetControlsOnPinValueInteraction) // if we need to retain the controls
+	{
+		ValueScope = MakeUnique<FControlValueScope>(this);
 	}
 
-#if WITH_EDITORONLY_DATA
-	ReferencedObjectPaths.Reset();
+	// for now we only allow one pin control at the same time
+	ClearTransientControls();
+
+	URigVMBlueprintGeneratedClass* RigClass = GetRigVMBlueprintGeneratedClass();
+	UControlRig* CDO = Cast<UControlRig>(RigClass->GetDefaultObject(true /* create if needed */));
 
-	TArray<IRigVMGraphFunctionHost*> ReferencedFunctionHosts = GetReferencedFunctionHosts(false);
-	for(IRigVMGraphFunctionHost* ReferencedFunctionHost : ReferencedFunctionHosts)
+	FRigElementKey SpaceKey;
+	FTransform OffsetTransform = FTransform::Identity;
+	if (URigVMUnitNode* UnitNode = Cast<URigVMUnitNode>(InPin->GetPinForLink()->GetNode()))
 	{
-		if (URigVMBlueprintGeneratedClass* BPGeneratedClass = Cast<URigVMBlueprintGeneratedClass>(ReferencedFunctionHost))
+		if (TSharedPtr<FStructOnScope> DefaultStructScope = UnitNode->ConstructStructInstance())
 		{
-			ReferencedObjectPaths.AddUnique(BPGeneratedClass);
+			FRigUnit* DefaultStruct = (FRigUnit*)DefaultStructScope->GetStructMemory();
+
+			FString PinPath = InPin->GetPinForLink()->GetPinPath();
+			FString Left, Right;
+
+			if (URigVMPin::SplitPinPathAtStart(PinPath, Left, Right))
+			{
+				SpaceKey = DefaultStruct->DetermineSpaceForPin(Right, Hierarchy);
+				
+				URigHierarchy* RigHierarchy = Hierarchy;
+
+				// use the active rig instead of the CDO rig because we want to access the evaluation result of the rig graph
+				// to calculate the offset transform, for example take a look at RigUnit_ModifyTransform
+				if (UControlRig* RigBeingDebugged = Cast<UControlRig>(GetObjectBeingDebugged()))
+				{
+					RigHierarchy = RigBeingDebugged->GetHierarchy();
+				}
+				
+				OffsetTransform = DefaultStruct->DetermineOffsetTransformForPin(Right, RigHierarchy);
+			}
 		}
 	}
 
-	for(const TSoftObjectPtr<UControlRigShapeLibrary>& ShapeLibraryPtr : ShapeLibraries)
+	FName ReturnName = NAME_None;
+	TArray<UObject*> ArchetypeInstances;
+	CDO->GetArchetypeInstances(ArchetypeInstances);
+	for (UObject* ArchetypeInstance : ArchetypeInstances)
 	{
-		if(ShapeLibraryPtr.IsValid())
+		UControlRig* InstancedControlRig = Cast<UControlRig>(ArchetypeInstance);
+		if (InstancedControlRig)
 		{
-			UControlRigShapeLibrary* ShapeLibrary = ShapeLibraryPtr.Get();
-			ReferencedObjectPaths.AddUnique(ShapeLibrary);
+			FName ControlName = InstancedControlRig->AddTransientControl(InPin, SpaceKey, OffsetTransform);
+			if (ReturnName == NAME_None)
+			{
+				ReturnName = ControlName;
+			}
 		}
 	}
-	ReferencedObjectPathsStored = true;
-#endif
 
-	bExposesAnimatableControls = false;
-	Hierarchy->ForEach<FRigControlElement>([this](FRigControlElement* ControlElement) -> bool
-    {
-		if (Hierarchy->IsAnimatable(ControlElement))
-		{
-			bExposesAnimatableControls = true;
-			return false;
-		}
-		return true;
-	});
-
-	FunctionReferenceNodeData = GetReferenceNodeData();
-	IAssetRegistry::GetChecked().AssetTagsFinalized(*this);
-
-	CachedAssetTags.Reset();
-}
-
-void UControlRigBlueprint::PostLoad()
-{
-	Super::PostLoad();
-
-	bVMRecompilationRequired = true;
-	{
-		TGuardValue<bool> IsCompilingGuard(bIsCompiling, true);
-		
-		TArray<IRigVMGraphFunctionHost*> ReferencedFunctionHosts = GetReferencedFunctionHosts(true);
-
-		// PostLoad all referenced function hosts so that their function data are fully loaded 
-		// and ready to be inlined into this BP during compilation
-		for (IRigVMGraphFunctionHost* FunctionHost : ReferencedFunctionHosts)
-		{
-			if (URigVMBlueprintGeneratedClass* BPGeneratedClass = Cast<URigVMBlueprintGeneratedClass>(FunctionHost))
-			{
-				if (BPGeneratedClass->HasAllFlags(RF_NeedPostLoad))
-				{
-					BPGeneratedClass->ConditionalPostLoad();
-				}
-			}
-		}
-		
-		// temporarily disable default value validation during load time, serialized values should always be accepted
-		TGuardValue<bool> DisablePinDefaultValueValidation(GetOrCreateController()->bValidatePinDefaults, false);
-
-		// correct the offset transforms
-		if (GetLinkerCustomVersion(FControlRigObjectVersion::GUID) < FControlRigObjectVersion::ControlOffsetTransform)
-		{
-			HierarchyContainer_DEPRECATED.ControlHierarchy.PostLoad();
-			if (HierarchyContainer_DEPRECATED.ControlHierarchy.Num() > 0)
-			{
-				bDirtyDuringLoad = true;
-			}
-
-			for (FRigControl& Control : HierarchyContainer_DEPRECATED.ControlHierarchy)
-			{
-				const FTransform PreviousOffsetTransform = Control.GetTransformFromValue(ERigControlValueType::Initial);
-				Control.OffsetTransform = PreviousOffsetTransform;
-				Control.InitialValue = Control.Value;
-
-				if (Control.ControlType == ERigControlType::Transform)
-				{
-					Control.InitialValue = FRigControlValue::Make<FTransform>(FTransform::Identity);
-				}
-				else if (Control.ControlType == ERigControlType::TransformNoScale)
-				{
-					Control.InitialValue = FRigControlValue::Make<FTransformNoScale>(FTransformNoScale::Identity);
-				}
-				else if (Control.ControlType == ERigControlType::EulerTransform)
-				{
-					Control.InitialValue = FRigControlValue::Make<FEulerTransform>(FEulerTransform::Identity);
-				}
-			}
-		}
-
-		// convert the hierarchy from V1 to V2
-		if (GetLinkerCustomVersion(FControlRigObjectVersion::GUID) < FControlRigObjectVersion::RigHierarchyV2)
-		{
-			Modify();
-			
-			TGuardValue<bool> SuspendNotifGuard(Hierarchy->GetSuspendNotificationsFlag(), true);
-			
-			Hierarchy->Reset();
-			GetHierarchyController()->ImportFromHierarchyContainer(HierarchyContainer_DEPRECATED, false);
-		}
-
-		// remove all non-controlrig-graphs
-		TArray<UEdGraph*> NewUberGraphPages;
-		for (UEdGraph* Graph : UbergraphPages)
-		{
-			UControlRigGraph* RigGraph = Cast<UControlRigGraph>(Graph);
-			if (RigGraph)
-			{
-				NewUberGraphPages.Add(RigGraph);
-			}
-			else
-			{
-				Graph->MarkAsGarbage();
-				Graph->Rename(nullptr, GetTransientPackage(), REN_ForceNoResetLoaders);
-			}
-		}
-		UbergraphPages = NewUberGraphPages;
-
-		TArray<TGuardValue<bool>> EditableGraphGuards;
-		{
-			for (URigVMGraph* Graph : GetAllModels())
-			{
-				EditableGraphGuards.Emplace(Graph->bEditable, true);
-			}
-		}
-		
-		InitializeModelIfRequired(false /* recompile vm */);
-		{
-			TGuardValue<bool> GuardNotifications(bSuspendModelNotificationsForSelf, true);
-			
-			const FRigVMClientPatchResult PatchResult = GetRigVMClient()->PatchModelsOnLoad();
-			if(PatchResult.RequiresToMarkPackageDirty())
-			{
-				(void)MarkPackageDirty();
-				bDirtyDuringLoad = true;
-			}
-			
-			PatchFunctionReferencesOnLoad();
-			PatchVariableNodesOnLoad();
-			PatchVariableNodesWithIncorrectType();
-			PatchRigElementKeyCacheOnLoad();
-			PatchBoundVariables();
-			PatchPropagateToChildren();
-			PatchParameterNodesOnLoad();
-			PatchLinksWithCast();
-			PatchFunctionsOnLoad();
-		}
-
-#if WITH_EDITOR
-
-		{
-			TGuardValue<bool> GuardNotifications(bSuspendModelNotificationsForSelf, true);
-
-			// refresh the graph such that the pin hierarchies matches their CPPTypeObject
-			// this step is needed everytime we open a BP in the editor, b/c even after load
-			// model data can change while the Control Rig BP is not opened
-			// for example, if a user defined struct changed after BP load,
-			// any pin that references the struct needs to be regenerated
-			RefreshAllModels();
-		}
-
-		// at this point we may still have links which are detached. we may or may not be able to 
-		// reattach them.
-		GetRigVMClient()->ProcessDetachedLinks();
-
-		GetControlRigBlueprintGeneratedClass()->GetRigVMGraphFunctionStore()->RemoveAllCompilationData();
-
-		// perform backwards compat value upgrades
-		TArray<URigVMGraph*> GraphsToValidate = GetAllModels();
-		for (int32 GraphIndex = 0; GraphIndex < GraphsToValidate.Num(); GraphIndex++)
-		{
-			URigVMGraph* GraphToValidate = GraphsToValidate[GraphIndex];
-			if(GraphToValidate == nullptr)
-			{
-				continue;
-			}
-
-			for(URigVMNode* Node : GraphToValidate->GetNodes())
-			{
-				URigVMController* Controller = GetOrCreateController(GraphToValidate);
-				FRigVMControllerNotifGuard NotifGuard(Controller, true);
-				Controller->RemoveUnusedOrphanedPins(Node, true);
-			}
-				
-			for(URigVMNode* Node : GraphToValidate->GetNodes())
-			{
-				TArray<URigVMPin*> Pins = Node->GetAllPinsRecursively();
-				for(URigVMPin* Pin : Pins)
-				{
-					if(Pin->GetCPPTypeObject() == StaticEnum<ERigElementType>())
-					{
-						if(Pin->GetDefaultValue() == TEXT("Space"))
-						{
-							if(URigVMController* Controller = GetController(GraphToValidate))
-							{
-								FRigVMControllerNotifGuard NotifGuard(Controller, true);
-								Controller->SetPinDefaultValue(Pin->GetPinPath(), TEXT("Null"), false, false, false);
-							}
-						}
-					}
-				}
-
-				
-				// avoid function reference related validation for temp assets, a temp asset may get generated during
-				// certain content validation process. It is usually just a simple file-level copy of the source asset
-				// so these references are usually not fixed-up properly. Thus, it is meaningless to validate them.
-				// They should not be allowed to dirty the source asset either.
-				if (!this->GetPackage()->GetName().StartsWith("/Temp/"))
-				{
-					if(URigVMFunctionReferenceNode* FunctionReferenceNode = Cast<URigVMFunctionReferenceNode>(Node))
-					{
-						if(URigVMBuildData* BuildData = URigVMBuildData::Get())
-						{
-							BuildData->RegisterFunctionReference(FunctionReferenceNode);
-						}
-					}
-				}
-			}
-		}
-
-		CompileLog.Messages.Reset();
-		CompileLog.NumErrors = CompileLog.NumWarnings = 0;
-#endif
-	}
-
-	// upgrade the gizmo libraries to shape libraries
-	if(!GizmoLibrary_DEPRECATED.IsNull() || GetLinkerCustomVersion(FControlRigObjectVersion::GUID) < FControlRigObjectVersion::RenameGizmoToShape)
-	{
-		// if it's an older file and it doesn't have the GizmoLibrary stored,
-		// refer to the previous default.
-		ShapeLibraries.Reset();
-
-		if(!GizmoLibrary_DEPRECATED.IsNull())
-		{
-			ShapeLibrariesToLoadOnPackageLoaded.Add(GizmoLibrary_DEPRECATED.ToString());
-		}
-		else
-		{
-			static const FString DefaultGizmoLibraryPath = TEXT("/ControlRig/Controls/DefaultGizmoLibrary.DefaultGizmoLibrary");
-			ShapeLibrariesToLoadOnPackageLoaded.Add(DefaultGizmoLibraryPath);
-		}
-
-		URigVMBlueprintGeneratedClass* RigClass = GetControlRigBlueprintGeneratedClass();
-		UControlRig* CDO = Cast<UControlRig>(RigClass->GetDefaultObject(false /* create if needed */));
-
-		TArray<UObject*> ArchetypeInstances;
-		CDO->GetArchetypeInstances(ArchetypeInstances);
-		ArchetypeInstances.Insert(CDO, 0);
-
-		for (UObject* Instance : ArchetypeInstances)
-		{
-			if (UControlRig* InstanceRig = Cast<UControlRig>(Instance))
-			{
-				InstanceRig->ShapeLibraries.Reset();
-				InstanceRig->GizmoLibrary_DEPRECATED.Reset();
-			}
-		}
-	}
-
-#if WITH_EDITOR
-	if(GIsEditor)
-	{
-		// delay compilation until the package has been loaded
-		FCoreUObjectDelegates::OnEndLoadPackage.AddUObject(this, &UControlRigBlueprint::HandlePackageDone);
-	}
-#else
-	RecompileVMIfRequired();
-#endif
-	RequestControlRigInit();
-
-	FCoreUObjectDelegates::OnObjectModified.RemoveAll(this);
-	OnChanged().RemoveAll(this);
-	FCoreUObjectDelegates::OnObjectModified.AddUObject(this, &UControlRigBlueprint::OnPreVariableChange);
-	OnChanged().AddUObject(this, &UControlRigBlueprint::OnPostVariableChange);
-
-	if (UPackage* Package = GetOutermost())
-	{
-		Package->SetDirtyFlag(bDirtyDuringLoad);
-	}
-
-#if WITH_EDITOR
-	// if we are running with -game we are in editor code,
-	// but GIsEditor is turned off
-	if(!GIsEditor)
-	{
-		HandlePackageDone();
-	}
-#endif
-
-	// RigVMRegistry changes can be triggered when new user defined types(structs/enums) are added/removed
-	// in which case we have to refresh the model
-	FRigVMRegistry::Get().OnRigVMRegistryChanged().RemoveAll(this);
-	FRigVMRegistry::Get().OnRigVMRegistryChanged().AddUObject(this, &UControlRigBlueprint::OnRigVMRegistryChanged);
-	
-	UEdGraphPin::ResolveAllPinReferences();
-}
-
-#if WITH_EDITORONLY_DATA
-void UControlRigBlueprint::DeclareConstructClasses(TArray<FTopLevelAssetPath>& OutConstructClasses, const UClass* SpecificSubclass)
-{
-	Super::DeclareConstructClasses(OutConstructClasses, SpecificSubclass);
-	OutConstructClasses.Add(FTopLevelAssetPath(URigVMController::StaticClass()));
-}
-#endif
-
-#if WITH_EDITOR
-void UControlRigBlueprint::HandlePackageDone(const FEndLoadPackageContext& Context)
-{
-	if (!Context.LoadedPackages.Contains(GetPackage()))
-	{
-		return;
-	}
-	HandlePackageDone();
-}
-
-void UControlRigBlueprint::HandlePackageDone()
-{
-	FCoreUObjectDelegates::OnEndLoadPackage.RemoveAll(this);
-
-	if (ShapeLibrariesToLoadOnPackageLoaded.Num() > 0)
-	{
-		for(const FString& ShapeLibraryToLoadOnPackageLoaded : ShapeLibrariesToLoadOnPackageLoaded)
-		{
-			ShapeLibraries.Add(LoadObject<UControlRigShapeLibrary>(nullptr, *ShapeLibraryToLoadOnPackageLoaded));
-		}
-
-		URigVMBlueprintGeneratedClass* RigClass = GetControlRigBlueprintGeneratedClass();
-		UControlRig* CDO = Cast<UControlRig>(RigClass->GetDefaultObject(false /* create if needed */));
-
-		TArray<UObject*> ArchetypeInstances;
-		CDO->GetArchetypeInstances(ArchetypeInstances);
-		ArchetypeInstances.Insert(CDO, 0);
-
-		for (UObject* Instance : ArchetypeInstances)
-		{
-			if (UControlRig* InstanceRig = Cast<UControlRig>(Instance))
-			{
-				InstanceRig->ShapeLibraries = ShapeLibraries;
-			}
-		}
-
-		ShapeLibrariesToLoadOnPackageLoaded.Reset();
-	}
-
-	if(URigVMBuildData* BuildData = URigVMBuildData::Get())
-	{
-		if(URigVMFunctionLibrary* FunctionLibrary = RigVMClient.GetFunctionLibrary())
-		{
-			// for backwards compatibility load the function references from the
-			// model's storage over to the centralized build data
-			if(!FunctionLibrary->FunctionReferences_DEPRECATED.IsEmpty())
-			{
-				// let's also update the asset data of the dependents
-				IAssetRegistry& AssetRegistry = IAssetRegistry::GetChecked();
-				
-				for(const TTuple< TObjectPtr<URigVMLibraryNode>, FRigVMFunctionReferenceArray >& Pair :
-					FunctionLibrary->FunctionReferences_DEPRECATED)
-				{
-					TSoftObjectPtr<URigVMLibraryNode> FunctionKey(Pair.Key);
-						
-					for(int32 ReferenceIndex = 0; ReferenceIndex < Pair.Value.Num(); ReferenceIndex++)
-					{
-						// update the build data
-						BuildData->RegisterFunctionReference(FunctionKey->GetFunctionIdentifier(), Pair.Value[ReferenceIndex]);
-
-						// find all control rigs matching the reference node
-						FAssetData AssetData = AssetRegistry.GetAssetByObjectPath(
-							Pair.Value[ReferenceIndex].ToSoftObjectPath().GetWithoutSubPath());
-
-						// if the asset has never been loaded - make sure to load it once and mark as dirty
-						if(AssetData.IsValid() && !AssetData.IsAssetLoaded())
-						{
-							if(UControlRigBlueprint* Dependent = Cast<UControlRigBlueprint>(AssetData.GetAsset()))
-							{
-								if(Dependent != this)
-								{
-									Dependent->MarkPackageDirty();
-								}
-							}
-						}
-					}
-				}
-				
-				FunctionLibrary->FunctionReferences_DEPRECATED.Reset();
-				MarkPackageDirty();
-			}
-		}
-
-		// update the build data from the current function references
-		const TArray<FRigVMReferenceNodeData> ReferenceNodeDatas = GetReferenceNodeData();
-		for(const FRigVMReferenceNodeData& ReferenceNodeData : ReferenceNodeDatas)
-		{
-			BuildData->RegisterFunctionReference(ReferenceNodeData);
-		}
-
-		BuildData->ClearInvalidReferences();
-	}
-	
-	PropagateHierarchyFromBPToInstances();
-	RemoveDeprecatedVMMemoryClass();
-	RecompileVM();
-	RequestControlRigInit();
-	BroadcastControlRigPackageDone();
-}
-
-void UControlRigBlueprint::BroadcastControlRigPackageDone()
-{
-	if (URigVMBlueprintGeneratedClass* RigClass = GetControlRigBlueprintGeneratedClass())
-	{
-		UControlRig* CDO = Cast<UControlRig>(RigClass->GetDefaultObject(true /* create if needed */));
-		CDO->BroadCastEndLoadPackage();
-
-		TArray<UObject*> ArchetypeInstances;
-		CDO->GetArchetypeInstances(ArchetypeInstances);
-		for (UObject* Instance : ArchetypeInstances)
-		{
-			if (UControlRig* InstanceRig = Cast<UControlRig>(Instance))
-			{
-				InstanceRig->BroadCastEndLoadPackage();
-			}
-		}
-	}
-}
-
-void UControlRigBlueprint::RemoveDeprecatedVMMemoryClass() const
-{
-	TArray<UObject*> Objects;
-	GetObjectsWithOuter(this, Objects, false);
-
-	for (UObject* Object : Objects)
-	{
-		if (URigVMMemoryStorageGeneratorClass* DeprecatedClass = Cast<URigVMMemoryStorageGeneratorClass>(Object))
-		{
-			DeprecatedClass->Rename(nullptr, GetTransientPackage(), REN_ForceNoResetLoaders | REN_DoNotDirty | REN_DontCreateRedirectors | REN_NonTransactional);
-		}
-	}
-}
-#endif
-
-void UControlRigBlueprint::RecompileVM()
-{
-	if(bIsCompiling)
-	{
-		return;
-	}
-
-	URigVMBlueprintGeneratedClass* RigClass = GetControlRigBlueprintGeneratedClass();
-	if(RigClass == nullptr)
-	{
-		return;
-	}
-
-	TGuardValue<bool> CompilingGuard(bIsCompiling, true);
-	
-	bErrorsDuringCompilation = false;
-
-	RigGraphDisplaySettings.MinMicroSeconds = RigGraphDisplaySettings.LastMinMicroSeconds = DBL_MAX;
-	RigGraphDisplaySettings.MaxMicroSeconds = RigGraphDisplaySettings.LastMaxMicroSeconds = (double)INDEX_NONE;
-
-	UControlRig* CDO = Cast<UControlRig>(RigClass->GetDefaultObject(true /* create if needed */));
-	if (CDO && CDO->VM != nullptr)
-	{
-		TGuardValue<bool> ReentrantGuardSelf(bSuspendModelNotificationsForSelf, true);
-		TGuardValue<bool> ReentrantGuardOthers(bSuspendModelNotificationsForOthers, true);
-
-		CDO->PostInitInstanceIfRequired();
-		CDO->VMRuntimeSettings = VMRuntimeSettings;
-		CDO->GetHierarchy()->CopyHierarchy(Hierarchy);
-
-		if (!HasAnyFlags(RF_Transient | RF_Transactional))
-		{
-			CDO->Modify(false);
-		}
-		CDO->VM->Reset();
-
-		// Clear all Errors
-		CompileLog.Messages.Reset();
-		CompileLog.NumErrors = CompileLog.NumWarnings = 0;
-		
-		TArray<UEdGraph*> EdGraphs;
-		GetAllGraphs(EdGraphs);
-
-		for (UEdGraph* Graph : EdGraphs)
-		{
-			UControlRigGraph* RigGraph = Cast<UControlRigGraph>(Graph);
-			if (RigGraph == nullptr)
-			{
-				continue;
-			}
-
-			for (UEdGraphNode* GraphNode : Graph->Nodes)
-			{
-				if (UControlRigGraphNode* ControlRigGraphNode = Cast<UControlRigGraphNode>(GraphNode))
-				{
-					ControlRigGraphNode->ClearErrorInfo();
-				}
-			}
-		}
-
-		URigVMCompiler* Compiler = URigVMCompiler::StaticClass()->GetDefaultObject<URigVMCompiler>();
-		VMCompileSettings.SetExecuteContextStruct(RigVMClient.GetExecuteContextStruct());
-		Compiler->Settings = (bCompileInDebugMode) ? FRigVMCompileSettings::Fast(VMCompileSettings.GetExecuteContextStruct()) : VMCompileSettings;
-		Compiler->Compile(RigVMClient.GetAllModels(false, false), GetOrCreateController(), CDO->VM, CDO->GetExtendedExecuteContext(), CDO->GetExternalVariablesImpl(false), &PinToOperandMap);
-
-		if (bErrorsDuringCompilation)
-		{
-			if(Compiler->Settings.SurpressErrors)
-			{
-				Compiler->Settings.Reportf(EMessageSeverity::Info, this,
-					TEXT("Compilation Errors may be suppressed for ControlRigBlueprint: %s. See VM Compile Setting in Class Settings for more Details"), *this->GetName());
-			}
-			bVMRecompilationRequired = false;
-			if(CDO->VM)
-			{
-				VMCompiledEvent.Broadcast(this, CDO->VM);
-			}
-			return;
-		}
-
-		TArray<UObject*> ArchetypeInstances;
-		CDO->GetArchetypeInstances(ArchetypeInstances);
-		for (UObject* Instance : ArchetypeInstances)
-		{
-			if (UControlRig* InstanceRig = Cast<UControlRig>(Instance))
-			{
-				// No objects should be created during load, so PostInitInstanceIfRequired, which creates a new VM and
-				// DynamicHierarchy, should not be called during load
-				if (!InstanceRig->HasAllFlags(RF_NeedPostLoad))
-				{
-					InstanceRig->PostInitInstanceIfRequired();
-				}
-				InstanceRig->InstantiateVMFromCDO();
-				InstanceRig->CopyExternalVariableDefaultValuesFromCDO();
-			}
-		}
-
-		bVMRecompilationRequired = false;
-		VMCompiledEvent.Broadcast(this, CDO->VM);
-
-#if WITH_EDITOR
-		RefreshControlRigBreakpoints();
-#endif
-	}
-}
-
-void UControlRigBlueprint::RecompileVMIfRequired()
-{
-	if (bVMRecompilationRequired)
-	{
-		RecompileVM();
-	}
-}
-
-void UControlRigBlueprint::RequestAutoVMRecompilation()
-{
-	bVMRecompilationRequired = true;
-	if (bAutoRecompileVM && VMRecompilationBracket == 0)
-	{
-		RecompileVMIfRequired();
-	}
-}
-
-void UControlRigBlueprint::IncrementVMRecompileBracket()
-{
-	VMRecompilationBracket++;
-}
-
-void UControlRigBlueprint::DecrementVMRecompileBracket()
-{
-	if (VMRecompilationBracket == 1)
-	{
-		if (bAutoRecompileVM)
-		{
-			RecompileVMIfRequired();
-		}
-		VMRecompilationBracket = 0;
-	}
-	else if (VMRecompilationBracket > 0)
-	{
-		VMRecompilationBracket--;
-	}
-}
-
-void UControlRigBlueprint::RefreshAllModels(EControlRigBlueprintLoadType InLoadType)
-{
-	const bool bIsPostLoad = InLoadType == EControlRigBlueprintLoadType::PostLoad;
-
-	// avoid any compute if the current structure hashes match with the serialized ones
-	if(CVarEnablePostLoadHashing->GetBool() && RigVMClient.GetStructureHash() == RigVMClient.GetSerializedStructureHash())
-	{
-		if(bIsPostLoad)
-		{
-			TArray<URigVMGraph*> ModelGraphs = RigVMClient.GetAllModels(true, true);
-			Algo::Reverse(ModelGraphs);
-			for (URigVMGraph* ModelGraph : ModelGraphs)
-			{
-				URigVMController* Controller = GetOrCreateController(ModelGraph);
-				URigVMController::FRestoreLinkedPathSettings Settings;
-				Settings.bFollowCoreRedirectors = true;
-				Settings.bRelayToOrphanPins = true;
-				Controller->ProcessDetachedLinks(Settings);
-			}
-		}
-		return;
-	}
-	
-	TGuardValue<bool> IsCompilingGuard(bIsCompiling, true);
-	TGuardValue<bool> ClientIgnoreModificationsGuard(RigVMClient.bIgnoreModelNotifications, true);
-	
-	TArray<URigVMGraph*> AllModelsLeavesFirst = RigVMClient.GetAllModelsLeavesFirst(true);
-	TMap<const URigVMGraph*, TArray<URigVMController::FLinkedPath>> LinkedPaths;
-
-	if (ensure(IsInGameThread()))
-	{
-		for (URigVMGraph* Graph : AllModelsLeavesFirst)
-		{
-			URigVMController* Controller = GetOrCreateController(Graph);
-			// temporarily disable default value validation during load time, serialized values should always be accepted
-			TGuardValue<bool> PerGraphDisablePinDefaultValueValidation(Controller->bValidatePinDefaults, false);
-			TGuardValue<bool> GuardEditGraph(Graph->bEditable, true);
-			FRigVMControllerNotifGuard NotifGuard(Controller, true);
-			LinkedPaths.Add(Graph, Controller->GetLinkedPaths());
-			Controller->FastBreakLinkedPaths(LinkedPaths.FindChecked(Graph));
-			TArray<URigVMNode*> Nodes = Graph->GetNodes();
-			for (URigVMNode* Node : Nodes)
-			{
-				Controller->RepopulatePinsOnNode(Node, true, true);
-			}
-		}
-		SetupPinRedirectorsForBackwardsCompatibility();
-	}
-
-	for (URigVMGraph* Graph : AllModelsLeavesFirst)
-	{
-		URigVMController* Controller = GetOrCreateController(Graph);
-		TGuardValue<bool> GuardEditGraph(Graph->bEditable, true);
-		FRigVMControllerNotifGuard NotifGuard(Controller, true);
-		{
-			URigVMController::FRestoreLinkedPathSettings Settings;
-			Settings.bFollowCoreRedirectors = true;
-			Controller->RestoreLinkedPaths(LinkedPaths.FindChecked(Graph));
-		}
-
-		for(URigVMNode* ModelNode : Graph->GetNodes())
-		{
-			Controller->RemoveUnusedOrphanedPins(ModelNode);
-		}
-
-		if(bIsPostLoad)
-		{
-			for(URigVMNode* ModelNode : Graph->GetNodes())
-			{
-				if (URigVMTemplateNode* TemplateNode = Cast<URigVMTemplateNode>(ModelNode))
-				{
-					TemplateNode->InvalidateCache();
-					TemplateNode->PostLoad();
-				}
-			}
-		}
-
-#if WITH_EDITOR
-
-		if(bIsPostLoad)
-		{
-			for(URigVMNode* ModelNode : Graph->GetNodes())
-			{
-				if(URigVMUnitNode* UnitNode = Cast<URigVMUnitNode>(ModelNode))
-				{
-					if (!UnitNode->HasWildCardPin())
-					{
-						UScriptStruct* ScriptStruct = UnitNode->GetScriptStruct(); 
-						if(ScriptStruct == nullptr)
-						{
-							Controller->FullyResolveTemplateNode(UnitNode, INDEX_NONE, false);
-						}
-
-						// Try to find a deprecated template
-						if (UnitNode->GetScriptStruct() == nullptr && !UnitNode->TemplateNotation.IsNone())
-						{
-							const FRigVMTemplate* Template = FRigVMRegistry::Get().FindTemplate(UnitNode->TemplateNotation, true);
-							FRigVMTemplate::FTypeMap TypeMap = UnitNode->GetTemplatePinTypeMap();
-
-							int32 Permutation;
-							if (Template->FullyResolve(TypeMap, Permutation))
-							{
-								const FRigVMFunction* Function = Template->GetPermutation(Permutation);
-								UnitNode->ResolvedFunctionName = Function->GetName();
-							}
-						}
-
-						if (UnitNode->GetScriptStruct() == nullptr)
-						{
-							static const TCHAR UnresolvedUnitNodeMessage[] = TEXT("Node %s could not be resolved.");
-							Controller->ReportErrorf(UnresolvedUnitNodeMessage, *ModelNode->GetNodePath(true));
-						}
-					}
-				}
-				if (URigVMDispatchNode* DispatchNode = Cast<URigVMDispatchNode>(ModelNode))
-				{
-					if (DispatchNode->GetFactory() == nullptr)
-					{
-						static const TCHAR UnresolvedDispatchNodeMessage[] = TEXT("Dispatch node %s has no factory..");
-						Controller->ReportErrorf(UnresolvedDispatchNodeMessage, *ModelNode->GetNodePath(true));
-					}
-					else if (!DispatchNode->HasWildCardPin())
-					{
-						if (DispatchNode->GetResolvedFunction() == nullptr)
-						{
-							Controller->FullyResolveTemplateNode(DispatchNode, INDEX_NONE, false);
-						}
-						if (DispatchNode->GetResolvedFunction() == nullptr)
-						{
-							static const TCHAR UnresolvedDispatchNodeMessage[] = TEXT("Node %s could not be resolved.");
-							Controller->ReportErrorf(UnresolvedDispatchNodeMessage, *ModelNode->GetNodePath(true));
-						}
-					}
-				}
-			}
-		}
-#endif
-
-	}
-}
-
-void UControlRigBlueprint::OnRigVMRegistryChanged()
-{
-	RefreshAllModels();
-	RebuildGraphFromModel();
-	// avoids slate crash
-	FRigVMBlueprintUtils::HandleRefreshAllNodes(this);
-}
-
-void UControlRigBlueprint::HandleReportFromCompiler(EMessageSeverity::Type InSeverity, UObject* InSubject, const FString& InMessage)
-{
-	UObject* SubjectForMessage = InSubject;
-	if(URigVMNode* ModelNode = Cast<URigVMNode>(SubjectForMessage))
-	{
-		if(UControlRigBlueprint* RigBlueprint = ModelNode->GetTypedOuter<UControlRigBlueprint>())
-		{
-			if(UControlRigGraph* EdGraph = Cast<UControlRigGraph>(RigBlueprint->GetEdGraph(ModelNode->GetGraph())))
-			{
-				if(UEdGraphNode* EdNode = EdGraph->FindNodeForModelNodeName(ModelNode->GetFName()))
-				{
-					SubjectForMessage = EdNode;
-				}
-			}
-		}
-	}
-
-	FCompilerResultsLog* Log = CurrentMessageLog ? CurrentMessageLog : &CompileLog;
-	if (InSeverity == EMessageSeverity::Error)
-	{
-		Status = BS_Error;
-		MarkPackageDirty();
-
-		// see UnitTest "ControlRig.Basics.OrphanedPins" to learn why errors are suppressed this way
-		if (VMCompileSettings.SurpressErrors)
-		{
-			Log->SetSilentMode(true);
-		}
-
-		if(InMessage.Contains(TEXT("@@")))
-		{
-			Log->Error(*InMessage, SubjectForMessage);
-		}
-		else
-		{
-			Log->Error(*InMessage);
-		}
-
-		BroadCastReportCompilerMessage(InSeverity, InSubject, InMessage);
-
-		// see UnitTest "ControlRig.Basics.OrphanedPins" to learn why errors are suppressed this way
-		if (!VMCompileSettings.SurpressErrors)
-		{ 
-			FScriptExceptionHandler::Get().HandleException(ELogVerbosity::Error, *InMessage, *FString());
-		}
-		
-		bErrorsDuringCompilation = true;
-	}
-	else if (InSeverity == EMessageSeverity::Warning)
-	{
-		if(InMessage.Contains(TEXT("@@")))
-		{
-			Log->Warning(*InMessage, SubjectForMessage);
-		}
-		else
-		{
-			Log->Warning(*InMessage);
-		}
-
-		BroadCastReportCompilerMessage(InSeverity, InSubject, InMessage);
-		FScriptExceptionHandler::Get().HandleException(ELogVerbosity::Warning, *InMessage, *FString());
-	}
-	else
-	{
-		if(InMessage.Contains(TEXT("@@")))
-		{
-			Log->Note(*InMessage, SubjectForMessage);
-		}
-		else
-		{
-			Log->Note(*InMessage);
-		}
-
-		UE_LOG(LogControlRigDeveloper, Display, TEXT("%s"), *InMessage);
-	}
-
-	if (UControlRigGraphNode* EdGraphNode = Cast<UControlRigGraphNode>(SubjectForMessage))
-	{
-		EdGraphNode->ErrorType = (int32)InSeverity;
-		EdGraphNode->ErrorMsg = InMessage;
-		EdGraphNode->bHasCompilerMessage = EdGraphNode->ErrorType <= int32(EMessageSeverity::Info);
-	}
-}
-
-TArray<IRigVMGraphFunctionHost*> UControlRigBlueprint::GetReferencedFunctionHosts(bool bForceLoad)
-{
-	TArray<IRigVMGraphFunctionHost*> ReferencedBlueprints;
-	
-	TArray<UEdGraph*> EdGraphs;
-	GetAllGraphs(EdGraphs);
-	for (UEdGraph* EdGraph : EdGraphs)
-	{
-		for(UEdGraphNode* Node : EdGraph->Nodes)
-		{
-			if(UControlRigGraphNode* RigNode = Cast<UControlRigGraphNode>(Node))
-			{
-				if(URigVMFunctionReferenceNode* FunctionRefNode = Cast<URigVMFunctionReferenceNode>(RigNode->GetModelNode()))
-				{
-					IRigVMGraphFunctionHost* Host = nullptr;
-					if (bForceLoad || FunctionRefNode->IsReferencedFunctionHostLoaded())
-					{
-						// Load the function host
-						Host = FunctionRefNode->GetReferencedFunctionHeader().GetFunctionHost();
-					}
-					else if (bForceLoad || FunctionRefNode->IsReferencedNodeLoaded())
-					{
-						// Load the reference library node
-						if(const URigVMLibraryNode* ReferencedNode = FunctionRefNode->LoadReferencedNode())
-						{
-							if(URigVMFunctionLibrary* ReferencedFunctionLibrary = ReferencedNode->GetLibrary())
-							{
-								FSoftObjectPath FunctionHostPath = ReferencedFunctionLibrary->GetFunctionHostObjectPath();
-								if (UObject* FunctionHostObj = FunctionHostPath.TryLoad())
-								{
-									Host = Cast<IRigVMGraphFunctionHost>(FunctionHostObj);									
-								}
-							}
-						}
-					}
-
-					if (Host != nullptr && Host != GetControlRigBlueprintGeneratedClass())
-					{
-						ReferencedBlueprints.Add(Host);
-					}
-				}
-			}
-		}
-	}
-	
-	return ReferencedBlueprints;
-}
-
-#if WITH_EDITOR
-
-void UControlRigBlueprint::ClearBreakpoints()
-{
-	for(URigVMNode* Node : RigVMBreakpointNodes)
-	{
-		Node->SetHasBreakpoint(false);		
-	}
-	
-	RigVMBreakpointNodes.Empty();
-	RefreshControlRigBreakpoints();
-}
-
-bool UControlRigBlueprint::AddBreakpoint(const FString& InBreakpointNodePath)
-{
-	URigVMLibraryNode* FunctionNode = nullptr;
-	
-	// Find the node in the graph
-	for(const URigVMGraph* Model : RigVMClient)
-	{
-		URigVMNode* BreakpointNode = Model->FindNode(InBreakpointNodePath);
-		if (BreakpointNode == nullptr)
-		{
-			// If we cannot find the node, it might be because it is inside a function
-			FString FunctionName = InBreakpointNodePath, Right;
-			URigVMNode::SplitNodePathAtStart(InBreakpointNodePath, FunctionName, Right);
-
-			// Look inside the local function library
-			if (URigVMLibraryNode* LibraryNode = GetLocalFunctionLibrary()->FindFunction(FName(FunctionName)))
-			{
-				BreakpointNode = LibraryNode->GetContainedGraph()->FindNode(Right);
-				FunctionNode = LibraryNode;
-			}
-		}
-
-		if(BreakpointNode)
-		{
-			return AddBreakpoint(BreakpointNode, FunctionNode);
-		}
-	}
-	return false;
-}
-
-bool UControlRigBlueprint::AddBreakpoint(URigVMNode* InBreakpointNode, URigVMLibraryNode* LibraryNode)
-{
-	if (InBreakpointNode == nullptr)
-	{
-		return false;
-	}
-
-	bool bSuccess = true;
-	if (LibraryNode)
-	{
-		// If the breakpoint node is inside a library node, find all references to the library node
-		TArray<TSoftObjectPtr<URigVMFunctionReferenceNode>> References = LibraryNode->GetLibrary()->GetReferencesForFunction(LibraryNode->GetFName());
-		for (TSoftObjectPtr<URigVMFunctionReferenceNode> Reference : References)
-		{
-			if (!Reference.IsValid())
-			{
-				continue;
-			}
-
-			UControlRigBlueprint* ReferenceBlueprint = Reference->GetTypedOuter<UControlRigBlueprint>();
-
-			// If the reference is not inside another function, add a breakpoint in the blueprint containing the
-			// reference, without a function specified
-			bool bIsInsideFunction = Reference->GetRootGraph()->IsA<URigVMFunctionLibrary>();
-			if(!bIsInsideFunction)
-			{
-				bSuccess &= ReferenceBlueprint->AddBreakpoint(InBreakpointNode);
-			}
-			else
-			{
-				// Otherwise, we need to add breakpoints to all the blueprints that reference this
-				// function (when the blueprint graph is flattened)
-				
-				// Get all the functions containing this reference
-				URigVMNode* Node = Reference.Get();
-				while (Node->GetGraph() != ReferenceBlueprint->GetLocalFunctionLibrary())
-				{
-					if (URigVMLibraryNode* ParentLibraryNode = Cast<URigVMLibraryNode>(Node->GetGraph()->GetOuter()))
-					{
-						// Recursively add breakpoints to the reference blueprint, specifying the parent function
-						bSuccess &= ReferenceBlueprint->AddBreakpoint(InBreakpointNode, ParentLibraryNode);
-					}
-
-					Node = Cast<URigVMNode>(Node->GetGraph()->GetOuter());
-				}
-			}
-		}
-	}
-	else
-	{
-		if (!RigVMBreakpointNodes.Contains(InBreakpointNode))
-		{
-			// Add the breakpoint to the VM
-			bSuccess = AddBreakpointToControlRig(InBreakpointNode);
-			BreakpointAddedEvent.Broadcast();
-		}
-	}
-
-	return bSuccess;
-}
-
-bool UControlRigBlueprint::AddBreakpointToControlRig(URigVMNode* InBreakpointNode)
-{
-	URigVMBlueprintGeneratedClass* RigClass = GetControlRigBlueprintGeneratedClass();
-	UControlRig* CDO = Cast<UControlRig>(RigClass->GetDefaultObject(false));
-	const FRigVMByteCode* ByteCode = GetController()->GetCurrentByteCode();
-	TSet<FString> AddedCallpaths;
-
-	if (CDO && ByteCode)
-	{
-		FRigVMInstructionArray Instructions = ByteCode->GetInstructions();	
-		
-		// For each instruction, see if the node is in the callpath
-		// Only add one breakpoint for each callpath related to this node (i.e. if a node produces multiple
-		// instructions, only add a breakpoint to the first instruction)
-		for (int32 i = 0; i< Instructions.Num(); ++i)
-		{
-			for(URigVMGraph* Model : RigVMClient)
-			{
-				const FRigVMASTProxy Proxy = FRigVMASTProxy::MakeFromCallPath(ByteCode->GetCallPathForInstruction(i), Model);
-				if (Proxy.GetCallstack().Contains(InBreakpointNode))
-				{
-					// Find the callpath related to the breakpoint node
-					FRigVMASTProxy BreakpointProxy = Proxy;
-					while(BreakpointProxy.GetSubject() != InBreakpointNode)
-					{
-						BreakpointProxy = BreakpointProxy.GetParent();
-					}
-					const FString& BreakpointCallPath = BreakpointProxy.GetCallstack().GetCallPath();
-
-					// Only add this callpath breakpoint once
-					if (!AddedCallpaths.Contains(BreakpointCallPath))
-					{
-						AddedCallpaths.Add(BreakpointCallPath);
-						CDO->AddBreakpoint(i, InBreakpointNode, BreakpointProxy.GetCallstack().Num());
-					}
-				}
-			}
-		}
-	}
-
-	if (AddedCallpaths.Num() > 0)
-	{
-		RigVMBreakpointNodes.AddUnique(InBreakpointNode);
-		return true;
-	}
-	
-	return false;
-}
-
-bool UControlRigBlueprint::RemoveBreakpoint(const FString& InBreakpointNodePath)
-{
-	// Find the node in the graph
-	URigVMNode* BreakpointNode = nullptr;
-	
-	for(URigVMGraph* Model : RigVMClient)
-	{
-		BreakpointNode = Model->FindNode(InBreakpointNodePath);
-		if (BreakpointNode == nullptr)
-		{
-			// If we cannot find the node, it might be because it is inside a function
-			FString FunctionName = InBreakpointNodePath, Right;
-			URigVMNode::SplitNodePathAtStart(InBreakpointNodePath, FunctionName, Right);
-
-			// Look inside the local function library
-			if (URigVMLibraryNode* LibraryNode = GetLocalFunctionLibrary()->FindFunction(FName(FunctionName)))
-			{
-				BreakpointNode = LibraryNode->GetContainedGraph()->FindNode(Right);
-			}
-		}
-		if(BreakpointNode)
-		{
-			break;
-		}
-	}
-
-	if(BreakpointNode)
-	{
-		bool bSuccess = RemoveBreakpoint(BreakpointNode);
-
-		// Remove the breakpoint from all the loaded dependent blueprints
-		TArray<UControlRigBlueprint*> DependentBlueprints = GetDependentBlueprints(true, true);
-		DependentBlueprints.Remove(this);
-		for (UControlRigBlueprint* Dependent : DependentBlueprints)
-		{
-			bSuccess &= Dependent->RemoveBreakpoint(BreakpointNode);
-		}
-		return bSuccess;
-	}
-
-	return false;
-}
-
-bool UControlRigBlueprint::RemoveBreakpoint(URigVMNode* InBreakpointNode)
-{
-	if (RigVMBreakpointNodes.Contains(InBreakpointNode))
-	{
-		RigVMBreakpointNodes.Remove(InBreakpointNode);
-
-		// Multiple breakpoint nodes might set a breakpoint to the same instruction. When we remove
-		// one of the breakpoint nodes, we do not want to remove the instruction breakpoint if there
-		// is another breakpoint node addressing it. For that reason, we just recompute all the
-		// breakpoint instructions.
-		// Refreshing breakpoints in the control rig will keep the state it had before.
-		RefreshControlRigBreakpoints();
-		return true;
-	}
-
-	return false;
-}
-
-
-void UControlRigBlueprint::RefreshControlRigBreakpoints()
-{
-	if (URigVMBlueprintGeneratedClass* RigClass = GetControlRigBlueprintGeneratedClass())
-	{
-		UControlRig* CDO = Cast<UControlRig>(RigClass->GetDefaultObject(false));
-		CDO->GetDebugInfo().Reset();
-		for (URigVMNode* Node : RigVMBreakpointNodes)
-		{
-			AddBreakpointToControlRig(Node);
-		}
-	}
-}
-
-TArray<FRigVMReferenceNodeData> UControlRigBlueprint::GetReferenceNodeData() const
-{
-	TArray<FRigVMReferenceNodeData> Data;
-	
-	const TArray<URigVMGraph*> AllModels = GetAllModels();
-	for (URigVMGraph* ModelToVisit : AllModels)
-	{
-		for(URigVMNode* Node : ModelToVisit->GetNodes())
-		{
-			if(URigVMFunctionReferenceNode* ReferenceNode = Cast<URigVMFunctionReferenceNode>(Node))
-			{
-				Data.Add(FRigVMReferenceNodeData(ReferenceNode));
-			}
-		}
-	}
-	return Data;
-}
-
-#endif
-
-void UControlRigBlueprint::RequestControlRigInit()
-{
-	if (URigVMBlueprintGeneratedClass* RigClass = GetControlRigBlueprintGeneratedClass())
-	{
-		UControlRig* CDO = Cast<UControlRig>(RigClass->GetDefaultObject(true /* create if needed */));
-		CDO->RequestInit();
-
-		TArray<UObject*> ArchetypeInstances;
-		CDO->GetArchetypeInstances(ArchetypeInstances);
-		for (UObject* Instance : ArchetypeInstances)
-		{
-			if (UControlRig* InstanceRig = Cast<UControlRig>(Instance))
-			{
-				InstanceRig->RequestInit();
-			}
-		}
-	}
-}
-
-URigVMGraph* UControlRigBlueprint::GetModel(const UEdGraph* InEdGraph) const
-{
-	if (InEdGraph == nullptr)
-	{
-		return GetDefaultModel();
-	}
-
-	if(InEdGraph->GetOutermost() != GetOutermost())
-	{
-		return nullptr;
-	}
-
-#if WITH_EDITORONLY_DATA
-	if (InEdGraph == FunctionLibraryEdGraph)
-	{
-		return RigVMClient.GetFunctionLibrary();
-	}
-#endif
-
-	const UControlRigGraph* RigGraph = Cast< UControlRigGraph>(InEdGraph);
-	check(RigGraph);
-	return GetModel(RigGraph->ModelNodePath);
-}
-
-URigVMGraph* UControlRigBlueprint::GetModel(const FString& InNodePath) const
-{
-	return RigVMClient.GetModel(InNodePath);
-}
-
-URigVMGraph* UControlRigBlueprint::GetDefaultModel() const
-{
-	return RigVMClient.GetDefaultModel();
-}
-
-TArray<URigVMGraph*> UControlRigBlueprint::GetAllModels() const
-{
-	return RigVMClient.GetAllModels(true, true);
-}
-
-URigVMFunctionLibrary* UControlRigBlueprint::GetLocalFunctionLibrary() const
-{
-	return RigVMClient.GetFunctionLibrary();
-}
-
-URigVMGraph* UControlRigBlueprint::AddModel(FString InName, bool bSetupUndoRedo, bool bPrintPythonCommand)
-{
-	const FString DesiredName = FString::Printf(TEXT("%s %s"),
-    	FRigVMClient::RigVMModelPrefix, *InName);
-
-	TGuardValue<bool> EnablePythonPrint(bSuspendPythonMessagesForRigVMClient, !bPrintPythonCommand);
-	return RigVMClient.AddModel(*DesiredName, bSetupUndoRedo);
-}
-
-bool UControlRigBlueprint::RemoveModel(FString InName, bool bSetupUndoRedo, bool bPrintPythonCommand)
-{
-	TGuardValue<bool> EnablePythonPrint(bSuspendPythonMessagesForRigVMClient, !bPrintPythonCommand);
-	return RigVMClient.RemoveModel(InName, bSetupUndoRedo);
-}
-
-URigVMGraph* UControlRigBlueprint::GetFocusedModel() const
-{
-#if WITH_EDITOR
-	if(OnGetFocusedGraphDelegate.IsBound())
-	{
-		return OnGetFocusedGraphDelegate.Execute(); 
-	}
-#endif
-	return RigVMClient.GetDefaultModel();
-}
-
-URigVMController* UControlRigBlueprint::GetController(const URigVMGraph* InGraph) const
-{
-	return RigVMClient.GetController(InGraph);
-}
-
-URigVMController* UControlRigBlueprint::GetControllerByName(const FString InGraphName) const
-{
-	if(InGraphName.IsEmpty())
-	{
-		if(const URigVMGraph* DefaultModel = GetRigVMClient()->GetDefaultModel())
-		{
-			return GetController(DefaultModel);
-		}
-	}
-	
-	for (const URigVMGraph* Graph : GetAllModels())
-	{
-		if (Graph->GetName() == InGraphName || Graph->GetGraphName() == InGraphName)
-		{
-			return GetController(Graph);
-		}
-	}
-	
-	return nullptr;
-}
-
-URigVMController* UControlRigBlueprint::GetOrCreateController(URigVMGraph* InGraph)
-{
-	return RigVMClient.GetOrCreateController(InGraph);
-}
-
-URigVMController* UControlRigBlueprint::GetController(const UEdGraph* InEdGraph) const
-{
-	return RigVMClient.GetController(InEdGraph);
-}
-
-URigVMController* UControlRigBlueprint::GetOrCreateController(const UEdGraph* InEdGraph)
-{
-	return RigVMClient.GetOrCreateController(InEdGraph);
-}
-
-TArray<FString> UControlRigBlueprint::GeneratePythonCommands(const FString InNewBlueprintName)
-{
-	TArray<FString> InternalCommands;
-	InternalCommands.Add(TEXT("unreal.load_module('ControlRigDeveloper')"));
-	InternalCommands.Add(TEXT("factory = unreal.ControlRigBlueprintFactory"));
-	InternalCommands.Add(FString::Printf(TEXT("blueprint = factory.create_new_control_rig_asset(desired_package_path = '%s')"), *InNewBlueprintName));
-	InternalCommands.Add(TEXT("library = blueprint.get_local_function_library()"));
-	InternalCommands.Add(TEXT("library_controller = blueprint.get_controller(library)"));
-	InternalCommands.Add(TEXT("hierarchy = blueprint.hierarchy"));
-	InternalCommands.Add(TEXT("hierarchy_controller = hierarchy.get_controller()"));
-	InternalCommands.Add(TEXT("blueprint.set_auto_vm_recompile(False)"));
-	
-
-	// Hierarchy
-	InternalCommands.Append(Hierarchy->GetController(true)->GeneratePythonCommands());
-		
-	// Add variables
-	for (const FBPVariableDescription& Variable : NewVariables)
-	{
-		const FRigVMExternalVariable ExternalVariable = RigVMTypeUtils::ExternalVariableFromBPVariableDescription(Variable);
-		FString CPPType;
-		UObject* CPPTypeObject = nullptr;
-		RigVMTypeUtils::CPPTypeFromExternalVariable(ExternalVariable, CPPType, &CPPTypeObject);
-		if (CPPTypeObject)
-		{
-			if (ExternalVariable.bIsArray)
-			{
-				CPPType = RigVMTypeUtils::ArrayTypeFromBaseType(CPPTypeObject->GetPathName());
-			}
-			else
-			{
-				CPPType = CPPTypeObject->GetPathName();
-			}
-		}
-		// FName AddMemberVariable(const FName& InName, const FString& InCPPType, bool bIsPublic = false, bool bIsReadOnly = false, FString InDefaultValue = TEXT(""));
-		InternalCommands.Add(FString::Printf(TEXT("blueprint.add_member_variable('%s', '%s', %s, %s)"),
-					*ExternalVariable.Name.ToString(),
-					*CPPType,
-					ExternalVariable.bIsPublic ? TEXT("True") : TEXT("False"),
-					ExternalVariable.bIsReadOnly ? TEXT("True") : TEXT("False")));	
-	}
-	
-	// Create graphs
-	{
-		TArray<URigVMGraph*> AllModels = GetAllModels();
-		AllModels.RemoveAll([](const URigVMGraph* GraphToRemove) -> bool
-		{
-			return GraphToRemove->GetTypedOuter<URigVMAggregateNode>() != nullptr;
-		});
-		
-		// Find all graphs to process and sort them by dependencies
-		TArray<URigVMGraph*> ProcessedGraphs;
-		while (ProcessedGraphs.Num() < AllModels.Num())
-		{
-			for (URigVMGraph* Graph : AllModels)
-			{
-				if (ProcessedGraphs.Contains(Graph))
-				{
-					continue;
-				}
-
-				bool bFoundUnprocessedReference = false;
-				for (auto Node : Graph->GetNodes())
-				{
-					if (URigVMFunctionReferenceNode* Reference = Cast<URigVMFunctionReferenceNode>(Node))
-					{
-						if (Reference->GetReferencedFunctionHeader().LibraryPointer.HostObject != GetControlRigBlueprintGeneratedClass())
-						{
-							continue;
-						}
-
-						URigVMLibraryNode* LibraryNode = Cast<URigVMLibraryNode>(Reference->GetReferencedFunctionHeader().LibraryPointer.LibraryNode.ResolveObject());
-						if (!ProcessedGraphs.Contains(LibraryNode->GetContainedGraph()))
-						{
-							bFoundUnprocessedReference = true;
-							break;
-						}
-					}
-					else if (URigVMCollapseNode* CollapseNode = Cast<URigVMCollapseNode>(Node))
-					{
-						if(!CollapseNode->IsA<URigVMAggregateNode>())
-						{
-							if (!ProcessedGraphs.Contains(CollapseNode->GetContainedGraph()))
-							{
-								bFoundUnprocessedReference = true;
-								break;
-							}
-						}
-					}
-				}
-
-				if (!bFoundUnprocessedReference)
-				{
-					ProcessedGraphs.Add(Graph);
-				}
-			}
-		}	
-
-		// Dump python commands for each graph
-		for (URigVMGraph* Graph : ProcessedGraphs)
-		{
-			if (Graph->IsA<URigVMFunctionLibrary>())
-			{
-				continue;
-			}
-
-			URigVMController* Controller = GetController(Graph);
-			if (Graph->GetParentGraph()) 
-			{
-				// Add them all as functions (even collapsed graphs)
-				// The controller will deal with deleting collapsed graph function when it creates the collapse node
-				{						
-					// Add Function
-					InternalCommands.Add(FString::Printf(TEXT("function_%s = library_controller.add_function_to_library('%s', mutable=%s)\ngraph = function_%s.get_contained_graph()"),
-							*RigVMPythonUtils::PythonizeName(Graph->GetGraphName()),
-							*Graph->GetGraphName(),
-							Graph->GetEntryNode()->IsMutable() ? TEXT("True") : TEXT("False"),
-							*RigVMPythonUtils::PythonizeName(Graph->GetGraphName())));
-
-					URigVMLibraryNode* LibraryNode = Cast<URigVMLibraryNode>(Graph->GetOuter());
-
-					InternalCommands.Add(FString::Printf(TEXT("library_controller.set_node_category_by_name('%s', '%s')"),
-						*Graph->GetGraphName(),
-						*LibraryNode->GetNodeCategory()));
-
-					InternalCommands.Add(FString::Printf(TEXT("library_controller.set_node_keywords_by_name('%s', '%s')"),
-						*Graph->GetGraphName(),
-						*LibraryNode->GetNodeKeywords() ));
-
-					InternalCommands.Add(FString::Printf(TEXT("library_controller.set_node_description_by_name('%s', '%s')"),
-						*Graph->GetGraphName(),
-						*LibraryNode->GetNodeDescription()));
-
-					InternalCommands.Add(FString::Printf(TEXT("library_controller.set_node_color_by_name('%s', %s)"),
-						*Graph->GetGraphName(),
-						*RigVMPythonUtils::LinearColorToPythonString(LibraryNode->GetNodeColor()) ));
-					
-					URigVMFunctionEntryNode* EntryNode = Graph->GetEntryNode();
-					URigVMFunctionReturnNode* ReturnNode = Graph->GetReturnNode();
-					
-					// Set Entry and Return nodes in the correct position
-					{
-						//bool SetNodePositionByName(const FName& InNodeName, const FVector2D& InPosition, bool bSetupUndoRedo = true, bool bMergeUndoAction = false);
-						InternalCommands.Add(FString::Printf(TEXT("blueprint.get_controller_by_name('%s').set_node_position_by_name('Entry', unreal.Vector2D(%f, %f))"),
-								*Graph->GetGraphName(),
-								EntryNode->GetPosition().X, 
-								EntryNode->GetPosition().Y));
-
-						InternalCommands.Add(FString::Printf(TEXT("blueprint.get_controller_by_name('%s').set_node_position_by_name('Return', unreal.Vector2D(%f, %f))"),
-								*Graph->GetGraphName(),
-								ReturnNode->GetPosition().X, 
-								ReturnNode->GetPosition().Y));
-					}
-					
-					// Add Exposed Pins
-					{
-
-						bool bHitFirstExecute = false;
-						bool bRenamedExecute = false;
-						for (auto Pin : EntryNode->GetPins())
-						{
-							if (Pin->GetDirection() != ERigVMPinDirection::Output)
-							{
-								continue;
-							}
-
-							if(Pin->IsExecuteContext())
-							{
-								if(!bHitFirstExecute)
-								{
-									bHitFirstExecute = true;
-									if (Pin->GetName() != FRigVMStruct::ExecuteContextName.ToString())
-									{
-										bRenamedExecute = true;
-										InternalCommands.Add(FString::Printf(TEXT("blueprint.get_controller_by_name('%s').rename_exposed_pin('%s', '%s')"),
-											*Graph->GetGraphName(),
-											*FRigVMStruct::ExecuteContextName.ToString(),
-											*Pin->GetName()));
-									}
-									continue;
-								}
-							}
-
-							// FName AddExposedPin(const FName& InPinName, ERigVMPinDirection InDirection, const FString& InCPPType, const FName& InCPPTypeObjectPath, const FString& InDefaultValue, bool bSetupUndoRedo = true);
-							InternalCommands.Add(FString::Printf(TEXT("blueprint.get_controller_by_name('%s').add_exposed_pin('%s', %s, '%s', '%s', '%s')"),
-									*Graph->GetGraphName(),
-									*Pin->GetName(),
-									*RigVMPythonUtils::EnumValueToPythonString<ERigVMPinDirection>((int64)ERigVMPinDirection::Input),
-									*Pin->GetCPPType(),
-									Pin->GetCPPTypeObject() ? *Pin->GetCPPTypeObject()->GetPathName() : TEXT(""),
-									*Pin->GetDefaultValue()));
-						}
-
-						bHitFirstExecute = false;
-						for (auto Pin : ReturnNode->GetPins())
-						{
-							if (Pin->GetDirection() != ERigVMPinDirection::Input)
-							{
-								continue;
-							}
-
-							if(Pin->IsExecuteContext())
-							{
-								if(!bHitFirstExecute)
-								{
-									bHitFirstExecute = true;
-									if (!bRenamedExecute && Pin->GetName() != FRigVMStruct::ExecuteContextName.ToString())
-									{
-										bRenamedExecute = true;
-										InternalCommands.Add(FString::Printf(TEXT("blueprint.get_controller_by_name('%s').rename_exposed_pin('%s', '%s')"),
-											*Graph->GetGraphName(),
-											*FRigVMStruct::ExecuteContextName.ToString(),
-											*Pin->GetName()));
-									}
-									continue;
-								}
-							}
-
-							// FName AddExposedPin(const FName& InPinName, ERigVMPinDirection InDirection, const FString& InCPPType, const FName& InCPPTypeObjectPath, const FString& InDefaultValue, bool bSetupUndoRedo = true);
-							InternalCommands.Add(FString::Printf(TEXT("blueprint.get_controller_by_name('%s').add_exposed_pin('%s', %s, '%s', '%s', '%s')"),
-									*Graph->GetGraphName(),
-									*Pin->GetName(),
-									*RigVMPythonUtils::EnumValueToPythonString<ERigVMPinDirection>((int64)ERigVMPinDirection::Output),
-									*Pin->GetCPPType(),
-									Pin->GetCPPTypeObject() ? *Pin->GetCPPTypeObject()->GetPathName() : TEXT(""),
-									*Pin->GetDefaultValue()));
-						}
-					}
-				}
-			}
-			else if(Graph != GetDefaultModel())
-			{
-				InternalCommands.Add(FString::Printf(TEXT("blueprint.add_model('%s')"),
-						*Graph->GetName()));
-			}
-								
-			InternalCommands.Append(Controller->GeneratePythonCommands());
-		}
-	}
-
-#if WITH_EDITORONLY_DATA
-	FString PreviewMeshPath = GetPreviewMesh()->GetPathName();
-	InternalCommands.Add(FString::Printf(TEXT("blueprint.set_preview_mesh(unreal.load_object(name='%s', outer=None))"),
-		*PreviewMeshPath));
-#endif
-
-	InternalCommands.Add(TEXT("blueprint.set_auto_vm_recompile(True)"));
-
-	// Split multiple commands into different array elements
-	TArray<FString> InnerFunctionCmds;
-	for (FString Cmd : InternalCommands)
-	{
-		FString Left, Right=Cmd;
-		while (Right.Split(TEXT("\n"), &Left, &Right))
-		{
-			InnerFunctionCmds.Add(Left);			
-		}
-		InnerFunctionCmds.Add(Right);
-	}
-
-	// Define a function and insert all the commands
-	// We do not want to pollute the global state with our definitions
-	TArray<FString> Commands;
-	Commands.Add(FString::Printf(TEXT("import unreal\n"
-		"def create_control_rig():\n")));
-	for (const FString& InnerCmd : InnerFunctionCmds)
-	{
-		Commands.Add(FString::Printf(TEXT("\t%s"), *InnerCmd));
-	}
-
-	Commands.Add(TEXT("create_control_rig()\n"));
-	return Commands;
-}
-
-
-URigVMGraph* UControlRigBlueprint::GetTemplateModel(bool bIsFunctionLibrary)
-{
-#if WITH_EDITORONLY_DATA
-	if (TemplateModel == nullptr)
-	{
-		if (bIsFunctionLibrary)
-		{
-			TemplateModel = NewObject<URigVMFunctionLibrary>(this, TEXT("TemplateFunctionLibrary"));
-		}
-		else
-		{
-			TemplateModel = NewObject<URigVMGraph>(this, TEXT("TemplateModel"));
-		}
-		TemplateModel->SetFlags(RF_Transient);
-		TemplateModel->SetExecuteContextStruct(RigVMClient.GetExecuteContextStruct());
-	}
-	return TemplateModel;
-#else
-	return nullptr;
-#endif
-}
-
-URigVMController* UControlRigBlueprint::GetTemplateController(bool bIsFunctionLibrary)
-{
-#if WITH_EDITORONLY_DATA
-	if (TemplateController == nullptr)
-	{
-		TemplateController = NewObject<URigVMController>(this, TEXT("TemplateController"));
-		TemplateController->SetGraph(GetTemplateModel(bIsFunctionLibrary));
-		TemplateController->EnableReporting(false);
-		TemplateController->SetFlags(RF_Transient);
-		TemplateController->SetSchema(RigVMClient.GetOrCreateSchema());
-	}
-	return TemplateController;
-#else
-	return nullptr;
-#endif
-}
-
-UEdGraph* UControlRigBlueprint::GetEdGraph(URigVMGraph* InModel) const
-{
-	return Cast<UEdGraph>(GetEditorObjectForRigVMGraph(InModel));
-}
-
-UEdGraph* UControlRigBlueprint::GetEdGraph(const FString& InNodePath) const
-{
-	if (URigVMGraph* ModelForNodePath = GetModel(InNodePath))
-	{
-		return GetEdGraph(ModelForNodePath);
-	}
-	return nullptr;
-}
-
-bool UControlRigBlueprint::IsFunctionPublic(const FName& InFunctionName) const
-{
-	return GetLocalFunctionLibrary()->IsFunctionPublic(InFunctionName);	
-}
-
-void UControlRigBlueprint::MarkFunctionPublic(const FName& InFunctionName, bool bIsPublic)
-{
-	if(IsFunctionPublic(InFunctionName) == bIsPublic)
-	{
-		return;
-	}
-	
-	URigVMController* Controller = RigVMClient.GetOrCreateController(GetLocalFunctionLibrary());
-	Controller->MarkFunctionAsPublic(InFunctionName, bIsPublic);
-}
-
-TArray<UControlRigBlueprint*> UControlRigBlueprint::GetDependencies(bool bRecursive) const
-{
-	TArray<UControlRigBlueprint*> Dependencies;
-
-	TArray<URigVMGraph*> Graphs = GetAllModels();
-	for(URigVMGraph* Graph : Graphs)
-	{
-		for(URigVMNode* Node : Graph->GetNodes())
-		{
-			if(URigVMFunctionReferenceNode* FunctionReferenceNode = Cast<URigVMFunctionReferenceNode>(Node))
-			{
-				if(const URigVMLibraryNode* LibraryNode = FunctionReferenceNode->LoadReferencedNode())
-				{
-					if(UControlRigBlueprint* DependencyBlueprint = LibraryNode->GetTypedOuter<UControlRigBlueprint>())
-					{
-						if(DependencyBlueprint != this)
-						{
-							if(!Dependencies.Contains(DependencyBlueprint))
-							{
-								Dependencies.Add(DependencyBlueprint);
-
-								if(bRecursive)
-								{
-									TArray<UControlRigBlueprint*> ChildDependencies = DependencyBlueprint->GetDependencies(true);
-									for(UControlRigBlueprint* ChildDependency : ChildDependencies)
-									{
-										Dependencies.AddUnique(ChildDependency);
-									}
-								}
-							}
-						}
-					}
-				}
-			}
-		}
-	}
-
-	return Dependencies;
-}
-
-TArray<FAssetData> UControlRigBlueprint::GetDependentAssets() const
-{
-	TArray<FAssetData> Dependents;
-	TArray<FSoftObjectPath> AssetPaths;
-
-	if(URigVMFunctionLibrary* FunctionLibrary = RigVMClient.GetFunctionLibrary())
-	{
-		IAssetRegistry& AssetRegistry = IAssetRegistry::GetChecked();
-		TArray<URigVMLibraryNode*> Functions = FunctionLibrary->GetFunctions();
-		for(URigVMLibraryNode* Function : Functions)
-		{
-			const FName FunctionName = Function->GetFName();
-			if(IsFunctionPublic(FunctionName))
-			{
-				TArray<TSoftObjectPtr<URigVMFunctionReferenceNode>> References = FunctionLibrary->GetReferencesForFunction(FunctionName);
-				for(const TSoftObjectPtr<URigVMFunctionReferenceNode>& Reference : References)
-				{
-					if (const URigVMFunctionReferenceNode* ReferencePtr = Reference.Get())
-					{
-						if (const UControlRigBlueprint* ControlRigBlueprint = ReferencePtr->GetTypedOuter<UControlRigBlueprint>())
-						{
-							const TSoftObjectPtr<UPackage> Blueprint = ControlRigBlueprint;
-							const FSoftObjectPath AssetPath = Blueprint.ToSoftObjectPath();
-							if(AssetPath.GetLongPackageName().StartsWith(TEXT("/Engine/Transient")))
-							{
-								continue;
-							}
-				
-							if(!AssetPaths.Contains(AssetPath))
-							{
-								AssetPaths.Add(AssetPath);
-
-								const FAssetData AssetData = AssetRegistry.GetAssetByObjectPath(AssetPath);
-								if(AssetData.IsValid())
-								{
-									Dependents.Add(AssetData);
-								}
-							}						
-						}
-					}
-				}
-			}
-		}
-	}
-
-	return Dependents;
-}
-
-TArray<UControlRigBlueprint*> UControlRigBlueprint::GetDependentBlueprints(bool bRecursive, bool bOnlyLoaded) const
-{
-	TArray<FAssetData> Assets = GetDependentAssets();
-	TArray<UControlRigBlueprint*> Dependents;
-
-	for(const FAssetData& Asset : Assets)
-	{
-		if (!bOnlyLoaded || Asset.IsAssetLoaded())
-		{
-			if(UControlRigBlueprint* Dependent = Cast<UControlRigBlueprint>(Asset.GetAsset()))
-			{
-				if(!Dependents.Contains(Dependent))
-				{
-					Dependents.Add(Dependent);
-
-					if(bRecursive && Dependent != this)
-					{
-						TArray<UControlRigBlueprint*> ParentDependents = Dependent->GetDependentBlueprints(true);
-						for(UControlRigBlueprint* ParentDependent : ParentDependents)
-						{
-							Dependents.AddUnique(ParentDependent);
-						}
-					}
-				}
-			}
-		}
-	}
-
-	return Dependents;
-}
-
-void UControlRigBlueprint::GetTypeActions(FBlueprintActionDatabaseRegistrar& ActionRegistrar) const
-{
-	DECLARE_SCOPE_HIERARCHICAL_COUNTER_FUNC()
-
-	IControlRigEditorModule::Get().GetTypeActions((UControlRigBlueprint*)this, ActionRegistrar);
-}
-
-void UControlRigBlueprint::GetInstanceActions(FBlueprintActionDatabaseRegistrar& ActionRegistrar) const
-{
-	DECLARE_SCOPE_HIERARCHICAL_COUNTER_FUNC()
-
-	IControlRigEditorModule::Get().GetInstanceActions((UControlRigBlueprint*)this, ActionRegistrar);
-}
-
-void UControlRigBlueprint::SetObjectBeingDebugged(UObject* NewObject)
-{
-	UControlRig* PreviousRigBeingDebugged = Cast<UControlRig>(GetObjectBeingDebugged());
-	if (PreviousRigBeingDebugged && PreviousRigBeingDebugged != NewObject)
-	{
-		PreviousRigBeingDebugged->DrawInterface.Reset();
-		PreviousRigBeingDebugged->RigVMLog = nullptr;
-	}
-
-	Super::SetObjectBeingDebugged(NewObject);
-
-	if (Validator)
-	{
-		if (Validator->GetControlRig() != nullptr)
-		{
-			Validator->SetControlRig(Cast<UControlRig>(GetObjectBeingDebugged()));
-		}
-	}
-}
-
-void UControlRigBlueprint::PostTransacted(const FTransactionObjectEvent& TransactionEvent)
-{
-	DECLARE_SCOPE_HIERARCHICAL_COUNTER_FUNC()
-	Super::PostTransacted(TransactionEvent);
-
-	if (TransactionEvent.GetEventType() == ETransactionObjectEventType::UndoRedo)
-	{
-		TArray<FName> PropertiesChanged = TransactionEvent.GetChangedProperties();
-		if (PropertiesChanged.Contains(GET_MEMBER_NAME_CHECKED(UControlRigBlueprint, Hierarchy)))
-		{
-			int32 TransactionIndex = GEditor->Trans->FindTransactionIndex(TransactionEvent.GetTransactionId());
-			const FTransaction* Transaction = GEditor->Trans->GetTransaction(TransactionIndex);
-
-			if (Transaction->GetTitle().BuildSourceString() == TEXT("Transform Gizmo"))
-			{
-				PropagatePoseFromBPToInstances();
-				return;
-			}
-
-			PropagateHierarchyFromBPToInstances();
-
-			// make sure the bone name list is up 2 date for the editor graph
-			for (UEdGraph* Graph : UbergraphPages)
-			{
-				UControlRigGraph* RigGraph = Cast<UControlRigGraph>(Graph);
-				if (RigGraph == nullptr)
-				{
-					continue;
-				}
-				RigGraph->CacheNameLists(Hierarchy, &DrawContainer, ShapeLibraries);
-			}
-
-			RequestAutoVMRecompilation();
-			MarkPackageDirty();
-		}
-
-		if (PropertiesChanged.Contains(GET_MEMBER_NAME_CHECKED(UControlRigBlueprint, DrawContainer)))
-		{
-			PropagateDrawInstructionsFromBPToInstances();
-		}
-
-		if (PropertiesChanged.Contains(GET_MEMBER_NAME_CHECKED(UControlRigBlueprint, VMRuntimeSettings)))
-		{
-			PropagateRuntimeSettingsFromBPToInstances();
-		}
-
-		if (PropertiesChanged.Contains(GET_MEMBER_NAME_CHECKED(UControlRigBlueprint, NewVariables)))
-		{
-			if (RefreshEditorEvent.IsBound())
-			{
-				RefreshEditorEvent.Broadcast(this);
-			}
-			MarkPackageDirty();			
-		}
-
-		if (PropertiesChanged.Contains(GET_MEMBER_NAME_CHECKED(UControlRigBlueprint, RigVMClient)) ||
-			PropertiesChanged.Contains(GET_MEMBER_NAME_CHECKED(UControlRigBlueprint, UbergraphPages)))
-		{
-			UbergraphPages.RemoveAll([](const UEdGraph* UberGraph) -> bool
-			{
- 				return UberGraph == nullptr || !IsValid(UberGraph);
-			});
-			RigVMClient.PostTransacted(TransactionEvent);
-
-			RecompileVM();
-			MarkPackageDirty();			
-		}
-	}
-}
-
-void UControlRigBlueprint::ReplaceDeprecatedNodes()
-{
-	TArray<UEdGraph*> EdGraphs;
-	GetAllGraphs(EdGraphs);
-
-	for (UEdGraph* EdGraph : EdGraphs)
-	{
-		if (UControlRigGraph* RigGraph = Cast<UControlRigGraph>(EdGraph))
-		{
-			RigGraph->Schema = UControlRigGraphSchema::StaticClass();
-		}
-	}
-
-	Super::ReplaceDeprecatedNodes();
-}
-
-void UControlRigBlueprint::PreDuplicate(FObjectDuplicationParameters& DupParams)
-{
-	Super::PreDuplicate(DupParams);
-	PreDuplicateAssetPath = GetPathName();
-	if(URigVMBlueprintGeneratedClass* CRGeneratedClass = GetControlRigBlueprintGeneratedClass())
-	{
-		PreDuplicateHostPath = FSoftObjectPath(CRGeneratedClass->GetPathName());
-	}
-	else
-	{
-		PreDuplicateHostPath.Reset();
-	}
-
-}
-
-void UControlRigBlueprint::PostDuplicate(bool bDuplicateForPIE)
-{
-	// assuming PostDuplicate is always followed by a PostLoad:
-	// so theoretically, PostDuplicate just makes corrections to the serialized data and does nothing more,
-	// while PostLoad looks at whatever is serialized and load it into memory according to the version of the editor used
-	// note: how to know if we have corrected everything?
-	// ans: check the reference viewer for the duplicated BP and make sure that the original BP does not appear in there
-	
-	{
-		// pause compilation because we need to patch some stuff first
-		TGuardValue<bool> CompilingGuard(bIsCompiling, true);
-		// this will create the new EMPTY generated class to be used as the function store for this BP
-		// it will be filled during PostLoad based on the graph model
-		Super::PostDuplicate(bDuplicateForPIE);
-	}
-	
-	if (URigHierarchyController* Controller = Hierarchy->GetController(true))
-	{
-		Controller->OnModified().RemoveAll(this);
-		Controller->OnModified().AddUObject(this, &UControlRigBlueprint::HandleHierarchyModified);
-	}
-
-	auto UpdateFunctionHeaders = [this](const FString& InOldPath, const FString& InNewPath)
-	{
-		if(InOldPath.IsEmpty() || InNewPath.IsEmpty())
-		{
-			return;
-		}
-		if(!InNewPath.Equals(InOldPath, ESearchCase::CaseSensitive))
-		{
-			if(URigVMBlueprintGeneratedClass* CRGeneratedClass = GetControlRigBlueprintGeneratedClass())
-			{
-				FRigVMGraphFunctionStore& Store = CRGeneratedClass->GraphFunctionStore;
-				// technically not needed, the store should be empty, it will not be populated until PostLoad
-				// this code is kept here just in case things change in the future
-				if (!(ensure(Store.PublicFunctions.Num() == 0) && ensure(Store.PrivateFunctions.Num() == 0)))
-				{
-					Store.PostDuplicateHost(InOldPath, InNewPath);
-				}
-			}
-			RigVMClient.PostDuplicateHost(InOldPath, InNewPath);
-		}
-	};
-
-	// update the paths once for the blueprint and once for the generated class
-	// make sure all function headers pointing to things in the old BP are changed to
-	// point to their duplicates in the new BP
-	UpdateFunctionHeaders(PreDuplicateAssetPath.ToString(), GetPathName());
-	if(const URigVMBlueprintGeneratedClass* CRGeneratedClass = GetControlRigBlueprintGeneratedClass())
-	{
-		UpdateFunctionHeaders(PreDuplicateHostPath.ToString(), CRGeneratedClass->GetPathName());
-	}
-
-	PreDuplicateAssetPath.Reset();
-	PreDuplicateHostPath.Reset();
-
-	FBlueprintEditorUtils::MarkBlueprintAsStructurallyModified(this);
-
-	// now that the data in the new BP is correct, make sure there is a call to PostLoad()
-	// to complete the loading and get the BP ready, including refresh all models, patch function store, and recompile VM
-	check(HasAnyFlags(RF_NeedPostLoad));
-}
-
-void UControlRigBlueprint::GetAssetRegistryTags(TArray<FAssetRegistryTag>& OutTags) const
-{
-	if (CachedAssetTags.IsEmpty())
-	{
-		Super::GetAssetRegistryTags(OutTags);
-		CachedAssetTags = OutTags;
-	}
-	else
-	{
-		OutTags = CachedAssetTags;
-	}
-}
-
-FRigVMGraphModifiedEvent& UControlRigBlueprint::OnModified()
-{
-	return ModifiedEvent;
-}
-
-
-FOnVMCompiledEvent& UControlRigBlueprint::OnVMCompiled()
-{
-	return VMCompiledEvent;
-}
-
-TArray<UControlRigBlueprint*> UControlRigBlueprint::GetCurrentlyOpenRigBlueprints()
-{
-	return sCurrentlyOpenedRigBlueprints;
-}
-
-UClass* UControlRigBlueprint::GetControlRigClass()
-{
-	return GeneratedClass;
-}
-
-UControlRig* UControlRigBlueprint::CreateControlRig()
-{
-	RecompileVMIfRequired();
-
-	UControlRig* Rig = NewObject<UControlRig>(this, GetControlRigClass());
-	Rig->Initialize(true);
-	return Rig;
-}
-
-TArray<UStruct*> UControlRigBlueprint::GetAvailableRigUnits()
-{
-	TArray<UStruct*> Structs;
-	UStruct* BaseStruct = FRigUnit::StaticStruct();
-
-	for (const FRigVMFunction& Function : FRigVMRegistry::Get().GetFunctions())
-	{
-		if (Function.Struct)
-		{
-			if (Function.Struct->IsChildOf(BaseStruct))
-			{
-				Structs.Add(Function.Struct);
-			}
-		}
-	}
-
-	return Structs;
-}
-
-#if WITH_EDITOR
-
-TArray<FRigVMGraphVariableDescription> UControlRigBlueprint::GetMemberVariables() const
-{
-	TArray<FRigVMGraphVariableDescription> Variables;
-	for (const FBPVariableDescription& BPVariable : NewVariables)
-	{
-		FRigVMGraphVariableDescription NewVariable;
-		NewVariable.Name = BPVariable.VarName;
-		NewVariable.DefaultValue = BPVariable.DefaultValue;
-		FString CPPType;
-		UObject* CPPTypeObject;
-		RigVMTypeUtils::CPPTypeFromPinType(BPVariable.VarType, CPPType, &CPPTypeObject);
-		NewVariable.CPPType = CPPType;
-		NewVariable.CPPTypeObject = CPPTypeObject;
-		Variables.Add(NewVariable);
-	}
-
-	return Variables;
-}
-
-FName UControlRigBlueprint::AddMemberVariable(const FName& InName, const FString& InCPPType, bool bIsPublic, bool bIsReadOnly, FString InDefaultValue)
-{
-	FRigVMExternalVariable Variable = RigVMTypeUtils::ExternalVariableFromCPPTypePath(InName, InCPPType, bIsPublic, bIsReadOnly);
-	FName Result = AddCRMemberVariableFromExternal(Variable, InDefaultValue);
-	if (!Result.IsNone())
-	{
-		FBPCompileRequest Request(this, EBlueprintCompileOptions::None, nullptr);
-		FBlueprintCompilationManager::CompileSynchronously(Request);
-	}
-	return Result;
-}
-
-bool UControlRigBlueprint::RemoveMemberVariable(const FName& InName)
-{
-	const int32 VarIndex = FBlueprintEditorUtils::FindNewVariableIndex(this, InName);
-	if (VarIndex == INDEX_NONE)
-	{
-		return false;
-	}
-	
-	FBlueprintEditorUtils::RemoveMemberVariable(this, InName);
-	return true;
-}
-
-bool UControlRigBlueprint::RenameMemberVariable(const FName& InOldName, const FName& InNewName)
-{
-	int32 VarIndex = FBlueprintEditorUtils::FindNewVariableIndex(this, InOldName);
-	if (VarIndex == INDEX_NONE)
-	{
-		return false;
-	}
-
-	VarIndex = FBlueprintEditorUtils::FindNewVariableIndex(this, InNewName);
-	if (VarIndex != INDEX_NONE)
-	{
-		return false;
-	}
-	
-	FBlueprintEditorUtils::RenameMemberVariable(this, InOldName, InNewName);
-	return true;
-}
-
-bool UControlRigBlueprint::ChangeMemberVariableType(const FName& InName, const FString& InCPPType, bool bIsPublic,
-	bool bIsReadOnly, FString InDefaultValue)
-{
-	int32 VarIndex = FBlueprintEditorUtils::FindNewVariableIndex(this, InName);
-	if (VarIndex == INDEX_NONE)
-	{
-		return false;
-	}
-
-	FRigVMExternalVariable Variable;
-	Variable.Name = InName;
-	Variable.bIsPublic = bIsPublic;
-	Variable.bIsReadOnly = bIsReadOnly;
-
-	FString CPPType = InCPPType;
-	if (CPPType.StartsWith(TEXT("TMap<")))
-	{
-		UE_LOG(LogControlRigDeveloper, Warning, TEXT("TMap Variables are not supported."));
-		return false;
-	}
-
-	Variable.bIsArray = RigVMTypeUtils::IsArrayType(CPPType);
-	if (Variable.bIsArray)
-	{
-		CPPType = RigVMTypeUtils::BaseTypeFromArrayType(CPPType);
-	}
-
-	if (CPPType == TEXT("bool"))
-	{
-		Variable.TypeName = *CPPType;
-		Variable.Size = sizeof(bool);
-	}
-	else if (CPPType == TEXT("float"))
-	{
-		Variable.TypeName = *CPPType;
-		Variable.Size = sizeof(float);
-	}
-	else if (CPPType == TEXT("double"))
-	{
-		Variable.TypeName = *CPPType;
-		Variable.Size = sizeof(double);
-	}
-	else if (CPPType == TEXT("int32"))
-	{
-		Variable.TypeName = *CPPType;
-		Variable.Size = sizeof(int32);
-	}
-	else if (CPPType == TEXT("FString"))
-	{
-		Variable.TypeName = *CPPType;
-		Variable.Size = sizeof(FString);
-	}
-	else if (CPPType == TEXT("FName"))
-	{
-		Variable.TypeName = *CPPType;
-		Variable.Size = sizeof(FName);
-	}
-	else if(UScriptStruct* ScriptStruct = RigVMTypeUtils::FindObjectFromCPPTypeObjectPath<UScriptStruct>(CPPType))
-	{
-		Variable.TypeName = *RigVMTypeUtils::GetUniqueStructTypeName(ScriptStruct);
-		Variable.TypeObject = ScriptStruct;
-		Variable.Size = ScriptStruct->GetStructureSize();
-	}
-	else if (UEnum* Enum= RigVMTypeUtils::FindObjectFromCPPTypeObjectPath<UEnum>(CPPType))
-	{
-		Variable.TypeName = *RigVMTypeUtils::CPPTypeFromEnum(Enum);
-		Variable.TypeObject = Enum;
-		Variable.Size = Enum->GetResourceSizeBytes(EResourceSizeMode::EstimatedTotal);
-	}
-
-	FEdGraphPinType PinType = RigVMTypeUtils::PinTypeFromExternalVariable(Variable);
-	if (!PinType.PinCategory.IsValid())
-	{
-		return false;
-	}
-
-	FBlueprintEditorUtils::ChangeMemberVariableType(this, InName, PinType);
-
-	return true;
-}
-
-const FControlRigShapeDefinition* UControlRigBlueprint::GetControlShapeByName(const FName& InName) const
-{
-	TMap<FString, FString> LibraryNameMap;
-	if(UControlRig* ControlRig = Cast<UControlRig>(GetObjectBeingDebugged()))
-	{
-		LibraryNameMap = ControlRig->ShapeLibraryNameMap;
-	}
-	return UControlRigShapeLibrary::GetShapeByName(InName, ShapeLibraries, LibraryNameMap);
-}
-
-FName UControlRigBlueprint::AddTransientControl(URigVMPin* InPin)
-{
-	TUniquePtr<FControlValueScope> ValueScope;
-	if (!UControlRigEditorSettings::Get()->bResetControlsOnPinValueInteraction) // if we need to retain the controls
-	{
-		ValueScope = MakeUnique<FControlValueScope>(this);
-	}
-
-	// for now we only allow one pin control at the same time
-	ClearTransientControls();
-
-	URigVMBlueprintGeneratedClass* RigClass = GetControlRigBlueprintGeneratedClass();
-	UControlRig* CDO = Cast<UControlRig>(RigClass->GetDefaultObject(true /* create if needed */));
-
-	FRigElementKey SpaceKey;
-	FTransform OffsetTransform = FTransform::Identity;
-	if (URigVMUnitNode* UnitNode = Cast<URigVMUnitNode>(InPin->GetPinForLink()->GetNode()))
-	{
-		if (TSharedPtr<FStructOnScope> DefaultStructScope = UnitNode->ConstructStructInstance())
-		{
-			FRigUnit* DefaultStruct = (FRigUnit*)DefaultStructScope->GetStructMemory();
-
-			FString PinPath = InPin->GetPinForLink()->GetPinPath();
-			FString Left, Right;
-
-			if (URigVMPin::SplitPinPathAtStart(PinPath, Left, Right))
-			{
-				SpaceKey = DefaultStruct->DetermineSpaceForPin(Right, Hierarchy);
-				
-				URigHierarchy* RigHierarchy = Hierarchy;
-
-				// use the active rig instead of the CDO rig because we want to access the evaluation result of the rig graph
-				// to calculate the offset transform, for example take a look at RigUnit_ModifyTransform
-				if (UControlRig* RigBeingDebugged = Cast<UControlRig>(GetObjectBeingDebugged()))
-				{
-					RigHierarchy = RigBeingDebugged->GetHierarchy();
-				}
-				
-				OffsetTransform = DefaultStruct->DetermineOffsetTransformForPin(Right, RigHierarchy);
-			}
-		}
-	}
-
-	FName ReturnName = NAME_None;
-	TArray<UObject*> ArchetypeInstances;
-	CDO->GetArchetypeInstances(ArchetypeInstances);
-	for (UObject* ArchetypeInstance : ArchetypeInstances)
-	{
-		UControlRig* InstancedControlRig = Cast<UControlRig>(ArchetypeInstance);
-		if (InstancedControlRig)
-		{
-			FName ControlName = InstancedControlRig->AddTransientControl(InPin, SpaceKey, OffsetTransform);
-			if (ReturnName == NAME_None)
-			{
-				ReturnName = ControlName;
-			}
-		}
-	}
-
-	return ReturnName;
-}
-
-FName UControlRigBlueprint::RemoveTransientControl(URigVMPin* InPin)
-{
-	TUniquePtr<FControlValueScope> ValueScope;
-	if (!UControlRigEditorSettings::Get()->bResetControlsOnPinValueInteraction) // if we need to retain the controls
-	{
-		ValueScope = MakeUnique<FControlValueScope>(this);
-	}
-
-	URigVMBlueprintGeneratedClass* RigClass = GetControlRigBlueprintGeneratedClass();
-	UControlRig* CDO = Cast<UControlRig>(RigClass->GetDefaultObject(true /* create if needed */));
-
-	FName RemovedName = NAME_None;
-	TArray<UObject*> ArchetypeInstances;
-	CDO->GetArchetypeInstances(ArchetypeInstances);
-	for (UObject* ArchetypeInstance : ArchetypeInstances)
-	{
-		UControlRig* InstancedControlRig = Cast<UControlRig>(ArchetypeInstance);
-		if (InstancedControlRig)
-		{
-			FName Name = InstancedControlRig->RemoveTransientControl(InPin);
-			if (RemovedName == NAME_None)
-	{
-				RemovedName = Name;
-			}
-		}
-	}
-
-	return RemovedName;
-}
-
-FName UControlRigBlueprint::AddTransientControl(const FRigElementKey& InElement)
-{
-	TUniquePtr<FControlValueScope> ValueScope;
-	if (!UControlRigEditorSettings::Get()->bResetControlsOnPinValueInteraction) // if we need to retain the controls
-	{
-		ValueScope = MakeUnique<FControlValueScope>(this);
-	}
-	URigVMBlueprintGeneratedClass* RigClass = GetControlRigBlueprintGeneratedClass();
-	UControlRig* CDO = Cast<UControlRig>(RigClass->GetDefaultObject(true /* create if needed */));
-
-	FName ReturnName = NAME_None;
-	TArray<UObject*> ArchetypeInstances;
-	CDO->GetArchetypeInstances(ArchetypeInstances);
-
-	// hierarchy transforms will be reset when ClearTransientControls() is called,
-	// so to retain any bone transform modifications we have to save them
-	TMap<UObject*, FTransform> SavedElementLocalTransforms;
-	for (UObject* ArchetypeInstance : ArchetypeInstances)
-	{
-		UControlRig* InstancedControlRig = Cast<UControlRig>(ArchetypeInstance);
-		if (InstancedControlRig)
-		{
-			if (InstancedControlRig->DynamicHierarchy)
-			{ 
-				SavedElementLocalTransforms.FindOrAdd(InstancedControlRig) = InstancedControlRig->DynamicHierarchy->GetLocalTransform(InElement);
-			}
-		}
-	}
-
-	// for now we only allow one pin control at the same time
-	ClearTransientControls();
-	
-	for (UObject* ArchetypeInstance : ArchetypeInstances)
-	{
-		UControlRig* InstancedControlRig = Cast<UControlRig>(ArchetypeInstance);
-		if (InstancedControlRig)
-		{
-			// restore the element transforms so that transient controls are created at the right place
-			if (const FTransform* SavedTransform = SavedElementLocalTransforms.Find(InstancedControlRig))
-			{
-				if (InstancedControlRig->DynamicHierarchy)
-				{ 
-					InstancedControlRig->DynamicHierarchy->SetLocalTransform(InElement, *SavedTransform);
-				}
-			}
-			
-			FName ControlName = InstancedControlRig->AddTransientControl(InElement);
-			if (ReturnName == NAME_None)
-			{
-				ReturnName = ControlName;
-			}
-		}
-	}
-
-	return ReturnName;
-
-}
-
-FName UControlRigBlueprint::RemoveTransientControl(const FRigElementKey& InElement)
-{
-	TUniquePtr<FControlValueScope> ValueScope;
-	if (!UControlRigEditorSettings::Get()->bResetControlsOnPinValueInteraction) // if we need to retain the controls
-	{
-		ValueScope = MakeUnique<FControlValueScope>(this);
-	}
-
-	URigVMBlueprintGeneratedClass* RigClass = GetControlRigBlueprintGeneratedClass();
-	UControlRig* CDO = Cast<UControlRig>(RigClass->GetDefaultObject(true /* create if needed */));
-
-	FName RemovedName = NAME_None;
-	TArray<UObject*> ArchetypeInstances;
-	CDO->GetArchetypeInstances(ArchetypeInstances);
-	for (UObject* ArchetypeInstance : ArchetypeInstances)
-	{
-		UControlRig* InstancedControlRig = Cast<UControlRig>(ArchetypeInstance);
-		if (InstancedControlRig)
-		{
-			FName Name = InstancedControlRig->RemoveTransientControl(InElement);
-			if (RemovedName == NAME_None)
-			{
-				RemovedName = Name;
-			}
-		}
-	}
-
-	return RemovedName;
-}
-
-void UControlRigBlueprint::ClearTransientControls()
-{
-	TUniquePtr<FControlValueScope> ValueScope;
-	if (!UControlRigEditorSettings::Get()->bResetControlsOnPinValueInteraction) // if we need to retain the controls
-	{
-		ValueScope = MakeUnique<FControlValueScope>(this);
-	}
-
-	if (URigVMBlueprintGeneratedClass* RigClass = GetControlRigBlueprintGeneratedClass())
-	{
-		UControlRig* CDO = Cast<UControlRig>(RigClass->GetDefaultObject(true /* create if needed */));
-
-		TArray<UObject*> ArchetypeInstances;
-		CDO->GetArchetypeInstances(ArchetypeInstances);
-		for (UObject* ArchetypeInstance : ArchetypeInstances)
-		{
-			UControlRig* InstancedControlRig = Cast<UControlRig>(ArchetypeInstance);
-			if (InstancedControlRig)
-			{
-				InstancedControlRig->ClearTransientControls();
-			}
-		}
-	}
-}
-
-#endif
-
-void UControlRigBlueprint::PopulateModelFromGraphForBackwardsCompatibility(UControlRigGraph* InGraph)
-{
-	DECLARE_SCOPE_HIERARCHICAL_COUNTER_FUNC()
-	
-	
-
-	// temporarily disable default value validation during load time, serialized values should always be accepted
-	TGuardValue<bool> DisablePinDefaultValueValidation(GetOrCreateController()->bValidatePinDefaults, false);
-
-	int32 LinkerVersion = GetLinkerCustomVersion(FControlRigObjectVersion::GUID);
-	if (LinkerVersion >= FControlRigObjectVersion::SwitchedToRigVM)
-	{
-		return;
-	}
-
-	bDirtyDuringLoad = true;
-
-	if (LinkerVersion < FControlRigObjectVersion::RemovalOfHierarchyRefPins)
-	{
-		UE_LOG(LogControlRigDeveloper, Warning, TEXT("Control Rig is too old (prior 4.23) - cannot automatically upgrade. Clearing graph."));
-		RebuildGraphFromModel();
-		return;
-	}
-
-	TGuardValue<bool> ReentrantGuardSelf(bSuspendModelNotificationsForSelf, true);
-	{
-		TGuardValue<bool> ReentrantGuardOthers(bSuspendModelNotificationsForOthers, true);
-
-		struct LocalHelpers
-		{
-			static FString FixUpPinPath(const FString& InPinPath)
-			{
-				FString PinPath = InPinPath;
-				if (!PinPath.Contains(TEXT(".")))
-				{
-					PinPath += TEXT(".Value");
-				}
-
-				PinPath = PinPath.Replace(TEXT("["), TEXT("."), ESearchCase::IgnoreCase);
-				PinPath = PinPath.Replace(TEXT("]"), TEXT(""), ESearchCase::IgnoreCase);
-
-				return PinPath;
-			}
-		};
-
-		for (UEdGraphNode* Node : InGraph->Nodes)
-		{
-			if (UControlRigGraphNode* RigNode = Cast<UControlRigGraphNode>(Node))
-			{
-				FName PropertyName = RigNode->PropertyName_DEPRECATED;
-				FVector2D NodePosition = FVector2D((float)RigNode->NodePosX, (float)RigNode->NodePosY);
-				FString StructPath = RigNode->StructPath_DEPRECATED;
-
-				if (StructPath.IsEmpty() && PropertyName != NAME_None)
-				{
-					if(FStructProperty* StructProperty = CastField<FStructProperty>(GetControlRigBlueprintGeneratedClass()->FindPropertyByName(PropertyName)))
-					{
-						StructPath = StructProperty->Struct->GetPathName();
-					}
-					else
-					{
-						// at this point the BP skeleton might not have been compiled,
-						// we should look into the new variables array to find the property
-						for (FBPVariableDescription NewVariable : NewVariables)
-						{
-							if (NewVariable.VarName == PropertyName && NewVariable.VarType.PinCategory == UEdGraphSchema_K2::PC_Struct)
-							{
-								if (UScriptStruct* Struct = Cast<UScriptStruct>(NewVariable.VarType.PinSubCategoryObject))
-								{
-									StructPath = Struct->GetPathName();
-									break;
-								}
-							}
-						}
-					}
-				}
-
-				URigVMNode* ModelNode = nullptr;
-
-				UScriptStruct* UnitStruct = RigVMTypeUtils::FindObjectFromCPPTypeObjectPath<UScriptStruct>(StructPath);
-				if (UnitStruct && UnitStruct->IsChildOf(FRigVMStruct::StaticStruct()))
-				{ 
-					ModelNode = GetOrCreateController()->AddUnitNode(UnitStruct, FRigUnit::GetMethodName(), NodePosition, PropertyName.ToString(), false);
-				}
-				else if (PropertyName != NAME_None) // check if this is a variable
-				{
-					bool bHasInputLinks = false;
-					bool bHasOutputLinks = false;
-					FString DefaultValue;
-
-					FEdGraphPinType PinType = RigNode->PinType_DEPRECATED;
-					if (RigNode->Pins.Num() > 0)
-				{
-						for (UEdGraphPin* Pin : RigNode->Pins)
-						{
-							if (!Pin->GetName().Contains(TEXT(".")))
-							{
-								PinType = Pin->PinType;
-
-								if (Pin->Direction == EGPD_Input)
-								{
-									bHasInputLinks = Pin->LinkedTo.Num() > 0;
-									DefaultValue = Pin->DefaultValue;
-								}
-								else if (Pin->Direction == EGPD_Output)
-								{
-									bHasOutputLinks = Pin->LinkedTo.Num() > 0;
-								}
-							}
-						}
-					}
-
-					FName DataType = PinType.PinCategory;
-
-					if (PinType.PinCategory == UEdGraphSchema_K2::PC_Real)
-					{
-						if (PinType.PinSubCategory == UEdGraphSchema_K2::PC_Float)
-						{
-							DataType = TEXT("float");
-						}
-						else if (PinType.PinSubCategory == UEdGraphSchema_K2::PC_Double)
-						{
-							DataType = TEXT("double");
-						}
-						else
-						{
-							ensure(false);
-						}
-					}
-
-					UObject* DataTypeObject = nullptr;
-					if (DataType == NAME_None)
-					{
-						continue;
-					}
-					if (DataType == UEdGraphSchema_K2::PC_Struct)
-					{
-						DataType = NAME_None;
-						if (UScriptStruct* DataStruct = Cast<UScriptStruct>(PinType.PinSubCategoryObject))
-						{
-							DataTypeObject = DataStruct;
-							DataType = *DataStruct->GetStructCPPName();
-						}
-					}
-
-					if (DataType == TEXT("int"))
-					{
-						DataType = TEXT("int32");
-					}
-					else if (DataType == TEXT("name"))
-						{
-						DataType = TEXT("FName");
-						}
-					else if (DataType == TEXT("string"))
-					{
-						DataType = TEXT("FString");
-					}
-
-					FProperty* ParameterProperty = GetControlRigBlueprintGeneratedClass()->FindPropertyByName(PropertyName);
-					if(ParameterProperty)
-					{
-						bool bIsInput = true;
-
-						if (ParameterProperty->HasMetaData(TEXT("AnimationInput")) || bHasOutputLinks)
-					{
-							bIsInput = true;
-						}
-						else if (ParameterProperty->HasMetaData(TEXT("AnimationOutput")))
-						{
-							bIsInput = false;
-						}
-
-						ModelNode = GetOrCreateController()->AddVariableNode(PropertyName, DataType.ToString(), DataTypeObject, bIsInput, FString(), NodePosition, PropertyName.ToString(), false);
-					}
-				}
-				else
-				{
-					continue;
-				}
-
-				if (ModelNode)
-				{
-					bool bWasReportingEnabled = GetOrCreateController()->IsReportingEnabled();
-					GetOrCreateController()->EnableReporting(false);
-
-					for (UEdGraphPin* Pin : RigNode->Pins)
-					{
-							FString PinPath = LocalHelpers::FixUpPinPath(Pin->GetName());
-
-							// check the material + mesh pins for deprecated control nodes
-							if (URigVMUnitNode* ModelUnitNode = Cast<URigVMUnitNode>(ModelNode))
-							{
-								const UScriptStruct* ScriptStruct = ModelUnitNode->GetScriptStruct();
-								if (ScriptStruct && ScriptStruct->IsChildOf(FRigUnit_Control::StaticStruct()))
-								{
-									if (Pin->GetName().EndsWith(TEXT(".StaticMesh")) || Pin->GetName().EndsWith(TEXT(".Materials")))
-									{
-										continue;
-									}
-								}
-							}
-
-						if (Pin->Direction == EGPD_Input && Pin->PinType.ContainerType == EPinContainerType::Array)
-						{
-							int32 ArraySize = Pin->SubPins.Num();
-							GetOrCreateController()->SetArrayPinSize(PinPath, ArraySize, FString(), false);
-						}
-
-						if (RigNode->ExpandedPins_DEPRECATED.Find(Pin->GetName()) != INDEX_NONE)
-						{
-								GetOrCreateController()->SetPinExpansion(PinPath, true, false);
-						}
-
-						if (Pin->SubPins.Num() == 0 && !Pin->DefaultValue.IsEmpty() && Pin->Direction == EGPD_Input)
-						{
-								GetOrCreateController()->SetPinDefaultValue(PinPath, Pin->DefaultValue, false, false, false);
-						}
-					}
-
-					GetOrCreateController()->EnableReporting(bWasReportingEnabled);
-				}
-
-				const int32 VarIndex = FBlueprintEditorUtils::FindNewVariableIndex(this, PropertyName);
-				if (VarIndex != INDEX_NONE)
-				{
-					// only remove the variable if it was backing up a rig unit
-					if(UnitStruct != nullptr)
-					{
-						NewVariables.RemoveAt(VarIndex);
-					}
-					FBlueprintEditorUtils::RemoveVariableNodes(this, PropertyName);
-				}
-			}
-			else if (const UEdGraphNode_Comment* CommentNode = Cast<UEdGraphNode_Comment>(Node))
-			{
-				FVector2D NodePosition = FVector2D((float)CommentNode->NodePosX, (float)CommentNode->NodePosY);
-				FVector2D NodeSize = FVector2D((float)CommentNode->NodeWidth, (float)CommentNode->NodeHeight);
-				GetOrCreateController()->AddCommentNode(CommentNode->NodeComment, NodePosition, NodeSize, CommentNode->CommentColor, CommentNode->GetName(), false);
-			}
-		}
-
-		SetupPinRedirectorsForBackwardsCompatibility();
-
-		for (UEdGraphNode* Node : InGraph->Nodes)
-		{
-			if (UControlRigGraphNode* RigNode = Cast<UControlRigGraphNode>(Node))
-			{
-				for (UEdGraphPin* Pin : RigNode->Pins)
-				{
-					if (Pin->Direction == EGPD_Input)
-					{
-						continue;
-					}
-
-					for (UEdGraphPin* LinkedPin : Pin->LinkedTo)
-					{
-						UControlRigGraphNode* LinkedRigNode = Cast<UControlRigGraphNode>(LinkedPin->GetOwningNode());
-						if (LinkedRigNode != nullptr)
-						{
-							FString SourcePinPath = LocalHelpers::FixUpPinPath(Pin->GetName());
-							FString TargetPinPath = LocalHelpers::FixUpPinPath(LinkedPin->GetName());
-							GetOrCreateController()->AddLink(SourcePinPath, TargetPinPath, false);
-						}
-					}
-				}
-			}
-		}
-	}
-
-	RebuildGraphFromModel();
-}
-
-void UControlRigBlueprint::SetupPinRedirectorsForBackwardsCompatibility()
-{
-	for(URigVMGraph* Model : RigVMClient)
-	{
-		for (URigVMNode* Node : Model->GetNodes())
-		{
-			if (URigVMUnitNode* UnitNode = Cast<URigVMUnitNode>(Node))
-			{
-				UScriptStruct* Struct = UnitNode->GetScriptStruct();
-				if (Struct == FRigUnit_SetBoneTransform::StaticStruct())
-				{
-					URigVMPin* TransformPin = UnitNode->FindPin(TEXT("Transform"));
-					URigVMPin* ResultPin = UnitNode->FindPin(TEXT("Result"));
-					GetOrCreateController()->AddPinRedirector(false, true, TransformPin->GetPinPath(), ResultPin->GetPinPath());
-				}
-			}
-		}
-	}
-}
-
-void UControlRigBlueprint::RebuildGraphFromModel()
-{
-	DECLARE_SCOPE_HIERARCHICAL_COUNTER_FUNC()
-
-	TGuardValue<bool> SelfGuard(bSuspendModelNotificationsForSelf, true);
-	TGuardValue<bool> ClientIgnoreModificationsGuard(RigVMClient.bIgnoreModelNotifications, true);
-	
-	verify(GetOrCreateController());
-
-	TArray<UEdGraph*> EdGraphs;
-	GetAllGraphs(EdGraphs);
-
-	for (UEdGraph* Graph : EdGraphs)
-	{
-		TArray<UEdGraphNode*> Nodes = Graph->Nodes;
-		for (UEdGraphNode* Node : Nodes)
-		{
-			Graph->RemoveNode(Node);
-		}
-
-		if (UControlRigGraph* RigGraph = Cast<UControlRigGraph>(Graph))
-		{
-			if (RigGraph->bIsFunctionDefinition)
-			{
-				FunctionGraphs.Remove(RigGraph);
-			}
-		}
-	}
-
-	if(FunctionLibraryEdGraph && RigVMClient.GetFunctionLibrary())
-	{
-		FunctionLibraryEdGraph->ModelNodePath = RigVMClient.GetFunctionLibrary()->GetNodePath();
-	}
-
-	TArray<URigVMGraph*> RigGraphs = RigVMClient.GetAllModels(true, true);
-
-	for (int32 RigGraphIndex = 0; RigGraphIndex < RigGraphs.Num(); RigGraphIndex++)
-	{
-		GetOrCreateController(RigGraphs[RigGraphIndex])->ResendAllNotifications();
-	}
-
-	for (int32 RigGraphIndex = 0; RigGraphIndex < RigGraphs.Num(); RigGraphIndex++)
-	{
-		URigVMGraph* RigGraph = RigGraphs[RigGraphIndex];
-
-		for (URigVMNode* RigNode : RigGraph->GetNodes())
-		{
-			if (URigVMCollapseNode* CollapseNode = Cast<URigVMCollapseNode>(RigNode))
-			{
-				CreateEdGraphForCollapseNodeIfNeeded(CollapseNode, true);
-			}
-		}
-	}
-
-	EdGraphs.Reset();
-	GetAllGraphs(EdGraphs);
-
-	for (UEdGraph* Graph : EdGraphs)
-	{
-		if (UControlRigGraph* RigGraph = Cast<UControlRigGraph>(Graph))
-		{
-			RigGraph->CacheNameLists(Hierarchy, &DrawContainer, ShapeLibraries);
-		}
-	}
-
-}
-
-void UControlRigBlueprint::Notify(ERigVMGraphNotifType InNotifType, UObject* InSubject)
-{
-	GetOrCreateController()->Notify(InNotifType, InSubject);
-}
-
-void UControlRigBlueprint::HandleModifiedEvent(ERigVMGraphNotifType InNotifType, URigVMGraph* InGraph, UObject* InSubject)
-{
-	DECLARE_SCOPE_HIERARCHICAL_COUNTER_FUNC()
-
-#if WITH_EDITOR
-
-	if (bSuspendAllNotifications)
-	{
-		return;
-	}
-
-	// since it's possible that a notification will be already sent / forwarded to the
-	// listening objects within the switch statement below - we keep a flag to mark
-	// the notify for still pending (or already sent)
-	bool bNotifForOthersPending = true;
-
-	auto MarkBlueprintAsStructurallyModified = [this]()
-	{
-		if(VMRecompilationBracket == 0)
-		{
-			if(bMarkBlueprintAsStructurallyModifiedPending)
-			{
-				bMarkBlueprintAsStructurallyModifiedPending = false;
-				FBlueprintEditorUtils::MarkBlueprintAsStructurallyModified(this);
-			}
-		}
-		else
-		{
-			bMarkBlueprintAsStructurallyModifiedPending = true;
-		}
-	};
-
-	if (!bSuspendModelNotificationsForSelf)
-	{
-		switch (InNotifType)
-		{
-			case ERigVMGraphNotifType::InteractionBracketOpened:
-			{
-				IncrementVMRecompileBracket();
-				break;
-			}
-			case ERigVMGraphNotifType::InteractionBracketClosed:
-			case ERigVMGraphNotifType::InteractionBracketCanceled:
-			{
-				DecrementVMRecompileBracket();
-				MarkBlueprintAsStructurallyModified();
-				break;
-			}
-			case ERigVMGraphNotifType::PinDefaultValueChanged:
-			{
-				if (URigVMPin* Pin = Cast<URigVMPin>(InSubject))
-				{
-					bool bRequiresRecompile = false;
-
-					URigVMPin* RootPin = Pin->GetRootPin();
-					static const FString ConstSuffix = TEXT(":Const");
-					const FString PinHash = RootPin->GetPinPath(true) + ConstSuffix;
-					
-					if (const FRigVMOperand* Operand = PinToOperandMap.Find(PinHash))
-					{
-						FRigVMASTProxy RootPinProxy = FRigVMASTProxy::MakeFromUObject(RootPin);
-						if(const FRigVMExprAST* Expression = InGraph->GetRuntimeAST()->GetExprForSubject(RootPinProxy))
-						{
-							bRequiresRecompile = Expression->NumParents() > 1;
-						}
-						else
-						{
-							bRequiresRecompile = true;
-						}
-
-						// If we are only changing a pin's default value, we need to
-						// check if there is a connection to a sub-pin of the root pin
-						// that has its value is directly stored in the root pin due to optimization, if so,
-						// we want to recompile to make sure the pin's new default value and values from other connections
-						// are both applied to the root pin because GetDefaultValue() alone cannot account for values
-						// from other connections.
-						if(!bRequiresRecompile)
-						{
-							TArray<URigVMPin*> SourcePins = RootPin->GetLinkedSourcePins(true);
-							for (const URigVMPin* SourcePin : SourcePins)
-							{
-								// check if the source node is optimized out, if so, only a recompile will allows us
-								// to re-query its value.
-								FRigVMASTProxy SourceNodeProxy = FRigVMASTProxy::MakeFromUObject(SourcePin->GetNode());
-								if (InGraph->GetRuntimeAST()->GetExprForSubject(SourceNodeProxy) == nullptr)
-								{
-									bRequiresRecompile = true;
-									break;
-								}
-							}
-						} 
-						
-						if(!bRequiresRecompile)
-						{
-							const FString DefaultValue = RootPin->GetDefaultValue();
-
-							URigVMBlueprintGeneratedClass* RigClass = GetControlRigBlueprintGeneratedClass();
-							UControlRig* CDO = Cast<UControlRig>(RigClass->GetDefaultObject(true /* create if needed */));
-							if (CDO->VM != nullptr)
-							{
-								CDO->VM->SetPropertyValueFromString(*Operand, DefaultValue);
-							}
-
-							TArray<UObject*> ArchetypeInstances;
-							CDO->GetArchetypeInstances(ArchetypeInstances);
-							for (UObject* ArchetypeInstance : ArchetypeInstances)
-							{
-								UControlRig* InstancedControlRig = Cast<UControlRig>(ArchetypeInstance);
-								if (InstancedControlRig)
-								{
-									if (InstancedControlRig->VM)
-									{
-										InstancedControlRig->VM->SetPropertyValueFromString(*Operand, DefaultValue);
-									}
-								}
-							}
-
-							if (Pin->IsDefinedAsConstant() || Pin->GetRootPin()->IsDefinedAsConstant())
-							{
-								// re-init the rigs
-								RequestControlRigInit();
-								bRequiresRecompile = true;
-							}
-						}
-					}
-					else
-					{
-						bRequiresRecompile = true;
-					}
-				
-					if(bRequiresRecompile)
-					{
-						RequestAutoVMRecompilation();
-					}
-
-					// check if this pin is part of an injected node, and if it is a visual debug node,
-					// we might need to recreate the control pin
-					if (UClass* MyControlRigClass = GeneratedClass)
-					{
-						if (UControlRig* DefaultObject = Cast<UControlRig>(MyControlRigClass->GetDefaultObject(false)))
-						{
-							TArray<UObject*> ArchetypeInstances;
-							DefaultObject->GetArchetypeInstances(ArchetypeInstances);
-							
-							for (UObject* ArchetypeInstance : ArchetypeInstances)
-							{
-								if (UControlRig* InstanceRig = Cast<UControlRig>(ArchetypeInstance))
-								{
-									Hierarchy->ForEach<FRigControlElement>([this, InstanceRig, Pin](FRigControlElement* ControlElement) -> bool
-                                    {
-										if(!ControlElement->Settings.bIsTransientControl)
-										{
-											return true;
-										}
-
-										for(URigVMGraph* Model : RigVMClient)
-										{
-											if (URigVMPin* ControlledPin = Model->FindPin(ControlElement->GetName().ToString()))
-											{
-												URigVMPin* ControlledPinForLink = ControlledPin->GetPinForLink();
-
-												if(ControlledPin->GetRootPin() == Pin->GetRootPin() ||
-												   ControlledPinForLink->GetRootPin() == Pin->GetRootPin())
-												{
-													InstanceRig->SetTransientControlValue(ControlledPin->GetPinForLink());
-												}
-												else if (ControlledPin->GetNode() == Pin->GetNode() ||
-														 ControlledPinForLink->GetNode() == Pin->GetNode())
-												{
-													InstanceRig->ClearTransientControls();
-													InstanceRig->AddTransientControl(ControlledPin);
-												}
-												return false;
-											}
-										}
-										return true;
-									});
-								}
-							}
-						}
-					}
-				}
-				MarkPackageDirty();
-				break;
-			}
-			case ERigVMGraphNotifType::NodeAdded:
-			case ERigVMGraphNotifType::NodeRemoved:
-			{
-				bool bAdded = InNotifType == ERigVMGraphNotifType::NodeAdded;
-				if (!bAdded)
-				{
-					if (URigVMNode* RigVMNode = Cast<URigVMNode>(InSubject))
-					{
-						RemoveBreakpoint(RigVMNode);
-					}
-				}
-
-				if (URigVMCollapseNode* CollapseNode = Cast<URigVMCollapseNode>(InSubject))
-				{
-					if (bAdded)
-					{
-						// If the controller for this graph already exist, make sure it is referencing the correct graph
-						if (URigVMController* Controller = RigVMClient.GetController(CollapseNode->GetContainedGraph()))
-						{
-							Controller->SetGraph(CollapseNode->GetContainedGraph());
-						}
-						
-						CreateEdGraphForCollapseNodeIfNeeded(CollapseNode);
-					}
-					else
-					{
-						bNotifForOthersPending = !RemoveEdGraphForCollapseNode(CollapseNode, true);
-
-						// Cannot remove from the Controllers array because we would lose the action stack on that graph
-						// Controllers.Remove(CollapseNode->GetContainedGraph();
-					}
-
-					ClearTransientControls();
-					RequestAutoVMRecompilation();
-
-					MarkPackageDirty();
-					MarkBlueprintAsStructurallyModified();
-					break;
-				}
-
-				if (URigVMNode* RigVMNode = Cast<URigVMNode>(InSubject))
-				{
-					if(RigVMNode->IsEvent() && RigVMNode->GetGraph()->IsRootGraph())
-					{
-						// let the UI know the title for the graph may have changed.
-						RigVMClient.NotifyOuterOfPropertyChange();
-
-						if(UControlRigGraph* EdGraph = Cast<UControlRigGraph>(GetEdGraph(RigVMNode->GetGraph())))
-						{
-							// decide if this graph should be renameable
-							const int32 NumberOfEvents = Algo::CountIf(RigVMNode->GetGraph()->GetNodes(), [](const URigVMNode* NodeToCount) -> bool
-							{
-								return NodeToCount->IsEvent() && NodeToCount->CanOnlyExistOnce();
-							});
-							EdGraph->bAllowRenaming = NumberOfEvents != 1;
-						}
-					}
-				}
-				// fall through to the next case
-			}
-			case ERigVMGraphNotifType::LinkAdded:
-			case ERigVMGraphNotifType::LinkRemoved:
-			case ERigVMGraphNotifType::PinArraySizeChanged:
-			case ERigVMGraphNotifType::PinDirectionChanged:
-			{
-				ClearTransientControls();
-				RequestAutoVMRecompilation();
-				MarkPackageDirty();
-
-				// we don't need to mark the blueprint as modified since we only
-				// need to recompile the VM here - unless we don't auto recompile.
-				if(!bAutoRecompileVM)
-				{
-					MarkBlueprintAsStructurallyModified();
-				}
-				break;
-			}
-			case ERigVMGraphNotifType::PinWatchedChanged:
-			{
-				if (UControlRig* CR = Cast<UControlRig>(GetObjectBeingDebugged()))
-				{
-					URigVMPin* Pin = CastChecked<URigVMPin>(InSubject)->GetRootPin(); 
-					URigVMCompiler* Compiler = URigVMCompiler::StaticClass()->GetDefaultObject<URigVMCompiler>();
-					Compiler->Settings = VMCompileSettings;
-
-					TSharedPtr<FRigVMParserAST> RuntimeAST = GetDefaultModel()->GetRuntimeAST();
-					
-					if(Pin->RequiresWatch())
-					{
-						// check if the node is optimized out - in that case we need to recompile
-						if(CR->GetVM()->GetByteCode().GetFirstInstructionIndexForSubject(Pin->GetNode()) == INDEX_NONE)
-						{
-							RequestAutoVMRecompilation();
-							MarkPackageDirty();
-						}
-						else
-						{
-							if(CR->GetVM()->GetDebugMemory()->Num() == 0)
-							{
-								RequestAutoVMRecompilation();
-								MarkPackageDirty();
-							}
-							else
-							{
-								Compiler->MarkDebugWatch(true, Pin, CR->GetVM(), &PinToOperandMap, RuntimeAST);
-							}
-						}
-					}
-					else
-					{
-						Compiler->MarkDebugWatch(false, Pin, CR->GetVM(), &PinToOperandMap, RuntimeAST);
-					}
-				}
-				// break; fall through
-			}
-			case ERigVMGraphNotifType::PinTypeChanged:
-			case ERigVMGraphNotifType::PinIndexChanged:
-			{
-				if (URigVMPin* ModelPin = Cast<URigVMPin>(InSubject))
-				{
-					if (UEdGraph* EdGraph = GetEdGraph(InGraph))
-					{
-						if (UControlRigGraph* Graph = Cast<UControlRigGraph>(EdGraph))
-						{
-							if (UEdGraphNode* EdNode = Graph->FindNodeForModelNodeName(ModelPin->GetNode()->GetFName()))
-							{
-								if (UEdGraphPin* EdPin = EdNode->FindPin(*ModelPin->GetPinPath()))
-								{
-									if (ModelPin->RequiresWatch())
-									{
-										if (!FKismetDebugUtilities::IsPinBeingWatched(this, EdPin))
-										{
-											FKismetDebugUtilities::AddPinWatch(this, FBlueprintWatchedPin(EdPin));
-										}
-									}
-									else
-									{
-										FKismetDebugUtilities::RemovePinWatch(this, EdPin);
-									}
-
-									if(InNotifType == ERigVMGraphNotifType::PinWatchedChanged)
-									{
-										return;
-									}
-									RequestAutoVMRecompilation();
-									MarkPackageDirty();
-								}
-							}
-						}
-					}
-				}
-				// fall through another time
-			}
-			case ERigVMGraphNotifType::PinAdded:
-			case ERigVMGraphNotifType::PinRemoved:
-			case ERigVMGraphNotifType::PinRenamed:
-			{
-				// exposed pin changes like this (as well as type change etc)
-				// require to mark the blueprint as structurally modified,
-				// so that the instance actions work out.
-				if (URigVMPin* ModelPin = Cast<URigVMPin>(InSubject))
-				{
-					if(URigVMCollapseNode* CollapseNode = Cast<URigVMCollapseNode>(ModelPin->GetNode()))
-					{
-						if(Cast<URigVMFunctionLibrary>(CollapseNode->GetOuter()))
-						{
-							MarkBlueprintAsStructurallyModified();
-						}
-					}
-				}
-				break;
-			}
-			case ERigVMGraphNotifType::PinBoundVariableChanged:
-			case ERigVMGraphNotifType::VariableRemappingChanged:
-			{
-				RequestAutoVMRecompilation();
-				MarkPackageDirty();
-				break;
-			}
-			case ERigVMGraphNotifType::NodeRenamed:
-			{
-				if (URigVMCollapseNode* CollapseNode = Cast<URigVMCollapseNode>(InSubject))
-				{
-					FString NewNodePath = CollapseNode->GetNodePath(true /* recursive */);
-					FString Left, Right = NewNodePath;
-					URigVMNode::SplitNodePathAtEnd(NewNodePath, Left, Right);
-					FString OldNodePath = CollapseNode->GetPreviousFName().ToString();
-					if (!Left.IsEmpty())
-					{
-						OldNodePath = URigVMNode::JoinNodePath(Left, OldNodePath);
-					}
-
-					HandleRigVMGraphRenamed(GetRigVMClient(), OldNodePath, NewNodePath);
-
-					if (UEdGraph* ContainedEdGraph = GetEdGraph(CollapseNode->GetContainedGraph()))
-					{
-						ContainedEdGraph->Rename(*CollapseNode->GetEditorSubGraphName(), nullptr);
-					}
-
-					MarkBlueprintAsStructurallyModified();
-				}
-				break;
-			}
-			case ERigVMGraphNotifType::NodeCategoryChanged:
-			case ERigVMGraphNotifType::NodeKeywordsChanged:
-			case ERigVMGraphNotifType::NodeDescriptionChanged:
-			{
-				MarkBlueprintAsStructurallyModified();
-				break;
-			}
-			default:
-			{
-				break;
-			}
-		}
-	}
-
-	// if the notification still has to be sent...
-	if (bNotifForOthersPending && !bSuspendModelNotificationsForOthers)
-	{
-		if (ModifiedEvent.IsBound())
-		{
-			ModifiedEvent.Broadcast(InNotifType, InGraph, InSubject);
-		}
-	}
-#endif
-}
-
-void UControlRigBlueprint::SuspendNotifications(bool bSuspendNotifs)
-{
-	if (bSuspendAllNotifications == bSuspendNotifs)
-	{
-		return;
-	}
-
-	bSuspendAllNotifications = bSuspendNotifs;
-	if (!bSuspendNotifs)
-	{
-		RebuildGraphFromModel();
-		RefreshEditorEvent.Broadcast(this);
-		RequestAutoVMRecompilation();
-	}
-}
-
-void UControlRigBlueprint::CreateMemberVariablesOnLoad()
-{
-#if WITH_EDITOR
-
-	int32 LinkerVersion = GetLinkerCustomVersion(FControlRigObjectVersion::GUID);
-	if (LinkerVersion < FControlRigObjectVersion::SwitchedToRigVM)
-	{
-		// ignore errors during the first potential compile of the VM
-		// since that this point variable nodes may still be ill-formed.
-		TGuardValue<FRigVMReportDelegate> SuspendReportDelegate(VMCompileSettings.ASTSettings.ReportDelegate,
-			FRigVMReportDelegate::CreateLambda([](EMessageSeverity::Type,  UObject*, const FString&)
-			{
-				// do nothing
-			})
-		);
-		InitializeModelIfRequired();
-	}
-
-	AddedMemberVariableMap.Reset();
-
-	for (int32 VariableIndex = 0; VariableIndex < NewVariables.Num(); VariableIndex++)
-	{
-		AddedMemberVariableMap.Add(NewVariables[VariableIndex].VarName, VariableIndex);
-	}
-
-	if (RigVMClient.Num() == 0)
-	{
-		return;
-	}
-
-	// setup variables on the blueprint based on the previous "parameters"
-	if (GetLinkerCustomVersion(FControlRigObjectVersion::GUID) < FControlRigObjectVersion::BlueprintVariableSupport)
-	{
-		TSharedPtr<FKismetNameValidator> NameValidator = MakeShareable(new FKismetNameValidator(this, NAME_None, nullptr));
-
-		for(URigVMGraph* Model : RigVMClient)
-		{
-			TArray<URigVMNode*> Nodes = Model->GetNodes();
-			for (URigVMNode* Node : Nodes)
-			{
-				if (URigVMVariableNode* VariableNode = Cast<URigVMVariableNode>(Node))
-				{
-					if (URigVMPin* VariablePin = VariableNode->FindPin(TEXT("Variable")))
-					{
-						if (VariablePin->GetDirection() != ERigVMPinDirection::Visible)
-						{
-							continue;
-						}
-					}
-
-					FRigVMGraphVariableDescription Description = VariableNode->GetVariableDescription();
-					if (AddedMemberVariableMap.Contains(Description.Name))
-					{
-						continue;
-					}
-
-					FEdGraphPinType PinType = RigVMTypeUtils::PinTypeFromExternalVariable(Description.ToExternalVariable());
-					if (!PinType.PinCategory.IsValid())
-					{
-						continue;
-					}
-
-					FName VarName = FindCRMemberVariableUniqueName(NameValidator, Description.Name.ToString());
-					int32 VariableIndex = AddCRMemberVariable(this, VarName, PinType, false, false, FString());
-					if (VariableIndex != INDEX_NONE)
-					{
-						AddedMemberVariableMap.Add(Description.Name, VariableIndex);
-						bDirtyDuringLoad = true;
-					}
-				}
-
-				// Leaving this for backwards compatibility, even though we don't support parameters anymore
-				// When a parameter node is found, we will create a variable
-				if (URigVMParameterNode* ParameterNode = Cast<URigVMParameterNode>(Node))
-				{
-					if (URigVMPin* ParameterPin = ParameterNode->FindPin(TEXT("Parameter")))
-					{
-						if (ParameterPin->GetDirection() != ERigVMPinDirection::Visible)
-						{
-							continue;
-						}
-					}
-
-					FRigVMGraphParameterDescription Description = ParameterNode->GetParameterDescription();
-					if (AddedMemberVariableMap.Contains(Description.Name))
-					{
-						continue;
-					}
-
-					FEdGraphPinType PinType = RigVMTypeUtils::PinTypeFromExternalVariable(Description.ToExternalVariable());
-					if (!PinType.PinCategory.IsValid())
-					{
-						continue;
-					}
-
-					FName VarName = FindCRMemberVariableUniqueName(NameValidator, Description.Name.ToString());
-					int32 VariableIndex = AddCRMemberVariable(this, VarName, PinType, true, !Description.bIsInput, FString());
-					if (VariableIndex != INDEX_NONE)
-					{
-						AddedMemberVariableMap.Add(Description.Name, VariableIndex);
-						bDirtyDuringLoad = true;
-					}
-				}
-			}
-		}
-	}
-
-#endif
-}
-
-#if WITH_EDITOR
-
-FName UControlRigBlueprint::FindCRMemberVariableUniqueName(TSharedPtr<FKismetNameValidator> InNameValidator, const FString& InBaseName)
-{
-	FString BaseName = InBaseName;
-	if (InNameValidator->IsValid(BaseName) == EValidatorResult::ContainsInvalidCharacters)
-	{
-		for (TCHAR& TestChar : BaseName)
-		{
-			for (TCHAR BadChar : UE_BLUEPRINT_INVALID_NAME_CHARACTERS)
-			{
-				if (TestChar == BadChar)
-				{
-					TestChar = TEXT('_');
-					break;
-				}
-			}
-		}
-	}
-
-	FString KismetName = BaseName;
-
-	int32 Suffix = 0;
-	while (InNameValidator->IsValid(KismetName) != EValidatorResult::Ok)
-	{
-		KismetName = FString::Printf(TEXT("%s_%d"), *BaseName, Suffix);
-		Suffix++;
-	}
-
-
-	return *KismetName;
-}
-
-int32 UControlRigBlueprint::AddCRMemberVariable(UControlRigBlueprint* InBlueprint, const FName& InVarName, FEdGraphPinType InVarType, bool bIsPublic, bool bIsReadOnly, FString InDefaultValue)
-{
-	FBPVariableDescription NewVar;
-
-	NewVar.VarName = InVarName;
-	NewVar.VarGuid = FGuid::NewGuid();
-	NewVar.FriendlyName = FName::NameToDisplayString(InVarName.ToString(), (InVarType.PinCategory == UEdGraphSchema_K2::PC_Boolean) ? true : false);
-	NewVar.VarType = InVarType;
-
-	NewVar.PropertyFlags |= (CPF_Edit | CPF_BlueprintVisible | CPF_DisableEditOnInstance);
-
-	if (bIsPublic)
-	{
-		NewVar.PropertyFlags &= ~CPF_DisableEditOnInstance;
-	}
-
-	if (bIsReadOnly)
-	{
-		NewVar.PropertyFlags |= CPF_BlueprintReadOnly;
-	}
-
-	NewVar.ReplicationCondition = COND_None;
-
-	NewVar.Category = UEdGraphSchema_K2::VR_DefaultCategory;
-
-	// user created variables should be none of these things
-	NewVar.VarType.bIsConst = false;
-	NewVar.VarType.bIsWeakPointer = false;
-	NewVar.VarType.bIsReference = false;
-
-	// Text variables, etc. should default to multiline
-	NewVar.SetMetaData(TEXT("MultiLine"), TEXT("true"));
-
-	NewVar.DefaultValue = InDefaultValue;
-
-	return InBlueprint->NewVariables.Add(NewVar);
-}
-
-FName UControlRigBlueprint::AddCRMemberVariableFromExternal(FRigVMExternalVariable InVariableToCreate, FString InDefaultValue)
-{
-	FEdGraphPinType PinType = RigVMTypeUtils::PinTypeFromExternalVariable(InVariableToCreate);
-	if (!PinType.PinCategory.IsValid())
-	{
-		return NAME_None;
-	}
-
-	Modify();
-
-	TSharedPtr<FKismetNameValidator> NameValidator = MakeShareable(new FKismetNameValidator(this, NAME_None, nullptr));
-	FName VarName = FindCRMemberVariableUniqueName(NameValidator, InVariableToCreate.Name.ToString());
-	int32 VariableIndex = AddCRMemberVariable(this, VarName, PinType, InVariableToCreate.bIsPublic, InVariableToCreate.bIsReadOnly, InDefaultValue);
-	if (VariableIndex != INDEX_NONE)
-	{
-		FBlueprintEditorUtils::MarkBlueprintAsStructurallyModified(this);
-		return VarName;
-	}
-
-	return NAME_None;
-}
-
-void UControlRigBlueprint::PatchFunctionReferencesOnLoad()
-{
-	// If the asset was copied from one project to another, the function referenced might have a different
-	// path, even if the function is internal to the contorl rig. In that case, let's try to find the function
-	// in the local function library.
-
-	for(URigVMGraph* Model : RigVMClient)
-	{
-		TArray<URigVMNode*> Nodes = Model->GetNodes();
-		for (URigVMLibraryNode* Library : RigVMClient.GetFunctionLibrary()->GetFunctions())
-		{
-			Nodes.Append(Library->GetContainedNodes());
-		}
-		
-		for (int32 i=0; i<Nodes.Num(); ++i)
-		{
-			URigVMNode* Node = Nodes[i];
-			if (URigVMFunctionReferenceNode* FunctionReferenceNode = Cast<URigVMFunctionReferenceNode>(Node))
-			{
-				if (!FunctionReferenceNode->ReferencedNodePtr_DEPRECATED.IsValid())
-				{
-					FunctionReferenceNode->ReferencedNodePtr_DEPRECATED.LoadSynchronous();
-				}
-				if (!FunctionReferenceNode->ReferencedNodePtr_DEPRECATED)
-				{
-					if(URigVMFunctionLibrary* FunctionLibrary = RigVMClient.GetFunctionLibrary())
-					{
-						FString FunctionPath = FunctionReferenceNode->ReferencedNodePtr_DEPRECATED.ToSoftObjectPath().GetSubPathString();
-						
-						FString Left, Right;
-						if(FunctionPath.Split(TEXT("."), &Left, &Right))
-						{
-							FString LibraryNodePath = FunctionLibrary->GetNodePath();
-							if(Left == FunctionLibrary->GetName())
-							{
-								if (URigVMLibraryNode* LibraryNode = Cast<URigVMLibraryNode>(FunctionLibrary->FindNode(Right)))
-								{
-									FunctionReferenceNode->ReferencedNodePtr_DEPRECATED = LibraryNode;
-								}
-							}
-						}
-					}
-				}
-
-				if (FunctionReferenceNode->ReferencedNodePtr_DEPRECATED.IsValid())
-				{
-					FunctionReferenceNode->ReferencedFunctionHeader = FunctionReferenceNode->ReferencedNodePtr_DEPRECATED->GetFunctionHeader();					
-				}
-				else if (!FunctionReferenceNode->ReferencedNodePtr_DEPRECATED.IsNull())
-				{
-					// At least lets make sure we store the path in the header
-					FunctionReferenceNode->ReferencedFunctionHeader.LibraryPointer.LibraryNode = FunctionReferenceNode->ReferencedNodePtr_DEPRECATED.ToSoftObjectPath();
-				}
-			}
-
-			if (URigVMCollapseNode* CollapseNode = Cast<URigVMCollapseNode>(Node))
-			{
-				Nodes.Append(CollapseNode->GetContainedNodes());
-			}
-			
-		}
-	}
-	FunctionReferenceNodeData = GetReferenceNodeData();
-}
-
-#endif
-
-void UControlRigBlueprint::PatchVariableNodesOnLoad()
-{
-#if WITH_EDITOR
-
-	// setup variables on the blueprint based on the previous "parameters"
-	if (GetLinkerCustomVersion(FControlRigObjectVersion::GUID) < FControlRigObjectVersion::BlueprintVariableSupport)
-	{
-		TGuardValue<bool> GuardNotifsSelf(bSuspendModelNotificationsForSelf, true);
-
-		check(GetDefaultModel());
-
-		for(URigVMGraph* Model : RigVMClient)
-		{
-			TArray<URigVMNode*> Nodes = Model->GetNodes();
-			for (URigVMNode* Node : Nodes)
-			{
-				if (URigVMVariableNode* VariableNode = Cast<URigVMVariableNode>(Node))
-				{
-					FRigVMGraphVariableDescription Description = VariableNode->GetVariableDescription();
-					if (!AddedMemberVariableMap.Contains(Description.Name))
-					{
-						continue;
-					}
-
-					int32 VariableIndex = AddedMemberVariableMap.FindChecked(Description.Name);
-					FName VarName = NewVariables[VariableIndex].VarName;
-
-					URigVMController* Controller = GetOrCreateController(Model);
-					Controller->RefreshVariableNode(VariableNode->GetFName(), VarName, Description.CPPType, Description.CPPTypeObject, false);
-					bDirtyDuringLoad = true;
-				}
-
-				if (URigVMParameterNode* ParameterNode = Cast<URigVMParameterNode>(Node))
-				{
-					FRigVMGraphParameterDescription Description = ParameterNode->GetParameterDescription();
-					if (!AddedMemberVariableMap.Contains(Description.Name))
-					{
-						continue;
-					}
-
-					int32 VariableIndex = AddedMemberVariableMap.FindChecked(Description.Name);
-					FName VarName = NewVariables[VariableIndex].VarName;
-					URigVMController* Controller = GetOrCreateController(Model);
-					Controller->ReplaceParameterNodeWithVariable(ParameterNode->GetFName(), VarName, Description.CPPType, Description.CPPTypeObject, false);
-					bDirtyDuringLoad = true;
-				}
-			}
-		}
-	}
-
-	AddedMemberVariableMap.Reset();
-	LastNewVariables = NewVariables;
-
-#endif
-}
-
-void UControlRigBlueprint::PatchRigElementKeyCacheOnLoad()
-{
-	if (GetLinkerCustomVersion(FControlRigObjectVersion::GUID) < FControlRigObjectVersion::RigElementKeyCache)
-	{
-		for (URigVMGraph* Graph : GetAllModels())
-		{
-			URigVMController* Controller = GetOrCreateController(Graph);
-			TGuardValue<bool> DisablePinDefaultValueValidation(Controller->bValidatePinDefaults, false);
-			FRigVMControllerNotifGuard NotifGuard(Controller, true);
-			for (URigVMNode* Node : Graph->GetNodes())
-			{
-				if (URigVMUnitNode* UnitNode = Cast<URigVMUnitNode>(Node))
-				{
-					UScriptStruct* ScriptStruct = UnitNode->GetScriptStruct();
-					FString FunctionName = FString::Printf(TEXT("%s::%s"), *ScriptStruct->GetStructCPPName(), *UnitNode->GetMethodName().ToString());
-					const FRigVMFunction* Function = FRigVMRegistry::Get().FindFunction(*FunctionName);
-					check(Function);
-					for (TFieldIterator<FProperty> It(Function->Struct); It; ++It)
-					{
-						if (It->GetCPPType() == TEXT("FCachedRigElement"))
-						{
-							if (URigVMPin* Pin = Node->FindPin(It->GetName()))
-							{
-								int32 BoneIndex = FCString::Atoi(*Pin->GetDefaultValue());
-								FRigElementKey Key = Hierarchy->GetKey(BoneIndex);
-								FCachedRigElement DefaultValueElement(Key, Hierarchy);
-								FString Result;
-								TBaseStructure<FCachedRigElement>::Get()->ExportText(Result, &DefaultValueElement, nullptr, nullptr, PPF_None, nullptr);								
-								Controller->SetPinDefaultValue(Pin->GetPinPath(), Result, true, false, false);
-								bDirtyDuringLoad = true;
-							}							
-						}
-					}
-				}
-			}
-		}
-	}
-}
-
-void UControlRigBlueprint::PatchBoundVariables()
-{
-	if (GetLinkerCustomVersion(FControlRigObjectVersion::GUID) < FControlRigObjectVersion::BoundVariableWithInjectionNode)
-	{
-		TGuardValue<bool> GuardNotifsSelf(bSuspendModelNotificationsForSelf, true);
-		
-		for (URigVMGraph* Graph : GetAllModels())
-		{
-			URigVMController* Controller = GetOrCreateController(Graph);
-			TArray<URigVMNode*> Nodes = Graph->GetNodes();
-			for (URigVMNode* Node : Nodes)
-			{
-				for (URigVMPin* Pin : Node->GetPins())
-				{
-					for (URigVMInjectionInfo* Info : Pin->GetInjectedNodes())
-					{
-						Info->Node = Info->UnitNode_DEPRECATED;
-						Info->UnitNode_DEPRECATED = nullptr;
-						bDirtyDuringLoad = true;						
-					}
-					
-					if (!Pin->BoundVariablePath_DEPRECATED.IsEmpty())
-					{
-						Controller->BindPinToVariable(Pin->GetPinPath(), Pin->BoundVariablePath_DEPRECATED, false);
-						Pin->BoundVariablePath_DEPRECATED = FString();
-						bDirtyDuringLoad = true;
-					}
-				}
-			}
-		}
-	}
-}
-
-void UControlRigBlueprint::PatchVariableNodesWithIncorrectType()
-{
-	TGuardValue<bool> GuardNotifsSelf(bSuspendModelNotificationsForSelf, true);
-
-	struct Local
-	{
-		static bool RefreshIfNeeded(URigVMController* Controller, URigVMVariableNode* VariableNode, const FString& CPPType, UObject* CPPTypeObject)
-		{
-			if (URigVMPin* ValuePin = VariableNode->GetValuePin())
-			{
-				if (ValuePin->GetCPPType() != CPPType || ValuePin->GetCPPTypeObject() != CPPTypeObject)
-				{
-					Controller->RefreshVariableNode(VariableNode->GetFName(), VariableNode->GetVariableName(), CPPType, CPPTypeObject, false);
-					if (RigVMTypeUtils::AreCompatible(*ValuePin->GetCPPType(), ValuePin->GetCPPTypeObject(), *CPPType, CPPTypeObject))
-					{
-						return false;
-					}
-					return true;
-				}
-			}
-			return false;
-		}
-	};
-
-	for (URigVMGraph* Graph : GetAllModels())
-	{
-		URigVMController* Controller = GetOrCreateController(Graph);
-		TArray<URigVMNode*> Nodes = Graph->GetNodes();
-		for (URigVMNode* Node : Nodes)
-		{
-			if (URigVMVariableNode* VariableNode = Cast<URigVMVariableNode>(Node))
-			{
-				if (VariableNode->IsInputArgument())
-				{
-					continue;
-				}
-				
-				FRigVMGraphVariableDescription Description = VariableNode->GetVariableDescription();
-
-				// Check for local variables
-				if (VariableNode->IsLocalVariable())
-				{
-					TArray<FRigVMGraphVariableDescription> LocalVariables = Graph->GetLocalVariables(false);
-					for (FRigVMGraphVariableDescription Variable : LocalVariables)
-					{
-						if (Variable.Name == Description.Name)
-						{
-							if (Local::RefreshIfNeeded(Controller, VariableNode, Variable.CPPType, Variable.CPPTypeObject))
-							{
-								bDirtyDuringLoad = true;
-							}
-							break;
-						}
-					}
-				}
-				else
-				{
-					for (struct FBPVariableDescription& Variable : NewVariables)
-					{
-						if (Variable.VarName == Description.Name)
-						{
-							FString CPPType;
-							UObject* CPPTypeObject = nullptr;
-							RigVMTypeUtils::CPPTypeFromPinType(Variable.VarType, CPPType, &CPPTypeObject);
-							if (Local::RefreshIfNeeded(Controller, VariableNode, CPPType, CPPTypeObject))
-							{
-								bDirtyDuringLoad = true;
-							}
-						}
-					}
-				}
-			}
-		}
-	}
-}
-
-// change the default value form False to True for transform nodes
-void UControlRigBlueprint::PatchPropagateToChildren()
-{
-	// no need to update default value past this version
-	if (GetLinkerCustomVersion(FControlRigObjectVersion::GUID) >= FControlRigObjectVersion::RenameGizmoToShape)
-	{
-		return;
-	}
-	
-	auto IsNullOrControl = [](const URigVMPin* InPin)
-	{
-		const bool bHasItem = InPin->GetCPPTypeObject() == FRigElementKey::StaticStruct() && InPin->GetName() == "Item";
-		if (!bHasItem)
-		{
-			return false;
-		}
-
-		if (const URigVMPin* TypePin = InPin->FindSubPin(TEXT("Type")))
-		{
-			const FString& TypeValue = TypePin->GetDefaultValue();
-			return TypeValue == TEXT("Null") || TypeValue == TEXT("Space") || TypeValue == TEXT("Control");
-		}
-		
-		return false;
-	};
-
-	auto IsPropagateChildren = [](const URigVMPin* InPin)
-	{
-		return InPin->GetCPPType() == TEXT("bool") && InPin->GetName() == TEXT("bPropagateToChildren");
-	};
-
-	auto FindPropagatePin = [IsNullOrControl, IsPropagateChildren](const URigVMNode* InNode)-> URigVMPin*
-	{
-		URigVMPin* PropagatePin = nullptr;
-		URigVMPin* ItemPin = nullptr;  
-		for (URigVMPin* Pin: InNode->GetPins())
-		{
-			// look for Item pin
-			if (!ItemPin && IsNullOrControl(Pin))
-			{
-				ItemPin = Pin;
-			}
-
-			// look for bPropagateToChildren pin
-			if (!PropagatePin && IsPropagateChildren(Pin))
-			{
-				PropagatePin = Pin;
-			}
-
-			// return propagation pin if both found
-			if (ItemPin && PropagatePin)
-			{
-				return PropagatePin;
-			}
-		}
-		return nullptr;
-	};
-
-	for (URigVMGraph* Graph : GetAllModels())
-	{
-		TArray< const URigVMPin* > PinsToUpdate;
-		for (const URigVMNode* Node : Graph->GetNodes())
-		{
-			if (const URigVMPin* PropagatePin = FindPropagatePin(Node))
-			{
-				PinsToUpdate.Add(PropagatePin);
-			}
-		}
-		
-		if (URigVMController* Controller = GetOrCreateController(Graph))
-		{
-			FRigVMControllerNotifGuard NotifGuard(Controller, true);
-			for (const URigVMPin* Pin: PinsToUpdate)
-			{
-				Controller->SetPinDefaultValue(Pin->GetPinPath(), TEXT("True"), false, false, false);
-			}
-		}
-	}
-}
-
-void UControlRigBlueprint::PatchParameterNodesOnLoad()
-{	
-#if WITH_EDITOR
-
-	// setup variables on the blueprint based on the previous "parameters"
-	if (GetLinkerCustomVersion(FControlRigObjectVersion::GUID) < FControlRigObjectVersion::RemoveParameters)
-	{
-		
-		TGuardValue<bool> GuardNotifsSelf(bSuspendModelNotificationsForSelf, true);
-
-		check(GetDefaultModel());
-
-		for(URigVMGraph* Model : RigVMClient)
-		{
-			TArray<URigVMNode*> Nodes = Model->GetNodes();
-			for (URigVMNode* Node : Nodes)
-			{
-				if (URigVMParameterNode* ParameterNode = Cast<URigVMParameterNode>(Node))
-				{
-					FRigVMGraphParameterDescription Description = ParameterNode->GetParameterDescription();
-
-					const FEdGraphPinType PinType = RigVMTypeUtils::PinTypeFromCPPType(*Description.CPPType, Description.CPPTypeObject);
-
-					int32 VariableIndex = INDEX_NONE;
-					bool bFoundName = false;
-					for (int32 i=0; i<NewVariables.Num(); ++i)
-					{
-						if (NewVariables[i].VarName == Description.Name)
-						{
-							if (NewVariables[i].VarType == PinType &&
-								NewVariables[i].DefaultValue == Description.DefaultValue)
-							{
-								VariableIndex = i;
-							}
-							bFoundName = true;
-							break;
-						}
-					}
-
-					if (VariableIndex == INDEX_NONE)
-					{
-						FName NewVariableName = Description.Name;
-						if (bFoundName)
-						{
-							bool bFound = true;
-							int32 Count = 0;
-							while (bFound)
-							{
-								bFound = false;
-								for (int32 i=0; i<NewVariables.Num(); ++i)
-								{
-									if (NewVariables[i].VarName == NewVariableName)
-									{
-										bFound = true;
-										NewVariableName = *FString::Printf(TEXT("%s_%d"), *Description.Name.ToString(), ++Count);
-										break;
-									}
-								}
-							}
-						}
-						
-						VariableIndex = AddCRMemberVariable(this, NewVariableName, PinType, false, false, Description.DefaultValue);
-					}
-					
-					FName VarName = NewVariables[VariableIndex].VarName;
-					GetOrCreateController()->ReplaceParameterNodeWithVariable(ParameterNode->GetFName(), VarName, Description.CPPType, Description.CPPTypeObject, false);
-					bDirtyDuringLoad = true;
-				}
-			}
-		}
-		LastNewVariables = NewVariables;
-	}
-
-#endif	
-}
-
-void UControlRigBlueprint::PatchLinksWithCast()
-{
-#if WITH_EDITOR
-
-	{
-		TGuardValue<bool> DisableAutoCompile(bAutoRecompileVM, false);
-
-		// find all links containing a cast
-		TArray<TTuple<URigVMGraph*,TWeakObjectPtr<URigVMLink>,FString,FString>> LinksWithCast;
-		for (URigVMGraph* Graph : GetAllModels())
-		{
-			for(URigVMLink* Link : Graph->GetLinks())
-			{
-				const URigVMPin* SourcePin = Link->GetSourcePin();
-				const URigVMPin* TargetPin = Link->GetTargetPin();
-				if (SourcePin && TargetPin)
-				{
-					const TRigVMTypeIndex SourceTypeIndex = SourcePin->GetTypeIndex();
-					const TRigVMTypeIndex TargetTypeIndex = TargetPin->GetTypeIndex();
-					
-					if(SourceTypeIndex != TargetTypeIndex)
-					{
-						if(!FRigVMRegistry::Get().CanMatchTypes(SourceTypeIndex, TargetTypeIndex, true))
-						{
-							LinksWithCast.Emplace(Graph, TWeakObjectPtr<URigVMLink>(Link), SourcePin->GetPinPath(), TargetPin->GetPinPath());
-						}
-					}
-				}
-			}
-		}
-
-		// remove all of those links
-		for(const auto& Tuple : LinksWithCast)
-		{
-			URigVMController* Controller = GetController(Tuple.Get<0>());
-
-			if(URigVMLink* Link = Tuple.Get<1>().Get())
-			{
-				// the link may be detached, attach it first so that removal works.
-				const URigVMPin* SourcePin = Link->GetSourcePin();
-				URigVMPin* TargetPin = Link->GetTargetPin();
-				if(!SourcePin->IsLinkedTo(TargetPin))
-				{
-					const TArray<URigVMController::FLinkedPath> LinkedPaths = Controller->GetLinkedPaths({Link});
-					Controller->RestoreLinkedPaths(LinkedPaths);
-				}
-			}
-
-			Controller->BreakLink(Tuple.Get<2>(), Tuple.Get<3>(), false);
-
-			// notify the user that the link has been broken.
-			UE_LOG(LogControlRigDeveloper, Warning,
-				TEXT("A link was removed in %s (%s) - it contained different types on source and target pin (former cast link?)."),
-				*Controller->GetGraph()->GetNodePath(),
-				*URigVMLink::GetPinPathRepresentation(Tuple.Get<2>(), Tuple.Get<3>())
-			);
-		}
-	}
-#endif
-}
-
-void UControlRigBlueprint::PatchFunctionsOnLoad()
-{
-	URigVMBlueprintGeneratedClass* CRGeneratedClass = GetControlRigBlueprintGeneratedClass();
-	FRigVMGraphFunctionStore& Store = CRGeneratedClass->GraphFunctionStore;
-	const URigVMFunctionLibrary* Library = GetLocalFunctionLibrary();
-
-	TMap<URigVMLibraryNode*, FRigVMGraphFunctionHeader> OldHeaders;
-
-	// Backwards compatibility. Store public access in the model
-	TArray<FName> BackwardsCompatiblePublicFunctions;
-	if (GetLinkerCustomVersion(FControlRigObjectVersion::GUID) < FControlRigObjectVersion::StoreFunctionsInGeneratedClass)
-	{
-		for (const FControlRigPublicFunctionData& Data : PublicFunctions_DEPRECATED)
-		{
-			BackwardsCompatiblePublicFunctions.Add(Data.Name);
-		}
-	}
-	else
-	{
-		if (GetLinkerCustomVersion(FUE5MainStreamObjectVersion::GUID) < FUE5MainStreamObjectVersion::RigVMSaveFunctionAccessInModel)
-		{
-			for (const FRigVMGraphFunctionData& FunctionData : Store.PublicFunctions)
-			{
-				BackwardsCompatiblePublicFunctions.Add(FunctionData.Header.Name);
-				URigVMLibraryNode* LibraryNode = Cast<URigVMLibraryNode>(FunctionData.Header.LibraryPointer.LibraryNode.ResolveObject());
-				OldHeaders.Add(LibraryNode, FunctionData.Header);
-			}
-		}
-	}
-
-	// Addressing issue where PublicGraphFunctions is populated, but the model PublicFunctionNames is not
-	URigVMFunctionLibrary* FunctionLibrary = GetLocalFunctionLibrary();
-	if (FunctionLibrary)
-	{
-		if (PublicGraphFunctions.Num() > FunctionLibrary->PublicFunctionNames.Num())
-		{
-			for (const FRigVMGraphFunctionHeader& PublicHeader : PublicGraphFunctions)
-			{
-				BackwardsCompatiblePublicFunctions.Add(PublicHeader.Name);
-			}
-		}
-	}
-
-	// Lets rebuild the FunctionStore from the model
-	if (FunctionLibrary)
-	{
-		Store.PublicFunctions.Reset();
-		Store.PrivateFunctions.Reset();
-
-		for (URigVMLibraryNode* LibraryNode : FunctionLibrary->GetFunctions())
-		{
-			bool bIsPublic = FunctionLibrary->IsFunctionPublic(LibraryNode->GetFName());
-			if (!bIsPublic)
-			{
-				bIsPublic = BackwardsCompatiblePublicFunctions.Contains(LibraryNode->GetFName());
-				if (bIsPublic)
-				{
-					FunctionLibrary->PublicFunctionNames.Add(LibraryNode->GetFName());
-				}
-			}
-
-			FRigVMGraphFunctionHeader Header = LibraryNode->GetFunctionHeader(CRGeneratedClass);
-			if (FRigVMGraphFunctionHeader* OldHeader = OldHeaders.Find(LibraryNode))
-			{				
-				Header.ExternalVariables = OldHeader->ExternalVariables;
-				Header.Dependencies = OldHeader->Dependencies;
-			}
-			Store.AddFunction(Header, bIsPublic);
-			
-		}
-	}
-
-	// Update dependencies and external variables if needed
-	for (URigVMLibraryNode* LibraryNode : Library->GetFunctions())
-	{
-		GetRigVMClient()->UpdateExternalVariablesForFunction(LibraryNode);
-		GetRigVMClient()->UpdateDependenciesForFunction(LibraryNode);
-	}
-}
-
-void UControlRigBlueprint::PropagatePoseFromInstanceToBP(UControlRig* InControlRig)
-{
-	check(InControlRig);
-	// current transforms in BP and CDO are meaningless, no need to copy them
-	// we use BP hierarchy to initialize CDO and instances' hierarchy, 
-	// so it should always be in the initial state.
-	Hierarchy->CopyPose(InControlRig->GetHierarchy(), false, true, false, true);
+	return ReturnName;
 }
 
-void UControlRigBlueprint::PropagatePoseFromBPToInstances()
+FName UControlRigBlueprint::RemoveTransientControl(URigVMPin* InPin)
 {
-	if (UClass* MyControlRigClass = GeneratedClass)
+	TUniquePtr<FControlValueScope> ValueScope;
+	if (!UControlRigEditorSettings::Get()->bResetControlsOnPinValueInteraction) // if we need to retain the controls
 	{
-		if (UControlRig* DefaultObject = Cast<UControlRig>(MyControlRigClass->GetDefaultObject(false)))
-		{
-			DefaultObject->PostInitInstanceIfRequired();
-			DefaultObject->GetHierarchy()->CopyPose(Hierarchy, true, true, true);
-
-			TArray<UObject*> ArchetypeInstances;
-			DefaultObject->GetArchetypeInstances(ArchetypeInstances);
-			for (UObject* ArchetypeInstance : ArchetypeInstances)
-			{
-				if (UControlRig* InstanceRig = Cast<UControlRig>(ArchetypeInstance))
-				{
-					InstanceRig->PostInitInstanceIfRequired();
-					InstanceRig->GetHierarchy()->CopyPose(Hierarchy, true, true, true);
-				}
-			}
-		}
+		ValueScope = MakeUnique<FControlValueScope>(this);
 	}
-}
-
-void UControlRigBlueprint::PropagateHierarchyFromBPToInstances()
-{
-	if (UClass* MyControlRigClass = GeneratedClass)
-	{
-		if (UControlRig* DefaultObject = Cast<UControlRig>(MyControlRigClass->GetDefaultObject(false)))
-		{
-			DefaultObject->PostInitInstanceIfRequired();
-			DefaultObject->GetHierarchy()->CopyHierarchy(Hierarchy);
-			DefaultObject->HierarchySettings = HierarchySettings;
-			DefaultObject->Initialize(true);
 
-			TArray<UObject*> ArchetypeInstances;
-			DefaultObject->GetArchetypeInstances(ArchetypeInstances);
-			for (UObject* ArchetypeInstance : ArchetypeInstances)
-			{
-				if (UControlRig* InstanceRig = Cast<UControlRig>(ArchetypeInstance))
-				{
-					InstanceRig->PostInitInstanceIfRequired();
-					InstanceRig->GetHierarchy()->CopyHierarchy(Hierarchy);
-					InstanceRig->HierarchySettings = HierarchySettings;
-					InstanceRig->Initialize(true);
-				}
-			}
-		}
-	}
-}
+	URigVMBlueprintGeneratedClass* RigClass = GetRigVMBlueprintGeneratedClass();
+	UControlRig* CDO = Cast<UControlRig>(RigClass->GetDefaultObject(true /* create if needed */));
 
-void UControlRigBlueprint::PropagateDrawInstructionsFromBPToInstances()
-{
-	if (UClass* MyControlRigClass = GeneratedClass)
-	{
-		if (UControlRig* DefaultObject = Cast<UControlRig>(MyControlRigClass->GetDefaultObject(false)))
+	FName RemovedName = NAME_None;
+	TArray<UObject*> ArchetypeInstances;
+	CDO->GetArchetypeInstances(ArchetypeInstances);
+	for (UObject* ArchetypeInstance : ArchetypeInstances)
 	{
-			DefaultObject->DrawContainer = DrawContainer;
-
-			TArray<UObject*> ArchetypeInstances;
-			DefaultObject->GetArchetypeInstances(ArchetypeInstances);
-
-			for (UObject* ArchetypeInstance : ArchetypeInstances)
-			{
-				if (UControlRig* InstanceRig = Cast<UControlRig>(ArchetypeInstance))
+		UControlRig* InstancedControlRig = Cast<UControlRig>(ArchetypeInstance);
+		if (InstancedControlRig)
+		{
+			FName Name = InstancedControlRig->RemoveTransientControl(InPin);
+			if (RemovedName == NAME_None)
 	{
-					InstanceRig->DrawContainer = DrawContainer;
-				}
+				RemovedName = Name;
 			}
 		}
 	}
 
-
-	// make sure the bone name list is up 2 date for the editor graph
-	for (UEdGraph* Graph : UbergraphPages)
-	{
-		UControlRigGraph* RigGraph = Cast<UControlRigGraph>(Graph);
-		if (RigGraph == nullptr)
-		{
-			continue;
-		}
-		RigGraph->CacheNameLists(Hierarchy, &DrawContainer, ShapeLibraries);
-	}
+	return RemovedName;
 }
 
-void UControlRigBlueprint::PropagateRuntimeSettingsFromBPToInstances()
+FName UControlRigBlueprint::AddTransientControl(const FRigElementKey& InElement)
 {
-	if (UClass* MyControlRigClass = GeneratedClass)
+	TUniquePtr<FControlValueScope> ValueScope;
+	if (!UControlRigEditorSettings::Get()->bResetControlsOnPinValueInteraction) // if we need to retain the controls
 	{
-		if (UControlRig* DefaultObject = Cast<UControlRig>(MyControlRigClass->GetDefaultObject(false)))
-		{
-			DefaultObject->VMRuntimeSettings = VMRuntimeSettings;
-
-			TArray<UObject*> ArchetypeInstances;
-			DefaultObject->GetArchetypeInstances(ArchetypeInstances);
-
-			for (UObject* ArchetypeInstance : ArchetypeInstances)
-			{
-				if (UControlRig* InstanceRig = Cast<UControlRig>(ArchetypeInstance))
-				{
-					InstanceRig->VMRuntimeSettings = VMRuntimeSettings;
-				}
-			}
-		}
+		ValueScope = MakeUnique<FControlValueScope>(this);
 	}
+	URigVMBlueprintGeneratedClass* RigClass = GetRigVMBlueprintGeneratedClass();
+	UControlRig* CDO = Cast<UControlRig>(RigClass->GetDefaultObject(true /* create if needed */));
 
-	TArray<UEdGraph*> EdGraphs;
-	GetAllGraphs(EdGraphs);
+	FName ReturnName = NAME_None;
+	TArray<UObject*> ArchetypeInstances;
+	CDO->GetArchetypeInstances(ArchetypeInstances);
 
-	for (UEdGraph* Graph : EdGraphs)
+	// hierarchy transforms will be reset when ClearTransientControls() is called,
+	// so to retain any bone transform modifications we have to save them
+	TMap<UObject*, FTransform> SavedElementLocalTransforms;
+	for (UObject* ArchetypeInstance : ArchetypeInstances)
 	{
-		TArray<UEdGraphNode*> Nodes = Graph->Nodes;
-		for (UEdGraphNode* Node : Nodes)
+		UControlRig* InstancedControlRig = Cast<UControlRig>(ArchetypeInstance);
+		if (InstancedControlRig)
 		{
-			if(UControlRigGraphNode* RigNode = Cast<UControlRigGraphNode>(Node))
-			{
-				RigNode->ReconstructNode_Internal(true);
+			if (InstancedControlRig->DynamicHierarchy)
+			{ 
+				SavedElementLocalTransforms.FindOrAdd(InstancedControlRig) = InstancedControlRig->DynamicHierarchy->GetLocalTransform(InElement);
 			}
 		}
 	}
-}
-
-void UControlRigBlueprint::PropagatePropertyFromBPToInstances(FRigElementKey InRigElement, const FProperty* InProperty)
-{
-	int32 ElementIndex = Hierarchy->GetIndex(InRigElement);
-	ensure(ElementIndex != INDEX_NONE);
-	check(InProperty);
 
-	if (UClass* MyControlRigClass = GeneratedClass)
+	// for now we only allow one pin control at the same time
+	ClearTransientControls();
+	
+	for (UObject* ArchetypeInstance : ArchetypeInstances)
 	{
-		if (UControlRig* DefaultObject = Cast<UControlRig>(MyControlRigClass->GetDefaultObject(false)))
+		UControlRig* InstancedControlRig = Cast<UControlRig>(ArchetypeInstance);
+		if (InstancedControlRig)
 		{
-			TArray<UObject*> ArchetypeInstances;
-			DefaultObject->GetArchetypeInstances(ArchetypeInstances);
-
-			const int32 PropertyOffset = InProperty->GetOffset_ReplaceWith_ContainerPtrToValuePtr();
-			const int32 PropertySize = InProperty->GetSize();
-
-			uint8* Source = ((uint8*)Hierarchy->Get(ElementIndex)) + PropertyOffset;
-			for (UObject* ArchetypeInstance : ArchetypeInstances)
+			// restore the element transforms so that transient controls are created at the right place
+			if (const FTransform* SavedTransform = SavedElementLocalTransforms.Find(InstancedControlRig))
 			{
-				if (UControlRig* InstanceRig = Cast<UControlRig>(ArchetypeInstance))
-				{
-					InstanceRig->PostInitInstanceIfRequired();
-					uint8* Dest = ((uint8*)InstanceRig->GetHierarchy()->Get(ElementIndex)) + PropertyOffset;
-					FMemory::Memcpy(Dest, Source, PropertySize);
+				if (InstancedControlRig->DynamicHierarchy)
+				{ 
+					InstancedControlRig->DynamicHierarchy->SetLocalTransform(InElement, *SavedTransform);
 				}
 			}
+			
+			FName ControlName = InstancedControlRig->AddTransientControl(InElement);
+			if (ReturnName == NAME_None)
+			{
+				ReturnName = ControlName;
+			}
 		}
 	}
-}
 
-void UControlRigBlueprint::PropagatePropertyFromInstanceToBP(FRigElementKey InRigElement, const FProperty* InProperty, UControlRig* InInstance)
-{
-	const int32 ElementIndex = Hierarchy->GetIndex(InRigElement);
-	ensure(ElementIndex != INDEX_NONE);
-	check(InProperty);
+	return ReturnName;
 
-	const int32 PropertyOffset = InProperty->GetOffset_ReplaceWith_ContainerPtrToValuePtr();
-	const int32 PropertySize = InProperty->GetSize();
-	uint8* Source = ((uint8*)InInstance->GetHierarchy()->Get(ElementIndex)) + PropertyOffset;
-	uint8* Dest = ((uint8*)Hierarchy->Get(ElementIndex)) + PropertyOffset;
-	FMemory::Memcpy(Dest, Source, PropertySize);
 }
 
-
-void UControlRigBlueprint::HandleHierarchyModified(ERigHierarchyNotification InNotification, URigHierarchy* InHierarchy, const FRigBaseElement* InElement)
+FName UControlRigBlueprint::RemoveTransientControl(const FRigElementKey& InElement)
 {
-#if WITH_EDITOR
-
-	if(bSuspendAllNotifications)
+	TUniquePtr<FControlValueScope> ValueScope;
+	if (!UControlRigEditorSettings::Get()->bResetControlsOnPinValueInteraction) // if we need to retain the controls
 	{
-		return;
+		ValueScope = MakeUnique<FControlValueScope>(this);
 	}
 
-	switch(InNotification)
-	{
-		case ERigHierarchyNotification::ElementRemoved:
-		{
-			Modify();
-			Influences.OnKeyRemoved(InElement->GetKey());
-			PropagateHierarchyFromBPToInstances();
-			break;
-		}
-		case ERigHierarchyNotification::ElementRenamed:
-		{
-			Modify();
-			Influences.OnKeyRenamed(FRigElementKey(InHierarchy->GetPreviousName(InElement->GetKey()), InElement->GetType()), InElement->GetKey());
-			PropagateHierarchyFromBPToInstances();
-			break;
-		}
-		case ERigHierarchyNotification::ElementAdded:
-		case ERigHierarchyNotification::ParentChanged:
-		case ERigHierarchyNotification::ElementReordered:
-		case ERigHierarchyNotification::HierarchyReset:
-		{
-			Modify();
-			PropagateHierarchyFromBPToInstances();
-			break;
-		}
-		case ERigHierarchyNotification::ElementSelected:
-		{
-			bool bClearTransientControls = true;
-			if (const FRigControlElement* ControlElement = Cast<FRigControlElement>(InElement))
-			{
-				if (ControlElement->Settings.bIsTransientControl)
-				{
-					bClearTransientControls = false;
-				}
-			}
-
-			if(bClearTransientControls)
-			{
-				if(UControlRig* RigBeingDebugged = Cast<UControlRig>(GetObjectBeingDebugged()))
-				{
-					const FName TransientControlName = UControlRig::GetNameForTransientControl(InElement->GetKey());
-					const FRigElementKey TransientControlKey(TransientControlName, ERigElementType::Control);
-					if (const FRigControlElement* ControlElement = RigBeingDebugged->GetHierarchy()->Find<FRigControlElement>(TransientControlKey))
-					{
-						if (ControlElement->Settings.bIsTransientControl)
-						{
-							bClearTransientControls = false;
-						}
-					}
-				}
-			}
+	URigVMBlueprintGeneratedClass* RigClass = GetRigVMBlueprintGeneratedClass();
+	UControlRig* CDO = Cast<UControlRig>(RigClass->GetDefaultObject(true /* create if needed */));
 
-			if(bClearTransientControls)
-			{
-				ClearTransientControls();
-			}
-			break;
-		}
-		default:
+	FName RemovedName = NAME_None;
+	TArray<UObject*> ArchetypeInstances;
+	CDO->GetArchetypeInstances(ArchetypeInstances);
+	for (UObject* ArchetypeInstance : ArchetypeInstances)
+	{
+		UControlRig* InstancedControlRig = Cast<UControlRig>(ArchetypeInstance);
+		if (InstancedControlRig)
 		{
-			break;
+			FName Name = InstancedControlRig->RemoveTransientControl(InElement);
+			if (RemovedName == NAME_None)
+			{
+				RemovedName = Name;
+			}
 		}
 	}
 
-	HierarchyModifiedEvent.Broadcast(InNotification, InHierarchy, InElement);
-	
-#endif
+	return RemovedName;
 }
 
-UControlRigBlueprint::FControlValueScope::FControlValueScope(UControlRigBlueprint* InBlueprint)
-: Blueprint(InBlueprint)
+void UControlRigBlueprint::ClearTransientControls()
 {
-#if WITH_EDITOR
-	check(Blueprint);
+	TUniquePtr<FControlValueScope> ValueScope;
+	if (!UControlRigEditorSettings::Get()->bResetControlsOnPinValueInteraction) // if we need to retain the controls
+	{
+		ValueScope = MakeUnique<FControlValueScope>(this);
+	}
 
-	if (UControlRig* CR = Cast<UControlRig>(Blueprint->GetObjectBeingDebugged()))
+	if (URigVMBlueprintGeneratedClass* RigClass = GetRigVMBlueprintGeneratedClass())
 	{
-		TArray<FRigControlElement*> Controls = CR->AvailableControls();
-		for (FRigControlElement* ControlElement : Controls)
+		UControlRig* CDO = Cast<UControlRig>(RigClass->GetDefaultObject(true /* create if needed */));
+
+		TArray<UObject*> ArchetypeInstances;
+		CDO->GetArchetypeInstances(ArchetypeInstances);
+		for (UObject* ArchetypeInstance : ArchetypeInstances)
 		{
-			ControlValues.Add(ControlElement->GetName(), CR->GetControlValue(ControlElement->GetName()));
+			UControlRig* InstancedControlRig = Cast<UControlRig>(ArchetypeInstance);
+			if (InstancedControlRig)
+			{
+				InstancedControlRig->ClearTransientControls();
+			}
 		}
 	}
-#endif
 }
 
-UControlRigBlueprint::FControlValueScope::~FControlValueScope()
+#endif
+
+void UControlRigBlueprint::SetupDefaultObjectDuringCompilation(URigVMHost* InCDO)
 {
-#if WITH_EDITOR
-	check(Blueprint);
+	Super::SetupDefaultObjectDuringCompilation(InCDO);
+	CastChecked<UControlRig>(InCDO)->GetHierarchy()->CopyHierarchy(Hierarchy);
+}
 
-	if (UControlRig* CR = Cast<UControlRig>(Blueprint->GetObjectBeingDebugged()))
+void UControlRigBlueprint::SetupPinRedirectorsForBackwardsCompatibility()
+{
+	for(URigVMGraph* Model : RigVMClient)
 	{
-		for (const TPair<FName, FRigControlValue>& Pair : ControlValues)
+		for (URigVMNode* Node : Model->GetNodes())
 		{
-			if (CR->FindControl(Pair.Key))
+			if (URigVMUnitNode* UnitNode = Cast<URigVMUnitNode>(Node))
 			{
-				CR->SetControlValue(Pair.Key, Pair.Value);
+				UScriptStruct* Struct = UnitNode->GetScriptStruct();
+				if (Struct == FRigUnit_SetBoneTransform::StaticStruct())
+				{
+					URigVMPin* TransformPin = UnitNode->FindPin(TEXT("Transform"));
+					URigVMPin* ResultPin = UnitNode->FindPin(TEXT("Result"));
+					GetOrCreateController()->AddPinRedirector(false, true, TransformPin->GetPinPath(), ResultPin->GetPinPath());
+				}
 			}
 		}
 	}
-#endif
 }
 
-#if WITH_EDITOR
+void UControlRigBlueprint::PathDomainSpecificContentOnLoad()
+{
+	PatchRigElementKeyCacheOnLoad();
+	PatchPropagateToChildren();
+}
 
-void UControlRigBlueprint::OnPreVariableChange(UObject* InObject)
+void UControlRigBlueprint::PatchRigElementKeyCacheOnLoad()
 {
-	if (InObject != this)
+	if (GetLinkerCustomVersion(FControlRigObjectVersion::GUID) < FControlRigObjectVersion::RigElementKeyCache)
 	{
-		return;
+		for (URigVMGraph* Graph : GetAllModels())
+		{
+			URigVMController* Controller = GetOrCreateController(Graph);
+			TGuardValue<bool> DisablePinDefaultValueValidation(Controller->bValidatePinDefaults, false);
+			FRigVMControllerNotifGuard NotifGuard(Controller, true);
+			for (URigVMNode* Node : Graph->GetNodes())
+			{
+				if (URigVMUnitNode* UnitNode = Cast<URigVMUnitNode>(Node))
+				{
+					UScriptStruct* ScriptStruct = UnitNode->GetScriptStruct();
+					FString FunctionName = FString::Printf(TEXT("%s::%s"), *ScriptStruct->GetStructCPPName(), *UnitNode->GetMethodName().ToString());
+					const FRigVMFunction* Function = FRigVMRegistry::Get().FindFunction(*FunctionName);
+					check(Function);
+					for (TFieldIterator<FProperty> It(Function->Struct); It; ++It)
+					{
+						if (It->GetCPPType() == TEXT("FCachedRigElement"))
+						{
+							if (URigVMPin* Pin = Node->FindPin(It->GetName()))
+							{
+								int32 BoneIndex = FCString::Atoi(*Pin->GetDefaultValue());
+								FRigElementKey Key = Hierarchy->GetKey(BoneIndex);
+								FCachedRigElement DefaultValueElement(Key, Hierarchy);
+								FString Result;
+								TBaseStructure<FCachedRigElement>::Get()->ExportText(Result, &DefaultValueElement, nullptr, nullptr, PPF_None, nullptr);								
+								Controller->SetPinDefaultValue(Pin->GetPinPath(), Result, true, false, false);
+								MarkDirtyDuringLoad();
+							}							
+						}
+					}
+				}
+			}
+		}
 	}
-	LastNewVariables = NewVariables;
 }
 
-void UControlRigBlueprint::OnPostVariableChange(UBlueprint* InBlueprint)
+// change the default value form False to True for transform nodes
+void UControlRigBlueprint::PatchPropagateToChildren()
 {
-	if (InBlueprint != this)
+	// no need to update default value past this version
+	if (GetLinkerCustomVersion(FControlRigObjectVersion::GUID) >= FControlRigObjectVersion::RenameGizmoToShape)
 	{
 		return;
 	}
-
-	if (bUpdatingExternalVariables)
+	
+	auto IsNullOrControl = [](const URigVMPin* InPin)
 	{
-		return;
-	}
-
-	TGuardValue<bool> UpdatingVariablesGuard(bUpdatingExternalVariables, true);
-	TArray<FBPVariableDescription> LocalLastNewVariables = LastNewVariables;
+		const bool bHasItem = InPin->GetCPPTypeObject() == FRigElementKey::StaticStruct() && InPin->GetName() == "Item";
+		if (!bHasItem)
+		{
+			return false;
+		}
 
-	TMap<FGuid, int32> NewVariablesByGuid;
-	for (int32 VarIndex = 0; VarIndex < NewVariables.Num(); VarIndex++)
-	{
-		NewVariablesByGuid.Add(NewVariables[VarIndex].VarGuid, VarIndex);
-	}
+		if (const URigVMPin* TypePin = InPin->FindSubPin(TEXT("Type")))
+		{
+			const FString& TypeValue = TypePin->GetDefaultValue();
+			return TypeValue == TEXT("Null") || TypeValue == TEXT("Space") || TypeValue == TEXT("Control");
+		}
+		
+		return false;
+	};
 
-	TMap<FGuid, int32> OldVariablesByGuid;
-	for (int32 VarIndex = 0; VarIndex < LocalLastNewVariables.Num(); VarIndex++)
+	auto IsPropagateChildren = [](const URigVMPin* InPin)
 	{
-		OldVariablesByGuid.Add(LocalLastNewVariables[VarIndex].VarGuid, VarIndex);
-	}
+		return InPin->GetCPPType() == TEXT("bool") && InPin->GetName() == TEXT("bPropagateToChildren");
+	};
 
-	for (const FBPVariableDescription& OldVariable : LocalLastNewVariables)
+	auto FindPropagatePin = [IsNullOrControl, IsPropagateChildren](const URigVMNode* InNode)-> URigVMPin*
 	{
-		if (!NewVariablesByGuid.Contains(OldVariable.VarGuid))
+		URigVMPin* PropagatePin = nullptr;
+		URigVMPin* ItemPin = nullptr;  
+		for (URigVMPin* Pin: InNode->GetPins())
 		{
-			OnVariableRemoved(OldVariable.VarName);
-			continue;
-		}
-	}
+			// look for Item pin
+			if (!ItemPin && IsNullOrControl(Pin))
+			{
+				ItemPin = Pin;
+			}
 
-	for (const FBPVariableDescription& NewVariable : NewVariables)
-	{
-		if (!OldVariablesByGuid.Contains(NewVariable.VarGuid))
-		{
-			OnVariableAdded(NewVariable.VarName);
-			continue;
+			// look for bPropagateToChildren pin
+			if (!PropagatePin && IsPropagateChildren(Pin))
+			{
+				PropagatePin = Pin;
+			}
+
+			// return propagation pin if both found
+			if (ItemPin && PropagatePin)
+			{
+				return PropagatePin;
+			}
 		}
+		return nullptr;
+	};
 
-		int32 OldVarIndex = OldVariablesByGuid.FindChecked(NewVariable.VarGuid);
-		const FBPVariableDescription& OldVariable = LocalLastNewVariables[OldVarIndex];
-		if (OldVariable.VarName != NewVariable.VarName)
+	for (URigVMGraph* Graph : GetAllModels())
+	{
+		TArray< const URigVMPin* > PinsToUpdate;
+		for (const URigVMNode* Node : Graph->GetNodes())
 		{
-			OnVariableRenamed(OldVariable.VarName, NewVariable.VarName);
+			if (const URigVMPin* PropagatePin = FindPropagatePin(Node))
+			{
+				PinsToUpdate.Add(PropagatePin);
+			}
 		}
-
-		if (OldVariable.VarType != NewVariable.VarType)
+		
+		if (URigVMController* Controller = GetOrCreateController(Graph))
 		{
-			OnVariableTypeChanged(NewVariable.VarName, OldVariable.VarType, NewVariable.VarType);
+			FRigVMControllerNotifGuard NotifGuard(Controller, true);
+			for (const URigVMPin* Pin: PinsToUpdate)
+			{
+				Controller->SetPinDefaultValue(Pin->GetPinPath(), TEXT("True"), false, false, false);
+			}
 		}
 	}
-
-	LastNewVariables = NewVariables;
 }
 
-void UControlRigBlueprint::OnVariableAdded(const FName& InVarName)
+void UControlRigBlueprint::PatchFunctionsOnLoad()
 {
-	FBPVariableDescription Variable;
-	for (FBPVariableDescription& NewVariable : NewVariables)
+	URigVMBlueprintGeneratedClass* CRGeneratedClass = GetRigVMBlueprintGeneratedClass();
+	FRigVMGraphFunctionStore& Store = CRGeneratedClass->GraphFunctionStore;
+	const URigVMFunctionLibrary* Library = GetLocalFunctionLibrary();
+
+	TMap<URigVMLibraryNode*, FRigVMGraphFunctionHeader> OldHeaders;
+
+	// Backwards compatibility. Store public access in the model
+	TArray<FName> BackwardsCompatiblePublicFunctions;
+	if (GetLinkerCustomVersion(FControlRigObjectVersion::GUID) < FControlRigObjectVersion::StoreFunctionsInGeneratedClass)
 	{
-		if (NewVariable.VarName == InVarName)
+		for (const FControlRigPublicFunctionData& Data : PublicFunctions_DEPRECATED)
 		{
-			Variable = NewVariable;
-			break;
+			BackwardsCompatiblePublicFunctions.Add(Data.Name);
 		}
 	}
-
-	const FRigVMExternalVariable ExternalVariable = RigVMTypeUtils::ExternalVariableFromBPVariableDescription(Variable);
-    FString CPPType;
-    UObject* CPPTypeObject = nullptr;
-    RigVMTypeUtils::CPPTypeFromExternalVariable(ExternalVariable, CPPType, &CPPTypeObject);
-	if (CPPTypeObject)
+	else
 	{
-		if (ExternalVariable.bIsArray)
-		{
-			CPPType = RigVMTypeUtils::ArrayTypeFromBaseType(CPPTypeObject->GetPathName());
-		}
-		else
+		if (GetLinkerCustomVersion(FUE5MainStreamObjectVersion::GUID) < FUE5MainStreamObjectVersion::RigVMSaveFunctionAccessInModel)
 		{
-			CPPType = CPPTypeObject->GetPathName();
+			for (const FRigVMGraphFunctionData& FunctionData : Store.PublicFunctions)
+			{
+				BackwardsCompatiblePublicFunctions.Add(FunctionData.Header.Name);
+				URigVMLibraryNode* LibraryNode = Cast<URigVMLibraryNode>(FunctionData.Header.LibraryPointer.LibraryNode.ResolveObject());
+				OldHeaders.Add(LibraryNode, FunctionData.Header);
+			}
 		}
 	}
 
-	// register the type in the registry
-	FRigVMRegistry::Get().FindOrAddType(FRigVMTemplateArgumentType(*CPPType, CPPTypeObject));
-	
-    RigVMPythonUtils::Print(GetFName().ToString(),
-		FString::Printf(TEXT("blueprint.add_member_variable('%s', '%s', %s, %s, '%s')"),
-			*InVarName.ToString(),
-			*CPPType,
-			(ExternalVariable.bIsPublic) ? TEXT("False") : TEXT("True"), 
-			(ExternalVariable.bIsReadOnly) ? TEXT("True") : TEXT("False"), 
-			*Variable.DefaultValue)); 
-	
-	BroadcastExternalVariablesChangedEvent();
-}
-
-void UControlRigBlueprint::OnVariableRemoved(const FName& InVarName)
-{
-	TArray<URigVMGraph*> AllGraphs = GetAllModels();
-	for (URigVMGraph* Graph : AllGraphs)
+	// Addressing issue where PublicGraphFunctions is populated, but the model PublicFunctionNames is not
+	URigVMFunctionLibrary* FunctionLibrary = GetLocalFunctionLibrary();
+	if (FunctionLibrary)
 	{
-		if (URigVMController* Controller = GetOrCreateController(Graph))
+		if (PublicGraphFunctions.Num() > FunctionLibrary->PublicFunctionNames.Num())
 		{
-#if WITH_EDITOR
-			const bool bSetupUndoRedo = !GIsTransacting;
-#else
-			const bool bSetupUndoRedo = false;
-#endif
-			Controller->OnExternalVariableRemoved(InVarName, bSetupUndoRedo);
+			for (const FRigVMGraphFunctionHeader& PublicHeader : PublicGraphFunctions)
+			{
+				BackwardsCompatiblePublicFunctions.Add(PublicHeader.Name);
+			}
 		}
 	}
 
-	RigVMPythonUtils::Print(GetFName().ToString(),
-		FString::Printf(TEXT("blueprint.remove_member_variable('%s')"),
-			*InVarName.ToString()));
-	
-	BroadcastExternalVariablesChangedEvent();
-}
-
-void UControlRigBlueprint::OnVariableRenamed(const FName& InOldVarName, const FName& InNewVarName)
-{
-	TArray<URigVMGraph*> AllGraphs = GetAllModels();
-	for (URigVMGraph* Graph : AllGraphs)
+	// Lets rebuild the FunctionStore from the model
+	if (FunctionLibrary)
 	{
-		if (URigVMController* Controller = GetOrCreateController(Graph))
+		Store.PublicFunctions.Reset();
+		Store.PrivateFunctions.Reset();
+
+		for (URigVMLibraryNode* LibraryNode : FunctionLibrary->GetFunctions())
 		{
-#if WITH_EDITOR
-			const bool bSetupUndoRedo = !GIsTransacting;
-#else
-			const bool bSetupUndoRedo = false;
-#endif
-			Controller->OnExternalVariableRenamed(InOldVarName, InNewVarName, bSetupUndoRedo);
+			bool bIsPublic = FunctionLibrary->IsFunctionPublic(LibraryNode->GetFName());
+			if (!bIsPublic)
+			{
+				bIsPublic = BackwardsCompatiblePublicFunctions.Contains(LibraryNode->GetFName());
+				if (bIsPublic)
+				{
+					FunctionLibrary->PublicFunctionNames.Add(LibraryNode->GetFName());
+				}
+			}
+
+			FRigVMGraphFunctionHeader Header = LibraryNode->GetFunctionHeader(CRGeneratedClass);
+			if (FRigVMGraphFunctionHeader* OldHeader = OldHeaders.Find(LibraryNode))
+			{				
+				Header.ExternalVariables = OldHeader->ExternalVariables;
+				Header.Dependencies = OldHeader->Dependencies;
+			}
+			Store.AddFunction(Header, bIsPublic);
+			
 		}
 	}
 
-	RigVMPythonUtils::Print(GetFName().ToString(),
-		FString::Printf(TEXT("blueprint.rename_member_variable('%s', '%s')"),
-			*InOldVarName.ToString(),
-			*InNewVarName.ToString()));
-	
-	BroadcastExternalVariablesChangedEvent();
+	// Update dependencies and external variables if needed
+	for (URigVMLibraryNode* LibraryNode : Library->GetFunctions())
+	{
+		GetRigVMClient()->UpdateExternalVariablesForFunction(LibraryNode);
+		GetRigVMClient()->UpdateDependenciesForFunction(LibraryNode);
+	}
+}
+
+void UControlRigBlueprint::PropagatePoseFromInstanceToBP(UControlRig* InControlRig) const
+{
+	check(InControlRig);
+	// current transforms in BP and CDO are meaningless, no need to copy them
+	// we use BP hierarchy to initialize CDO and instances' hierarchy, 
+	// so it should always be in the initial state.
+	Hierarchy->CopyPose(InControlRig->GetHierarchy(), false, true, false, true);
 }
 
-void UControlRigBlueprint::OnVariableTypeChanged(const FName& InVarName, FEdGraphPinType InOldPinType, FEdGraphPinType InNewPinType)
+void UControlRigBlueprint::PropagatePoseFromBPToInstances() const
 {
-	FString CPPType;
-	UObject* CPPTypeObject = nullptr;
-	RigVMTypeUtils::CPPTypeFromPinType(InNewPinType, CPPType, &CPPTypeObject);
-	
-	TArray<URigVMGraph*> AllGraphs = GetAllModels();
-	for (URigVMGraph* Graph : AllGraphs)
+	if (UClass* MyControlRigClass = GeneratedClass)
 	{
-		if (URigVMController* Controller = GetOrCreateController(Graph))
+		if (UControlRig* DefaultObject = Cast<UControlRig>(MyControlRigClass->GetDefaultObject(false)))
 		{
-#if WITH_EDITOR
-			const bool bSetupUndoRedo = !GIsTransacting;
-#else
-			const bool bSetupUndoRedo = false;
-#endif
+			DefaultObject->PostInitInstanceIfRequired();
+			DefaultObject->GetHierarchy()->CopyPose(Hierarchy, true, true, true);
 
-			if (!CPPType.IsEmpty())
-			{
-				Controller->OnExternalVariableTypeChanged(InVarName, CPPType, CPPTypeObject, bSetupUndoRedo);
-			}
-			else
+			TArray<UObject*> ArchetypeInstances;
+			DefaultObject->GetArchetypeInstances(ArchetypeInstances);
+			for (UObject* ArchetypeInstance : ArchetypeInstances)
 			{
-				Controller->OnExternalVariableRemoved(InVarName, bSetupUndoRedo);
+				if (UControlRig* InstanceRig = Cast<UControlRig>(ArchetypeInstance))
+				{
+					InstanceRig->PostInitInstanceIfRequired();
+					InstanceRig->GetHierarchy()->CopyPose(Hierarchy, true, true, true);
+				}
 			}
 		}
 	}
+}
 
-	if(UScriptStruct* ScriptStruct = Cast<UScriptStruct>(CPPTypeObject))
+void UControlRigBlueprint::PropagateHierarchyFromBPToInstances() const
+{
+	if (UClass* MyControlRigClass = GeneratedClass)
 	{
-		for (auto Var : NewVariables)
+		if (UControlRig* DefaultObject = Cast<UControlRig>(MyControlRigClass->GetDefaultObject(false)))
 		{
-			if (Var.VarName == InVarName)
+			DefaultObject->PostInitInstanceIfRequired();
+			DefaultObject->GetHierarchy()->CopyHierarchy(Hierarchy);
+			DefaultObject->HierarchySettings = HierarchySettings;
+			DefaultObject->Initialize(true);
+
+			TArray<UObject*> ArchetypeInstances;
+			DefaultObject->GetArchetypeInstances(ArchetypeInstances);
+			for (UObject* ArchetypeInstance : ArchetypeInstances)
 			{
-				CPPType = ScriptStruct->GetName();
+				if (UControlRig* InstanceRig = Cast<UControlRig>(ArchetypeInstance))
+				{
+					InstanceRig->PostInitInstanceIfRequired();
+					InstanceRig->GetHierarchy()->CopyHierarchy(Hierarchy);
+					InstanceRig->HierarchySettings = HierarchySettings;
+					InstanceRig->Initialize(true);
+				}
 			}
 		}
 	}
-	else if (UEnum* Enum = Cast<UEnum>(CPPTypeObject))
+}
+
+void UControlRigBlueprint::PropagateDrawInstructionsFromBPToInstances() const
+{
+	if (UClass* MyControlRigClass = GeneratedClass)
 	{
-		for (auto Var : NewVariables)
-		{
-			if (Var.VarName == InVarName)
+		if (UControlRig* DefaultObject = Cast<UControlRig>(MyControlRigClass->GetDefaultObject(false)))
+	{
+			DefaultObject->DrawContainer = DrawContainer;
+
+			TArray<UObject*> ArchetypeInstances;
+			DefaultObject->GetArchetypeInstances(ArchetypeInstances);
+
+			for (UObject* ArchetypeInstance : ArchetypeInstances)
 			{
-				CPPType = Enum->GetName();
+				if (UControlRig* InstanceRig = Cast<UControlRig>(ArchetypeInstance))
+	{
+					InstanceRig->DrawContainer = DrawContainer;
+				}
 			}
 		}
 	}
 
-	// register the type in the registry
-	FRigVMRegistry::Get().FindOrAddType(FRigVMTemplateArgumentType(*CPPType, CPPTypeObject));
 
-	RigVMPythonUtils::Print(GetFName().ToString(),
-		FString::Printf(TEXT("blueprint.change_member_variable_type('%s', '%s')"),
-		*InVarName.ToString(),
-		*CPPType));
-
-	BroadcastExternalVariablesChangedEvent();
-}
-
-void UControlRigBlueprint::BroadcastExternalVariablesChangedEvent()
-{
-	if (URigVMBlueprintGeneratedClass* RigClass = GetControlRigBlueprintGeneratedClass())
+	// make sure the bone name list is up 2 date for the editor graph
+	for (UEdGraph* Graph : UbergraphPages)
 	{
-		if (UControlRig* CDO = Cast<UControlRig>(RigClass->GetDefaultObject(true /* create if needed */)))
+		UControlRigGraph* RigGraph = Cast<UControlRigGraph>(Graph);
+		if (RigGraph == nullptr)
 		{
-			ExternalVariablesChangedEvent.Broadcast(CDO->GetExternalVariables());
+			continue;
 		}
+		RigGraph->CacheNameLists(Hierarchy, &DrawContainer, ShapeLibraries);
 	}
 }
 
-void UControlRigBlueprint::BroadcastNodeDoubleClicked(URigVMNode* InNode)
-{
-	NodeDoubleClickedEvent.Broadcast(this, InNode);
-}
-
-void UControlRigBlueprint::BroadcastGraphImported(UEdGraph* InGraph)
-{
-	GraphImportedEvent.Broadcast(InGraph);
-}
-
-void UControlRigBlueprint::BroadcastPostEditChangeChainProperty(FPropertyChangedChainEvent& PropertyChangedChainEvent)
-{
-	PostEditChangeChainPropertyEvent.Broadcast(PropertyChangedChainEvent);
-}
-
-void UControlRigBlueprint::BroadcastRequestLocalizeFunctionDialog(FRigVMGraphFunctionIdentifier InFunction, bool bForce)
-{
-	RequestLocalizeFunctionDialog.Broadcast(InFunction, this, bForce);
-}
-
-void UControlRigBlueprint::BroadCastReportCompilerMessage(EMessageSeverity::Type InSeverity, UObject* InSubject, const FString& InMessage)
-{
-	ReportCompilerMessageEvent.Broadcast(InSeverity, InSubject, InMessage);
-}
-
-#endif
-
-UEdGraph* UControlRigBlueprint::CreateEdGraph(URigVMGraph* InModel, bool bForce)
+void UControlRigBlueprint::PropagatePropertyFromBPToInstances(FRigElementKey InRigElement, const FProperty* InProperty) const
 {
-	check(InModel);
-
-#if WITH_EDITORONLY_DATA
-	if(InModel->IsA<URigVMFunctionLibrary>())
-	{
-		return FunctionLibraryEdGraph;
-	}
-#endif
-	
-	if(bForce)
-	{
-		RemoveEdGraph(InModel);
-	}
-
-	FString GraphName = InModel->GetName();
-	GraphName.RemoveFromStart(FRigVMClient::RigVMModelPrefix);
-	GraphName.TrimStartAndEndInline();
+	int32 ElementIndex = Hierarchy->GetIndex(InRigElement);
+	ensure(ElementIndex != INDEX_NONE);
+	check(InProperty);
 
-	if(GraphName.IsEmpty())
+	if (UClass* MyControlRigClass = GeneratedClass)
 	{
-		GraphName = UControlRigGraphSchema::GraphName_ControlRig.ToString();
-	}
-
-	GraphName = RigVMClient.GetUniqueName(*GraphName).ToString();
-
-	UControlRigGraph* ControlRigGraph = NewObject<UControlRigGraph>(this, *GraphName, RF_Transactional);
-	ControlRigGraph->Schema = UControlRigGraphSchema::StaticClass();
-	ControlRigGraph->bAllowDeletion = true;
-	ControlRigGraph->ModelNodePath = InModel->GetNodePath();
-	ControlRigGraph->Initialize(this);
-	
-	FBlueprintEditorUtils::AddUbergraphPage(this, ControlRigGraph);
-	LastEditedDocuments.AddUnique(ControlRigGraph);
+		if (UControlRig* DefaultObject = Cast<UControlRig>(MyControlRigClass->GetDefaultObject(false)))
+		{
+			TArray<UObject*> ArchetypeInstances;
+			DefaultObject->GetArchetypeInstances(ArchetypeInstances);
 
-	return ControlRigGraph;
-}
+			const int32 PropertyOffset = InProperty->GetOffset_ReplaceWith_ContainerPtrToValuePtr();
+			const int32 PropertySize = InProperty->GetSize();
 
-bool UControlRigBlueprint::RemoveEdGraph(URigVMGraph* InModel)
-{
-	if(UControlRigGraph* RigGraph = Cast<UControlRigGraph>(GetEdGraph(InModel)))
-	{
-		if(UbergraphPages.Contains(RigGraph))
-		{
-			Modify();
-			UbergraphPages.Remove(RigGraph);
+			uint8* Source = ((uint8*)Hierarchy->Get(ElementIndex)) + PropertyOffset;
+			for (UObject* ArchetypeInstance : ArchetypeInstances)
+			{
+				if (UControlRig* InstanceRig = Cast<UControlRig>(ArchetypeInstance))
+				{
+					InstanceRig->PostInitInstanceIfRequired();
+					uint8* Dest = ((uint8*)InstanceRig->GetHierarchy()->Get(ElementIndex)) + PropertyOffset;
+					FMemory::Memcpy(Dest, Source, PropertySize);
+				}
+			}
 		}
-		DestroyObject(RigGraph);
-		return true;
 	}
-	return false;
-}
-
-void UControlRigBlueprint::DestroyObject(UObject* InObject)
-{
-	RigVMClient.DestroyObject(InObject);
 }
 
-void UControlRigBlueprint::RenameGraph(const FString& InNodePath, const FName& InNewName)
+void UControlRigBlueprint::PropagatePropertyFromInstanceToBP(FRigElementKey InRigElement, const FProperty* InProperty, UControlRig* InInstance) const
 {
-	FName OldName = NAME_None;
-	UEdGraph* EdGraph = GetEdGraph(InNodePath);
-	if(EdGraph)
-	{
-		OldName = EdGraph->GetFName();
-	}
-	
-	RigVMClient.RenameModel(InNodePath, InNewName, true);
+	const int32 ElementIndex = Hierarchy->GetIndex(InRigElement);
+	ensure(ElementIndex != INDEX_NONE);
+	check(InProperty);
 
-	if(EdGraph)
-	{
-		NotifyGraphRenamed(EdGraph, OldName, EdGraph->GetFName());
-	}
+	const int32 PropertyOffset = InProperty->GetOffset_ReplaceWith_ContainerPtrToValuePtr();
+	const int32 PropertySize = InProperty->GetSize();
+	uint8* Source = ((uint8*)InInstance->GetHierarchy()->Get(ElementIndex)) + PropertyOffset;
+	uint8* Dest = ((uint8*)Hierarchy->Get(ElementIndex)) + PropertyOffset;
+	FMemory::Memcpy(Dest, Source, PropertySize);
 }
 
-void UControlRigBlueprint::CreateEdGraphForCollapseNodeIfNeeded(URigVMCollapseNode* InNode, bool bForce)
+
+void UControlRigBlueprint::HandleHierarchyModified(ERigHierarchyNotification InNotification, URigHierarchy* InHierarchy, const FRigBaseElement* InElement)
 {
-	check(InNode);
+#if WITH_EDITOR
 
-	if (bForce)
+	if(bSuspendAllNotifications)
 	{
-		RemoveEdGraphForCollapseNode(InNode, false);
+		return;
 	}
 
-	if (InNode->GetGraph()->IsA<URigVMFunctionLibrary>())
+	switch(InNotification)
 	{
-		if (URigVMGraph* ContainedGraph = InNode->GetContainedGraph())
+		case ERigHierarchyNotification::ElementRemoved:
+		{
+			Modify();
+			Influences.OnKeyRemoved(InElement->GetKey());
+			PropagateHierarchyFromBPToInstances();
+			break;
+		}
+		case ERigHierarchyNotification::ElementRenamed:
+		{
+			Modify();
+			Influences.OnKeyRenamed(FRigElementKey(InHierarchy->GetPreviousName(InElement->GetKey()), InElement->GetType()), InElement->GetKey());
+			PropagateHierarchyFromBPToInstances();
+			break;
+		}
+		case ERigHierarchyNotification::ElementAdded:
+		case ERigHierarchyNotification::ParentChanged:
+		case ERigHierarchyNotification::ElementReordered:
+		case ERigHierarchyNotification::HierarchyReset:
+		{
+			Modify();
+			PropagateHierarchyFromBPToInstances();
+			break;
+		}
+		case ERigHierarchyNotification::ElementSelected:
 		{
-			bool bFunctionGraphExists = false;
-			for (UEdGraph* FunctionGraph : FunctionGraphs)
+			bool bClearTransientControls = true;
+			if (const FRigControlElement* ControlElement = Cast<FRigControlElement>(InElement))
 			{
-				if (UControlRigGraph* RigFunctionGraph = Cast<UControlRigGraph>(FunctionGraph))
+				if (ControlElement->Settings.bIsTransientControl)
 				{
-					if (RigFunctionGraph->ModelNodePath == ContainedGraph->GetNodePath())
-					{
-						bFunctionGraphExists = true;
-						break;
-					}
+					bClearTransientControls = false;
 				}
 			}
 
-			if (!bFunctionGraphExists)
-			{
-				// create a sub graph
-				UControlRigGraph* RigFunctionGraph = NewObject<UControlRigGraph>(this, *InNode->GetEditorSubGraphName(), RF_Transactional);
-				RigFunctionGraph->Schema = UControlRigGraphSchema::StaticClass();
-				RigFunctionGraph->bAllowRenaming = 1;
-				RigFunctionGraph->bEditable = 1;
-				RigFunctionGraph->bAllowDeletion = 1;
-				RigFunctionGraph->ModelNodePath = ContainedGraph->GetNodePath();
-				RigFunctionGraph->bIsFunctionDefinition = true;
-
-				FunctionGraphs.Add(RigFunctionGraph);
-
-				RigFunctionGraph->Initialize(this);
-
-				GetOrCreateController(ContainedGraph)->ResendAllNotifications();
-			}
-
-		}
-	}
-	else if (UControlRigGraph* RigGraph = Cast<UControlRigGraph>(GetEdGraph(InNode->GetGraph())))
-	{
-		if (URigVMGraph* ContainedGraph = InNode->GetContainedGraph())
-		{
-			bool bSubGraphExists = false;
-			for (UEdGraph* SubGraph : RigGraph->SubGraphs)
+			if(bClearTransientControls)
 			{
-				if (UControlRigGraph* SubRigGraph = Cast<UControlRigGraph>(SubGraph))
+				if(UControlRig* RigBeingDebugged = Cast<UControlRig>(GetObjectBeingDebugged()))
 				{
-					if (SubRigGraph->ModelNodePath == ContainedGraph->GetNodePath())
+					const FName TransientControlName = UControlRig::GetNameForTransientControl(InElement->GetKey());
+					const FRigElementKey TransientControlKey(TransientControlName, ERigElementType::Control);
+					if (const FRigControlElement* ControlElement = RigBeingDebugged->GetHierarchy()->Find<FRigControlElement>(TransientControlKey))
 					{
-						bSubGraphExists = true;
-						break;
+						if (ControlElement->Settings.bIsTransientControl)
+						{
+							bClearTransientControls = false;
+						}
 					}
 				}
 			}
 
-			if (!bSubGraphExists)
+			if(bClearTransientControls)
 			{
-				bool bEditable = true;
-				if (InNode->IsA<URigVMAggregateNode>())
-				{
-					bEditable = false;
-				}
-				
-				// create a sub graph
-				UControlRigGraph* SubRigGraph = NewObject<UControlRigGraph>(RigGraph, *InNode->GetEditorSubGraphName(), RF_Transactional);
-				SubRigGraph->Schema = UControlRigGraphSchema::StaticClass();
-				SubRigGraph->bAllowRenaming = 1;
-				SubRigGraph->bEditable = bEditable;
-				SubRigGraph->bAllowDeletion = 1;
-				SubRigGraph->ModelNodePath = ContainedGraph->GetNodePath();
-				SubRigGraph->bIsFunctionDefinition = false;
-
-				RigGraph->SubGraphs.Add(SubRigGraph);
-
-				SubRigGraph->Initialize(this);
-
-				GetOrCreateController(ContainedGraph)->ResendAllNotifications();
+				ClearTransientControls();
 			}
+			break;
+		}
+		default:
+		{
+			break;
 		}
 	}
+
+	HierarchyModifiedEvent.Broadcast(InNotification, InHierarchy, InElement);
+	
+#endif
 }
 
-bool UControlRigBlueprint::RemoveEdGraphForCollapseNode(URigVMCollapseNode* InNode, bool bNotify)
+UControlRigBlueprint::FControlValueScope::FControlValueScope(UControlRigBlueprint* InBlueprint)
+: Blueprint(InBlueprint)
 {
-	check(InNode);
+#if WITH_EDITOR
+	check(Blueprint);
 
-	if (InNode->GetGraph()->IsA<URigVMFunctionLibrary>())
+	if (UControlRig* CR = Cast<UControlRig>(Blueprint->GetObjectBeingDebugged()))
 	{
-		if (URigVMGraph* ContainedGraph = InNode->GetContainedGraph())
+		TArray<FRigControlElement*> Controls = CR->AvailableControls();
+		for (FRigControlElement* ControlElement : Controls)
 		{
-			for (UEdGraph* FunctionGraph : FunctionGraphs)
-			{
-				if (UControlRigGraph* RigFunctionGraph = Cast<UControlRigGraph>(FunctionGraph))
-				{
-					if (RigFunctionGraph->ModelNodePath == ContainedGraph->GetNodePath())
-					{
-						if (URigVMController* SubController = GetController(ContainedGraph))
-						{
-							SubController->OnModified().RemoveAll(RigFunctionGraph);
-						}
-
-						if (ModifiedEvent.IsBound() && bNotify)
-						{
-							ModifiedEvent.Broadcast(ERigVMGraphNotifType::NodeRemoved, InNode->GetGraph(), InNode);
-						}
-
-						FunctionGraphs.Remove(RigFunctionGraph);
-						RigFunctionGraph->Rename(nullptr, GetTransientPackage(), REN_ForceNoResetLoaders | REN_DontCreateRedirectors);
-						RigFunctionGraph->MarkAsGarbage();
-						return bNotify;
-					}
-				}
-			}
+			ControlValues.Add(ControlElement->GetName(), CR->GetControlValue(ControlElement->GetName()));
 		}
 	}
-	else if (UControlRigGraph* RigGraph = Cast<UControlRigGraph>(GetEdGraph(InNode->GetGraph())))
+#endif
+}
+
+UControlRigBlueprint::FControlValueScope::~FControlValueScope()
+{
+#if WITH_EDITOR
+	check(Blueprint);
+
+	if (UControlRig* CR = Cast<UControlRig>(Blueprint->GetObjectBeingDebugged()))
 	{
-		if (URigVMGraph* ContainedGraph = InNode->GetContainedGraph())
+		for (const TPair<FName, FRigControlValue>& Pair : ControlValues)
 		{
-			for (UEdGraph* SubGraph : RigGraph->SubGraphs)
+			if (CR->FindControl(Pair.Key))
 			{
-				if (UControlRigGraph* SubRigGraph = Cast<UControlRigGraph>(SubGraph))
-				{
-					if (SubRigGraph->ModelNodePath == ContainedGraph->GetNodePath())
-					{
-						if (URigVMController* SubController = GetController(ContainedGraph))
-						{
-							SubController->OnModified().RemoveAll(SubRigGraph);
-						}
-
-						if (ModifiedEvent.IsBound() && bNotify)
-						{
-							ModifiedEvent.Broadcast(ERigVMGraphNotifType::NodeRemoved, InNode->GetGraph(), InNode);
-						}
-
-						RigGraph->SubGraphs.Remove(SubRigGraph);
-						SubRigGraph->Rename(nullptr, GetTransientPackage(), REN_ForceNoResetLoaders | REN_DontCreateRedirectors);
-						SubRigGraph->MarkAsGarbage();
-						return bNotify;
-					}
-				}
+				CR->SetControlValue(Pair.Key, Pair.Value);
 			}
 		}
 	}
-
-	return false;
+#endif
 }
 
 #undef LOCTEXT_NAMESPACE
