RigVM: First steps towards rebasing ControlRigBlueprint

diff --git a/Engine/Plugins/Animation/ControlRig/Source/ControlRigDeveloper/Private/Graph/ControlRigGraph.cpp b/Engine/Plugins/Animation/ControlRig/Source/ControlRigDeveloper/Private/Graph/ControlRigGraph.cpp
--- a/Engine/Plugins/Animation/ControlRig/Source/ControlRigDeveloper/Private/Graph/ControlRigGraph.cpp
+++ b/Engine/Plugins/Animation/ControlRig/Source/ControlRigDeveloper/Private/Graph/ControlRigGraph.cpp
@@ -2,150 +2,42 @@
 
 #include "Graph/ControlRigGraph.h"
 #include "ControlRigBlueprint.h"
-#include "RigVMBlueprintGeneratedClass.h"
-#include "Graph/ControlRigGraphSchema.h"
 #include "ControlRig.h"
 #include "RigVMModel/RigVMGraph.h"
-#include "ControlRigObjectVersion.h"
 #include "Units/RigUnit.h"
-#include "EdGraphNode_Comment.h"
-#include "GraphEditAction.h"
-#include "Units/Execution/RigUnit_BeginExecution.h"
-#include "RigVMModel/Nodes/RigVMLibraryNode.h"
-#include "RigVMModel/Nodes/RigVMFunctionEntryNode.h"
-#include "RigVMModel/Nodes/RigVMFunctionReturnNode.h"
-#include "RigVMModel/Nodes/RigVMRerouteNode.h"
 
 #include UE_INLINE_GENERATED_CPP_BY_NAME(ControlRigGraph)
 
-#if WITH_EDITOR
-#include "Kismet2/BlueprintEditorUtils.h"
-#include "RigVMBlueprintUtils.h"
-#include "BlueprintCompilationManager.h"
-#include "Widgets/Notifications/SNotificationList.h"
-#include "Framework/Notifications/NotificationManager.h"
-#endif
-
 #define LOCTEXT_NAMESPACE "ControlRigGraph"
 
 UControlRigGraph::UControlRigGraph()
 {
 	bSuspendModelNotifications = false;
 	bIsTemporaryGraphForCopyPaste = false;
-	bIsSelecting = false;
 	LastHierarchyTopologyVersion = INDEX_NONE;
 	bIsFunctionDefinition = false;
 }
 
-FRigVMClient* UControlRigGraph::GetRigVMClient() const
-{
-	if (const IRigVMClientHost* Host = GetImplementingOuter<IRigVMClientHost>())
-	{
-		return (FRigVMClient*)Host->GetRigVMClient();
-	}
-	return nullptr;
-}
-
-FString UControlRigGraph::GetRigVMNodePath() const
-{
-	return ModelNodePath;
-}
-
-void UControlRigGraph::HandleRigVMGraphRenamed(const FString& InOldNodePath, const FString& InNewNodePath)
-{
-	static constexpr TCHAR NodePathPrefixFormat[] = TEXT("%s|");
-	const FString OldPrefix = FString::Printf(NodePathPrefixFormat, *InOldNodePath); 
-	const FString NewPrefix = FString::Printf(NodePathPrefixFormat, *InNewNodePath);
-	
-	if(ModelNodePath == InOldNodePath)
-	{
-		Modify();
-		ModelNodePath = InNewNodePath;
-
-		FString GraphName;
-		if(!ModelNodePath.Split(TEXT("|"), nullptr, &GraphName, ESearchCase::CaseSensitive, ESearchDir::FromEnd))
-		{
-			GraphName = ModelNodePath;
-		}
-		GraphName.RemoveFromEnd(TEXT("::"));
-		GraphName.RemoveFromStart(FRigVMClient::RigVMModelPrefix);
-		GraphName.TrimStartAndEndInline();
-
-		if(GraphName.IsEmpty())
-		{
-			GraphName = UControlRigGraphSchema::GraphName_ControlRig.ToString(); 
-		}
-		GraphName = FRigVMClient::GetUniqueName(GetOuter(), *GraphName).ToString();
-
-		Rename(*GraphName, nullptr, REN_ForceNoResetLoaders | REN_DontCreateRedirectors);
-	}
-	else if(ModelNodePath.StartsWith(OldPrefix))
-	{
-		Modify();
-		ModelNodePath = NewPrefix + ModelNodePath.RightChop(OldPrefix.Len() - 1);
-	}
-	else
-	{
-		return;
-	}
-
-	for(UEdGraphNode* Node : Nodes)
-	{
-		if(UControlRigGraphNode* RigNode = Cast<UControlRigGraphNode>(Node))
-		{
-			if(RigNode->ModelNodePath.StartsWith(OldPrefix))
-			{
-				RigNode->Modify();
-				RigNode->ModelNodePath = NewPrefix + RigNode->ModelNodePath.RightChop(OldPrefix.Len() - 1);
-			}
-		}
-	}
-}
-
-void UControlRigGraph::Initialize(UControlRigBlueprint* InBlueprint)
+void UControlRigGraph::InitializeFromBlueprint(URigVMBlueprint* InBlueprint)
 {
 	DECLARE_SCOPE_HIERARCHICAL_COUNTER_FUNC()
 
-	InBlueprint->OnModified().RemoveAll(this);
-	InBlueprint->OnModified().AddUObject(this, &UControlRigGraph::HandleModifiedEvent);
-	InBlueprint->OnVMCompiled().RemoveAll(this);
-	InBlueprint->OnVMCompiled().AddUObject(this, &UControlRigGraph::HandleVMCompiledEvent);
+	Super::InitializeFromBlueprint(InBlueprint);
 
-	URigHierarchy* Hierarchy = InBlueprint->Hierarchy;
+	const UControlRigBlueprint* ControlRigBlueprint = CastChecked<UControlRigBlueprint>(InBlueprint);
+	URigHierarchy* Hierarchy = ControlRigBlueprint->Hierarchy;
 
-	if(UControlRig* ControlRig = Cast<UControlRig>(InBlueprint->GetObjectBeingDebugged()))
+	if(UControlRig* ControlRig = Cast<UControlRig>(ControlRigBlueprint->GetObjectBeingDebugged()))
 	{
 		Hierarchy = ControlRig->GetHierarchy();
 	}
 
 	if(Hierarchy)
 	{
-		CacheNameLists(Hierarchy, &InBlueprint->DrawContainer, InBlueprint->ShapeLibraries);
+		CacheNameLists(Hierarchy, &ControlRigBlueprint->DrawContainer, ControlRigBlueprint->ShapeLibraries);
 	}
 }
 
-const UControlRigGraphSchema* UControlRigGraph::GetControlRigGraphSchema()
-{
-	return CastChecked<const UControlRigGraphSchema>(GetSchema());
-}
-
-#if WITH_EDITORONLY_DATA
-void UControlRigGraph::Serialize(FArchive& Ar)
-{
-	Super::Serialize(Ar);
-
-	Ar.UsingCustomVersion(FControlRigObjectVersion::GUID);
-
-	if (Ar.IsLoading())
-	{
-		if(Schema == nullptr || !Schema->IsChildOf(UControlRigGraphSchema::StaticClass()))
-		{
-			Schema = UControlRigGraphSchema::StaticClass();
-		}
-	}
-}
-#endif
-
 #if WITH_EDITOR
 
 TArray<TSharedPtr<FString>> UControlRigGraph::EmptyElementNameList;
@@ -194,7 +86,7 @@ void UControlRigGraph::CacheNameLists(URigHierarchy* InHierarchy, const FRigVMDr
 	EntryNameList.Reset();
 	EntryNameList.Add(MakeShared<FString>(FName(NAME_None).ToString()));
 
-	if(const UControlRigBlueprint* Blueprint = GetBlueprint())
+	if(const UControlRigBlueprint* Blueprint = CastChecked<UControlRigBlueprint>(GetBlueprint()))
 	{
 		const TArray<FName> EntryNames = Blueprint->GetRigVMClient()->GetEntryNames();
 		for (const FName& EntryName : EntryNames)
@@ -254,7 +146,7 @@ const TArray<TSharedPtr<FString>>* UControlRigGraph::GetElementNameList(ERigElem
 	
 	if(!ElementNameLists.Contains(InElementType))
 	{
-		UControlRigBlueprint* Blueprint = GetBlueprint();
+		const UControlRigBlueprint* Blueprint = CastChecked<UControlRigBlueprint>(GetBlueprint());
 		if(Blueprint == nullptr)
 		{
 			return &EmptyElementNameList;
@@ -304,7 +196,7 @@ const TArray<TSharedPtr<FString>> UControlRigGraph::GetSelectedElementsNameList(
 		return OuterGraph->GetSelectedElementsNameList();
 	}
 	
-	UControlRigBlueprint* Blueprint = GetBlueprint();
+	const UControlRigBlueprint* Blueprint = CastChecked<UControlRigBlueprint>(GetBlueprint());
 	if(Blueprint == nullptr)
 	{
 		return Result;
@@ -348,387 +240,27 @@ const TArray<TSharedPtr<FString>>* UControlRigGraph::GetShapeNameList(URigVMPin*
 	return &ShapeNameList;
 }
 
+#endif
 
-void UControlRigGraph::HandleModifiedEvent(ERigVMGraphNotifType InNotifType, URigVMGraph* InGraph, UObject* InSubject)
+bool UControlRigGraph::HandleModifiedEvent_Internal(ERigVMGraphNotifType InNotifType, URigVMGraph* InGraph, UObject* InSubject)
 {
 	DECLARE_SCOPE_HIERARCHICAL_COUNTER_FUNC()
 
-	if (bSuspendModelNotifications)
-	{
-		return;
-	}
-
-	// only make sure to receive notifs for this graph - unless
-	// we are on a template graph (used by node spawners)
-	if (GetModel() != InGraph && TemplateController == nullptr)
-	{
-		return;
-	}
-
-	if(UControlRigGraphSchema* ControlRigSchema = (UControlRigGraphSchema*)GetControlRigGraphSchema())
+	if(!Super::HandleModifiedEvent_Internal(InNotifType, InGraph, InSubject))
 	{
-		ControlRigSchema->HandleModifiedEvent(InNotifType, InGraph, InSubject);
-	}
-
-	if(TemplateController)
-	{
-		switch(InNotifType)
-		{
-			case ERigVMGraphNotifType::NodeRemoved:
-			case ERigVMGraphNotifType::PinTypeChanged:
-			case ERigVMGraphNotifType::PinDefaultValueChanged:
-			case ERigVMGraphNotifType::PinRemoved:
-			case ERigVMGraphNotifType::PinRenamed:
-			case ERigVMGraphNotifType::VariableRemoved:
-			case ERigVMGraphNotifType::VariableRenamed:
-			case ERigVMGraphNotifType::GraphChanged:
-			{
-				return;
-			}
-			default:
-			{
-				break;
-			}
-		}
-	}
-
-	// increment the node topology version for any interaction
-	// with a node.
-	{
-		UControlRigGraphNode* EdNode = nullptr;
-		if (URigVMNode* ModelNode = Cast<URigVMNode>(InSubject))
-		{
-			EdNode = Cast<UControlRigGraphNode>(FindNodeForModelNodeName(ModelNode->GetFName()));
-		}
-		else if (URigVMPin* ModelPin = Cast<URigVMPin>(InSubject))
-		{
-			EdNode = Cast<UControlRigGraphNode>(FindNodeForModelNodeName(ModelPin->GetNode()->GetFName()));
-		}
-
-		if(EdNode)
-		{
-			EdNode->NodeTopologyVersion++;
-		}
+		return false;
 	}
 
 	switch (InNotifType)
 	{
-		case ERigVMGraphNotifType::GraphChanged:
-		{
-			ModelNodePathToEdNode.Reset();
-
-			for (URigVMNode* Node : InGraph->GetNodes())
-			{
-				UEdGraphNode* EdNode = FindNodeForModelNodeName(Node->GetFName(), false);
-				if (EdNode != nullptr)
-				{
-					RemoveNode(EdNode);
-				}
-			}
-			break;
-		}
-		case ERigVMGraphNotifType::NodeSelectionChanged:
-		{
-			if (bIsSelecting)
-			{
-				return;
-			}
-			TGuardValue<bool> SelectionGuard(bIsSelecting, true);
-
-			TSet<const UEdGraphNode*> NodeSelection;
-			for (FName NodeName : InGraph->GetSelectNodes())
-			{
-				if (UEdGraphNode* EdNode = FindNodeForModelNodeName(NodeName))
-				{
-					NodeSelection.Add(EdNode);
-				}
-			}
-			SelectNodeSet(NodeSelection);
-			break;
-		}
-		case ERigVMGraphNotifType::NodeAdded:
-		{
-			if (URigVMNode* ModelNode = Cast<URigVMNode>(InSubject))
-			{
-				if (!ModelNode->IsVisibleInUI())
-				{
-					if (URigVMInjectionInfo* Injection = ModelNode->GetInjectionInfo())
-					{
-						if (URigVMPin* ModelPin = Injection->GetPin())
-						{
-							URigVMNode* ParentModelNode = ModelPin->GetNode();
-							if (ParentModelNode)
-							{
-								UEdGraphNode* EdNode = FindNodeForModelNodeName(ParentModelNode->GetFName());
-								if (EdNode)
-								{
-									if (UControlRigGraphNode* RigNode = Cast<UControlRigGraphNode>(EdNode))
-									{
-										RigNode->ModelPinsChanged(true);
-									}
-								}
-							}
-						}
-					}
-					break;
-				}
-				else
-				{
-					// check if the node is already part of the graph
-					if(UEdGraphNode* NodeToAdd = FindNodeForModelNodeName(ModelNode->GetFName()))
-					{
-						// Notify the GraphPanel a new node has been added
-						EEdGraphActionType AddNodeAction = GRAPHACTION_AddNode;
-						FEdGraphEditAction Action(AddNodeAction, this, NodeToAdd, false);
-						NotifyGraphChanged(Action);
-						break;
-					}
-				}
-
-				if (URigVMCommentNode* CommentModelNode = Cast<URigVMCommentNode>(ModelNode))
-				{
-					UEdGraphNode_Comment* NewNode = NewObject<UEdGraphNode_Comment>(this, CommentModelNode->GetFName());
-					AddNode(NewNode, false, false);
-
-					NewNode->CreateNewGuid();
-					NewNode->PostPlacedNewNode();
-					NewNode->AllocateDefaultPins();
-
-					NewNode->NodePosX = ModelNode->GetPosition().X;
-					NewNode->NodePosY = ModelNode->GetPosition().Y;
-					NewNode->NodeWidth = ModelNode->GetSize().X;
-					NewNode->NodeHeight = ModelNode->GetSize().Y;
-					NewNode->CommentColor = ModelNode->GetNodeColor();
-					NewNode->NodeComment = CommentModelNode->GetCommentText();
-					NewNode->FontSize = CommentModelNode->GetCommentFontSize();
-					NewNode->bCommentBubbleVisible = CommentModelNode->GetCommentBubbleVisible();
-					NewNode->bCommentBubbleVisible_InDetailsPanel = CommentModelNode->GetCommentBubbleVisible();
-					NewNode->bCommentBubblePinned = CommentModelNode->GetCommentBubbleVisible();
-					NewNode->bColorCommentBubble = CommentModelNode->GetCommentColorBubble();
-					NewNode->SetFlags(RF_Transactional);
-					NewNode->GetNodesUnderComment();
-
-					ModelNodePathToEdNode.Add(ModelNode->GetFName(), NewNode);
-				}
-				else // struct, library, parameter + variable
-				{
-					UControlRigGraphNode* NewNode = NewObject<UControlRigGraphNode>(this, GetControlRigGraphSchema()->GetGraphNodeClass(), ModelNode->GetFName());
-					AddNode(NewNode, false, false);
-
-					NewNode->ModelNodePath = ModelNode->GetNodePath();
-					NewNode->CreateNewGuid();
-					NewNode->PostPlacedNewNode();
-					NewNode->AllocateDefaultPins();
-					NewNode->PostReconstructNode();
-
-					NewNode->NodePosX = ModelNode->GetPosition().X;
-					NewNode->NodePosY = ModelNode->GetPosition().Y;
-					if (URigVMRerouteNode* RerouteModelNode = Cast<URigVMRerouteNode>(ModelNode))
-					{
-						if (URigVMPin* ModelPin = ModelNode->FindPin("Value"))
-						{
-							if (UEdGraphPin* ValuePin = NewNode->FindPin(ModelPin->GetPinPath()))
-							{
-								NewNode->SetColorFromModel(GetSchema()->GetPinTypeColor(ValuePin->PinType));
-							}
-						}
-					}
-					else
-					{
-						NewNode->SetColorFromModel(ModelNode->GetNodeColor());
-					}
-					NewNode->SetFlags(RF_Transactional);
-
-					ModelNodePathToEdNode.Add(ModelNode->GetFName(), NewNode);
-				}
-			}
-			break;
-		}
-		case ERigVMGraphNotifType::NodeRemoved:
-		{
-			if (URigVMNode* ModelNode = Cast<URigVMNode>(InSubject))
-			{
-				if (URigVMInjectionInfo* Injection = ModelNode->GetInjectionInfo())
-				{
-					if (URigVMPin* ModelPin = Injection->GetPin())
-					{
-						if (URigVMNode* ParentModelNode = ModelPin->GetNode())
-						{
-							if (UEdGraphNode* EdNode = FindNodeForModelNodeName(ParentModelNode->GetFName()))
-							{
-								if (UControlRigGraphNode* RigNode = Cast<UControlRigGraphNode>(EdNode))
-								{
-									RigNode->ModelPinsChanged(true);
-								}
-							}
-						}
-					}
-					break;
-				}
-
-				if (UEdGraphNode* EdNode = FindNodeForModelNodeName(ModelNode->GetFName(), false))
-				{
-					RemoveNode(EdNode);
-				}
-
-				ModelNodePathToEdNode.Remove(ModelNode->GetFName());
-			}
-			break;
-		}
-		case ERigVMGraphNotifType::NodePositionChanged:
-		{
-			if (URigVMNode* ModelNode = Cast<URigVMNode>(InSubject))
-			{
-				UEdGraphNode* EdNode = FindNodeForModelNodeName(ModelNode->GetFName());
-				if (EdNode)
-				{
-					// No need to call Node->Modify(), since control rig has its own undo/redo system see RigVMControllerActions.cpp
-					EdNode->NodePosX = (int32)ModelNode->GetPosition().X;
-					EdNode->NodePosY = (int32)ModelNode->GetPosition().Y;
-				}
-			}
-			break;
-		}
-		case ERigVMGraphNotifType::NodeSizeChanged:
-		{
-			if (URigVMNode* ModelNode = Cast<URigVMNode>(InSubject))
-			{
-				UEdGraphNode_Comment* EdNode = Cast<UEdGraphNode_Comment>(FindNodeForModelNodeName(ModelNode->GetFName()));
-				if (EdNode)
-				{
-					// No need to call Node->Modify(), since control rig has its own undo/redo system see RigVMControllerActions.cpp
-					EdNode->NodeWidth = (int32)ModelNode->GetSize().X;
-					EdNode->NodeHeight = (int32)ModelNode->GetSize().Y;
-				}
-			}
-			break;
-		}
-		case ERigVMGraphNotifType::NodeDescriptionChanged:
-		case ERigVMGraphNotifType::NodeCategoryChanged:
-		{
-			if (URigVMNode* ModelNode = Cast<URigVMNode>(InSubject))
-			{
-				if (UControlRigGraphNode* RigNode = Cast<UControlRigGraphNode>(FindNodeForModelNodeName(ModelNode->GetFName())))
-				{
-					RigNode->SyncGraphNodeTitleWithModelNodeTitle();
-				}
-			}
-			break;
-		}
-		case ERigVMGraphNotifType::NodeColorChanged:
-		{
-			if (URigVMNode* ModelNode = Cast<URigVMNode>(InSubject))
-			{
-				if (ModelNode->IsA<URigVMLibraryNode>() || ModelNode->IsA<URigVMTemplateNode>())
-				{
-					if (UControlRigGraphNode* RigNode = Cast<UControlRigGraphNode>(FindNodeForModelNodeName(ModelNode->GetFName())))
-					{
-						RigNode->SetColorFromModel(ModelNode->GetNodeColor());
-					}
-				}
-				else if(UEdGraphNode_Comment * EdComment = Cast<UEdGraphNode_Comment>(FindNodeForModelNodeName(ModelNode->GetFName())))
-				{
-					EdComment->CommentColor = ModelNode->GetNodeColor();
-				}
-			}
-			break;
-		}
-		case ERigVMGraphNotifType::CommentTextChanged:
-		{
-			if (URigVMCommentNode* ModelNode = Cast<URigVMCommentNode>(InSubject))
-			{
-				UEdGraphNode_Comment* EdNode = Cast<UEdGraphNode_Comment>(FindNodeForModelNodeName(ModelNode->GetFName()));
-				if (EdNode)
-				{
-					EdNode->NodeComment = ModelNode->GetCommentText();
-					EdNode->FontSize = ModelNode->GetCommentFontSize();
-					EdNode->bCommentBubbleVisible = ModelNode->GetCommentBubbleVisible();
-					EdNode->bCommentBubbleVisible_InDetailsPanel = ModelNode->GetCommentBubbleVisible();
-					EdNode->bColorCommentBubble = ModelNode->GetCommentColorBubble();
-				}
-			}
-			break;
-		}
-		case ERigVMGraphNotifType::LinkAdded:
-		case ERigVMGraphNotifType::LinkRemoved:
-		{
-			bool AddLink = InNotifType == ERigVMGraphNotifType::LinkAdded;
-
-			if (URigVMLink* Link = Cast<URigVMLink>(InSubject))
-			{
-				URigVMPin* SourcePin = Link->GetSourcePin();
-				URigVMPin* TargetPin = Link->GetTargetPin();
-
-				if (SourcePin)
-				{
-					SourcePin = SourcePin->GetOriginalPinFromInjectedNode();
-				}
-				if (TargetPin)
-				{
-					TargetPin = TargetPin->GetOriginalPinFromInjectedNode();
-				}
-
-				if (SourcePin && TargetPin && SourcePin != TargetPin)
-				{
-					UControlRigGraphNode* SourceRigNode = Cast<UControlRigGraphNode>(FindNodeForModelNodeName(SourcePin->GetNode()->GetFName()));
-					UControlRigGraphNode* TargetRigNode = Cast<UControlRigGraphNode>(FindNodeForModelNodeName(TargetPin->GetNode()->GetFName()));
-
-					if (SourceRigNode != nullptr && TargetRigNode != nullptr)
-					{
-						FString SourcePinPath = SourcePin->GetPinPath();
-						FString TargetPinPath = TargetPin->GetPinPath();
-						UEdGraphPin* SourceRigPin = SourceRigNode->FindPin(*SourcePinPath, EGPD_Output);
-						UEdGraphPin* TargetRigPin = TargetRigNode->FindPin(*TargetPinPath, EGPD_Input);
-
-						if (SourceRigPin != nullptr && TargetRigPin != nullptr)
-						{
-							if (AddLink)
-							{
-								SourceRigPin->MakeLinkTo(TargetRigPin);
-							}
-							else
-							{
-								SourceRigPin->BreakLinkTo(TargetRigPin);
-							}
-
-							SourceRigPin->LinkedTo.Remove(nullptr);
-							TargetRigPin->LinkedTo.Remove(nullptr);
-						}
-					}
-				}
-			}
-			break;
-		}
 		case ERigVMGraphNotifType::PinDefaultValueChanged:
 		{
 			if (URigVMPin* ModelPin = Cast<URigVMPin>(InSubject))
 			{
 				if (UControlRigGraphNode* RigNode = Cast<UControlRigGraphNode>(FindNodeForModelNodeName(ModelPin->GetNode()->GetFName())))
 				{
-					UEdGraphPin* RigNodePin = RigNode->FindPin(ModelPin->GetPinPath());
-					if (RigNodePin == nullptr)
+					if (Cast<URigVMUnitNode>(ModelPin->GetNode()))
 					{
-						if(ModelPin->GetNode()->IsEvent())
-						{
-							RigNode->SyncGraphNodeTitleWithModelNodeTitle();
-						}
-						break;
-					}
-
-					RigNode->SetupPinDefaultsFromModel(RigNodePin);
-
-					if (Cast<URigVMVariableNode>(ModelPin->GetNode()))
-					{
-						if (ModelPin->GetName() == TEXT("Variable"))
-						{
-							RigNode->SyncGraphNodeTitleWithModelNodeTitle();
-							RigNode->SynchronizeGraphPinValueWithModelPin(ModelPin);
-						}
-					}
-					else if (Cast<URigVMUnitNode>(ModelPin->GetNode()))
-					{
-						RigNode->SyncGraphNodeTitleWithModelNodeTitle();
-
 						// if the node contains a rig element key - invalidate the node
 						if(RigNode->GetAllPins().ContainsByPredicate([](UEdGraphPin* Pin) -> bool
 						{
@@ -740,170 +272,6 @@ void UControlRigGraph::HandleModifiedEvent(ERigVMGraphNotifType InNotifType, URi
 						}
 					}
 				}
-				else if (URigVMInjectionInfo* Injection = ModelPin->GetNode()->GetInjectionInfo())
-				{
-					if (Injection->InputPin != ModelPin->GetRootPin())
-					{
-						if (URigVMPin* InjectionPin = Injection->GetPin())
-						{
-							URigVMNode* ParentModelNode = InjectionPin->GetNode();
-							if (ParentModelNode)
-							{
-								UEdGraphNode* HostEdNode = FindNodeForModelNodeName(ParentModelNode->GetFName());
-								if (HostEdNode)
-								{
-									if (UControlRigGraphNode* HostRigNode = Cast<UControlRigGraphNode>(HostEdNode))
-									{
-										HostRigNode->ModelPinsChanged();
-									}
-								}
-							}
-						}
-					}
-				}
-			}
-			break;
-		}
-		case ERigVMGraphNotifType::PinAdded:
-		{
-			if (URigVMPin* ModelPin = Cast<URigVMPin>(InSubject))
-			{
-				if(URigVMNode* ModelNode = ModelPin->GetNode())
-				{
-					if (UControlRigGraphNode* RigNode = Cast<UControlRigGraphNode>(FindNodeForModelNodeName(ModelNode->GetFName())))
-					{
-						RigNode->ModelPinAdded(ModelPin);
-					}
-				}
-			}
-			break;
-		}
-		case ERigVMGraphNotifType::PinRemoved:
-		{
-			if (URigVMPin* ModelPin = Cast<URigVMPin>(InSubject))
-			{
-				if(URigVMNode* ModelNode = ModelPin->GetNode())
-				{
-					if (UControlRigGraphNode* RigNode = Cast<UControlRigGraphNode>(FindNodeForModelNodeName(ModelNode->GetFName())))
-					{
-						RigNode->ModelPinRemoved(ModelPin);
-					}
-				}
-			}
-			break;
-		}
-		case ERigVMGraphNotifType::PinArraySizeChanged:
-		{
-			// don't do anything here - the UI will update based on the
-			// PinAdded and PinRemoved notifs
-			break;
-		}
-		case ERigVMGraphNotifType::PinDirectionChanged:
-		case ERigVMGraphNotifType::PinIndexChanged:
-		case ERigVMGraphNotifType::PinBoundVariableChanged:
-		{
-			if (URigVMPin* ModelPin = Cast<URigVMPin>(InSubject))
-			{
-				if (UControlRigGraphNode* RigNode = Cast<UControlRigGraphNode>(FindNodeForModelNodeName(ModelPin->GetNode()->GetFName())))
-				{
-					RigNode->ModelPinsChanged();
-				}
-			}
-			break;
-		}
-		case ERigVMGraphNotifType::LibraryTemplateChanged:
-		{
-			if (URigVMNode* LibraryNode = Cast<URigVMNode>(InSubject))
-			{
-				if (UControlRigGraphNode* RigNode = Cast<UControlRigGraphNode>(FindNodeForModelNodeName(LibraryNode->GetFName())))
-				{
-					RigNode->ModelPinsChanged(true);
-				}
-			}
-			break;
-		}
-		case ERigVMGraphNotifType::PinTypeChanged:
-		{
-			if (URigVMPin* ModelPin = Cast<URigVMPin>(InSubject))
-			{
-				if (UControlRigGraphNode* RigNode = Cast<UControlRigGraphNode>(FindNodeForModelNodeName(ModelPin->GetNode()->GetFName())))
-				{
-					RigNode->SynchronizeGraphPinTypeWithModelPin(ModelPin);
-				}
-			}
-			break;
-		}
-		case ERigVMGraphNotifType::PinRenamed:
-		{
-			if (URigVMPin* ModelPin = Cast<URigVMPin>(InSubject))
-			{
-				if (UControlRigGraphNode* RigNode = Cast<UControlRigGraphNode>(FindNodeForModelNodeName(ModelPin->GetNode()->GetFName())))
-				{
-					RigNode->SynchronizeGraphPinNameWithModelPin(ModelPin);
-				}
-			}
-			break;
-		}
-		case ERigVMGraphNotifType::NodeRenamed:
-		{
-			if (URigVMNode* ModelNode = Cast<URigVMNode>(InSubject))
-			{
-				UEdGraphNode* EdNode = FindNodeForModelNodeName(ModelNode->GetPreviousFName());
-				ModelNodePathToEdNode.Remove(ModelNode->GetPreviousFName());
-				if (UControlRigGraphNode* RigNode = Cast<UControlRigGraphNode>(EdNode))
-				{
-					RigNode->SyncGraphNodeNameWithModelNodeName(ModelNode);
-					ModelNodePathToEdNode.Add(ModelNode->GetFName(), RigNode);
-				}
-			}
-			break;
-		}
-		case ERigVMGraphNotifType::VariableRenamed:
-		case ERigVMGraphNotifType::NodeReferenceChanged:
-		{
-			if (URigVMNode* ModelNode = Cast<URigVMNode>(InSubject))
-			{
-				if (UControlRigGraphNode* RigNode = Cast<UControlRigGraphNode>(FindNodeForModelNodeName(ModelNode->GetFName())))
-				{
-					RigNode->SyncGraphNodeTitleWithModelNodeTitle();
-				}
-			}
-			break;
-		}
-		case ERigVMGraphNotifType::NodeSelected:
-		{
-			if (URigVMCommentNode* ModelNode = Cast<URigVMCommentNode>(InSubject))
-			{
-				// UEdGraphNode_Comment cannot access RigVMCommentNode's selection state, so we have to manually toggle its selection state
-				// UControlRigGraphNode does not need this step because it overrides the IsSelectedInEditor() method
-				UEdGraphNode_Comment* EdNode = Cast<UEdGraphNode_Comment>(FindNodeForModelNodeName(ModelNode->GetFName()));
-				if (EdNode)
-				{
-					EdNode->SetSelectionState(UEdGraphNode_Comment::ESelectionState::Selected);
-				}
-			}
-			break;
-		}
-		case ERigVMGraphNotifType::NodeDeselected:
-		{
-			if (URigVMCommentNode* ModelNode = Cast<URigVMCommentNode>(InSubject))
-			{
-				UEdGraphNode_Comment* EdNode = Cast<UEdGraphNode_Comment>(FindNodeForModelNodeName(ModelNode->GetFName()));
-				if (EdNode)
-				{
-					EdNode->SetSelectionState(UEdGraphNode_Comment::ESelectionState::Deselected);
-				}
-			}
-			break;
-		}
-		case ERigVMGraphNotifType::PinExpansionChanged:
-		{
-			if (URigVMPin* ModelPin = Cast<URigVMPin>(InSubject))
-			{
-				if (UControlRigGraphNode* RigNode = Cast<UControlRigGraphNode>(FindNodeForModelNodeName(ModelPin->GetNode()->GetFName())))
-				{
-					RigNode->OnNodePinExpansionChanged().Broadcast();
-				}
 			}
 			break;
 		}
@@ -912,296 +280,9 @@ void UControlRigGraph::HandleModifiedEvent(ERigVMGraphNotifType InNotifType, URi
 			break;
 		}
 	}
-}
-
-int32 UControlRigGraph::GetInstructionIndex(const UControlRigGraphNode* InNode, bool bAsInput)
-{
-	if (const TPair<int32, int32>* FoundIndex = CachedInstructionIndices.Find(InNode->GetModelNode()))
-	{
-		return bAsInput ? FoundIndex->Key : FoundIndex->Value;
-	}
-
-	struct Local
-	{
-		static int32 GetInstructionIndex(URigVMNode* InModelNode, const FRigVMByteCode* InByteCode, TMap<URigVMNode*, TPair<int32, int32>>& Indices, bool bAsInput)
-		{
-			if (InModelNode == nullptr)
-			{
-				return INDEX_NONE;
-			}
-
-			if (const TPair<int32, int32>* ExistingIndex = Indices.Find(InModelNode))
-			{
-				const int32 Index = bAsInput ? ExistingIndex->Key : ExistingIndex->Value;
-				if(Index != INDEX_NONE)
-				{
-					return Index;
-				}
-			}
-
-			if(const URigVMRerouteNode* RerouteNode = Cast<URigVMRerouteNode>(InModelNode))
-			{
-				int32 InstructionIndex = INDEX_NONE;
-				if(bAsInput)
-				{
-					TArray<URigVMNode*> SourceNodes = RerouteNode->GetLinkedSourceNodes();
-					for(URigVMNode* SourceNode : SourceNodes)
-					{
-						InstructionIndex = GetInstructionIndex(SourceNode, InByteCode, Indices, bAsInput);
-						if(InstructionIndex != INDEX_NONE)
-						{
-							break;
-						}
-					}
-					Indices.FindOrAdd(InModelNode).Key = InstructionIndex;
-				}
-				else
-				{
-					TArray<URigVMNode*> TargetNodes = RerouteNode->GetLinkedTargetNodes();
-					for(URigVMNode* TargetNode : TargetNodes)
-					{
-						InstructionIndex = GetInstructionIndex(TargetNode, InByteCode, Indices, bAsInput);
-						if(InstructionIndex != INDEX_NONE)
-						{
-							break;
-						}
-					}
-					Indices.FindOrAdd(InModelNode).Value = InstructionIndex;
-				}
-
-				return InstructionIndex;
-			}
-			else if(URigVMFunctionEntryNode* EntryNode = Cast<URigVMFunctionEntryNode>(InModelNode))
-			{
-				int32 InstructionIndex = INDEX_NONE;
-				if(!bAsInput)
-				{
-					TArray<URigVMNode*> TargetNodes = EntryNode->GetLinkedTargetNodes();
-					for(URigVMNode* TargetNode : TargetNodes)
-					{
-						InstructionIndex = GetInstructionIndex(TargetNode, InByteCode, Indices, bAsInput);
-						if(InstructionIndex != INDEX_NONE)
-						{
-							break;
-						}
-					}
-					Indices.FindOrAdd(InModelNode).Key = InstructionIndex;
-				}
-				return InstructionIndex;
-			}
-			else if(URigVMFunctionReturnNode* ReturnNode = Cast<URigVMFunctionReturnNode>(InModelNode))
-			{
-				int32 InstructionIndex = INDEX_NONE;
-				if(bAsInput)
-				{
-					TArray<URigVMNode*> SourceNodes = ReturnNode->GetLinkedSourceNodes();
-					for(URigVMNode* SourceNode : SourceNodes)
-					{
-						InstructionIndex = GetInstructionIndex(SourceNode, InByteCode, Indices, bAsInput);
-						if(InstructionIndex != INDEX_NONE)
-						{
-							break;
-						}
-					}
-					Indices.FindOrAdd(InModelNode).Key = InstructionIndex;
-				}
-				return InstructionIndex;
-			}
-
-			Indices.FindOrAdd(InModelNode, TPair<int32, int32>(INDEX_NONE, INDEX_NONE));
-
-			int32 InstructionIndex = InByteCode->GetFirstInstructionIndexForSubject(InModelNode);
-			if (InstructionIndex != INDEX_NONE)
-			{
-				TPair<int32, int32> Pair(InstructionIndex, InstructionIndex);
-				if(bAsInput)
-				{
-					Indices.FindOrAdd(InModelNode).Key = InstructionIndex;
-				}
-				else
-				{
-					Indices.FindOrAdd(InModelNode).Value = InstructionIndex;
-				}
-				return InstructionIndex;
-			}
-
-			FRigVMInstructionArray Instructions = InByteCode->GetInstructions();
-			for (int32 i = 0; i < Instructions.Num(); ++i)
-			{
-				const FRigVMASTProxy Proxy = FRigVMASTProxy::MakeFromCallPath(InByteCode->GetCallPathForInstruction(i), InModelNode->GetRootGraph());
-				if (Proxy.GetCallstack().Contains(InModelNode))
-				{
-					if(bAsInput)
-					{
-						Indices.FindOrAdd(InModelNode).Key = i;
-					}
-					else
-					{
-						Indices.FindOrAdd(InModelNode).Value = i;
-					}
-					return i;
-				}
-			}
-			return INDEX_NONE;
-		}
-	};
-
-	if (const FRigVMByteCode* ByteCode = GetController()->GetCurrentByteCode())
-	{
-		const int32 SourceInstructionIndex = Local::GetInstructionIndex(InNode->GetModelNode(), ByteCode, CachedInstructionIndices, true);
-		const int32 TargetInstructionIndex = Local::GetInstructionIndex(InNode->GetModelNode(), ByteCode, CachedInstructionIndices, false);
-		return bAsInput ? SourceInstructionIndex : TargetInstructionIndex;
-	}
 
-	return INDEX_NONE;
+	return true;
 }
 
-UEdGraphNode* UControlRigGraph::FindNodeForModelNodeName(const FName& InModelNodeName, const bool bCacheIfRequired)
-{
-	DECLARE_SCOPE_HIERARCHICAL_COUNTER_FUNC()
-
-	if(UEdGraphNode** MappedNode = ModelNodePathToEdNode.Find(InModelNodeName))
-	{
-		return *MappedNode;
-	}
-	
-	const FString InModelNodePath = InModelNodeName.ToString();
-	for (UEdGraphNode* EdNode : Nodes)
-	{
-		if (UControlRigGraphNode* RigNode = Cast<UControlRigGraphNode>(EdNode))
-		{
-			if (RigNode->ModelNodePath == InModelNodePath)
-			{
-				if (RigNode->GetOuter() == this)
-				{
-					if (bCacheIfRequired)
-					{
-						ModelNodePathToEdNode.Add(InModelNodeName, EdNode);
-					}
-					return EdNode;
-				}
-			}
-		}
-		else if (EdNode)
-		{
-			if (EdNode->GetFName() == InModelNodeName)
-			{
-				if (EdNode->GetOuter() == this)
-				{
-					if (bCacheIfRequired)
-					{
-						ModelNodePathToEdNode.Add(InModelNodeName, EdNode);
-					}
-					return EdNode;
-				}
-			}
-		}
-	}
-	return nullptr;
-}
-
-UControlRigBlueprint* UControlRigGraph::GetBlueprint() const
-{
-	if (UControlRigGraph* OuterGraph = Cast<UControlRigGraph>(GetOuter()))
-	{
-		return OuterGraph->GetBlueprint();
-	}
-	return Cast<UControlRigBlueprint>(GetOuter());
-}
-
-URigVMGraph* UControlRigGraph::GetModel() const
-{
-	if(CachedModelGraph.IsValid())
-	{
-		return CachedModelGraph.Get();
-	}
-	
-	if (const FRigVMClient* Client = GetRigVMClient())
-	{
-		URigVMGraph* Model = Client->GetModel(this);
-		CachedModelGraph = Model;
-		return Model;
-	}
-	return nullptr;
-}
-
-URigVMController* UControlRigGraph::GetController() const
-{
-	if (FRigVMClient* Client = GetRigVMClient())
-	{
-		return Client->GetOrCreateController(GetModel());
-	}
-	return nullptr;
-}
-
-const UControlRigGraph* UControlRigGraph::GetRootGraph() const
-{
-	if(const UControlRigGraph* ParentGraph = Cast<UControlRigGraph>(GetOuter()))
-	{
-		return ParentGraph->GetRootGraph();
-	}
-	return this;
-}
-
-void UControlRigGraph::AddNode(UEdGraphNode* NodeToAdd, bool bUserAction, bool bSelectNewNode)
-{
-	// Comments are added outside of the ControlRigEditor, so we add here the node to the model
-	if (const UEdGraphNode_Comment* CommentNode = Cast<const UEdGraphNode_Comment>(NodeToAdd))
-	{
-		if (URigVMController* Controller = GetBlueprint()->GetOrCreateController(GetModel()))
-		{
-			if (GetModel()->FindNodeByName(NodeToAdd->GetFName()) == nullptr) // When recreating nodes at RebuildGraphFromModel, the model node already exists
-			{
-				TGuardValue<bool> BlueprintNotifGuard(GetBlueprint()->bSuspendModelNotificationsForOthers, true);
-				FVector2D NodePos(CommentNode->NodePosX, CommentNode->NodePosY);
-				FVector2D NodeSize(CommentNode->NodeWidth, CommentNode->NodeHeight);
-				FLinearColor NodeColor = CommentNode->CommentColor;
-				Controller->AddCommentNode(CommentNode->NodeComment, NodePos, NodeSize, NodeColor, CommentNode->GetName(), GIsTransacting == false, true);
-			}
-		}
-	}
-
-	Super::AddNode(NodeToAdd, bUserAction, bSelectNewNode);
-}
-
-void UControlRigGraph::RemoveNode(UEdGraphNode* InNode)
-{
-	// Make sure EdGraph is not part of the transaction
-	TGuardValue<ITransaction*> TransactionGuard(GUndo, nullptr);
-
-	if(UControlRigGraphNode* RigNode = Cast<UControlRigGraphNode>(InNode))
-	{
-		RigNode->OnNodeBeginRemoval().Broadcast();
-	}
-
-	// clear out the pin relationships
-	for(UEdGraphPin* Pin : InNode->Pins)
-	{
-		Pin->MarkAsGarbage();
-	}
-	InNode->Pins.Reset();
-					
-	// this also subsequently calls NotifyGraphChanged
-	Super::RemoveNode(InNode);
-}
-
-URigVMController* UControlRigGraph::GetTemplateController()
-{
-	if (TemplateController == nullptr)
-	{
-		TemplateController = GetBlueprint()->GetTemplateController();
-		TemplateController->OnModified().RemoveAll(this);
-		TemplateController->OnModified().AddUObject(this, &UControlRigGraph::HandleModifiedEvent);
-	}
-	return TemplateController;
-}
-
-void UControlRigGraph::HandleVMCompiledEvent(UObject* InCompiledObject, URigVM* InVM)
-{
-	CachedInstructionIndices.Reset();
-}
-
-#endif
-
 #undef LOCTEXT_NAMESPACE
 
diff --git a/Engine/Plugins/Animation/ControlRig/Source/ControlRigDeveloper/Private/Graph/ControlRigGraphNode.cpp b/Engine/Plugins/Animation/ControlRig/Source/ControlRigDeveloper/Private/Graph/ControlRigGraphNode.cpp
--- a/Engine/Plugins/Animation/ControlRig/Source/ControlRigDeveloper/Private/Graph/ControlRigGraphNode.cpp
+++ b/Engine/Plugins/Animation/ControlRig/Source/ControlRigDeveloper/Private/Graph/ControlRigGraphNode.cpp
@@ -1,1790 +1,27 @@
 // Copyright Epic Games, Inc. All Rights Reserved.
 
 #include "Graph/ControlRigGraphNode.h"
-#include "EdGraph/EdGraphPin.h"
 #include "Graph/ControlRigGraph.h"
 #include "Graph/ControlRigGraphSchema.h"
-#include "Kismet2/BlueprintEditorUtils.h"
-#include "Kismet2/KismetDebugUtilities.h"
-#include "Kismet2/WatchedPin.h"
-#include "KismetCompiler.h"
-#include "BlueprintNodeSpawner.h"
-#include "EditorCategoryUtils.h"
-#include "BlueprintActionDatabaseRegistrar.h"
-#include "ControlRig.h"
-#include "Textures/SlateIcon.h"
-#include "Units/RigUnit.h"
-#include "ControlRigBlueprint.h"
-#include "PropertyPathHelpers.h"
-#include "Kismet2/Kismet2NameValidators.h"
-#include "ScopedTransaction.h"
-#include "StructReference.h"
-#include "UObject/PropertyPortFlags.h"
-#include "RigVMBlueprintUtils.h"
-#include "Curves/CurveFloat.h"
-#include "RigVMCore/RigVMExecuteContext.h"
-#include "ControlRigDeveloper.h"
-#include "ControlRigObjectVersion.h"
-#include "GraphEditAction.h"
-#include "RigVMModel/Nodes/RigVMAggregateNode.h"
-#include "RigVMModel/Nodes/RigVMFunctionReferenceNode.h"
-#include "RigVMModel/Nodes/RigVMFunctionEntryNode.h"
-#include "RigVMModel/Nodes/RigVMFunctionReturnNode.h"
-#include "RigVMModel/Nodes/RigVMCollapseNode.h"
-#include "RigVMModel/Nodes/RigVMInvokeEntryNode.h"
-#include "RigVMModel/Nodes/RigVMDispatchNode.h"
+#include "IControlRigEditorModule.h"
 
 #include UE_INLINE_GENERATED_CPP_BY_NAME(ControlRigGraphNode)
 
-#if WITH_EDITOR
-#include "IControlRigEditorModule.h"
-#endif //WITH_EDITOR
-
 #define LOCTEXT_NAMESPACE "ControlRigGraphNode"
 
-TAutoConsoleVariable<bool> CVarControlRigDisableCompactNodes(TEXT("ControlRig.Graph.DisableCompactNodes"), false, TEXT("When true all nodes are going to be drawn as full nodes."));
-
 UControlRigGraphNode::UControlRigGraphNode()
-: Dimensions(0.0f, 0.0f)
-, NodeTitle(FText::GetEmpty())
-, FullNodeTitle(FText::GetEmpty())
-, NodeTopologyVersion(INDEX_NONE)
-, CachedTitleColor(FLinearColor(0.f, 0.f, 0.f, 0.f))
-, CachedNodeColor(FLinearColor(0.f, 0.f, 0.f, 0.f))
-#if WITH_EDITOR
-, bEnableProfiling(false)
-#endif
-, CachedTemplate(nullptr)
-{
-	bHasCompilerMessage = false;
-	ErrorType = (int32)EMessageSeverity::Info + 1;
-}
-
-FText UControlRigGraphNode::GetNodeTitle(ENodeTitleType::Type TitleType) const
-{
-	if(NodeTitle.IsEmpty())
-	{
-		FString SubTitle;
-		if(URigVMNode* ModelNode = GetModelNode())
-		{
-			if (URigVMUnitNode* UnitNode = Cast<URigVMUnitNode>(ModelNode))
-			{
-				const UScriptStruct* ScriptStruct = UnitNode->GetScriptStruct();
-				if (ScriptStruct && ScriptStruct->IsChildOf(FRigVMStruct::StaticStruct()))
-				{
-					if (TSharedPtr<FStructOnScope> StructOnScope = UnitNode->ConstructStructInstance())
-					{
-						FRigVMStruct* RigUnit = (FRigVMStruct*)StructOnScope->GetStructMemory();
-						NodeTitle = FText::FromString(RigUnit->GetUnitLabel());
-					}
-				}
-			}
-
-			else if(URigVMFunctionReferenceNode* FunctionReferenceNode = Cast<URigVMFunctionReferenceNode>(ModelNode))
-			{
-				const FRigVMGraphFunctionHeader& ReferencedHeader = FunctionReferenceNode->GetReferencedFunctionHeader();
-				{
-					TSoftObjectPtr<URigVMFunctionReferenceNode> RefNodePtr(FunctionReferenceNode);
-					const FString& PackagePath = ReferencedHeader.LibraryPointer.LibraryNode.GetLongPackageName();
-					if(PackagePath != RefNodePtr.GetLongPackageName())
-					{
-						SubTitle = FString::Printf(TEXT("From %s"), *PackagePath);
-					}
-					else
-					{
-						static const FString LocalFunctionString = TEXT("Local Function");
-						SubTitle = LocalFunctionString;
-					}
-				}
-			}
-
-			else if(URigVMCollapseNode* CollapseNode = Cast<URigVMCollapseNode>(ModelNode))
-			{
-				if(!CollapseNode->IsA<URigVMAggregateNode>())
-				{
-					static const FString CollapseNodeString = TEXT("Collapsed Graph");
-					SubTitle = CollapseNodeString;
-				}
-			}
-
-			else if(URigVMVariableNode* VariableNode = Cast<URigVMVariableNode>(ModelNode))
-			{
-				if(VariableNode->IsLocalVariable())
-				{
-					static const FString LocalVariableString = TEXT("Local Variable");
-					const FString DefaultValue = VariableNode->GetVariableDescription().DefaultValue;
-					if(DefaultValue.IsEmpty())
-					{
-						SubTitle = LocalVariableString;
-					}
-					else
-					{
-						SubTitle = FString::Printf(TEXT("%s\nDefault %s"), *LocalVariableString, *DefaultValue);
-					}
-				}
-				else if (VariableNode->IsInputArgument())
-				{
-					SubTitle = TEXT("Input parameter");
-				}
-				else
-				{
-					if(UBlueprint* Blueprint = GetBlueprint())
-					{
-						const FName VariableName = VariableNode->GetVariableName();
-						for(const FBPVariableDescription& NewVariable : Blueprint->NewVariables)
-						{
-							if(NewVariable.VarName == VariableName)
-							{
-								FString DefaultValue = NewVariable.DefaultValue;
-								if(DefaultValue.IsEmpty())
-								{
-									static const FString VariableString = TEXT("Variable");
-									SubTitle = VariableString;
-								}
-								else
-								{
-									// Change the order of values in rotators so that they match the pin order
-									if (!NewVariable.VarType.IsContainer() && NewVariable.VarType.PinSubCategoryObject == TBaseStructure<FRotator>::Get())
-									{
-										TArray<FString> Values;
-										DefaultValue.ParseIntoArray(Values, TEXT(","));
-										if (Values.Num() == 3)
-										{
-											Values.Swap(0, 1);
-											Values.Swap(0, 2);
-										}
-										DefaultValue = FString::Join(Values, TEXT(","));										
-									}
-									SubTitle = FString::Printf(TEXT("Default %s"), *DefaultValue);
-								}
-								break;
-							}
-						}
-					}
-				}
-
-				if(SubTitle.Len() > 40)
-				{
-					SubTitle = SubTitle.Left(36) + TEXT(" ...");
-				}
-			}
-
-			if (NodeTitle.IsEmpty())
-			{
-				NodeTitle = FText::FromString(ModelNode->GetNodeTitle());
-			}
-		}
-
-		if(IsDeprecated())
-		{
-			NodeTitle = FText::FromString(FString::Printf(TEXT("%s (Deprecated)"), *NodeTitle.ToString()));
-		}
-
-		FullNodeTitle = NodeTitle;
-
-		if(!SubTitle.IsEmpty())
-		{
-			FullNodeTitle = FText::FromString(FString::Printf(TEXT("%s\n%s"), *NodeTitle.ToString(), *SubTitle));
-		}
-	}
-
-	if(TitleType == ENodeTitleType::FullTitle)
-	{
-		return FullNodeTitle;
-	}
-	return NodeTitle;
-}
-
-void UControlRigGraphNode::ReconstructNode()
-{
-	ReconstructNode_Internal();
-}
-
-void UControlRigGraphNode::ReconstructNode_Internal(bool bForce)
-{
-	DECLARE_SCOPE_HIERARCHICAL_COUNTER_FUNC()
-
-	UControlRigGraph* RigGraph = Cast<UControlRigGraph>(GetOuter()); 
-	if (RigGraph && !bForce)
-	{
-		if (RigGraph->bIsTemporaryGraphForCopyPaste)
-		{
-			return;
-		}
-
-		// if this node has been saved prior to our custom version,
-		// don't reset the node
-		int32 LinkerVersion = RigGraph->GetLinkerCustomVersion(FControlRigObjectVersion::GUID);
-		if (LinkerVersion < FControlRigObjectVersion::SwitchedToRigVM)
-		{
-			return;
-		}
-	}
-
-#if WITH_EDITOR
-	bEnableProfiling = false;
-	if(RigGraph)
-	{
-		if(UControlRigBlueprint* RigBlueprint = RigGraph->GetBlueprint())
-		{
-			bEnableProfiling = RigBlueprint->VMRuntimeSettings.bEnableProfiling;
-		}
-	}
-#endif
-
-	// Clear previously set messages
-	ErrorMsg.Reset();
-
-	// Move the existing pins to a saved array.
-	// This way we can reuse them later
-	LastEdGraphPins = Pins;
-
-	// Also make sure to clean up any indirections
-	// on them.
-	for(UEdGraphPin* LastEdGraphPin : LastEdGraphPins)
-	{
-		LastEdGraphPin->ParentPin = nullptr;
-		LastEdGraphPin->SubPins.Reset();
-	}
-	Pins.Reset();
-
-	// Recreate the new pins
-	ReallocatePinsDuringReconstruction(LastEdGraphPins);
-
-	// Maintain watches up to date
-	if (URigVMNode* Node = GetModelNode())
-	{
-		UBlueprint* Blueprint = GetBlueprint();
-		for (UEdGraphPin* NewPin : Pins)
-		{
-			const FString PinName = NewPin->GetName();
-			FString Left, Right = PinName;
-			URigVMPin::SplitPinPathAtStart(PinName, Left, Right);
-			if (URigVMPin* ModelPin = Node->FindPin(Right))
-			{
-				if (ModelPin->RequiresWatch())
-				{
-					FKismetDebugUtilities::AddPinWatch(Blueprint, FBlueprintWatchedPin(NewPin));
-				}
-			}
-		}
-	}
-	
-	// LogBlueprint: Warning: UEdGraphPin::DestroyImpl BasicIKSetup_0.UpperBone? causing an issue?
-	// we may need to clean these pins in the opposite order? child pins first?
-	RewireOldPinsToNewPins(LastEdGraphPins, Pins);
-	LastEdGraphPins.Reset();
-
-	DrawAsCompactNodeCache.Reset();
-
-	// Let subclasses do any additional work
-	PostReconstructNode();
-
-	if(bForce)
-	{
-		if (RigGraph)
-		{
-			RigGraph->NotifyGraphChanged();
-		}
-	}
-	else
-	{
-		InvalidateNodeTitle();
-		OnNodePinsChanged().Broadcast();
-	}
-}
-
-bool UControlRigGraphNode::IsDeprecated() const
-{
-	if(URigVMNode* ModelNode = GetModelNode())
-	{
-		if(URigVMUnitNode* StructModelNode = Cast<URigVMUnitNode>(ModelNode))
-		{
-			return StructModelNode->IsDeprecated();
-		}
-	}
-	return Super::IsDeprecated();
-}
-
-FEdGraphNodeDeprecationResponse UControlRigGraphNode::GetDeprecationResponse(EEdGraphNodeDeprecationType DeprecationType) const
-{
-	FEdGraphNodeDeprecationResponse Response = Super::GetDeprecationResponse(DeprecationType);
-
-	if(URigVMNode* ModelNode = GetModelNode())
-	{
-		if(URigVMUnitNode* StructModelNode = Cast<URigVMUnitNode>(ModelNode))
-		{
-			FString DeprecatedMetadata = StructModelNode->GetDeprecatedMetadata();
-			if (!DeprecatedMetadata.IsEmpty())
-			{
-				FFormatNamedArguments Args;
-				Args.Add(TEXT("DeprecatedMetadata"), FText::FromString(DeprecatedMetadata));
-				Response.MessageText = FText::Format(LOCTEXT("ControlRigGraphNodeDeprecationMessage", "Warning: This node is deprecated from: {DeprecatedMetadata}"), Args);
-			}
-		}
-	}
-
-	return Response;
-}
-
-void UControlRigGraphNode::ReallocatePinsDuringReconstruction(const TArray<UEdGraphPin*>& OldPins)
-{
-	AllocateDefaultPins();
-}
-
-void UControlRigGraphNode::RewireOldPinsToNewPins(TArray<UEdGraphPin*>& InOldPins, TArray<UEdGraphPin*>& InNewPins)
-{
-	DECLARE_SCOPE_HIERARCHICAL_COUNTER_FUNC()
-
-	for(UEdGraphPin* OldPin : InOldPins)
-	{
-		for(UEdGraphPin* NewPin : InNewPins)
-		{
-			if(OldPin->PinName == NewPin->PinName && OldPin->Direction == NewPin->Direction)
-			{
-				if (OldPin->PinType == NewPin->PinType ||
-					OldPin->PinType.PinSubCategoryObject == RigVMTypeUtils::GetWildCardCPPTypeObject() ||
-					NewPin->PinType.PinSubCategoryObject == RigVMTypeUtils::GetWildCardCPPTypeObject())
-				{
-					// make sure to remove invalid entries from the linked to list
-					OldPin->LinkedTo.Remove(nullptr);
-				
-					NewPin->MovePersistentDataFromOldPin(*OldPin);
-					break;
-				}
-			}
-		}
-	}
-
-	DestroyPinList(InOldPins);
-}
-
-void UControlRigGraphNode::DestroyPinList(TArray<UEdGraphPin*>& InPins)
-{
-	DECLARE_SCOPE_HIERARCHICAL_COUNTER_FUNC()
-
-	UBlueprint* Blueprint = GetBlueprint();
-	bool bNotify = false;
-	if (Blueprint != nullptr)
-	{
-		bNotify = !Blueprint->bIsRegeneratingOnLoad;
-	}
-
-	// Throw away the original pins
-	for (UEdGraphPin* Pin : InPins)
-	{
-		// previously serialized graphs potentially contain
-		// nullptr in the subpins
-		Pin->SubPins.Remove(nullptr);
-		
-		Pin->BreakAllPinLinks(bNotify);
-		Pin->SubPins.Remove(nullptr);
-		UEdGraphNode::DestroyPin(Pin);
-	}
-}
-
-void UControlRigGraphNode::PostReconstructNode()
-{
-	DECLARE_SCOPE_HIERARCHICAL_COUNTER_FUNC()
-
-	for (UEdGraphPin* Pin : Pins)
-	{
-		SetupPinDefaultsFromModel(Pin);
-	}
-
-	bCanRenameNode = false;
-
-	if(URigVMNode* ModelNode = GetModelNode())
-	{
-		SetColorFromModel(ModelNode->GetNodeColor());
-	}
-}
-
-void UControlRigGraphNode::SetColorFromModel(const FLinearColor& InColor)
-{
-	static const FLinearColor TitleToNodeColor(0.35f, 0.35f, 0.35f, 1.f);
-	CachedNodeColor = InColor * TitleToNodeColor;
-	CachedTitleColor = InColor;
-}
-
-void UControlRigGraphNode::HandleClearArray(FString InPinPath)
-{
-	DECLARE_SCOPE_HIERARCHICAL_COUNTER_FUNC()
-
-	if(URigVMController* Controller = GetController())
-	{
-		Controller->ClearArrayPin(InPinPath);
-	}
-}
-
-void UControlRigGraphNode::HandleAddArrayElement(FString InPinPath)
-{
-	DECLARE_SCOPE_HIERARCHICAL_COUNTER_FUNC()
-
-	if (URigVMController* Controller = GetController())
-	{
-		Controller->OpenUndoBracket(TEXT("Add Array Pin"));
-		FString PinPath = Controller->AddArrayPin(InPinPath, FString(), true, true);
-		Controller->SetPinExpansion(InPinPath, true);
-		Controller->SetPinExpansion(PinPath, true);
-		Controller->CloseUndoBracket();
-	}
-}
-
-void UControlRigGraphNode::HandleRemoveArrayElement(FString InPinPath)
-{
-	DECLARE_SCOPE_HIERARCHICAL_COUNTER_FUNC()
-
-	if (URigVMController* Controller = GetController())
-	{
-		Controller->RemoveArrayPin(InPinPath, true, true);
-	}
-}
-
-void UControlRigGraphNode::HandleInsertArrayElement(FString InPinPath)
-{
-	DECLARE_SCOPE_HIERARCHICAL_COUNTER_FUNC()
-
-	if (URigVMController* Controller = GetController())
-	{
-		if (URigVMPin* ArrayElementPin = GetModelPinFromPinPath(InPinPath))
-			{
-			if (URigVMPin* ArrayPin = ArrayElementPin->GetParentPin())
-				{
-				Controller->OpenUndoBracket(TEXT("Add Array Pin"));
-				FString PinPath = Controller->InsertArrayPin(InPinPath, ArrayElementPin->GetPinIndex() + 1, FString(), true, true);
-				Controller->SetPinExpansion(InPinPath, true);
-				Controller->SetPinExpansion(PinPath, true);
-				Controller->CloseUndoBracket();
-			}
-		}
-	}
-}
-
-int32 UControlRigGraphNode::GetInstructionIndex(bool bAsInput) const
-{
-	if (UControlRigGraph* RigGraph = Cast<UControlRigGraph>(GetGraph()))
-	{
-		return RigGraph->GetInstructionIndex(this, bAsInput);
-	}
-	return INDEX_NONE;
-}
-
-const FRigVMTemplate* UControlRigGraphNode::GetTemplate() const
-{
-	if(CachedTemplate == nullptr)
-	{
-		const FRigVMRegistry& Registry = FRigVMRegistry::Get();
-		if(URigVMTemplateNode* TemplateNode = Cast<URigVMTemplateNode>(GetModelNode()))
-		{
-			CachedTemplate = TemplateNode->GetTemplate();
-		}
-		else if(const FRigVMTemplate* Template = Registry.FindTemplate(*ModelNodePath))
-		{
-			CachedTemplate = Template; 
-		}
-	}
-	return CachedTemplate;
-}
-
-void UControlRigGraphNode::ClearErrorInfo()
-{
-	bHasCompilerMessage = false;
-	// SControlRigGraphNode only updates if the error types do not match so we have
-	// clear the error type as well, see SControlRigGraphNode::RefreshErrorInfo()
-	ErrorType = (int32)EMessageSeverity::Info + 1;
-	ErrorMsg = FString();	
-}
-
-URigVMPin* UControlRigGraphNode::FindModelPinFromGraphPin(const UEdGraphPin* InGraphPin) const
-{
-	if(InGraphPin == nullptr)
-	{
-		return nullptr;
-	}
-	
-	for(const auto& Pair : CachedPins)
-	{
-		if(Pair.Value.InputPin == InGraphPin ||
-			Pair.Value.OutputPin == InGraphPin)
-		{
-			return Pair.Key;
-		}
-	}
-
-	return GetModelPinFromPinPath(InGraphPin->GetName());
-}
-
-UEdGraphPin* UControlRigGraphNode::FindGraphPinFromModelPin(const URigVMPin* InModelPin, bool bAsInput) const
-{
-	if(InModelPin == nullptr)
-	{
-		return nullptr;
-	}
-	
-	if(const FPinPair* Pair = CachedPins.Find(InModelPin))
-	{
-		return bAsInput ? Pair->InputPin : Pair->OutputPin;
-	}
-
-	const FString PinPath = InModelPin->GetPinPath();
-	for(UEdGraphPin* GraphPin : Pins)
-	{
-		if((GraphPin->Direction == EGPD_Input) == bAsInput)
-		{
-			if(GraphPin->GetName() == PinPath)
-			{
-				return GraphPin;
-			}
-		}
-	}
-	
-	return nullptr;
-}
-
-void UControlRigGraphNode::SynchronizeGraphPinNameWithModelPin(const URigVMPin* InModelPin, bool bNotify)
-{
-	auto SyncGraphPinLambda = [this](const URigVMPin* InModelPin, bool bAsInput) -> bool
-	{
-		if(UEdGraphPin* GraphPin = FindGraphPinFromModelPin(InModelPin, bAsInput))
-		{
-			const FString OldPinName = GraphPin->PinName.ToString();
-			const FString NewPinName = InModelPin->GetPinPath();
-
-			if(OldPinName != NewPinName)
-			{
-				PinPathToModelPin.Remove(OldPinName);
-				
-				GraphPin->PinName = *NewPinName;
-				GraphPin->PinFriendlyName = FText::FromName(InModelPin->GetDisplayName());
-
-				PinPathToModelPin.Add(NewPinName, (URigVMPin*)InModelPin);
-				return true;
-			}
-		}
-
-		return false;
-	};
-
-	bool bResult = false;
-	
-	if(SyncGraphPinLambda(InModelPin, true))
-	{
-		bResult = true;
-	}
-	if(SyncGraphPinLambda(InModelPin, false))
-	{
-		bResult = true;
-	}
-
-	if(bResult)
-	{
-		for (const URigVMPin* ModelSubPin : InModelPin->GetSubPins())
-		{
-			SynchronizeGraphPinNameWithModelPin(ModelSubPin, false);
-		}
-
-		if(bNotify)
-		{
-			// Notify the node widget that the pins have changed.
-			OnNodePinsChanged().Broadcast();
-		}
-	}
-}
-
-void UControlRigGraphNode::SynchronizeGraphPinValueWithModelPin(const URigVMPin* InModelPin)
-{
-	auto SyncGraphPinLambda = [this](const URigVMPin* InModelPin, bool bAsInput)
-	{
-		// If the pin has sub-pins, we may need to remove or rebuild the sub-pins.
-		if (UEdGraphPin* GraphPin = FindGraphPinFromModelPin(InModelPin, bAsInput))
-		{
-			SetupPinDefaultsFromModel(GraphPin, InModelPin);
-		}
-	};
-
-	SyncGraphPinLambda(InModelPin, true);
-	SyncGraphPinLambda(InModelPin, false);
-}
-
-void UControlRigGraphNode::SynchronizeGraphPinTypeWithModelPin(const URigVMPin* InModelPin)
-{
-	bool bNotify = false;
-
-	auto SyncGraphPinLambda = [this, &bNotify](const URigVMPin* InModelPin, bool bAsInput)
-	{
-		// If the pin has sub-pins, we may need to remove or rebuild the sub-pins.
-		if (UEdGraphPin* GraphPin = FindGraphPinFromModelPin(InModelPin, bAsInput))
-		{
-			const FEdGraphPinType NewGraphPinType = GetPinTypeForModelPin(InModelPin);
-			if(NewGraphPinType != GraphPin->PinType)
-			{
-				GraphPin->PinType = NewGraphPinType;
-				ConfigurePin(GraphPin, InModelPin);
-
-				// Create new sub-pins, if required, to reflect the new type.
-				TArray<UEdGraphPin*> GraphSubPinsToKeep;
-				if (!InModelPin->GetSubPins().IsEmpty())
-				{
-					for (const URigVMPin* ModelSubPin : InModelPin->GetSubPins())
-					{
-						if(UEdGraphPin* GraphSubPin = FindGraphPinFromModelPin(ModelSubPin, bAsInput))
-						{
-							GraphSubPinsToKeep.Add(GraphSubPin);
-
-							const FEdGraphPinType NewGraphSubPinType = GetPinTypeForModelPin(ModelSubPin);
-							if(NewGraphSubPinType != GraphSubPin->PinType)
-							{
-								GraphSubPin->PinType = NewGraphSubPinType;
-								ConfigurePin(GraphSubPin, ModelSubPin);
-							}
-						}
-						else
-						{
-							CreateGraphPinFromModelPin(ModelSubPin, GraphPin->Direction, GraphPin);
-						}
-					}
-				}
-
-				// If the graph node had other sub-pins, we need to remove those.
-				RemoveGraphSubPins(GraphPin, GraphSubPinsToKeep);
-
-				bNotify = true;
-			}
-		}
-	};
-
-	SyncGraphPinLambda(InModelPin, true);
-	SyncGraphPinLambda(InModelPin, false);
-
-	if(bNotify)
-	{
-		// Notify the node widget that the pins have changed.
-		OnNodePinsChanged().Broadcast();
-	}
-}
-
-void UControlRigGraphNode::SynchronizeGraphPinExpansionWithModelPin(const URigVMPin* InModelPin)
-{
-	OnNodePinExpansionChanged().Broadcast();
-}
-
-void UControlRigGraphNode::SyncGraphNodeTitleWithModelNodeTitle()
-{
-	InvalidateNodeTitle();
-}
-
-void UControlRigGraphNode::SyncGraphNodeNameWithModelNodeName(const URigVMNode* InModelNode)
-{
-	Rename(*InModelNode->GetName());
-	ModelNodePath = InModelNode->GetNodePath();
-	SyncGraphNodeTitleWithModelNodeTitle();
-
-	TArray<URigVMPin*> AllModelPins = InModelNode->GetAllPinsRecursively();
-	for(const URigVMPin* ModelPin : AllModelPins)
-	{
-		SynchronizeGraphPinNameWithModelPin(ModelPin);
-	}
-}
-
-bool UControlRigGraphNode::CreateGraphPinFromModelPin(const URigVMPin* InModelPin, EEdGraphPinDirection InDirection,
-	UEdGraphPin* InParentPin)
-{
-	// don't create output pins for array elements
-	if(InDirection == EGPD_Output && InModelPin->IsArrayElement() && !InModelPin->GetParentPin()->IsFixedSizeArray())
-	{
-		return false;
-	}
-
-	const FPinPair PairConst = CachedPins.FindOrAdd((URigVMPin*)InModelPin);
-
-	auto CreatePinLambda = [this](const URigVMPin* InModelPin, EEdGraphPinDirection InDirection, UEdGraphPin* InParentPin) -> UEdGraphPin*
-	{
-		// check if we already have a pin like this in the last pins arrays
-		const FString PinPath = InModelPin->GetPinPath();
-		const FEdGraphPinType PinType = GetPinTypeForModelPin(InModelPin);
-
-		UEdGraphPin* GraphPin = nullptr;
-		UEdGraphPin** ExistingEdGraphPinPtr = LastEdGraphPins.FindByPredicate([PinPath, InDirection, PinType](const UEdGraphPin* ExistingPin) -> bool
-		{
-			return !ExistingPin->bWasTrashed &&
-				ExistingPin->GetName() == PinPath &&
-				ExistingPin->Direction == InDirection &&
-				ExistingPin->PinType == PinType;
-		});
-		if(ExistingEdGraphPinPtr)
-		{
-			GraphPin = *ExistingEdGraphPinPtr;
-			GraphPin->ParentPin = nullptr;
-			GraphPin->SubPins.Reset();
-			ExistingEdGraphPinPtr = nullptr;
-			LastEdGraphPins.Remove(GraphPin);
-			Pins.Add(GraphPin);
-		}
-
-		if(GraphPin == nullptr)
-		{
-			GraphPin = CreatePin(InDirection, GetPinTypeForModelPin(InModelPin), FName(*InModelPin->GetPinPath()));
-		}
-		
-		if (GraphPin)
-		{
-			ConfigurePin(GraphPin, InModelPin);
-
-			if (InParentPin)
-			{
-				InParentPin->SubPins.Add(GraphPin);
-				GraphPin->ParentPin = InParentPin;
-			}
-			
-			for(const URigVMPin* ModelSubPin : InModelPin->GetSubPins())
-			{
-				CreateGraphPinFromModelPin(ModelSubPin, InDirection, GraphPin);
-			}
-
-			SetupPinDefaultsFromModel(GraphPin, InModelPin);
-		}
-
-		return GraphPin;
-	};
-
-	UEdGraphPin* InputPin = nullptr;
-	UEdGraphPin* OutputPin = nullptr;
-
-	bool bResult = false;
-	if (InDirection == EGPD_Input && PairConst.InputPin == nullptr)
-	{
-		InputPin = CreatePinLambda(InModelPin, EGPD_Input, InParentPin);
-		bResult = true;
-	}
-	if (InDirection == EGPD_Output && PairConst.OutputPin == nullptr)
-	{
-		OutputPin = CreatePinLambda(InModelPin, EGPD_Output, InParentPin);
-		bResult = true;
-	}
-
-	FPinPair& Pair = CachedPins.FindChecked((URigVMPin*)InModelPin);
-	Pair.InputPin = InputPin != nullptr ? InputPin : Pair.InputPin;
-	Pair.OutputPin = OutputPin != nullptr ? OutputPin : Pair.OutputPin;
-	return Pair.IsValid() && bResult;
-}
-
-void UControlRigGraphNode::RemoveGraphSubPins(UEdGraphPin* InParentPin, const TArray<UEdGraphPin*>& InPinsToKeep)
-{
-	TArray<UEdGraphPin*> SubPins = InParentPin->SubPins;
-
-	TArray<URigVMPin*> ModelSubPins;
-	for (const UEdGraphPin* SubPin: SubPins)
-	{
-		if(InPinsToKeep.Contains(SubPin))
-		{
-			continue;
-		}
-		for(const auto& Pair : CachedPins)
-		{
-			if(Pair.Value.InputPin == SubPin ||
-				Pair.Value.OutputPin == SubPin)
-			{
-				ModelSubPins.Add(Pair.Key);
-			}
-		}
-	}
-
-	for(const URigVMPin* ModelSubPin : ModelSubPins)
-	{
-		PinPathToModelPin.Remove(ModelSubPin->GetPinPath());
-
-		if(const FPinPair* PinPair = CachedPins.Find(ModelSubPin))
-		{
-			auto Traverse = [this](UEdGraphPin* SubPin)
-			{
-				if(SubPin == nullptr)
-				{
-					return;
-				}
-				
-				// Remove this pin from our owned pins
-				Pins.Remove(SubPin);
-
-				if (!SubPin->SubPins.IsEmpty())
-				{
-					RemoveGraphSubPins(SubPin);
-				}
-
-				SubPin->MarkAsGarbage();
-			};
-			
-			Traverse(PinPair->InputPin);
-			Traverse(PinPair->OutputPin);
-		}
-
-		CachedPins.Remove(ModelSubPin);
-	}
-	InParentPin->SubPins.Reset();
-}
-
-bool UControlRigGraphNode::ModelPinsChanged(bool bForce)
-{
-	UpdatePinLists();
-
-	// check if any of the pins need to be added / removed
-	auto AddMissingPins = [this](const TArray<URigVMPin*>& InModelPins)
-	{
-		int32 PinsAdded = 0;
-		for(const URigVMPin* ModelPin : InModelPins)
-		{
-			if(FindGraphPinFromModelPin(ModelPin, true) == nullptr && 
-				FindGraphPinFromModelPin(ModelPin, false))
-			{
-				PinsAdded += ModelPinAdded_Internal(ModelPin) ? 1 : 0;
-			}
-		}
-		return PinsAdded;
-	};
-
-	auto RemoveObsoletePins = [this]()
-	{
-		TArray<URigVMPin*> PinsToRemove; 
-		for(const auto& Pair : CachedPins)
-		{
-			URigVMPin* ModelPin = Pair.Key;
-			TArray<URigVMPin*>& List = PinListForPin(ModelPin);
-			bool bRemove = true;
-			if (URigVMPin** FoundPin = List.FindByKey(ModelPin->GetRootPin()))
-			{
-				if (ModelPin->GetParentPin() == nullptr ||
-					(*FoundPin)->FindSubPin(ModelPin->GetSegmentPath()))
-				{
-					bRemove = false;
-				}
-			}
-			if(bRemove)
-			{
-				PinsToRemove.Add(ModelPin);
-			}
-		}
-
-		int32 PinsRemoved = 0;
-		for(const URigVMPin* ModelPin : PinsToRemove)
-		{
-			PinsRemoved += ModelPinRemoved_Internal(ModelPin) ? 1 : 0;
-		}
-		return PinsRemoved;
-	};
-
-	auto OrderPins = [this](const TArray<URigVMPin*>& InModelPins) -> int32
-	{
-		if(InModelPins.Num() < 2)
-		{
-			return 0;
-		}
-
-		TArray<UEdGraphPin*> OrderedGraphPins;
-		OrderedGraphPins.Reserve(InModelPins.Num() * 2);
-		int32 LastInputIndex = INDEX_NONE;
-		int32 LastOutputIndex = INDEX_NONE;
-		int32 PinsToReorder = 0;
-
-		for(const URigVMPin* ModelPin : InModelPins)
-		{
-			if(UEdGraphPin* InputGraphPin = FindGraphPinFromModelPin(ModelPin, true))
-			{
-				OrderedGraphPins.Add(InputGraphPin);
-				const int32 InputPinIndex = Pins.Find(InputGraphPin);
-				if(LastInputIndex > InputPinIndex)
-				{
-					PinsToReorder++;
-				}
-				LastInputIndex = InputPinIndex;
-			}
-			
-			if(UEdGraphPin* OutputGraphPin = FindGraphPinFromModelPin(ModelPin, false))
-			{
-				OrderedGraphPins.Add(OutputGraphPin);
-				const int32 OutputPinIndex = Pins.Find(OutputGraphPin);
-				if(LastOutputIndex > OutputPinIndex)
-				{
-					PinsToReorder++;
-				}
-				LastOutputIndex = OutputPinIndex;
-			}
-		}
-
-		if(PinsToReorder > 0)
-		{
-			for(UEdGraphPin* GraphPinInOrder : OrderedGraphPins)
-			{
-				Pins.Remove(GraphPinInOrder);
-			}
-
-			OrderedGraphPins.Append(Pins);
-			Swap(OrderedGraphPins, Pins);
-		}
-
-		return PinsToReorder;
-	};
-	
-	int32 PinsAdded = 0;
-	PinsAdded += AddMissingPins(ExecutePins);
-	PinsAdded += AddMissingPins(InputPins);
-	PinsAdded += AddMissingPins(InputOutputPins);
-	PinsAdded += AddMissingPins(OutputPins);
-
-	const int32 PinsRemoved = RemoveObsoletePins();
-
-	// working through it in the opposite order
-	// due to the use of Append within the lambda
-	int32 PinsReordered = 0;
-	PinsReordered += OrderPins(OutputPins);
-	PinsReordered += OrderPins(InputOutputPins);
-	PinsReordered += OrderPins(InputPins);
-	PinsReordered += OrderPins(ExecutePins);
-
-	const bool bResult = (PinsAdded > 0) || (PinsRemoved > 0) || (PinsReordered > 0); 
-	if(bResult || bForce)
-	{
-		OnNodePinsChanged().Broadcast();
-	}
-	return bResult;
-}
-
-bool UControlRigGraphNode::ModelPinAdded(const URigVMPin* InModelPin)
-{
-	if(InModelPin == nullptr)
-	{
-		return false;
-	}
-
-	UpdatePinLists();
-
-	bool bResult = ModelPinAdded_Internal(InModelPin);
-	if(bResult)
-	{
-		OnNodePinsChanged().Broadcast();
-	}
-	return bResult;
-}
-
-bool UControlRigGraphNode::ModelPinAdded_Internal(const URigVMPin* InModelPin)
-{
-	bool bResult = false;
-
-	// conversion nodes don't show sub pins
-	if(!InModelPin->IsRootPin())
-	{
-		if(DrawAsCompactNode())
-		{
-			return false;
-		}
-	}
-
-	UEdGraphPin* InputParentPin = FindGraphPinFromModelPin(InModelPin->GetParentPin(), true);
-	UEdGraphPin* OutputParentPin = FindGraphPinFromModelPin(InModelPin->GetParentPin(), false);
-		
-	if(InModelPin->GetDirection() == ERigVMPinDirection::Input ||
-		InModelPin->GetDirection() == ERigVMPinDirection::Visible ||
-		InModelPin->GetDirection() == ERigVMPinDirection::IO)
-	{
-		if(CreateGraphPinFromModelPin(InModelPin, EGPD_Input, InputParentPin))
-		{
-			bResult = true;
-		}
-	}
-	
-	if(InModelPin->GetDirection() == ERigVMPinDirection::Output ||
-		InModelPin->GetDirection() == ERigVMPinDirection::IO)
-	{
-		if(CreateGraphPinFromModelPin(InModelPin, EGPD_Output, OutputParentPin))
-		{
-			bResult = true;
-		}
-	}
-	
-	return bResult;
-}
-
-bool UControlRigGraphNode::ModelPinRemoved(const URigVMPin* InModelPin)
-{
-	if(InModelPin == nullptr)
-	{
-		return false;
-	}
-
-	UpdatePinLists();
-
-	const bool bResult = ModelPinRemoved_Internal(InModelPin);
-	if(bResult)
-	{
-		OnNodePinsChanged().Broadcast();
-	}
-	return bResult;
-}
-
-bool UControlRigGraphNode::DrawAsCompactNode() const
-{
-	if(CVarControlRigDisableCompactNodes.GetValueOnAnyThread())
-	{
-		return false;
-	}
-	
-	if(!DrawAsCompactNodeCache.IsSet())
-	{
-		DrawAsCompactNodeCache = false;
-		if(const URigVMTemplateNode* TemplateModelNode = Cast<URigVMTemplateNode>(GetModelNode()))
-		{
-			if(TemplateModelNode->GetNotation() == RigVMTypeUtils::GetCastTemplateNotation())
-			{
-				DrawAsCompactNodeCache = true;
-				
-				// if the node has links on any subpin - we can't draw it as a compact node.
-				const TArray<URigVMLink*> Links = TemplateModelNode->GetLinks();
-				if(Links.ContainsByPredicate([TemplateModelNode](const URigVMLink* Link)
-				{
-					if(const URigVMPin* SourcePin = Link->GetSourcePin())
-					{
-						if(SourcePin->GetNode() == TemplateModelNode)
-						{
-							return !SourcePin->IsRootPin();
-						}
-					}
-					if(const URigVMPin* TargetPin = Link->GetTargetPin())
-					{
-						if(TargetPin->GetNode() == TemplateModelNode)
-						{
-							return !TargetPin->IsRootPin();
-						}
-					}
-					return false;
-				}))
-				{
-					DrawAsCompactNodeCache = false;
-				}
-
-				// if the node has any injected nodes - we can't draw it as compact
-				if(TemplateModelNode->GetPins().ContainsByPredicate([](const URigVMPin* Pin)
-				{
-					return Pin->HasInjectedNodes();
-				}))
-				{
-					DrawAsCompactNodeCache = false;
-				}
-			}
-		}
-	}
-	return DrawAsCompactNodeCache.GetValue();
-}
-
-bool UControlRigGraphNode::ModelPinRemoved_Internal(const URigVMPin* InModelPin)
-{
-	auto RemoveGraphPin = [this](const URigVMPin* InModelPin, bool bAsInput)
-	{
-		if(UEdGraphPin* GraphPin = FindGraphPinFromModelPin(InModelPin, bAsInput))
-		{
-			RemoveGraphSubPins(GraphPin);
-			Pins.Remove(GraphPin);
-			GraphPin->MarkAsGarbage();
-		}
-	};
-	RemoveGraphPin(InModelPin, true);
-	RemoveGraphPin(InModelPin, false);
-		
-	const bool bResult = PinPathToModelPin.Remove(InModelPin->GetPinPath()) > 0;
-	CachedPins.Remove(InModelPin);
-	return bResult;
-}
-
-FLinearColor UControlRigGraphNode::GetNodeProfilingColor() const
-{
-#if WITH_EDITOR
-	if(bEnableProfiling)
-	{
-		if(UControlRigBlueprint* Blueprint = GetTypedOuter<UControlRigBlueprint>())
-		{
-			if(UControlRig* DebuggedControlRig = Cast<UControlRig>(Blueprint->GetObjectBeingDebugged()))
-			{
-				if(URigVMNode* ModelNode = GetModelNode())
-				{
-					const double MicroSeconds = ModelNode->GetInstructionMicroSeconds(DebuggedControlRig->GetExtendedExecuteContext(), DebuggedControlRig->GetVM(), FRigVMASTProxy());
-					if(MicroSeconds >= 0.0)
-					{
-						if(Blueprint->RigGraphDisplaySettings.bAutoDetermineRange)
-						{
-							if(MicroSeconds < Blueprint->RigGraphDisplaySettings.MinMicroSeconds)
-							{
-								Blueprint->RigGraphDisplaySettings.MinMicroSeconds = MicroSeconds;
-							}
-							if(MicroSeconds > Blueprint->RigGraphDisplaySettings.MaxMicroSeconds)
-							{
-								Blueprint->RigGraphDisplaySettings.MaxMicroSeconds = MicroSeconds;
-							}
-						}
-							
-						const double MinMicroSeconds = Blueprint->RigGraphDisplaySettings.LastMinMicroSeconds;
-						const double MaxMicroSeconds = Blueprint->RigGraphDisplaySettings.LastMaxMicroSeconds;
-						if(MaxMicroSeconds <= MinMicroSeconds)
-						{
-							return FLinearColor::Black;
-						}
-			
-						const FLinearColor& MinColor = Blueprint->RigGraphDisplaySettings.MinDurationColor;
-						const FLinearColor& MaxColor = Blueprint->RigGraphDisplaySettings.MaxDurationColor;
-
-						const double T = (MicroSeconds - MinMicroSeconds) / (MaxMicroSeconds - MinMicroSeconds);
-						return FMath::Lerp<FLinearColor>(MinColor, MaxColor, (float)T);
-					}
-				}
-			}
-		}
-	}
-#endif
-	return FLinearColor::Black;
-}
-
-void UControlRigGraphNode::UpdatePinLists()
-{
-	ExecutePins.Reset();
-	InputPins.Reset();
-	InputOutputPins.Reset();
-	OutputPins.Reset();
-
-	if (URigVMNode* ModelNode = GetModelNode())
-	{
-		for(int32 PinListIndex=0; PinListIndex<2; PinListIndex++)
-		{
-			const TArray<URigVMPin*>& ModelPins = PinListIndex == 0 ? ModelNode->GetPins() : ModelNode->GetOrphanedPins();
-			for (URigVMPin* ModelPin : ModelPins)
-			{
-				if (ModelPin->ShowInDetailsPanelOnly())
-				{
-					continue;
-				}
-				if (ModelPin->GetDirection() == ERigVMPinDirection::IO)
-				{
-					if (ModelPin->IsStruct())
-					{
-						const UScriptStruct* ScriptStruct = ModelPin->GetScriptStruct();
-						if (ScriptStruct && ScriptStruct->IsChildOf(FRigVMExecuteContext::StaticStruct()))
-						{
-							ExecutePins.Add(ModelPin);
-							continue;
-						}
-					}
-					InputOutputPins.Add(ModelPin);
-				}
-				else if (ModelPin->GetDirection() == ERigVMPinDirection::Input || 
-					ModelPin->GetDirection() == ERigVMPinDirection::Visible)
-				{
-					InputPins.Add(ModelPin);
-				}
-				else if (ModelPin->GetDirection() == ERigVMPinDirection::Output)
-				{
-					OutputPins.Add(ModelPin);
-				}
-			}
-		}
-	}
-}
-
-void UControlRigGraphNode::AllocateDefaultPins()
-{
-	UpdatePinLists();
-
-	CachedPins.Reset();
-	PinPathToModelPin.Reset();
-
-	auto CreateGraphPins = [this](const TArray<URigVMPin*>& InModelPins)
-	{
-		for (const URigVMPin* ModelPin : InModelPins)
-		{
-			ModelPinAdded_Internal(ModelPin);
-		}
-	};
-
-	CreateGraphPins(ExecutePins);
-	CreateGraphPins(InputPins);
-	CreateGraphPins(InputOutputPins);
-	CreateGraphPins(OutputPins);
-
-	// Fill the variable list
-	ExternalVariables.Reset();
-	if (URigVMFunctionReferenceNode* FunctionReferenceNode = Cast<URigVMFunctionReferenceNode>(GetModelNode()))
-	{
-		if(FunctionReferenceNode->RequiresVariableRemapping())
-		{
-			TArray<FRigVMExternalVariable> CurrentExternalVariables = FunctionReferenceNode->GetExternalVariables(false);
-			for(const FRigVMExternalVariable& CurrentExternalVariable : CurrentExternalVariables)
-			{
-				ExternalVariables.Add(MakeShared<FRigVMExternalVariable>(CurrentExternalVariable));
-			}
-		}
-	}
-}
-
-UClass* UControlRigGraphNode::GetControlRigGeneratedClass() const
-{
-	const UControlRigBlueprint* Blueprint = GetTypedOuter<UControlRigBlueprint>();
-	if(Blueprint)
-	{
-		if (Blueprint->GeneratedClass)
-		{
-			check(Blueprint->GeneratedClass->IsChildOf(UControlRig::StaticClass()));
-			return Blueprint->GeneratedClass;
-		}
-	}
-
-	return nullptr;
-}
-
-UClass* UControlRigGraphNode::GetControlRigSkeletonGeneratedClass() const
-{
-	const UControlRigBlueprint* Blueprint = GetTypedOuter<UControlRigBlueprint>();
-	if(Blueprint)
-	{
-		if (Blueprint->SkeletonGeneratedClass)
-		{
-			check(Blueprint->SkeletonGeneratedClass->IsChildOf(UControlRig::StaticClass()));
-			return Blueprint->SkeletonGeneratedClass;
-		}
-	}
-	return nullptr;
-}
-
-FLinearColor UControlRigGraphNode::GetNodeOpacityColor() const
+: URigVMEdGraphNode()
 {
-	if (URigVMNode* ModelNode = GetModelNode())
-	{
-		if (Cast<URigVMVariableNode>(ModelNode))
-		{
-			return FLinearColor::White;
-		}
-
-		if(GetInstructionIndex(true) == INDEX_NONE)
-		{
-			return FLinearColor(0.35f, 0.35f, 0.35f, 0.35f);
-		}
-	}
-	return FLinearColor::White;
 }
 
-FLinearColor UControlRigGraphNode::GetNodeTitleColor() const
-{
-	// return a darkened version of the default node's color
-	return CachedTitleColor * GetNodeOpacityColor();
-}
-
-FLinearColor UControlRigGraphNode::GetNodeBodyTintColor() const
-{
-#if WITH_EDITOR
-	if(bEnableProfiling)
-	{
-		return GetNodeProfilingColor();
-	}
-#endif
-	
-	return CachedNodeColor * GetNodeOpacityColor();
-}
-
-bool UControlRigGraphNode::ShowPaletteIconOnNode() const
-{
-	if (URigVMNode* ModelNode = GetModelNode())
-	{
-		return ModelNode->IsEvent() ||
-			ModelNode->IsA<URigVMInvokeEntryNode>() ||
-			ModelNode->IsA<URigVMFunctionEntryNode>() ||
-			ModelNode->IsA<URigVMFunctionReturnNode>() ||
-			ModelNode->IsA<URigVMFunctionReferenceNode>() ||
-			ModelNode->IsA<URigVMCollapseNode>() ||
-			ModelNode->IsA<URigVMUnitNode>() ||
-			ModelNode->IsA<URigVMDispatchNode>() ||
-			ModelNode->IsLoopNode();
-	}
-	return false;
-}
-
-FSlateIcon UControlRigGraphNode::GetIconAndTint(FLinearColor& OutColor) const
-{
-	OutColor = FLinearColor::White;
-
-	static FSlateIcon FunctionIcon(FAppStyle::GetAppStyleSetName(), "Kismet.AllClasses.FunctionIcon");
-	static FSlateIcon EventIcon(FAppStyle::GetAppStyleSetName(), "GraphEditor.Event_16x");
-	static FSlateIcon EntryReturnIcon(FAppStyle::GetAppStyleSetName(), "GraphEditor.Default_16x");
-	static FSlateIcon CollapsedNodeIcon(FAppStyle::GetAppStyleSetName(), "GraphEditor.SubGraph_16x");
-	static FSlateIcon TemplateNodeIcon("ControlRigEditorStyle", "ControlRig.Template");
-
-	if (URigVMNode* ModelNode = GetModelNode())
-	{
-		if (ModelNode->IsEvent() || ModelNode->IsA<URigVMInvokeEntryNode>())
-		{
-			return EventIcon;
-		}
-
-		URigVMNode* InnerNode = ModelNode;
-		while(const URigVMAggregateNode* AggregateNode = Cast<URigVMAggregateNode>(InnerNode))
-		{
-			InnerNode = AggregateNode->GetFirstInnerNode();
-		}
-		if (InnerNode)
-		{
-			ModelNode = InnerNode;
-		}
-		else
-		{
-			if(!ModelNode->HasAnyFlags(RF_NeedLoad))
-			{
-				const FString Message = FString::Printf(TEXT("Could not find first inner node in aggregate node %s in package %s"), *ModelNode->GetPathName(), *GetPackage()->GetPathName());
-				FScriptExceptionHandler::Get().HandleException(ELogVerbosity::Warning, *Message, *FString());
-			}
-		}
-
-		if (ModelNode->IsA<URigVMFunctionReferenceNode>())
-		{ 
-			return FunctionIcon;
-		}
-
-		if (ModelNode->IsA<URigVMCollapseNode>())
-		{
-			return CollapsedNodeIcon;
-		}
-
-		if (ModelNode->IsA<URigVMFunctionEntryNode>() || 
-            ModelNode->IsA<URigVMFunctionReturnNode>())
-		{
-			return EntryReturnIcon;
-		}
-
-		const UScriptStruct* MetadataScriptStruct = nullptr;
-		if (const URigVMUnitNode* UnitNode = Cast<URigVMUnitNode>(ModelNode))
-		{
-			MetadataScriptStruct = UnitNode->GetScriptStruct();
-		}
-		else if (const URigVMDispatchNode* DispatchNode = Cast<URigVMDispatchNode>(ModelNode))
-		{
-			MetadataScriptStruct = DispatchNode->GetFactory()->GetScriptStruct();
-		}
-
-		if(MetadataScriptStruct && MetadataScriptStruct->HasMetaDataHierarchical(FRigVMStruct::IconMetaName))
-		{
-			FString IconPath;
-			const int32 NumOfIconPathNames = 4;
-			
-			FName IconPathNames[NumOfIconPathNames] = {
-				NAME_None, // StyleSetName
-				NAME_None, // StyleName
-				NAME_None, // SmallStyleName
-				NAME_None  // StatusOverlayStyleName
-			};
-
-			// icon path format: StyleSetName|StyleName|SmallStyleName|StatusOverlayStyleName
-			// the last two names are optional, see FSlateIcon() for reference
-			MetadataScriptStruct->GetStringMetaDataHierarchical(FRigVMStruct::IconMetaName, &IconPath);
-
-			int32 NameIndex = 0;
-
-			while (!IconPath.IsEmpty() && NameIndex < NumOfIconPathNames)
-			{
-				FString Left;
-				FString Right;
-
-				if (!IconPath.Split(TEXT("|"), &Left, &Right))
-				{
-					Left = IconPath;
-				}
-
-				IconPathNames[NameIndex] = FName(*Left);
-
-				NameIndex++;
-				IconPath = Right;
-			}
-			
-			return FSlateIcon(IconPathNames[0], IconPathNames[1], IconPathNames[2], IconPathNames[3]);
-		}
-
-		if (const URigVMTemplateNode* TemplateNode = Cast<URigVMTemplateNode>(ModelNode))
-		{
-			if(const FRigVMTemplate* Template = TemplateNode->GetTemplate())
-			{
-				if(Template->NumPermutations() > 1)
-				{
-					return TemplateNodeIcon;
-				}
-			}
-		}
-	}
-
-	return FunctionIcon;
-}
-
-void UControlRigGraphNode::GetNodeContextMenuActions(class UToolMenu* Menu, class UGraphNodeContextMenuContext* Context) const
+void UControlRigGraphNode::GetNodeContextMenuActions(UToolMenu* Menu, UGraphNodeContextMenuContext* Context) const
 {
 #if WITH_EDITOR
+	// todo - part of this is likely generic 
 	const UControlRigGraphSchema* Schema = Cast<UControlRigGraphSchema>(GetSchema());
 	IControlRigEditorModule::Get().GetContextMenuActions(Schema, Menu, Context);
 #endif
 }
 
-bool UControlRigGraphNode::IsPinExpanded(const FString& InPinPath)
-{
-	if (URigVMPin* ModelPin = GetModelPinFromPinPath(InPinPath))
-	{
-		return ModelPin->IsExpanded();
-	}
-	return false;
-}
-
-void UControlRigGraphNode::DestroyNode()
-{
-	DECLARE_SCOPE_HIERARCHICAL_COUNTER_FUNC()
-
-	if(UControlRigGraph* Graph = Cast<UControlRigGraph>(GetOuter()))
-	{
-		BreakAllNodeLinks();
-		
-		UControlRigBlueprint* ControlRigBlueprint = Cast<UControlRigBlueprint>(Graph->GetOuter());
-		if(ControlRigBlueprint)
-		{
-			if(PropertyName_DEPRECATED.IsValid())
-			{
-				FRigVMBlueprintUtils::RemoveMemberVariableIfNotUsed(ControlRigBlueprint, PropertyName_DEPRECATED);
-			}
-		}
-	}
-
-	UEdGraphNode::DestroyNode();
-}
-
-void UControlRigGraphNode::PinDefaultValueChanged(UEdGraphPin* Pin)
-{
-	CopyPinDefaultsToModel(Pin, true, true);
-}
-
-void UControlRigGraphNode::CopyPinDefaultsToModel(UEdGraphPin* Pin, bool bUndo, bool bPrintPythonCommand)
-{
-	DECLARE_SCOPE_HIERARCHICAL_COUNTER_FUNC()
-
-	if (Pin->Direction != EGPD_Input)
-	{
-		return;
-	}
-
-	const UEdGraphSchema_K2* K2Schema = GetDefault<UEdGraphSchema_K2>();
-	if (URigVMPin* ModelPin = FindModelPinFromGraphPin(Pin))
-	{
-		if (ModelPin->GetSubPins().Num() > 0)
-		{
-			return;
-		}
-
-		FString DefaultValue = Pin->DefaultValue;
-
-		if(DefaultValue.IsEmpty() && (
-			Pin->PinType.PinCategory == UEdGraphSchema_K2::PC_Object ||
-			Pin->PinType.PinCategory == UEdGraphSchema_K2::PC_SoftObject ||
-			Pin->PinType.PinCategory == UEdGraphSchema_K2::AllObjectTypes ||
-			Pin->PinType.PinCategory == UEdGraphSchema_K2::PC_Interface
-			))
-		{
-			if(Pin->DefaultObject)
-			{
-				DefaultValue = Pin->DefaultObject->GetPathName();
-			}
-		}
-		
-		if (DefaultValue == FName(NAME_None).ToString() && Pin->PinType.PinSubCategory == UEdGraphSchema_K2::PC_Name)
-		{
-			DefaultValue = FString();
-		}
-
-		if (ModelPin->GetDefaultValue() != DefaultValue)
-		{
-			if (URigVMController* Controller = GetController())
-			{
-				Controller->SetPinDefaultValue(ModelPin->GetPinPath(), DefaultValue, false, true, false, bPrintPythonCommand);
-			}
-		}
-	}
-}
-
-UControlRigBlueprint* UControlRigGraphNode::GetBlueprint() const
-{
-	if(UControlRigGraph* Graph = Cast<UControlRigGraph>(GetOuter()))
-	{
-		return Graph->GetBlueprint();
-		if (UControlRigGraph* OuterGraph = Cast<UControlRigGraph>(Graph->GetOuter()))
-		{
-			return OuterGraph->GetBlueprint();
-		}
-		return Cast<UControlRigBlueprint>(Graph->GetOuter());
-	}
-	return nullptr;
-}
-
-URigVMGraph* UControlRigGraphNode::GetModel() const
-{
-	if (UControlRigGraph* Graph = Cast<UControlRigGraph>(GetOuter()))
-	{
-		return Graph->GetModel();
-	}
-	return nullptr;
-}
-
-URigVMController* UControlRigGraphNode::GetController() const
-{
-	if (UControlRigGraph* Graph = Cast<UControlRigGraph>(GetOuter()))
-	{
-		return Graph->GetController();
-	}
-	return nullptr;
-}
-
-URigVMNode* UControlRigGraphNode::GetModelNode() const
-{
-	UControlRigGraphNode* MutableThis = (UControlRigGraphNode*)this;
-	if (CachedModelNode.IsValid())
-	{
-		if (CachedModelNode.Get()->GetOuter() == GetTransientPackage())
-		{
-			MutableThis->CachedModelNode.Reset();
-		}
-		else
-		{
-			return CachedModelNode.Get();
-		}
-	}
-
-	if (UControlRigGraph* Graph = Cast<UControlRigGraph>(GetOuter()))
-	{
-#if WITH_EDITOR
-
-		if (Graph->TemplateController != nullptr)
-		{
-			MutableThis->CachedModelNode = TWeakObjectPtr<URigVMNode>(Graph->TemplateController->GetGraph()->FindNode(ModelNodePath));
-			return MutableThis->CachedModelNode.Get();
-		}
-
-#endif
-
-		if (URigVMGraph* Model = GetModel())
-		{
-			MutableThis->CachedModelNode = TWeakObjectPtr<URigVMNode>(Model->FindNode(ModelNodePath));
-			return MutableThis->CachedModelNode.Get();
-		}
-	}
-
-	return nullptr;
-}
-
-FName UControlRigGraphNode::GetModelNodeName() const
-{
-	if (URigVMNode* ModelNode = GetModelNode())
-	{
-		return ModelNode->GetFName();
-	}
-	return NAME_None;
-}
-
-URigVMPin* UControlRigGraphNode::GetModelPinFromPinPath(const FString& InPinPath) const
-{
-	if (TWeakObjectPtr<URigVMPin> const* CachedModelPinPtr = PinPathToModelPin.Find(InPinPath))
-	{
-		if(CachedModelPinPtr->IsValid())
-		{
-			URigVMPin* CachedModelPin = CachedModelPinPtr->Get();
-			if (!CachedModelPin->HasAnyFlags(RF_Transient) && CachedModelPin->GetNode())
-			{
-				return CachedModelPin;
-			}
-		}
-	}
-
-	if (const URigVMNode* ModelNode = GetModelNode())
-	{
-		const FString PinPath = InPinPath.RightChop(ModelNode->GetNodePath().Len() + 1);
-		URigVMPin* ModelPin = ModelNode->FindPin(PinPath);
-		if (ModelPin)
-		{
-			UControlRigGraphNode* MutableThis = (UControlRigGraphNode*)this;
-			MutableThis->PinPathToModelPin.FindOrAdd(InPinPath) = ModelPin;
-		}
-		return ModelPin;
-	}
-	
-	return nullptr;
-}
-
-void UControlRigGraphNode::HandleAddAggregateElement(const FString& InNodePath)
-{
-	DECLARE_SCOPE_HIERARCHICAL_COUNTER_FUNC()
-
-	if (URigVMController* Controller = GetController())
-	{
-		Controller->AddAggregatePin(InNodePath, FString(), FString(), true, true);
-	}	
-}
-
-void UControlRigGraphNode::SetupPinDefaultsFromModel(UEdGraphPin* Pin, const URigVMPin* InModelPin)
-{
-	DECLARE_SCOPE_HIERARCHICAL_COUNTER_FUNC()
-
-	if (Pin->Direction != EGPD_Input)
-	{
-		return;
-	}
-
-	if(InModelPin == nullptr)
-	{
-		InModelPin = FindModelPinFromGraphPin(Pin);
-	}
-
-	// remove stale subpins
-	Pin->SubPins.Remove(nullptr);
-
-	const UEdGraphSchema_K2* K2Schema = GetDefault<UEdGraphSchema_K2>();
-	if (InModelPin && IsValid(InModelPin))
-	{
-		if (InModelPin->GetSubPins().Num() > 0)
-		{
-			return;
-		}
-
-		FString DefaultValueString = InModelPin->GetDefaultValue();
-		if (DefaultValueString.IsEmpty() && InModelPin->GetCPPType() == TEXT("FName"))
-		{
-			DefaultValueString = FName(NAME_None).ToString();
-		}
-		K2Schema->GetPinDefaultValuesFromString(Pin->PinType, Pin->GetOwningNodeUnchecked(), DefaultValueString, Pin->DefaultValue, Pin->DefaultObject, Pin->DefaultTextValue);
-	}
-}
-
-FText UControlRigGraphNode::GetTooltipText() const
-{
-	if(URigVMNode* ModelNode = GetModelNode())
-	{
-		return ModelNode->GetToolTipText();
-	}
-	return FText::FromString(ModelNodePath);
-}
-
-void UControlRigGraphNode::InvalidateNodeTitle() const
-{
-	NodeTitle = FText();
-	FullNodeTitle = FText();
-	NodeTitleDirtied.Broadcast();
-}
-
-bool UControlRigGraphNode::CanCreateUnderSpecifiedSchema(const UEdGraphSchema* InSchema) const
-{
-	return InSchema->IsA<UControlRigGraphSchema>();
-}
-
-void UControlRigGraphNode::AutowireNewNode(UEdGraphPin* FromPin)
-{
-	DECLARE_SCOPE_HIERARCHICAL_COUNTER_FUNC()
-
-	Super::AutowireNewNode(FromPin);
-
-	const UControlRigGraphSchema* Schema = GetDefault<UControlRigGraphSchema>();
-	Schema->AutowireNewNode(this, FromPin);
-}
-
-bool UControlRigGraphNode::IsSelectedInEditor() const
-{
-	URigVMNode* ModelNode = GetModelNode();
-	if (ModelNode)
-	{
-		return ModelNode->IsSelected();
-	}
-	return false;
-}
-
-bool UControlRigGraphNode::ShouldDrawNodeAsControlPointOnly(int32& OutInputPinIndex, int32& OutOutputPinIndex) const
-{
-	if (Cast<URigVMRerouteNode>(GetModelNode()))
-	{
-		OutInputPinIndex = 0;
-		OutOutputPinIndex = Pins.Num()/2;
-
-		if (Pins[OutOutputPinIndex]->Direction != EGPD_Output || Pins[OutOutputPinIndex]->ParentPin != nullptr)
-		{
-			for (int32 i=0; i<Pins.Num(); ++i)
-			{
-				if (Pins[i]->Direction == EGPD_Output && Pins[i]->ParentPin == nullptr)
-				{
-					OutOutputPinIndex = i;
-					break;
-				}
-			}
-		}
-		return true;
-	}
-	return false;
-}
-
-void UControlRigGraphNode::BeginDestroy()
-{
-	for(UEdGraphPin* Pin : Pins)
-	{
-		Pin->SubPins.Remove(nullptr);
-	}
-	Super::BeginDestroy();
-}
-
-FEdGraphPinType UControlRigGraphNode::GetPinTypeForModelPin(const URigVMPin* InModelPin)
-{
-	FEdGraphPinType PinType = RigVMTypeUtils::PinTypeFromCPPType(*InModelPin->GetCPPType(), InModelPin->GetCPPTypeObject());
-	PinType.bIsConst = InModelPin->IsDefinedAsConstant();
-	return PinType;
-}
-
-void UControlRigGraphNode::ConfigurePin(UEdGraphPin* EdGraphPin, const URigVMPin* ModelPin)
-{
-	const bool bConnectable =
-		ModelPin->GetDirection() == ERigVMPinDirection::Input ||
-		ModelPin->GetDirection() == ERigVMPinDirection::Output || 
-		ModelPin->GetDirection() == ERigVMPinDirection::IO;
-
-	// hide sub-pins on a compacted (knot) reroute node
-	bool bHidden = ModelPin->GetDirection() == ERigVMPinDirection::Hidden;
-	if(!bHidden && !ModelPin->IsRootPin())
-	{
-		if(Cast<URigVMRerouteNode>(ModelPin->GetNode()))
-		{
-			bHidden = true;
-		}
-	}
-
-	EdGraphPin->bHidden = bHidden;
-	EdGraphPin->PinFriendlyName = FText::FromName(ModelPin->GetDisplayName());
-	EdGraphPin->bNotConnectable = !bConnectable  || ModelPin->IsFixedSizeArray();
-	EdGraphPin->bOrphanedPin = ModelPin->IsOrphanPin() ? 1 : 0; 
-	EdGraphPin->bDisplayAsMutableRef = ModelPin->IsWildCard();
-}
-
-TArray<URigVMPin*>& UControlRigGraphNode::PinListForPin(const URigVMPin* InModelPin)
-{
-	if(IsValid(InModelPin))
-	{
-		if(InModelPin->IsExecuteContext() && InModelPin->GetDirection() == ERigVMPinDirection::IO)
-		{
-			return ExecutePins;
-		}
-		if(InModelPin->GetDirection() == ERigVMPinDirection::Input || InModelPin->GetDirection() == ERigVMPinDirection::Visible)
-		{
-			return InputPins;
-		}
-		if(InModelPin->GetDirection() == ERigVMPinDirection::IO)
-		{
-			return InputOutputPins;
-		}
-		if(InModelPin->GetDirection() == ERigVMPinDirection::Output)
-		{
-			return OutputPins;
-		}
-
-		checkNoEntry();
-	}
-	static TArray<URigVMPin*> EmptyList;
-	return EmptyList;
-}
-
 #undef LOCTEXT_NAMESPACE
 
