RigVM: First steps towards rebasing ControlRigBlueprint

diff --git a/Engine/Plugins/Animation/ControlRig/Source/ControlRigDeveloper/Private/Graph/ControlRigGraphSchema.cpp b/Engine/Plugins/Animation/ControlRig/Source/ControlRigDeveloper/Private/Graph/ControlRigGraphSchema.cpp
--- a/Engine/Plugins/Animation/ControlRig/Source/ControlRigDeveloper/Private/Graph/ControlRigGraphSchema.cpp
+++ b/Engine/Plugins/Animation/ControlRig/Source/ControlRigDeveloper/Private/Graph/ControlRigGraphSchema.cpp
@@ -2,2434 +2,53 @@
 
 #include "Graph/ControlRigGraphSchema.h"
 #include "Graph/ControlRigGraph.h"
-#include "Graph/ControlRigGraphNode.h"
-#include "IControlRigEditorModule.h"
-#include "UObject/UObjectIterator.h"
-#include "Units/RigUnit.h"
-#include "Kismet2/BlueprintEditorUtils.h"
-#include "EdGraphNode_Comment.h"
-#include "EdGraphSchema_K2_Actions.h"
-#include "ScopedTransaction.h"
-#include "Framework/MultiBox/MultiBoxBuilder.h"
-#include "GraphEditorActions.h"
-#include "ControlRig.h"
-#include "ControlRigBlueprint.h"
-#include "RigVMBlueprintGeneratedClass.h"
-#include "ControlRigDeveloper.h"
-#include "Widgets/Notifications/SNotificationList.h"
-#include "Framework/Application/SlateApplication.h"
-#include "Framework/Notifications/NotificationManager.h"
-#include "EulerTransform.h"
-#include "RigVMStringUtils.h"
-#include "Curves/CurveFloat.h"
-#include "RigVMModel/Nodes/RigVMLibraryNode.h"
-#include "RigVMModel/Nodes/RigVMCollapseNode.h"
-#include "RigVMModel/Nodes/RigVMFunctionEntryNode.h"
-#include "RigVMModel/Nodes/RigVMFunctionReturnNode.h"
-#include "RigVMModel/RigVMVariableDescription.h"
-#include "RigVMCore/RigVMUnknownType.h"
-#include "Kismet2/Kismet2NameValidators.h"
-#include "Algo/Count.h"
-#include "RigVMFunctions/RigVMDispatch_Array.h"
-#include "RigVMFunctions/RigVMDispatch_Constant.h"
-#include "RigVMFunctions/RigVMDispatch_MakeStruct.h"
-#include "RigVMModel/Nodes/RigVMDispatchNode.h"
 #include "Units/Execution/RigUnit_BeginExecution.h"
 #include "Units/Execution/RigUnit_PrepareForExecution.h"
 #include "Units/Execution/RigUnit_InverseExecution.h"
 #include "Units/Execution/RigUnit_InteractionExecution.h"
 
-#include UE_INLINE_GENERATED_CPP_BY_NAME(ControlRigGraphSchema)
-
-#if WITH_EDITOR
-#include "Editor/SControlRigFunctionLocalizationWidget.h"
-#include "Misc/MessageDialog.h"
-#include "Editor/Transactor.h"
-#endif
-
-#define LOCTEXT_NAMESPACE "ControlRigGraphSchema"
-
-const FName UControlRigGraphSchema::GraphName_ControlRig(TEXT("Rig Graph"));
-
-FControlRigLocalVariableNameValidator::FControlRigLocalVariableNameValidator(const UBlueprint* Blueprint, const URigVMGraph* Graph, FName InExistingName)
-	: FStringSetNameValidator(InExistingName.ToString())
-{
-	if (Blueprint)
-	{
-		TSet<FName> NamesTemp;
-		// We allow local variables with same name as blueprint variable
-		
-		FBlueprintEditorUtils::GetFunctionNameList(Blueprint, NamesTemp);
-		FBlueprintEditorUtils::GetAllGraphNames(Blueprint, NamesTemp);
-		FBlueprintEditorUtils::GetSCSVariableNameList(Blueprint, NamesTemp);
-		FBlueprintEditorUtils::GetImplementingBlueprintsFunctionNameList(Blueprint, NamesTemp);
-
-		for (FName & Name : NamesTemp)
-		{
-			Names.Add(Name.ToString());
-		}
-	}
-
-	if (Graph)
-	{
-		for (const FRigVMGraphVariableDescription& LocalVariable : Graph->GetLocalVariables())
-		{
-			Names.Add(LocalVariable.Name.ToString());
-		}
-
-		for (const FRigVMGraphVariableDescription& InputArgument : Graph->GetInputArguments())
-		{
-			Names.Add(InputArgument.Name.ToString());
-		}
-
-		for (const FRigVMGraphVariableDescription& OutputArgument : Graph->GetOutputArguments())
-		{
-			Names.Add(OutputArgument.Name.ToString());
-		}
-	}
-}
-
-EValidatorResult FControlRigLocalVariableNameValidator::IsValid(const FString& Name, bool bOriginal)
-{
-	EValidatorResult Result = FStringSetNameValidator::IsValid(Name, bOriginal);
-	if (Result == EValidatorResult::Ok)
-	{
-		const URigVMSchema* RigSchema = Cast<URigVMSchema>(URigVMSchema::StaticClass()->GetDefaultObject(true));
-		if (RigSchema->GetSanitizedName(Name, false, true) == Name)
-		{
-			return Result;
-		}
-
-		return EValidatorResult::ContainsInvalidCharacters;
-	}
-	return Result;
-}
-
-EValidatorResult FControlRigLocalVariableNameValidator::IsValid(const FName& Name, bool bOriginal)
-{
-	return IsValid(Name.ToString(), bOriginal);
-}
-
-FControlRigNameValidator::FControlRigNameValidator(const UBlueprint* Blueprint, const UStruct* ValidationScope, FName InExistingName)
-: FStringSetNameValidator(InExistingName.ToString())
-{
-	if (Blueprint)
-	{
-		TSet<FName> NamesTemp;
-		FBlueprintEditorUtils::GetClassVariableList(Blueprint, NamesTemp, true);
-		FBlueprintEditorUtils::GetFunctionNameList(Blueprint, NamesTemp);
-		FBlueprintEditorUtils::GetAllGraphNames(Blueprint, NamesTemp);
-		FBlueprintEditorUtils::GetSCSVariableNameList(Blueprint, NamesTemp);
-		FBlueprintEditorUtils::GetImplementingBlueprintsFunctionNameList(Blueprint, NamesTemp);
-
-		for (FName & Name : NamesTemp)
-		{
-			Names.Add(Name.ToString());
-		}
-	}
-}
-
-EValidatorResult FControlRigNameValidator::IsValid(const FString& Name, bool bOriginal)
-{
-	EValidatorResult Result = FStringSetNameValidator::IsValid(Name, bOriginal);
-	if (Result == EValidatorResult::Ok)
-	{
-		const URigVMSchema* RigSchema = Cast<URigVMSchema>(URigVMSchema::StaticClass()->GetDefaultObject(true));
-		if (RigSchema->GetSanitizedName(Name, false, true) == Name)
-		{
-			return Result;
-		}
-
-		return EValidatorResult::ContainsInvalidCharacters;
-	}
-	return Result;
-}
-
-EValidatorResult FControlRigNameValidator::IsValid(const FName& Name, bool bOriginal)
-{
-	return IsValid(Name.ToString(), bOriginal);
-}
-
-FEdGraphPinType FControlRigGraphSchemaAction_LocalVar::GetPinType() const
-{
-	if (UControlRigGraph* Graph = Cast<UControlRigGraph>(GetVariableScope()))
-	{
-		for (FRigVMGraphVariableDescription Variable : Graph->GetModel()->GetLocalVariables())
-		{
-			if (Variable.Name == GetVariableName())
-			{
-				return Variable.ToPinType();
-			}
-		}
-
-		for (FRigVMGraphVariableDescription Variable : Graph->GetModel()->GetInputArguments())
-		{
-			if (Variable.Name == GetVariableName())
-			{
-				return Variable.ToPinType();
-			}
-		}
-	}
-
-	return FEdGraphPinType();
-}
-
-void FControlRigGraphSchemaAction_LocalVar::ChangeVariableType(const FEdGraphPinType& NewPinType)
-{
-	if (UControlRigGraph* Graph = Cast<UControlRigGraph>(GetVariableScope()))
-	{
-		FString NewCPPType;
-		UObject* NewCPPTypeObject = nullptr;
-		RigVMTypeUtils::CPPTypeFromPinType(NewPinType, NewCPPType, &NewCPPTypeObject);
-		Graph->GetController()->SetLocalVariableType(GetVariableName(), NewCPPType, NewCPPTypeObject, true, true);
-	}
-}
-
-void FControlRigGraphSchemaAction_LocalVar::RenameVariable(const FName& NewName)
-{
-	const FName OldName = GetVariableName();
-	if (OldName == NewName)
-	{
-		return;
-	}
-	
-	if (UControlRigGraph* Graph = Cast<UControlRigGraph>(GetVariableScope()))
-	{
-		if (Graph->GetController()->RenameLocalVariable(OldName, NewName, true, true))
-		{
-			SetVariableInfo(NewName, GetVariableScope(), GetPinType().PinCategory == TEXT("bool"));		
-		}
-	}	
-}
-
-bool FControlRigGraphSchemaAction_LocalVar::IsValidName(const FName& NewName, FText& OutErrorMessage) const
-{
-	if (UControlRigGraph* ControlRigGraph = Cast<UControlRigGraph>(GetVariableScope()))
-	{
-		FControlRigLocalVariableNameValidator NameValidator(ControlRigGraph->GetBlueprint(), ControlRigGraph->GetModel(), GetVariableName());
-		EValidatorResult Result = NameValidator.IsValid(NewName.ToString(), false);
-		if (Result != EValidatorResult::Ok && Result != EValidatorResult::ExistingName)
-		{
-			OutErrorMessage = FText::FromString(TEXT("Name with invalid format"));
-			return false;
-		}
-	}
-	return true;
-}
-
-void FControlRigGraphSchemaAction_LocalVar::DeleteVariable() 
-{
-	if (UControlRigGraph* Graph = Cast<UControlRigGraph>(GetVariableScope()))
-	{
-#if WITH_EDITOR
-
-		if (GEditor)
-		{
-			GEditor->CancelTransaction(0);
-		}
-	
-#endif
-		
-		Graph->GetController()->RemoveLocalVariable(GetVariableName(), true, true);
-	}
-}
-
-bool FControlRigGraphSchemaAction_LocalVar::IsVariableUsed()
-{
-	if (UControlRigGraph* ControlRigGraph = Cast<UControlRigGraph>(GetVariableScope()))
-	{
-		const FString VarNameStr = GetVariableName().ToString();
-		for (URigVMNode* Node : ControlRigGraph->GetModel()->GetNodes())
-		{
-			if (URigVMVariableNode* VarNode = Cast<URigVMVariableNode>(Node))
-			{
-				if (VarNode->FindPin(TEXT("Variable"))->GetDefaultValue() == VarNameStr)
-				{
-					return true;		
-				}
-			}
-		}
-	}
-	return false;
-}
-
-FControlRigGraphSchemaAction_PromoteToVariable::FControlRigGraphSchemaAction_PromoteToVariable(UEdGraphPin* InEdGraphPin, bool InLocalVariable)
-: FEdGraphSchemaAction(	FText(), 
-						InLocalVariable ? LOCTEXT("PromoteToLocalVariable", "Promote to local variable") : LOCTEXT("PromoteToVariable", "Promote to variable"),
-						InLocalVariable ? LOCTEXT("PromoteToLocalVariable", "Promote to local variable") : LOCTEXT("PromoteToVariable", "Promote to variable"),
-						1)
-, EdGraphPin(InEdGraphPin)
-, bLocalVariable(InLocalVariable)
-{
-}
-
-UEdGraphNode* FControlRigGraphSchemaAction_PromoteToVariable::PerformAction(UEdGraph* ParentGraph, UEdGraphPin* FromPin,
-	const FVector2D Location, bool bSelectNewNode)
-{
-	UControlRigGraph* RigGraph = Cast<UControlRigGraph>(ParentGraph);
-	if(RigGraph == nullptr)
-	{
-		return nullptr;
-	}
-
-	UControlRigBlueprint* Blueprint = RigGraph->GetBlueprint();
-	URigVMGraph* Model = RigGraph->GetModel();
-	URigVMController* Controller = RigGraph->GetController();
-	if((Blueprint == nullptr) ||
-		(Model == nullptr) ||
-		(Controller == nullptr))
-	{
-		return nullptr;
-	}
-	
-	URigVMPin* ModelPin = Model->FindPin(FromPin->GetName());
-
-	FName VariableName(NAME_None);
-
-	const FScopedTransaction Transaction(
-		bLocalVariable ?
-		LOCTEXT("GraphEd_PromoteToLocalVariable", "Promote Pin To Local Variable") :
-		LOCTEXT("GraphEd_PromoteToVariable", "Promote Pin To Variable"));
-
-	if(bLocalVariable)
-	{
-#if WITH_EDITOR
-
-		if (GEditor)
-		{
-			GEditor->CancelTransaction(0);
-		}
-	
-#endif
-		const FRigVMGraphVariableDescription VariableDescription = Controller->AddLocalVariable(
-			*ModelPin->GetPinPath(),
-			ModelPin->GetCPPType(),
-			ModelPin->GetCPPTypeObject(),
-			ModelPin->GetDefaultValue(),
-			true,
-			true
-		);
-
-		VariableName = VariableDescription.Name;
-	}
-	else
-	{
-		Blueprint->Modify();
-
-		FString DefaultValue = ModelPin->GetDefaultValue();
-		if(!DefaultValue.IsEmpty())
-		{
-			if(UScriptStruct* ScriptStruct = Cast<UScriptStruct>(ModelPin->GetCPPTypeObject()))
-			{
-				if(ScriptStruct == TBaseStructure<FVector2D>::Get())
-				{
-					FVector2D Value = FVector2D::ZeroVector;
-					ScriptStruct->ImportText(*DefaultValue, &Value, nullptr, PPF_None, nullptr, FString());
-					DefaultValue = Value.ToString();
-				}
-				if(ScriptStruct == TBaseStructure<FVector>::Get())
-				{
-					FVector Value = FVector::ZeroVector;
-					ScriptStruct->ImportText(*DefaultValue, &Value, nullptr, PPF_None, nullptr, FString());
-					DefaultValue = Value.ToString();
-				}
-				if(ScriptStruct == TBaseStructure<FQuat>::Get())
-				{
-					FQuat Value = FQuat::Identity;
-					ScriptStruct->ImportText(*DefaultValue, &Value, nullptr, PPF_None, nullptr, FString());
-					DefaultValue = Value.ToString();
-				}
-				if(ScriptStruct == TBaseStructure<FRotator>::Get())
-				{
-					FRotator Value = FRotator::ZeroRotator;
-					ScriptStruct->ImportText(*DefaultValue, &Value, nullptr, PPF_None, nullptr, FString());
-					DefaultValue = Value.ToString();
-				}
-				if(ScriptStruct == TBaseStructure<FTransform>::Get())
-				{
-					FTransform Value = FTransform::Identity;
-					ScriptStruct->ImportText(*DefaultValue, &Value, nullptr, PPF_None, nullptr, FString());
-					DefaultValue = Value.ToString();
-				}
-			}
-		}
-
-		FRigVMExternalVariable ExternalVariable;
-		ExternalVariable.Name = FromPin->GetFName();
-		ExternalVariable.bIsArray = ModelPin->IsArray();
-		ExternalVariable.TypeName = ModelPin->IsArray() ? *ModelPin->GetArrayElementCppType() : *ModelPin->GetCPPType();
-		ExternalVariable.TypeObject = ModelPin->GetCPPTypeObject();
-		
-		VariableName = Blueprint->AddCRMemberVariableFromExternal(
-			ExternalVariable,
-			DefaultValue
-		);
-	}
-
-	if(!VariableName.IsNone())
-	{
-		URigVMNode* ModelNode = Controller->AddVariableNode(
-			VariableName,
-			ModelPin->GetCPPType(),
-			ModelPin->GetCPPTypeObject(),
-			FromPin->Direction == EGPD_Input,
-			ModelPin->GetDefaultValue(),
-			Location,
-			FString(),
-			true,
-			true
-		);
-
-		if(ModelNode)
-		{
-			if(FromPin->Direction == EGPD_Input)
-			{
-				Controller->AddLink(ModelNode->FindPin(TEXT("Value")), ModelPin, true);
-			}
-			else
-			{
-				Controller->AddLink(ModelPin, ModelNode->FindPin(TEXT("Value")), true);
-			}
-			return RigGraph->FindNodeForModelNodeName(ModelNode->GetFName());
-		}
-	}
-	
-	return nullptr;
-}
-
-FControlRigGraphSchemaAction_PromoteToExposedPin::FControlRigGraphSchemaAction_PromoteToExposedPin(UEdGraphPin* InEdGraphPin)
-: FEdGraphSchemaAction(	FText(), 
-						LOCTEXT("PromoteToExposedPin", "Promote to exposed pin"),
-						LOCTEXT("PromoteToExposedPin", "Promote to exposed pin"),
-						1)
-, EdGraphPin(InEdGraphPin)
-{
-}
-
-UEdGraphNode* FControlRigGraphSchemaAction_PromoteToExposedPin::PerformAction(UEdGraph* ParentGraph, UEdGraphPin* FromPin, const FVector2D Location, bool bSelectNewNode)
-{
-	UControlRigGraph* RigGraph = Cast<UControlRigGraph>(ParentGraph);
-	if(RigGraph == nullptr)
-	{
-		return nullptr;
-	}
-
-	URigVMGraph* Model = RigGraph->GetModel();
-	URigVMController* Controller = RigGraph->GetController();
-	if(Model == nullptr || Controller == nullptr)
-	{
-		return nullptr;
-	}
-	
-	URigVMPin* ModelPin = Model->FindPin(FromPin->GetName());
-	if (ModelPin == nullptr)
-	{
-		return nullptr;
-	}
-
-	const FScopedTransaction Transaction(
-		LOCTEXT("GraphEd_PromoteToExposedPin", "Promote To Exposed Pin"));
-	
-#if WITH_EDITOR
-	if (GEditor)
-	{
-		GEditor->CancelTransaction(0);
-	}	
-#endif
-
-	Controller->OpenUndoBracket(TEXT("Promote to Exposed Pin"));
-
-	const UObject* CPPTypeObject = ModelPin->GetCPPTypeObject();
-	const FString PinName = Controller->AddExposedPin(
-		*ModelPin->GetName(),
-		ModelPin->GetDirection(),
-		ModelPin->GetCPPType(),
-		CPPTypeObject ? (FName)*CPPTypeObject->GetPathName() : NAME_None,
-		ModelPin->GetDefaultValue(),
-		true,
-		true
-	).ToString();
-
-	UEdGraphNode* Result = nullptr;
-	if(!PinName.IsEmpty())
-	{
-		if (ModelPin->GetDirection() == ERigVMPinDirection::Input)
-		{
-			Controller->AddLink(Model->GetEntryNode()->FindPin(PinName), ModelPin, true);
-			Result = RigGraph->FindNodeForModelNodeName(Model->GetEntryNode()->GetFName());
-		}
-		else if(ModelPin->GetDirection() == ERigVMPinDirection::Output)
-		{
-			Controller->AddLink(ModelPin, Model->GetReturnNode()->FindPin(PinName), true);
-			Result = RigGraph->FindNodeForModelNodeName(Model->GetReturnNode()->GetFName());
-		}
-		else if (ModelPin->GetDirection() == ERigVMPinDirection::IO)
-		{
-			URigVMFunctionEntryNode* EntryNode = Model->GetEntryNode();
-			Controller->AddLink(EntryNode->FindPin(PinName), ModelPin, true);
-			Controller->AddLink(ModelPin, Model->GetReturnNode()->FindPin(PinName), true);
-			Result = RigGraph->FindNodeForModelNodeName(EntryNode->GetFName());
-		}		
-	}
-
-	Controller->CloseUndoBracket();
-	
-	return Result;
-}
-
-FControlRigGraphSchemaAction_Event::FControlRigGraphSchemaAction_Event(const FName& InEventName, const FString& InNodePath, const FText& InNodeCategory)
-: FEdGraphSchemaAction(	InNodeCategory, 
-						FText::FromName(InEventName),
-						FText(),
-						1)
-, NodePath(InNodePath)
-{
-}
-
-FReply FControlRigGraphSchemaAction_Event::OnDoubleClick(UBlueprint* InBlueprint)
-{
-	if (UControlRigBlueprint* Blueprint = Cast<UControlRigBlueprint>(InBlueprint))
-	{
-		if(const URigVMNode* ModelNode = Blueprint->GetRigVMClient()->FindNode(NodePath))
-		{
-			Blueprint->OnRequestJumpToHyperlink().Execute(ModelNode);
-			return FReply::Handled();
-		}
-	}
-	return FReply::Unhandled();;
-}
-
-FSlateBrush const* FControlRigGraphSchemaAction_Event::GetPaletteIcon() const
-{
-	static FSlateIcon EventIcon(FAppStyle::GetAppStyleSetName(), "GraphEditor.Event_16x");
-	return EventIcon.GetIcon();
-}
-
-FReply FControlRigFunctionDragDropAction::DroppedOnPanel(const TSharedRef< class SWidget >& Panel, FVector2D ScreenPosition, FVector2D GraphPosition, UEdGraph& Graph)
-{
-	// For local variables
-	if (SourceAction->GetTypeId() == FControlRigGraphSchemaAction_LocalVar::StaticGetTypeId())
-	{
-		if (UControlRigGraph* TargetRigGraph = Cast<UControlRigGraph>(&Graph))
-		{
-			if (TargetRigGraph == SourceRigGraph)
-			{
-				FControlRigGraphSchemaAction_LocalVar* VarAction = (FControlRigGraphSchemaAction_LocalVar*) SourceAction.Get();
-				for (FRigVMGraphVariableDescription LocalVariable : TargetRigGraph->GetModel()->GetLocalVariables())
-				{
-					if (LocalVariable.Name == VarAction->GetVariableName())
-					{
-						URigVMController* Controller = TargetRigGraph->GetController();
-						FMenuBuilder MenuBuilder(true, NULL);
-						const FText VariableNameText = FText::FromName( LocalVariable.Name );
-
-						MenuBuilder.BeginSection("BPVariableDroppedOn", VariableNameText );
-
-						MenuBuilder.AddMenuEntry(
-							FText::Format( LOCTEXT("CreateGetVariable", "Get {0}"), VariableNameText ),
-							FText::Format( LOCTEXT("CreateVariableGetterToolTip", "Create Getter for variable '{0}'\n(Ctrl-drag to automatically create a getter)"), VariableNameText ),
-							FSlateIcon(),
-							FUIAction(
-							FExecuteAction::CreateLambda([Controller, LocalVariable, GraphPosition]()
-							{
-								Controller->AddVariableNode(LocalVariable.Name, LocalVariable.CPPType, LocalVariable.CPPTypeObject, true, LocalVariable.DefaultValue, GraphPosition, FString(), true, true);
-							}),
-							FCanExecuteAction()));
-
-						MenuBuilder.AddMenuEntry(
-							FText::Format( LOCTEXT("CreateSetVariable", "Set {0}"), VariableNameText ),
-							FText::Format( LOCTEXT("CreateVariableSetterToolTip", "Create Setter for variable '{0}'\n(Alt-drag to automatically create a setter)"), VariableNameText ),
-							FSlateIcon(),
-							FUIAction(
-							FExecuteAction::CreateLambda([Controller, LocalVariable, GraphPosition]()
-							{
-								Controller->AddVariableNode(LocalVariable.Name, LocalVariable.CPPType, LocalVariable.CPPTypeObject, false, LocalVariable.DefaultValue, GraphPosition, FString(), true, true);
-							}),
-							FCanExecuteAction()));
-
-						TSharedRef< SWidget > PanelWidget = Panel;
-						// Show dialog to choose getter vs setter
-						FSlateApplication::Get().PushMenu(
-							PanelWidget,
-							FWidgetPath(),
-							MenuBuilder.MakeWidget(),
-							ScreenPosition,
-							FPopupTransitionEffect( FPopupTransitionEffect::ContextMenu)
-							);
-
-						MenuBuilder.EndSection();
-					}	
-				}				
-			}
-		}
-	}
-	// For functions
-	else if (UControlRigGraph* TargetRigGraph = Cast<UControlRigGraph>(&Graph))
-	{
-		if (UControlRigBlueprint* TargetRigBlueprint = Cast<UControlRigBlueprint>(FBlueprintEditorUtils::FindBlueprintForGraph(TargetRigGraph)))
-		{
-			if (URigVMGraph* FunctionDefinitionGraph = SourceRigBlueprint->GetModel(SourceRigGraph))
-			{
-				if (URigVMLibraryNode* FunctionDefinitionNode = Cast<URigVMLibraryNode>(FunctionDefinitionGraph->GetOuter()))
-				{
-					if(URigVMController* TargetController = TargetRigBlueprint->GetController(TargetRigGraph))
-					{
-						if(URigVMFunctionLibrary* FunctionLibrary = Cast<URigVMFunctionLibrary>(FunctionDefinitionNode->GetOuter()))
-						{
-							if(UControlRigBlueprint* FunctionRigBlueprint = Cast<UControlRigBlueprint>(FunctionLibrary->GetOuter()))
-							{
-#if WITH_EDITOR
-								if(FunctionRigBlueprint != TargetRigBlueprint)
-								{
-									if(!FunctionRigBlueprint->IsFunctionPublic(FunctionDefinitionNode->GetFName()))
-									{
-										FRigVMGraphFunctionIdentifier FunctionIdentifier = FunctionDefinitionNode->GetFunctionIdentifier();
-										TargetRigBlueprint->BroadcastRequestLocalizeFunctionDialog(FunctionIdentifier);
-										FunctionDefinitionNode = TargetRigBlueprint->GetLocalFunctionLibrary()->FindPreviouslyLocalizedFunction(FunctionIdentifier);
-									}
-								}
-#endif
-								TargetController->AddFunctionReferenceNode(FunctionDefinitionNode, GraphPosition, FString(), true, true);
-							}
-						}
-					}
-				}
-			}
-		}
-	}
-
-	
-	
-	return FReply::Unhandled();
-}
-
-FReply FControlRigFunctionDragDropAction::DroppedOnPin(FVector2D ScreenPosition, FVector2D GraphPosition)
-{
-	return FReply::Unhandled();
-}
-
-FReply FControlRigFunctionDragDropAction::DroppedOnAction(TSharedRef<FEdGraphSchemaAction> Action)
-{
-	return FReply::Unhandled();
-}
-
-FReply FControlRigFunctionDragDropAction::DroppedOnCategory(FText Category)
-{
-	// todo
-	/*
-	if (SourceAction.IsValid())
-	{
-		SourceAction->MovePersistentItemToCategory(Category);
-	}
-	*/
-	return FReply::Unhandled();
-}
-
-void FControlRigFunctionDragDropAction::HoverTargetChanged()
-{
-	// todo - see FMyBlueprintItemDragDropAction
-	FGraphSchemaActionDragDropAction::HoverTargetChanged();
-
-	// check for category + graph, everything else we won't allow for now.
-
-	bDropTargetValid = true;
-}
-
-FControlRigFunctionDragDropAction::FControlRigFunctionDragDropAction()
-	: FGraphSchemaActionDragDropAction()
-	, SourceRigBlueprint(nullptr)
-	, SourceRigGraph(nullptr)
-	, bControlDrag(false)
-	, bAltDrag(false)
-{
-}
-
-TSharedRef<FControlRigFunctionDragDropAction> FControlRigFunctionDragDropAction::New(TSharedPtr<FEdGraphSchemaAction> InAction, UControlRigBlueprint* InRigBlueprint, UControlRigGraph* InRigGraph)
-{
-	TSharedRef<FControlRigFunctionDragDropAction> Action = MakeShareable(new FControlRigFunctionDragDropAction);
-	Action->SourceAction = InAction;
-	Action->SourceRigBlueprint = InRigBlueprint;
-	Action->SourceRigGraph = InRigGraph;
-	Action->Construct();
-	return Action;
-}
-
-UControlRigGraphSchema::UControlRigGraphSchema()
-{
-}
-
-void UControlRigGraphSchema::GetGraphContextActions(FGraphContextMenuBuilder& ContextMenuBuilder) const
-{
-
-}
-
-void UControlRigGraphSchema::GetContextMenuActions(class UToolMenu* Menu, class UGraphNodeContextMenuContext* Context) const
-{
-	/*
-	// this seems to be taken care of by ControlRigGraphNode
-#if WITH_EDITOR
-	return IControlRigEditorModule::Get().GetContextMenuActions(this, Menu, Context);
-#else
-	check(0);
-#endif
-	*/
-}
-
-bool UControlRigGraphSchema::TryCreateConnection(UEdGraphPin* PinA, UEdGraphPin* PinB) const
-{
-#if WITH_EDITOR
-
-	if (GEditor)
-	{
-		GEditor->CancelTransaction(0);
-	}
-	
-#endif
-
-	if (PinA == PinB)
-	{
-		return false;
-	}
-
-	if (PinA->GetOwningNode() == PinB->GetOwningNode())
-	{
-		return false;
-	}
-
-	UControlRigGraphSchema* MutableThis = (UControlRigGraphSchema*)this;
-
-	IRigVMClientHost* Host = PinA->GetOwningNode()->GetImplementingOuter<IRigVMClientHost>();
-	if (Host != nullptr)
-	{
-		if (URigVMController* Controller = Host->GetRigVMClient()->GetOrCreateController(PinA->GetOwningNode()->GetGraph()))
-		{
-			ERigVMPinDirection UserLinkDirection = ERigVMPinDirection::Output;
-			if (PinA->Direction == EGPD_Input)
-			{
-				Swap(PinA, PinB);
-				UserLinkDirection = ERigVMPinDirection::Input;
-			}
-
-			const bool bPinAIsWildCard = PinA->PinType.PinSubCategoryObject == RigVMTypeUtils::GetWildCardCPPTypeObject();
-			const bool bPinBIsWildCard = PinB->PinType.PinSubCategoryObject == RigVMTypeUtils::GetWildCardCPPTypeObject();
-			if(bPinAIsWildCard != bPinBIsWildCard)
-			{
-				// switch the user link direction if only one of the pins is a wildcard 
-				if(UserLinkDirection == ERigVMPinDirection::Input && bPinBIsWildCard)
-				{
-					UserLinkDirection = ERigVMPinDirection::Output;
-				}
-				else if(UserLinkDirection == ERigVMPinDirection::Output && bPinAIsWildCard)
-				{
-					UserLinkDirection = ERigVMPinDirection::Input;
-				}
-			}
-
-#if WITH_EDITOR
-
-			static const FText LoopMessage = LOCTEXT("LinkingLoopNotRecommended", "Linking a function return within a loop is not recommended.\nAre you sure?");
-			
-			// check if we are trying to connect a loop iteration pin to a return
-			if(URigVMGraph* Graph = Controller->GetGraph())
-			{
-				if(URigVMPin* TargetPin = Graph->FindPin(PinB->GetName()))
-				{
-					if(TargetPin->IsExecuteContext() && TargetPin->GetNode()->IsA<URigVMFunctionReturnNode>())
-					{
-						bool bIsInLoopIteration = false;
-						if(URigVMPin* SourcePin = Graph->FindPin(PinA->GetName()))
-						{
-							while(SourcePin)
-							{
-								if(!SourcePin->IsExecuteContext())
-								{
-									break;
-								}
-								URigVMPin* CurrentSourcePin = SourcePin;
-								SourcePin = nullptr;
-								
-								if(URigVMUnitNode* UnitNode = Cast<URigVMUnitNode>(CurrentSourcePin->GetNode()))
-								{
-									TSharedPtr<FStructOnScope> UnitScope = UnitNode->ConstructStructInstance();
-									if(UnitScope.IsValid())
-									{
-										FRigVMStruct* Unit = (FRigVMStruct*)UnitScope->GetStructMemory();
-										if(Unit->IsForLoop())
-										{
-											if(CurrentSourcePin->GetFName() != FRigVMStruct::ForLoopCompletedPinName)
-											{
-												bIsInLoopIteration = true;
-												break;
-											}
-										}
-									}
-									
-								}
-
-								for(URigVMPin* PinOnSourceNode : CurrentSourcePin->GetNode()->GetPins())
-								{
-									if(!PinOnSourceNode->IsExecuteContext())
-									{
-										continue;
-									}
-
-									if(PinOnSourceNode->GetDirection() != ERigVMPinDirection::Input &&
-										PinOnSourceNode->GetDirection() != ERigVMPinDirection::IO)
-									{
-										continue;
-									}
-
-									TArray<URigVMPin*> NextSourcePins = PinOnSourceNode->GetLinkedSourcePins();
-									if(NextSourcePins.Num() > 0)
-									{
-										SourcePin = NextSourcePins[0];
-										break;
-									}
-								}
-							}
-						}
-
-						if(bIsInLoopIteration)
-						{
-							const EAppReturnType::Type Answer = FMessageDialog::Open( EAppMsgType::YesNo, LoopMessage );
-							if(Answer == EAppReturnType::No)
-							{
-								return false;
-							}
-						}
-					}
-				}
-			}
-#endif
-
-			const bool bCreateCastNode = !FSlateApplication::Get().GetModifierKeys().IsAltDown();
-			return Controller->AddLink(PinA->GetName(), PinB->GetName(), true, true, UserLinkDirection, bCreateCastNode);
-		}
-	}
-	return false;
-}
-
-static bool HasParentConnection_Recursive(const UEdGraphPin* InPin)
-{
-	if(InPin->ParentPin)
-	{
-		return InPin->ParentPin->LinkedTo.Num() > 0 || HasParentConnection_Recursive(InPin->ParentPin);
-	}
-
-	return false;
-}
-
-static bool HasChildConnection_Recursive(const UEdGraphPin* InPin)
-{
-	for(const UEdGraphPin* SubPin : InPin->SubPins)
-	{
-		if(SubPin->LinkedTo.Num() > 0 || HasChildConnection_Recursive(SubPin))
-		{
-			return true;
-		}
-	}
-
-	return false;
-}
-
-const FPinConnectionResponse UControlRigGraphSchema::CanCreateConnection(const UEdGraphPin* A, const UEdGraphPin* B) const
-{
-	if (A == nullptr || B == nullptr)
-	{
-		return FPinConnectionResponse(CONNECT_RESPONSE_DISALLOW, TEXT("One of the Pins is NULL"));
-	}
-	
-	UControlRigGraphNode* RigNodeA = Cast<UControlRigGraphNode>(A->GetOwningNode());
-	UControlRigGraphNode* RigNodeB = Cast<UControlRigGraphNode>(B->GetOwningNode());
-
-	if (RigNodeA && RigNodeB && RigNodeA != RigNodeB)
-	{
-		URigVMPin* PinA = RigNodeA->GetModelPinFromPinPath(A->GetName());
-		if (PinA)
-		{
-			PinA = PinA->GetPinForLink();
-			RigNodeA->GetModel()->PrepareCycleChecking(PinA, A->Direction == EGPD_Input);
-		}
-
-		URigVMPin* PinB = RigNodeB->GetModelPinFromPinPath(B->GetName());
-		if (PinB)
-		{
-			PinB = PinB->GetPinForLink();
-		}
-
-		ERigVMPinDirection UserLinkDirection = ERigVMPinDirection::Output;
-		if (A->Direction == EGPD_Input)
-		{
-			Swap(PinA, PinB);
-			UserLinkDirection = ERigVMPinDirection::Input;
-		}
-
-		if (!PinA)
-		{
-			return FPinConnectionResponse(ECanCreateConnectionResponse::CONNECT_RESPONSE_DISALLOW, FString::Printf(TEXT("Pin %s not found"), *A->GetName()));
-		}
-
-		if (!PinB)
-		{
-			return FPinConnectionResponse(ECanCreateConnectionResponse::CONNECT_RESPONSE_DISALLOW, FString::Printf(TEXT("Pin %s not found"), *B->GetName()));
-		}
-
-		if(PinA->IsWildCard() != PinB->IsWildCard())
-		{
-			// switch the user link direction if only one of the pins is a wildcard 
-			if(UserLinkDirection == ERigVMPinDirection::Input && PinB->IsWildCard())
-			{
-				UserLinkDirection = ERigVMPinDirection::Output;
-			}
-			else if(UserLinkDirection == ERigVMPinDirection::Output && PinA->IsWildCard())
-			{
-				UserLinkDirection = ERigVMPinDirection::Input;
-			}
-		}
-
-		const FRigVMByteCode* ByteCode = RigNodeA->GetController()->GetCurrentByteCode();
-
-		const bool bEnableTypeCasting = !FSlateApplication::Get().GetModifierKeys().IsAltDown();
-		FString FailureReason;
-		bool bResult = RigNodeA->GetModel()->CanLink(PinA, PinB, &FailureReason, ByteCode, UserLinkDirection, bEnableTypeCasting);
-		if (!bResult)
-		{
-			return FPinConnectionResponse(CONNECT_RESPONSE_DISALLOW, FText::FromString(FailureReason));
-		}
-		return FPinConnectionResponse(CONNECT_RESPONSE_MAKE, LOCTEXT("ConnectResponse_Allowed", "Connect"));
-	}
-
-	return FPinConnectionResponse(CONNECT_RESPONSE_DISALLOW, LOCTEXT("ConnectResponse_Disallowed_Unexpected", "Unexpected error"));
-}
-
-FLinearColor UControlRigGraphSchema::GetPinTypeColor(const FEdGraphPinType& PinType) const
-{
-	const FName& TypeName = PinType.PinCategory;
-	if (TypeName == UEdGraphSchema_K2::PC_Struct)
-	{
-		if (UStruct* Struct = Cast<UStruct>(PinType.PinSubCategoryObject))
-		{
-			if (Struct->IsChildOf(FRigVMExecuteContext::StaticStruct()))
-			{
-				return FLinearColor::White;
-			}
-
-			if (Struct->IsChildOf(FRigVMDecorator::StaticStruct()))
-			{
-				return FLinearColor::White;
-			}
-
-			if (Struct->IsChildOf(RigVMTypeUtils::GetWildCardCPPTypeObject()))
-			{
-				return FLinearColor(FVector3f::OneVector * 0.25f);
-			}
-
-			if (Struct == FRigElementKey::StaticStruct() || Struct == FRigElementKeyCollection::StaticStruct())
-			{
-				return FLinearColor(0.0, 0.6588, 0.9490);
-			}
-
-			if (Struct == FRigElementKey::StaticStruct() || Struct == FRigPose::StaticStruct())
-			{
-				return FLinearColor(0.0, 0.3588, 0.5490);
-			}
-
-			// external types can register their own colors, check if there are any
-			if (IControlRigDeveloperModule* Module = FModuleManager::GetModulePtr<IControlRigDeveloperModule>("ControlRigDeveloper"))
-			{
-				if (const FLinearColor* Color = Module->FindPinTypeColor(Struct))
-				{
-					return *Color;
-				}
-			}
-		}
-	}
-	
-	return GetDefault<UEdGraphSchema_K2>()->GetPinTypeColor(PinType);
-}
-
-void UControlRigGraphSchema::InsertAdditionalActions(TArray<UBlueprint*> InBlueprints, TArray<UEdGraph*> EdGraphs,
-	TArray<UEdGraphPin*> EdGraphPins, FGraphActionListBuilderBase& OutAllActions) const
-{
-	Super::InsertAdditionalActions(InBlueprints, EdGraphs, EdGraphPins, OutAllActions);
-
-	if(EdGraphPins.Num() > 0)
-	{
-		if(UControlRigGraphNode* RigNode = Cast<UControlRigGraphNode>(EdGraphPins[0]->GetOwningNode()))
-		{
-			if(URigVMPin* ModelPin = RigNode->GetModelPinFromPinPath(EdGraphPins[0]->GetName()))
-			{
-				const bool bIsRootGraph = ModelPin->GetGraph()->IsRootGraph();
-				if(!ModelPin->IsExecuteContext() && !ModelPin->IsWildCard())
-				{
-					if(!ModelPin->GetNode()->IsA<URigVMVariableNode>())
-					{
-						OutAllActions.AddAction(TSharedPtr<FControlRigGraphSchemaAction_PromoteToVariable>(
-							new FControlRigGraphSchemaAction_PromoteToVariable(EdGraphPins[0], false)
-						));
-
-						if(!bIsRootGraph && !ModelPin->IsWildCard())
-						{
-							OutAllActions.AddAction(TSharedPtr<FControlRigGraphSchemaAction_PromoteToVariable>(
-								new FControlRigGraphSchemaAction_PromoteToVariable(EdGraphPins[0], true)
-							));
-						}
-					}
-				}
-
-				if (!bIsRootGraph)
-				{
-					if (!ModelPin->GetGraph()->GetRootGraph()->IsA<URigVMFunctionLibrary>() || !ModelPin->IsWildCard())
-					{
-						OutAllActions.AddAction(TSharedPtr<FControlRigGraphSchemaAction_PromoteToExposedPin>(
-								   new FControlRigGraphSchemaAction_PromoteToExposedPin(EdGraphPins[0])
-							   ));
-					}
-				}
-			}
-		}
-	}
-}
-
-TSharedPtr<INameValidatorInterface> UControlRigGraphSchema::GetNameValidator(const UBlueprint* BlueprintObj, const FName& OriginalName, const UStruct* ValidationScope, const FName& ActionTypeId) const
-{
-	if (ActionTypeId == FControlRigGraphSchemaAction_LocalVar::StaticGetTypeId())
-	{
-		if (const UControlRigGraph* ControlRigGraph = Cast<UControlRigGraph>(ValidationScope))
-		{
-			if (const URigVMGraph* Graph = ControlRigGraph->GetModel())
-			{
-				return MakeShareable(new FControlRigLocalVariableNameValidator(BlueprintObj, Graph, OriginalName));
-			}
-		}
-	}
-
-	return MakeShareable(new FControlRigNameValidator(BlueprintObj, ValidationScope, OriginalName));		
-}
-
-bool UControlRigGraphSchema::SupportsPinType(UScriptStruct* ScriptStruct) const
-{
-	if(!ScriptStruct)
-	{
-		return false;
-	}
-
-	for (TFieldIterator<FProperty> It(ScriptStruct); It; ++It)
-	{
-		FName PropertyName = It->GetFName();
-		FProperty* Property = *It;
-
-		if (FArrayProperty* ArrayProperty = CastField<FArrayProperty>(Property))
-		{
-			Property = ArrayProperty->Inner;
-		}
-
-		FString CPPType = Property->GetCPPType();
-		if (CPPType == TEXT("bool") ||
-			CPPType == TEXT("float") ||
-			CPPType == TEXT("double") ||
-			CPPType == TEXT("int32") ||
-			CPPType == TEXT("FString") ||
-			CPPType == TEXT("FName") ||
-			CPPType == TEXT("uint16"))
-		{
-			continue;
-		}		
-		
-		if (FStructProperty* StructProperty = CastField<FStructProperty>(Property))
-		{
-			if (SupportsPinType(StructProperty->Struct))
-			{
-				continue;
-			}
-		}
-		else if (FEnumProperty* EnumProperty = CastField<FEnumProperty>(Property))
-		{
-			continue;
-		}
-		else if (FByteProperty* ByteProperty = CastField<FByteProperty>(Property))
-		{
-			if (ByteProperty->Enum)
-			{
-				continue;
-			}
-		}
-		else if (CastField<FObjectProperty>(Property) && RigVMCore::SupportsUObjects())
-		{
-			continue;
-		}
-		else if (CastField<FInterfaceProperty>(Property) && RigVMCore::SupportsUInterfaces())
-		{
-			continue;
-		}
-	
-		return false;
-	}
-	
-	return true;
-}
-
-bool UControlRigGraphSchema::SupportsPinType(TWeakPtr<const FEdGraphSchemaAction> SchemaAction, const FEdGraphPinType& PinType) const
-{
-	if (PinType.IsContainer())
-	{
-		return false;
-	}
-	
-	const FName TypeName = PinType.PinCategory;
-
-	if (TypeName == UEdGraphSchema_K2::PC_Boolean ||
-		TypeName == UEdGraphSchema_K2::PC_Int ||
-		TypeName == UEdGraphSchema_K2::PC_Real ||
-		TypeName == UEdGraphSchema_K2::PC_Name ||
-		TypeName == UEdGraphSchema_K2::PC_String ||
-		TypeName == UEdGraphSchema_K2::PC_Enum)
-	{
-		return true;
-	}
-
-	if(RigVMCore::SupportsUObjects())
-	{
-		if (PinType.PinCategory == UEdGraphSchema_K2::PC_Object ||
-			PinType.PinCategory == UEdGraphSchema_K2::PC_SoftObject ||
-			PinType.PinCategory == UEdGraphSchema_K2::AllObjectTypes)
-		{
-			if (PinType.PinSubCategoryObject.IsValid())
-			{
-				return PinType.PinSubCategoryObject->IsA<UClass>();
-			}
-		}
-	}
-
-	if (RigVMCore::SupportsUInterfaces())
-	{
-		if (PinType.PinCategory == UEdGraphSchema_K2::PC_Interface)
-		{
-			if (PinType.PinSubCategoryObject.IsValid())
-			{
-				return PinType.PinSubCategoryObject->IsA<UInterface>();
-			}
-		}
-	}
-
-	if (PinType.PinCategory == UEdGraphSchema_K2::PC_Struct)
-	{
-		if (UScriptStruct* ScriptStruct = Cast<UScriptStruct>(PinType.PinSubCategoryObject))
-		{
-			if(SchemaAction.IsValid() && SchemaAction.Pin()->IsAVariable())
-			{
-				if(ScriptStruct->IsChildOf(FRigVMExecuteContext::StaticStruct()))
-				{
-					return false;
-				}
-			}
-			return SupportsPinType(ScriptStruct);
-		}
-		else if (PinType.PinSubCategoryObject == UUserDefinedStruct::StaticClass())
-		{
-			// if a user defined struct hasn't been loaded yet,
-			// its PinSubCategoryObject equals UUserDefinedStruct::StaticClass()
-			// and since it is not practical to load every user defined struct to check if they only contain supported types,
-			// we always return true so that they at least show up in the drop down menu
-			// if they contain members of unsupported type, the spawned node will generate error
-			return true;
-		}
-	}
-
-	if (PinType.PinCategory == UEdGraphSchema_K2::PC_Byte)
-	{
-		if (PinType.PinSubCategoryObject.IsValid())
-		{
-			return PinType.PinSubCategoryObject->IsA<UEnum>();
-		}
-	}
-
-	return false;
-}
-
-bool UControlRigGraphSchema::SupportsPinTypeContainer(TWeakPtr<const FEdGraphSchemaAction> SchemaAction,
-	const FEdGraphPinType& PinType, const EPinContainerType& ContainerType) const
-{
-	// Do not allow containers for execute context type
-	if(const UScriptStruct* ExecuteContextScriptStruct = Cast<UScriptStruct>(PinType.PinSubCategoryObject))
-	{
-		if (ExecuteContextScriptStruct->IsChildOf(FRigVMExecuteContext::StaticStruct()))
-		{
-			return ContainerType == EPinContainerType::None;
-		}
-	}
-	
-	return ContainerType == EPinContainerType::None || ContainerType == EPinContainerType::Array;
-}
-
-void UControlRigGraphSchema::BreakPinLinks(UEdGraphPin& TargetPin, bool bSendsNodeNotifcation) const
-{
-	//const FScopedTransaction Transaction( LOCTEXT("GraphEd_BreakPinLinks", "Break Pin Links") );
-
-	// cache this here, as BreakPinLinks can trigger a node reconstruction invalidating the TargetPin referenceS
-	if (UControlRigGraphNode* Node = Cast< UControlRigGraphNode>(TargetPin.GetOwningNode()))
-	{
-		Node->GetController()->BreakAllLinks(TargetPin.GetName(), TargetPin.Direction == EGPD_Input, true, true);
-	}
-}
-
-void UControlRigGraphSchema::BreakSinglePinLink(UEdGraphPin* SourcePin, UEdGraphPin* TargetPin) const
-{
-	//const FScopedTransaction Transaction(LOCTEXT("GraphEd_BreakSinglePinLink", "Break Pin Link") );
-
-	if (UControlRigGraphNode* Node = Cast< UControlRigGraphNode>(TargetPin->GetOwningNode()))
-	{
-		if (SourcePin->Direction == EGPD_Input)
-		{
-			UEdGraphPin* Temp = TargetPin;
-			TargetPin = SourcePin;
-			SourcePin = Temp;
-		}
-		
-		Node->GetController()->BreakLink(SourcePin->GetName(), TargetPin->GetName(), true, true);
-	}
-}
-
-bool UControlRigGraphSchema::CanGraphBeDropped(TSharedPtr<FEdGraphSchemaAction> InAction) const
-{
-	if (!InAction.IsValid())
-	{
-		return false;
-	}
-
-	if (InAction->GetTypeId() == FEdGraphSchemaAction_K2Graph::StaticGetTypeId())
-	{
-		FEdGraphSchemaAction_K2Graph* FuncAction = (FEdGraphSchemaAction_K2Graph*)InAction.Get();
-		if (UControlRigGraph* RigGraph = Cast<UControlRigGraph>((UEdGraph*)FuncAction->EdGraph))
-		{
-			return true;
-		}
-	}
-	else if (InAction->GetTypeId() == FControlRigGraphSchemaAction_LocalVar::StaticGetTypeId())
-	{
-		FControlRigGraphSchemaAction_LocalVar* VarAction = (FControlRigGraphSchemaAction_LocalVar*)InAction.Get();
-		if (UControlRigGraph* RigGraph = Cast<UControlRigGraph>((UEdGraph*)VarAction->GetVariableScope()))
-		{
-			return true;
-		}
-	}
-	
-	return false;
-}
-
-FString UControlRigGraphSchema::GetFindReferenceSearchTerm(const FEdGraphSchemaAction* InGraphAction) const
-{
-	if(InGraphAction)
-	{
-		if(InGraphAction->GetTypeId() == FEdGraphSchemaAction_K2Var::StaticGetTypeId())
-		{
-			const FEdGraphSchemaAction_K2Var* VarAction = (const FEdGraphSchemaAction_K2Var*)InGraphAction;
-			return VarAction->GetVariableName().ToString();
-		}
-		else if (InGraphAction->GetTypeId() == FControlRigGraphSchemaAction_LocalVar::StaticGetTypeId())
-		{
-			const FControlRigGraphSchemaAction_LocalVar* VarAction = (const FControlRigGraphSchemaAction_LocalVar*)InGraphAction;
-			return VarAction->GetVariableName().ToString();
-		}
-	}
-	return Super::GetFindReferenceSearchTerm(InGraphAction);
-}
-
-FReply UControlRigGraphSchema::BeginGraphDragAction(TSharedPtr<FEdGraphSchemaAction> InAction, const FPointerEvent& MouseEvent) const
-{
-	if (!InAction.IsValid())
-	{
-		return FReply::Unhandled();
-	}
-
-	if (InAction->GetTypeId() == FEdGraphSchemaAction_K2Graph::StaticGetTypeId())
-	{
-		FEdGraphSchemaAction_K2Graph* FuncAction = (FEdGraphSchemaAction_K2Graph*)InAction.Get();
-		if (UControlRigGraph* RigGraph = Cast<UControlRigGraph>(FuncAction->EdGraph))
-		{
-			if (UControlRigBlueprint* RigBlueprint = Cast<UControlRigBlueprint>(FBlueprintEditorUtils::FindBlueprintForGraph(RigGraph)))
-			{
-				TSharedRef<FControlRigFunctionDragDropAction> Action = FControlRigFunctionDragDropAction::New(InAction, RigBlueprint, RigGraph);
-				Action->SetAltDrag(MouseEvent.IsAltDown());
-				Action->SetCtrlDrag(MouseEvent.IsControlDown());
-				return FReply::Handled().BeginDragDrop(Action);
-			}
-		}
-	}
-	else if(InAction->GetTypeId() == FControlRigGraphSchemaAction_LocalVar::StaticGetTypeId())
-	{
-		FControlRigGraphSchemaAction_LocalVar* VarAction = (FControlRigGraphSchemaAction_LocalVar*)InAction.Get();
-		if (UControlRigGraph* RigGraph = Cast<UControlRigGraph>(VarAction->GetVariableScope()))
-		{
-			if (UControlRigBlueprint* RigBlueprint = Cast<UControlRigBlueprint>(FBlueprintEditorUtils::FindBlueprintForGraph(RigGraph)))
-			{
-				TSharedRef<FControlRigFunctionDragDropAction> Action = FControlRigFunctionDragDropAction::New(InAction, RigBlueprint, RigGraph);
-				Action->SetAltDrag(MouseEvent.IsAltDown());
-				Action->SetCtrlDrag(MouseEvent.IsControlDown());
-				return FReply::Handled().BeginDragDrop(Action);
-			}
-		}
-	}
-	return FReply::Unhandled();
-}
-
-FConnectionDrawingPolicy* UControlRigGraphSchema::CreateConnectionDrawingPolicy(int32 InBackLayerID, int32 InFrontLayerID, float InZoomFactor, const FSlateRect& InClippingRect, class FSlateWindowElementList& InDrawElements, class UEdGraph* InGraphObj) const
-{
-#if WITH_EDITOR
-	return IControlRigEditorModule::Get().CreateConnectionDrawingPolicy(InBackLayerID, InFrontLayerID, InZoomFactor, InClippingRect, InDrawElements, InGraphObj);
-#else
-	check(0);
-	return nullptr;
-#endif
-}
-
-bool UControlRigGraphSchema::ShouldHidePinDefaultValue(UEdGraphPin* Pin) const
-{
-	// we should hide default values if any of our parents are connected
-	if(HasParentConnection_Recursive(Pin))
-	{
-		return true;
-	}
-
-	if(UControlRigGraphNode* RigGraphNode = Cast<UControlRigGraphNode>(Pin->GetOwningNode()))
-	{
-		if(RigGraphNode->DrawAsCompactNode())
-		{
-			return true;
-		}
-	}
-
-	return false;
-}
-
-bool UControlRigGraphSchema::IsPinBeingWatched(UEdGraphPin const* Pin) const
-{
-	if (UControlRigGraphNode* Node = Cast< UControlRigGraphNode>(Pin->GetOwningNode()))
-	{
-		if (URigVMPin* ModelPin = Node->GetModel()->FindPin(Pin->GetName()))
-		{
-			return ModelPin->RequiresWatch();
-		}
-	}
-	return false;
-}
-
-void UControlRigGraphSchema::ClearPinWatch(UEdGraphPin const* Pin) const
-{
-	if (UControlRigGraphNode* Node = Cast< UControlRigGraphNode>(Pin->GetOwningNode()))
-	{
-		Node->GetController()->SetPinIsWatched(Pin->GetName(), false);
-	}
-}
-
-void UControlRigGraphSchema::OnPinConnectionDoubleCicked(UEdGraphPin* PinA, UEdGraphPin* PinB, const FVector2D& GraphPosition) const
-{
-	if (UControlRigGraphNode* Node = Cast< UControlRigGraphNode>(PinA->GetOwningNode()))
-	{
-		if (URigVMLink* Link = Node->GetModel()->FindLink(FString::Printf(TEXT("%s -> %s"), *PinA->GetName(), *PinB->GetName())))
-		{
-			Node->GetController()->AddRerouteNodeOnLink(Link, GraphPosition, FString(), true, true);
-		}
-	}
-}
-
-bool UControlRigGraphSchema::MarkBlueprintDirtyFromNewNode(UBlueprint* InBlueprint, UEdGraphNode* InEdGraphNode) const
-{
-	if (InBlueprint == nullptr || InEdGraphNode == nullptr)
-	{
-		return false;
-	}
-	return true;
-}
-
-bool UControlRigGraphSchema::IsStructEditable(UStruct* InStruct) const
-{
-	if (InStruct == FRuntimeFloatCurve::StaticStruct())
-	{
-		return true;
-	}
-	return false;
-}
-
-void UControlRigGraphSchema::SetNodePosition(UEdGraphNode* Node, const FVector2D& Position) const
-{
-	return SetNodePosition(Node, Position, true);
-}
-
-void UControlRigGraphSchema::SetNodePosition(UEdGraphNode* Node, const FVector2D& Position, bool bSetupUndo) const
-{
-	StartGraphNodeInteraction(Node);
-
-	if (UControlRigGraphNode* RigNode = Cast<UControlRigGraphNode>(Node))
-	{
-		RigNode->GetController()->SetNodePosition(RigNode->GetModelNode(), Position, bSetupUndo, false, false);
-	}
-	
-	if (UEdGraphNode_Comment* CommentNode = Cast<UEdGraphNode_Comment>(Node))
-	{
-		if(UControlRigGraph* Graph = CommentNode->GetTypedOuter<UControlRigGraph>())
-		{
-			Graph->GetController()->SetNodePositionByName(CommentNode->GetFName(), Position, bSetupUndo, false, false);
-		}
-	}
-}
-
-void UControlRigGraphSchema::GetGraphDisplayInformation(const UEdGraph& Graph, /*out*/ FGraphDisplayInfo& DisplayInfo) const
-{
-	Super::GetGraphDisplayInformation(Graph, DisplayInfo);
-
-	if (UControlRigGraph* RigGraph = Cast<UControlRigGraph>((UEdGraph*)&Graph))
-	{
-		if(URigVMGraph* Model = RigGraph->GetModel())
-		{
-			TArray<FString> NodePathParts;
-			if (URigVMNode::SplitNodePath(RigGraph->ModelNodePath, NodePathParts))
-			{
-				if(NodePathParts.Num() > 1)
-				{
-					DisplayInfo.DisplayName = FText::FromString(NodePathParts.Last());
-					DisplayInfo.PlainName = DisplayInfo.DisplayName;
-
-					static const FText LocalFunctionText = FText::FromString(TEXT("A local function.")); 
-					DisplayInfo.Tooltip = LocalFunctionText;
-				}
-
-				// if this is a riggraph within a collapse node - let's use that for the tooltip
-				if(URigVMCollapseNode* CollapseNode = Model->GetTypedOuter<URigVMCollapseNode>())
-				{
-					DisplayInfo.Tooltip = CollapseNode->GetToolTipText();
-				}
-			}
-
-			if(Model->IsRootGraph())
-			{
-				// let's see if there is only one event
-				FString EventName;
-				if(Algo::CountIf(Model->GetNodes(), [&EventName](const URigVMNode* NodeToCount) -> bool
-				{
-					if(NodeToCount->IsEvent() && NodeToCount->CanOnlyExistOnce())
-					{
-						if(EventName.IsEmpty())
-						{
-							if(const URigVMUnitNode* UnitNode = Cast<URigVMUnitNode>(NodeToCount))
-							{
-								if(UnitNode->GetScriptStruct())
-								{
-									EventName = UnitNode->GetScriptStruct()->GetDisplayNameText().ToString();
-								}
-							}
-						}
-						return true;
-					}
-					return false;
-				}) == 1)
-				{
-					static constexpr TCHAR EventGraphNameFormat[] = TEXT("%s Graph");
-					const FString DesiredGraphName = FString::Printf(EventGraphNameFormat, *EventName);
-					DisplayInfo.DisplayName = FText::FromString(DesiredGraphName);
-					DisplayInfo.PlainName = DisplayInfo.DisplayName;
-				}
-				else
-				{
-					static const FText MainGraphText = FText::FromString(TEXT("A top level graph for the rig."));
-					DisplayInfo.Tooltip = MainGraphText;
-
-					if(!NodePathParts.IsEmpty())
-					{
-						static constexpr TCHAR NodePathSuffix[] = TEXT("::");
-						FString NodePath = NodePathParts[0];
-						NodePath.RemoveFromEnd(NodePathSuffix);
-						NodePath.RemoveFromStart(FRigVMClient::RigVMModelPrefix);
-						NodePath.TrimStartAndEndInline();
-						DisplayInfo.DisplayName = FText::FromString(NodePath);
-						DisplayInfo.PlainName = DisplayInfo.DisplayName;
-					}
-				}
-			}
-		}
-	}
-}
-
-bool UControlRigGraphSchema::GetLocalVariables(const UEdGraph* InGraph, TArray<FBPVariableDescription>& OutLocalVariables) const
-{
-	OutLocalVariables.Reset();
-	if (UControlRigGraph* RigGraph = Cast<UControlRigGraph>((UEdGraph*)InGraph))
-	{
-		if (URigVMGraph* Model = RigGraph->GetModel())
-		{
-			TArray<FRigVMGraphVariableDescription> LocalVariables = Model->GetLocalVariables();
-			for (FRigVMGraphVariableDescription LocalVariable : LocalVariables)
-			{
-				FBPVariableDescription VariableDescription;
-				VariableDescription.VarName = LocalVariable.Name;
-				VariableDescription.FriendlyName = LocalVariable.Name.ToString();
-				VariableDescription.DefaultValue = LocalVariable.DefaultValue;
-				VariableDescription.VarType = LocalVariable.ToPinType();
-				VariableDescription.PropertyFlags |= CPF_BlueprintVisible;
-				OutLocalVariables.Add(VariableDescription);
-			}
-		}
-	}
-	return true;
-}
-
-TSharedPtr<FEdGraphSchemaAction> UControlRigGraphSchema::MakeActionFromVariableDescription(const UEdGraph* InEdGraph,
-	const FBPVariableDescription& Variable) const
-{
-	if (UControlRigGraph* RigGraph = Cast<UControlRigGraph>((UEdGraph*)InEdGraph))
-	{
-		FText Category = Variable.Category;
-		if (Variable.Category.EqualTo(UEdGraphSchema_K2::VR_DefaultCategory))
-		{
-			Category = FText::GetEmpty();
-		}
-
-		TSharedPtr<FControlRigGraphSchemaAction_LocalVar> Action = MakeShareable(new FControlRigGraphSchemaAction_LocalVar(Category, FText::FromName(Variable.VarName), FText::GetEmpty(), 0, NodeSectionID::LOCAL_VARIABLE));
-		Action->SetVariableInfo(Variable.VarName, RigGraph, Variable.VarType.PinCategory == UEdGraphSchema_K2::PC_Boolean);
-		return Action;
-	}
-	return nullptr;
-}
-
-FText UControlRigGraphSchema::GetGraphCategory(const UEdGraph* InGraph) const
-{
-	if (UControlRigGraph* RigGraph = Cast<UControlRigGraph>((UEdGraph*)InGraph))
-	{
-		if (URigVMGraph* Model = RigGraph->GetModel())
-		{
-			if (URigVMCollapseNode* CollapseNode = Cast<URigVMCollapseNode>(Model->GetOuter()))
-			{
-				return FText::FromString(CollapseNode->GetNodeCategory());
-			}
-		}
-	}
-	return FText();
-}
-
-EGraphType UControlRigGraphSchema::GetGraphType(const UEdGraph* TestEdGraph) const
-{
-	if (UControlRigGraph* RigGraph = Cast<UControlRigGraph>((UEdGraph*)TestEdGraph))
-	{
-		if (UControlRigBlueprint* RigBlueprint = Cast<UControlRigBlueprint>(FBlueprintEditorUtils::FindBlueprintForGraph(RigGraph)))
-		{
-			if (URigVMGraph* Model = RigGraph->GetModel())
-			{
-				if(Model->IsRootGraph())
-				{
-					if(Model->IsA<URigVMFunctionLibrary>())
-					{
-						return EGraphType::GT_Function;
-					}
-					else
-					{
-						return EGraphType::GT_Ubergraph;
-					}
-				}
-				else if(URigVMLibraryNode* LibraryNode = Cast<URigVMLibraryNode>(Model->GetOuter()))
-				{
-					if(LibraryNode->GetGraph()->IsA<URigVMFunctionLibrary>())
-					{
-						return EGraphType::GT_Function;
-					}
-					else
-					{
-						// collapse nodes show up as uber graphs
-						return EGraphType::GT_Ubergraph;
-					}
-				}
-			}
-		}
-	}
-	return Super::GetGraphType(TestEdGraph);
-}
-
-FReply UControlRigGraphSchema::TrySetGraphCategory(const UEdGraph* InGraph, const FText& InCategory)
-{
-	if (UControlRigGraph* RigGraph = Cast<UControlRigGraph>((UEdGraph*)InGraph))
-	{
-		if (UControlRigBlueprint* RigBlueprint = Cast<UControlRigBlueprint>(FBlueprintEditorUtils::FindBlueprintForGraph(RigGraph)))
-		{
-			if (URigVMGraph* Model = RigGraph->GetModel())
-			{
-				if (URigVMCollapseNode* CollapseNode = Cast<URigVMCollapseNode>(Model->GetOuter()))
-				{
-					if (URigVMController* Controller = RigBlueprint->GetOrCreateController(CollapseNode->GetGraph()))
-					{
-						if (Controller->SetNodeCategory(CollapseNode, InCategory.ToString(), true, false, true))
-						{
-							return FReply::Handled();
-						}
-					}
-				}
-			}
-		}
-	}
-	return FReply::Unhandled();
-}
-
-bool UControlRigGraphSchema::TryDeleteGraph(UEdGraph* GraphToDelete) const
-{
-	if (UControlRigGraph* RigGraph = Cast<UControlRigGraph>(GraphToDelete))
-	{
-		if (UControlRigBlueprint* RigBlueprint = Cast<UControlRigBlueprint>(FBlueprintEditorUtils::FindBlueprintForGraph(RigGraph)))
-		{
-			if (URigVMGraph* Model = RigBlueprint->GetModel(GraphToDelete))
-			{
-				if(Model->IsRootGraph())
-				{
-					RigBlueprint->RemoveModel(Model->GetNodePath());
-					return true;
-				}
-				else if (URigVMCollapseNode* LibraryNode = Cast<URigVMCollapseNode>(Model->GetOuter()))
-				{
-					if (URigVMController* Controller = RigBlueprint->GetOrCreateController(LibraryNode->GetGraph()))
-					{
-						// check if there is a "bulk remove function" transaction going on.
-						// which implies that a category is being deleted
-						if (GEditor->CanTransact())
-						{
-							if (GEditor->Trans->GetQueueLength() > 0)
-							{
-								const FTransaction* LastTransaction = GEditor->Trans->GetTransaction(GEditor->Trans->GetQueueLength() - 1);
-								if (LastTransaction)
-								{
-									if (LastTransaction->GetTitle().ToString() == TEXT("Bulk Remove Functions"))
-									{
-										// instead of deleting the graph, let's set its category to none
-										// and thus moving it to the top of the tree
-										return Controller->SetNodeCategory(LibraryNode, FString());
-									}
-								}
-							}
-						}
-
-						bool bSetupUndoRedo = true;
-
-						// if the element to remove is a function, check if it is public and referenced. If so,
-						// warn the user about a bulk remove
-						if (URigVMFunctionLibrary* Library = Cast<URigVMFunctionLibrary>(LibraryNode->GetGraph()))
-						{
-							const FName& FunctionName = LibraryNode->GetFName();
-							if (RigBlueprint->IsFunctionPublic(FunctionName))
-							{
-								for (auto Reference : Library->GetReferencesForFunction(FunctionName))
-								{
-									if (Reference.IsValid())
-									{
-										UControlRigBlueprint* OtherBlueprint = Reference->GetTypedOuter<UControlRigBlueprint>(); 
-										if (OtherBlueprint != RigBlueprint)
-										{											
-											if(RigBlueprint->OnRequestBulkEditDialog().IsBound())
-											{
-												URigVMController* FunctionController = RigBlueprint->GetController(LibraryNode->GetContainedGraph());
-												FRigVMController_BulkEditResult Result = RigBlueprint->OnRequestBulkEditDialog().Execute(RigBlueprint, FunctionController, LibraryNode, ERigVMControllerBulkEditType::RemoveFunction);
-												if(Result.bCanceled)
-												{
-													return false;
-												}
-												bSetupUndoRedo = Result.bSetupUndoRedo;
-											}
-											break;	
-										}
-									}
-								}
-							}
-						}
-						
-						return Controller->RemoveNode(LibraryNode, bSetupUndoRedo, false);
-					}
-				}
-			}
-		}
-	}
-	return false;
-}
-
-bool UControlRigGraphSchema::TryRenameGraph(UEdGraph* GraphToRename, const FName& InNewName) const
-{
-	if (UControlRigGraph* RigGraph = Cast<UControlRigGraph>(GraphToRename))
-	{
-		if (UControlRigBlueprint* RigBlueprint = Cast<UControlRigBlueprint>(FBlueprintEditorUtils::FindBlueprintForGraph(RigGraph)))
-		{
-			if (URigVMGraph* Model = RigGraph->GetModel())
-			{
-				if(Model->IsRootGraph())
-				{
-					const FString NewName = FString::Printf(TEXT("%s %s"), FRigVMClient::RigVMModelPrefix, *InNewName.ToString()); 
-					RigBlueprint->RenameGraph(Model->GetNodePath(), *NewName);
-				}
-				else if (URigVMGraph* RootModel = Model->GetRootGraph())
-				{
-					URigVMLibraryNode* LibraryNode = Cast<URigVMLibraryNode>(RootModel->FindNode(RigGraph->ModelNodePath));
-					if (LibraryNode)
-					{
-						if (URigVMController* Controller = RigBlueprint->GetOrCreateController(LibraryNode->GetGraph()))
-						{
-							Controller->RenameNode(LibraryNode, InNewName, true, true);
-							return true;
-						}
-					}
-				}
-			}
-		}
-	}
-	return false;
-}
-
-bool UControlRigGraphSchema::TryToGetChildEvents(const UEdGraph* Graph, const int32 SectionId,
-	TArray<TSharedPtr<FEdGraphSchemaAction>>& Actions, const FText& ParentCategory) const
-{
-	check(Graph);
-
-	if(const UControlRigGraph* RigGraph = Cast<UControlRigGraph>(Graph))
-	{
-		if(URigVMGraph* Model = RigGraph->GetModel())
-		{
-			TArray<FName> EventNames;
-			TMap<FName, const URigVMNode*> EventNameToNode;
-			for(const URigVMNode* Node : Model->GetNodes())
-			{
-				if(Node->IsEvent())
-				{
-					if(!EventNames.Contains(Node->GetEventName()))
-					{
-						EventNames.Add(Node->GetEventName());
-						EventNameToNode.Add(Node->GetEventName(), Node);
-					}
-				}
-			}
-
-			if(!EventNames.IsEmpty())
-			{
-				FGraphDisplayInfo EdGraphDisplayInfo;
-				GetGraphDisplayInformation(*Graph, EdGraphDisplayInfo);
-
-				FText EdGraphDisplayName = EdGraphDisplayInfo.DisplayName;
-				FText ActionCategory;
-				if (!ParentCategory.IsEmpty())
-				{
-					ActionCategory = FText::Format(FText::FromString(TEXT("{0}|{1}")), ParentCategory, EdGraphDisplayName);
-				}
-				else
-				{
-					ActionCategory = MoveTemp(EdGraphDisplayName);
-				}
-
-				for(const FName& EventName : EventNames)
-				{
-					const URigVMNode* Node = EventNameToNode.FindChecked(EventName);
-					TSharedPtr<FEdGraphSchemaAction> Action = MakeShareable(
-						new FControlRigGraphSchemaAction_Event(EventName, Node->GetNodePath(true), ActionCategory)
-					);
-					Action->SectionID = SectionId;
-					Actions.Add(Action);
-				}
-			}
-			return true;
-		}
-	}
-
-	return false;
-}
 
-UEdGraphPin* UControlRigGraphSchema::DropPinOnNode(UEdGraphNode* InTargetNode, const FName& InSourcePinName, const FEdGraphPinType& InSourcePinType, EEdGraphPinDirection InSourcePinDirection) const
-{
-	FString NewPinName;
-
-	if (UControlRigBlueprint* RigBlueprint = Cast<UControlRigBlueprint>(FBlueprintEditorUtils::FindBlueprintForNode(InTargetNode)))
-	{
-		if (UControlRigGraphNode* RigNode = Cast<UControlRigGraphNode>(InTargetNode))
-		{
-			if (URigVMNode* ModelNode = RigNode->GetModelNode())
-			{
-				URigVMGraph* Model = nullptr;
-				ERigVMPinDirection PinDirection = InSourcePinDirection == EGPD_Input ? ERigVMPinDirection::Input : ERigVMPinDirection::Output;
-
-				if (URigVMCollapseNode* CollapseNode = Cast<URigVMCollapseNode>(ModelNode))
-				{
-					Model = CollapseNode->GetContainedGraph();
-					PinDirection = PinDirection == ERigVMPinDirection::Output ? ERigVMPinDirection::Input : ERigVMPinDirection::Output;
-				}
-				else if (ModelNode->IsA<URigVMFunctionEntryNode>() ||
-					ModelNode->IsA<URigVMFunctionReturnNode>())
-				{
-					Model = ModelNode->GetGraph();
-				}
-
-				if (Model)
-				{
-					ensure(!Model->IsTopLevelGraph());
-
-					FRigVMExternalVariable ExternalVar = RigVMTypeUtils::ExternalVariableFromPinType(InSourcePinName, InSourcePinType);
-					if (ExternalVar.IsValid(true /* allow null memory */))
-					{
-						if (URigVMController* Controller = RigBlueprint->GetController(Model))
-						{
-							FName SourcePinName = InSourcePinName; 
-							FString TypeName = ExternalVar.TypeName.ToString();
-							if (ExternalVar.bIsArray)
-							{
-								TypeName = RigVMTypeUtils::ArrayTypeFromBaseType(*TypeName);
-							}
-							FName TypeObjectPathName = NAME_None;
-							if (ExternalVar.TypeObject)
-							{
-								TypeObjectPathName = *ExternalVar.TypeObject->GetPathName();
-
-								if(const UScriptStruct* CPPTypeStruct = Cast<UScriptStruct>(ExternalVar.TypeObject))
-								{
-									if(CPPTypeStruct->IsChildOf(FRigVMExecuteContext::StaticStruct()))
-									{
-										SourcePinName = FRigVMStruct::ExecuteContextName;
-										PinDirection = ERigVMPinDirection::IO;										
-									}
-								}
-							}
-
-							FString DefaultValue;
-							if (PinBeingDropped)
-							{
-								if (UControlRigGraphNode* SourceNode = Cast<UControlRigGraphNode>(PinBeingDropped->GetOwningNode()))
-								{
-									if (URigVMPin* SourcePin = SourceNode->GetModelPinFromPinPath(PinBeingDropped->GetName()))
-									{
-										DefaultValue = SourcePin->GetDefaultValue();
-									}
-								}
-							}
-
-							FName ExposedPinName = Controller->AddExposedPin(
-								SourcePinName,
-								PinDirection,
-								TypeName,
-								TypeObjectPathName,
-								DefaultValue,
-								true,
-								true
-							);
-							
-							if (!ExposedPinName.IsNone())
-							{
-								NewPinName = ExposedPinName.ToString();
-							}
-						}
-					}
-				}
-
-				if (!NewPinName.IsEmpty())
-				{
-					if (URigVMPin* NewModelPin = ModelNode->FindPin(NewPinName))
-					{
-						return RigNode->FindPin(*NewModelPin->GetPinPath());
-					}
-				}
-			}
-		}
-	}
-
-	return nullptr;
-}
-
-bool UControlRigGraphSchema::SupportsDropPinOnNode(UEdGraphNode* InTargetNode, const FEdGraphPinType& InSourcePinType, EEdGraphPinDirection InSourcePinDirection, FText& OutErrorMessage) const
-{
-	if (UControlRigGraphNode* RigNode = Cast<UControlRigGraphNode>(InTargetNode))
-	{
-		if(URigVMNode* ModelNode = RigNode->GetModelNode())
-		{
-			if (ModelNode->IsA<URigVMFunctionEntryNode>())
-			{
-				if (InSourcePinDirection == EGPD_Output)
-				{
-					OutErrorMessage = LOCTEXT("AddPinToReturnNode", "Add Pin to Return Node");
-					return false;
-				}
-				return true;
-			}
-			else if (ModelNode->IsA<URigVMFunctionReturnNode>())
-			{
-				if (InSourcePinDirection == EGPD_Input)
-				{
-					OutErrorMessage = LOCTEXT("AddPinToEntryNode", "Add Pin to Entry Node");
-					return false;
-				}
-				return true;
-			}
-			else if (ModelNode->IsA<URigVMCollapseNode>())
-			{
-				return true;
-			}
-		}
-	}
-
-	return false;
-}
+#include UE_INLINE_GENERATED_CPP_BY_NAME(ControlRigGraphSchema)
 
-UControlRigGraphNode* UControlRigGraphSchema::CreateGraphNode(UControlRigGraph* InGraph, const FName& InPropertyName) const
-{
-	const bool bSelectNewNode = true;
-	FGraphNodeCreator<UControlRigGraphNode> GraphNodeCreator(*InGraph);
-	UControlRigGraphNode* ControlRigGraphNode = GraphNodeCreator.CreateNode(bSelectNewNode);
-	ControlRigGraphNode->ModelNodePath = InPropertyName.ToString();
-	GraphNodeCreator.Finalize();
 
-	return ControlRigGraphNode;
-}
 
-void UControlRigGraphSchema::TrySetDefaultValue(UEdGraphPin& InPin, const FString& InNewDefaultValue, bool bMarkAsModified) const
-{
-#if WITH_EDITOR
-	if (GEditor)
-	{
-		GEditor->CancelTransaction(0);
-	}
-#endif
-	GetDefault<UEdGraphSchema_K2>()->TrySetDefaultValue(InPin, InNewDefaultValue, false);
-}
+#define LOCTEXT_NAMESPACE "ControlRigGraphSchema"
 
-void UControlRigGraphSchema::TrySetDefaultObject(UEdGraphPin& InPin, UObject* InNewDefaultObject, bool bMarkAsModified) const
-{
-#if WITH_EDITOR
-	if (GEditor)
-	{
-		GEditor->CancelTransaction(0);
-	}
-#endif
-	GetDefault<UEdGraphSchema_K2>()->TrySetDefaultObject(InPin, InNewDefaultObject, false);
-}
+const FName UControlRigGraphSchema::GraphName_ControlRig(TEXT("Rig"));
 
-void UControlRigGraphSchema::TrySetDefaultText(UEdGraphPin& InPin, const FText& InNewDefaultText, bool bMarkAsModified) const
+UControlRigGraphSchema::UControlRigGraphSchema()
 {
-#if WITH_EDITOR
-	if (GEditor)
-	{
-		GEditor->CancelTransaction(0);
-	}
-#endif
-	GetDefault<UEdGraphSchema_K2>()->TrySetDefaultText(InPin, InNewDefaultText, false);
 }
 
-bool UControlRigGraphSchema::ArePinsCompatible(const UEdGraphPin* PinA, const UEdGraphPin* PinB, const UClass* CallingContext, bool bIgnoreArray /*= false*/) const
+FLinearColor UControlRigGraphSchema::GetPinTypeColor(const FEdGraphPinType& PinType) const
 {
-	// filter out pins which have a parent
-	if (PinB->ParentPin != nullptr)
-	{
-		return false;
-	}
-
-	// if we are looking at a polymorphic node
-	if((PinA->PinType.ContainerType == PinB->PinType.ContainerType) ||
-		(PinA->PinType.PinSubCategoryObject != PinB->PinType.PinSubCategoryObject))
+	const FName& TypeName = PinType.PinCategory;
+	if (TypeName == UEdGraphSchema_K2::PC_Struct)
 	{
-		auto IsPinCompatibleWithType = [](const UEdGraphPin* InPin, const FEdGraphPinType& InPinType) -> bool
-		{
-			if(const UControlRigGraphNode* RigNode = Cast<UControlRigGraphNode>(InPin->GetOwningNode()))
-			{
-				FString CPPType;
-				UObject* CPPTypeObject = nullptr;
-				if(RigVMTypeUtils::CPPTypeFromPinType(InPinType, CPPType, &CPPTypeObject))
-				{
-					FString PinPath, PinName;
-					URigVMPin::SplitPinPathAtEnd(InPin->GetName(), PinPath, PinName);
-						
-					if((InPin->ParentPin != nullptr) &&
-						(InPin->ParentPin->ParentPin == nullptr) &&
-						(InPin->ParentPin->PinType.ContainerType == EPinContainerType::Array))
-					{
-						URigVMPin::SplitPinPathAtEnd(InPin->ParentPin->GetName(), PinPath, PinName);
-						CPPType = RigVMTypeUtils::ArrayTypeFromBaseType(CPPType);
-					}
-
-					const FRigVMTemplateArgumentType Type(*CPPType, CPPTypeObject);
-					const TRigVMTypeIndex TypeIndex = FRigVMRegistry::Get().GetTypeIndex(Type);
-					if(const FRigVMTemplate* Template = RigNode->GetTemplate())
-					{
-						if(const FRigVMTemplateArgument* Argument = Template->FindArgument(*PinName))
-						{
-							if(Argument->SupportsTypeIndex(TypeIndex))
-							{
-								return true;
-							}
-					
-							const TArray<TRigVMTypeIndex>& AvailableCasts = RigVMTypeUtils::GetAvailableCasts(TypeIndex, InPin->Direction == EGPD_Output);
-							for(const TRigVMTypeIndex& AvailableCast : AvailableCasts)
-							{
-								if(Argument->SupportsTypeIndex(AvailableCast))
-								{
-									return true;
-								}
-							}
-						}
-					}
-				}
-			}
-			return false;
-		};
-
-		auto IsTemplateNodePin = [](const UEdGraphPin* InPin)
-		{
-			if(const UControlRigGraphNode* RigNode = Cast<UControlRigGraphNode>(InPin->GetOwningNode()))
-			{
-				return RigNode->GetTemplate() != nullptr;
-			};
-			return false;
-		};
-		
-		if(PinA->PinType.PinSubCategoryObject == RigVMTypeUtils::GetWildCardCPPTypeObject() || IsTemplateNodePin(PinA))
+		if (UStruct* Struct = Cast<UStruct>(PinType.PinSubCategoryObject))
 		{
-			if(IsPinCompatibleWithType(PinA, PinB->PinType))
+			if (Struct == FRigElementKey::StaticStruct() || Struct == FRigElementKeyCollection::StaticStruct())
 			{
-				UControlRigGraphSchema* MutableThis = (UControlRigGraphSchema*)this;
-				MutableThis->LastPinForCompatibleCheck = PinB;
-				MutableThis->bLastPinWasInput = PinB->Direction == EGPD_Input;
-				return true;
+				return FLinearColor(0.0, 0.6588, 0.9490);
 			}
-		}
- 		if(PinB->PinType.PinSubCategoryObject == RigVMTypeUtils::GetWildCardCPPTypeObject() || IsTemplateNodePin(PinB))
-		{
- 			if(IsPinCompatibleWithType(PinB, PinA->PinType))
- 			{
- 				UControlRigGraphSchema* MutableThis = (UControlRigGraphSchema*)this;
- 				MutableThis->LastPinForCompatibleCheck = PinA;
- 				MutableThis->bLastPinWasInput = PinA->Direction == EGPD_Input;
- 				return true;
- 			}
-		}
-	}
-
-	// for large world coordinate support we should allow connections
-	// between float and double
-	if(PinA->PinType.ContainerType == EPinContainerType::None &&
-		PinB->PinType.ContainerType == EPinContainerType::None)
-	{
-		if((PinA->PinType.PinCategory == UEdGraphSchema_K2::PC_Float &&
-			PinB->PinType.PinCategory == UEdGraphSchema_K2::PC_Double) ||
-			(PinA->PinType.PinCategory == UEdGraphSchema_K2::PC_Double &&
-			PinB->PinType.PinCategory == UEdGraphSchema_K2::PC_Float))
-		{
-			return true;
-		}
-	}
-
-	if(GetDefault<UEdGraphSchema_K2>()->ArePinsCompatible(PinA, PinB, CallingContext, bIgnoreArray))
-	{
-		return true;
-	}
-
-	// also check if there's a cast available for the type
-	const TRigVMTypeIndex TypeIndexA = RigVMTypeUtils::TypeIndexFromPinType(PinA->Direction == EGPD_Output ? PinA->PinType : PinB->PinType);
-	const TRigVMTypeIndex TypeIndexB = RigVMTypeUtils::TypeIndexFromPinType(PinA->Direction == EGPD_Output ? PinB->PinType : PinA->PinType);
-	return RigVMTypeUtils::CanCastTypes(TypeIndexA, TypeIndexB);
-}
-
-void UControlRigGraphSchema::RenameNode(UControlRigGraphNode* Node, const FName& InNewNodeName) const
-{
-	Node->NodeTitle = FText::FromName(InNewNodeName);
-	Node->Modify();
-}
-
-void UControlRigGraphSchema::ResetPinDefaultsRecursive(UEdGraphPin* InPin) const
-{
-	UControlRigGraphNode* RigNode = Cast<UControlRigGraphNode>(InPin->GetOwningNode());
-	if (RigNode == nullptr)
-	{
-		return;
-	}
-
-	RigNode->CopyPinDefaultsToModel(InPin);
-	for (UEdGraphPin* SubPin : InPin->SubPins)
-	{
-		ResetPinDefaultsRecursive(SubPin);
-	}
-}
-
-void UControlRigGraphSchema::GetVariablePinTypes(TArray<FEdGraphPinType>& PinTypes) const
-{
-	PinTypes.Add(FEdGraphPinType(UEdGraphSchema_K2::PC_Boolean, FName(NAME_None), nullptr, EPinContainerType::None, false, FEdGraphTerminalType()));
-	PinTypes.Add(FEdGraphPinType(UEdGraphSchema_K2::PC_Real, FName(NAME_None), nullptr, EPinContainerType::None, false, FEdGraphTerminalType()));
-	PinTypes.Add(FEdGraphPinType(UEdGraphSchema_K2::PC_Int, FName(NAME_None), nullptr, EPinContainerType::None, false, FEdGraphTerminalType()));
-	PinTypes.Add(FEdGraphPinType(UEdGraphSchema_K2::PC_Struct, FName(NAME_None), TBaseStructure<FVector>::Get(), EPinContainerType::None, false, FEdGraphTerminalType()));
-	PinTypes.Add(FEdGraphPinType(UEdGraphSchema_K2::PC_Struct, FName(NAME_None), TBaseStructure<FVector2D>::Get(), EPinContainerType::None, false, FEdGraphTerminalType()));
-	PinTypes.Add(FEdGraphPinType(UEdGraphSchema_K2::PC_Struct, FName(NAME_None), TBaseStructure<FRotator>::Get(), EPinContainerType::None, false, FEdGraphTerminalType()));
-	PinTypes.Add(FEdGraphPinType(UEdGraphSchema_K2::PC_Struct, FName(NAME_None), TBaseStructure<FTransform>::Get(), EPinContainerType::None, false, FEdGraphTerminalType()));
-	PinTypes.Add(FEdGraphPinType(UEdGraphSchema_K2::PC_Struct, FName(NAME_None), TBaseStructure<FEulerTransform>::Get(), EPinContainerType::None, false, FEdGraphTerminalType()));
-	PinTypes.Add(FEdGraphPinType(UEdGraphSchema_K2::PC_Struct, FName(NAME_None), TBaseStructure<FLinearColor>::Get(), EPinContainerType::None, false, FEdGraphTerminalType()));
-}
-
-bool UControlRigGraphSchema::SafeDeleteNodeFromGraph(UEdGraph* Graph, UEdGraphNode* Node) const
-{
-	if (UControlRigGraphNode* RigNode = Cast<UControlRigGraphNode>(Node))
-	{
-		if (GEditor)
-		{
-			GEditor->CancelTransaction(0);
-		}
-		return RigNode->GetController()->RemoveNode(RigNode->GetModelNode(), true, true);
-	}
-	return false;
-}
-
-bool UControlRigGraphSchema::CanVariableBeDropped(UEdGraph* InGraph, FProperty* InVariableToDrop) const
-{
-	FRigVMExternalVariable ExternalVariable = FRigVMExternalVariable::Make(InVariableToDrop, nullptr);
-	return ExternalVariable.IsValid(true /* allow nullptr */);
-}
-
-bool UControlRigGraphSchema::RequestVariableDropOnPanel(UEdGraph* InGraph, FProperty* InVariableToDrop, const FVector2D& InDropPosition, const FVector2D& InScreenPosition)
-{
-#if WITH_EDITOR
-	if (CanVariableBeDropped(InGraph, InVariableToDrop))
-	{
-		FRigVMExternalVariable ExternalVariable = FRigVMExternalVariable::Make(InVariableToDrop, nullptr);
 
-		UBlueprint* Blueprint = FBlueprintEditorUtils::FindBlueprintForGraph(InGraph);
-		UControlRigBlueprint* RigBlueprint = Cast<UControlRigBlueprint>(Blueprint);
-		if (RigBlueprint != nullptr)
-		{
-			RigBlueprint->OnVariableDropped().Broadcast(InGraph, InVariableToDrop, InDropPosition, InScreenPosition);
-			return true;
-		}
-	}
-#endif
-
-	return false;
-}
-
-bool UControlRigGraphSchema::RequestVariableDropOnPin(UEdGraph* InGraph, FProperty* InVariableToDrop, UEdGraphPin* InPin, const FVector2D& InDropPosition, const FVector2D& InScreenPosition)
-{
-#if WITH_EDITOR
-	if (CanVariableBeDropped(InGraph, InVariableToDrop))
-	{
-		if(UControlRigGraph* Graph = Cast<UControlRigGraph>(InGraph))
-		{
-			if (URigVMPin* ModelPin = Graph->GetModel()->FindPin(InPin->GetName()))
+			if (Struct == FRigElementKey::StaticStruct() || Struct == FRigPose::StaticStruct())
 			{
-				FRigVMExternalVariable ExternalVariable = FRigVMExternalVariable::Make(InVariableToDrop, nullptr);
-				if (ModelPin->CanBeBoundToVariable(ExternalVariable))
-				{
-					FModifierKeysState KeyState = FSlateApplication::Get().GetModifierKeys();
-					if (KeyState.IsAltDown())
-					{
-						return Graph->GetController()->BindPinToVariable(ModelPin->GetPinPath(), InVariableToDrop->GetName(), true, true);
-					}
-					else
-					{
-						Graph->GetController()->OpenUndoBracket(TEXT("Bind Variable to Pin"));
-						if (URigVMVariableNode* VariableNode = Graph->GetController()->AddVariableNode(ExternalVariable.Name, ExternalVariable.TypeName.ToString(), ExternalVariable.TypeObject, true, FString(), InDropPosition + FVector2D(0.f, -34.f)))
-						{
-							Graph->GetController()->AddLink(VariableNode->FindPin(TEXT("Value"))->GetPinPath(), ModelPin->GetPinPath(), true);
-						}
-						Graph->GetController()->CloseUndoBracket();
-						return true;
-					}
-				}
+				return FLinearColor(0.0, 0.3588, 0.5490);
 			}
 		}
 	}
-#endif
-
-	return false;
-}
-
-void UControlRigGraphSchema::StartGraphNodeInteraction(UEdGraphNode* InNode) const
-{
-#if WITH_EDITOR
-
-	check(InNode);
-
-	if(NodesBeingInteracted.Contains(InNode))
-	{
-		return;
-	}
 	
-	NodePositionsDuringStart.Reset();
-	NodesBeingInteracted.Reset();
-
-	UControlRigGraph* Graph = Cast<UControlRigGraph>(InNode->GetOuter());
-	if (Graph == nullptr)
-	{
-		return;
-	}
-
-	check(Graph->GetController());
-	check(Graph->GetModel());
-
-	NodesBeingInteracted = GetNodesToMoveForNode(InNode);
-
-	for (UEdGraphNode* NodeToMove : NodesBeingInteracted)
-	{
-		FName NodeName = NodeToMove->GetFName();
-		if (URigVMNode* ModelNode = Graph->GetModel()->FindNodeByName(NodeName))
-		{
-			NodePositionsDuringStart.FindOrAdd(NodeName, ModelNode->GetPosition());
-		}
-	}
-
-#endif
+	return Super::GetPinTypeColor(PinType);
 }
 
-void UControlRigGraphSchema::EndGraphNodeInteraction(UEdGraphNode* InNode) const
+bool UControlRigGraphSchema::IsRigVMDefaultEvent(const FName& InEventName) const
 {
-#if WITH_EDITOR
-
-	UControlRigGraph* Graph = Cast<UControlRigGraph>(InNode->GetOuter());
-	if (Graph == nullptr)
+	if(Super::IsRigVMDefaultEvent(InEventName))
 	{
-		return;
+		return true;
 	}
-
-	check(Graph->GetController());
-	check(Graph->GetModel());
-
-	TArray<UEdGraphNode*> NodesToMove = GetNodesToMoveForNode(InNode);
 	
-	bool bMovedSomething = false;
-
-	Graph->GetController()->OpenUndoBracket(TEXT("Move Nodes"));
-
-	for (UEdGraphNode* NodeToMove : NodesToMove)
-	{
-		FName NodeName = NodeToMove->GetFName();
-		if (URigVMNode* ModelNode = Graph->GetModel()->FindNodeByName(NodeName))
-		{
-			FVector2D NewPosition(NodeToMove->NodePosX, NodeToMove->NodePosY);
-
-			if(FVector2D* OldPosition = NodePositionsDuringStart.Find(NodeName))
-			{
-				TGuardValue<bool> SuspendNotification(Graph->bSuspendModelNotifications, true);
-				Graph->GetController()->SetNodePositionByName(NodeName, *OldPosition, false, false);
-			}
-			
-			if(Graph->GetController()->SetNodePositionByName(NodeName, NewPosition, true, false, true))
-			{
-				bMovedSomething = true;
-			}
-		}
-	}
-
-	if (bMovedSomething)
-	{
-		if (GEditor)
-		{
-			GEditor->CancelTransaction(0);
-		}
-
-		Graph->GetController()->CloseUndoBracket();
-	}
-	else
-	{
-		Graph->GetController()->CancelUndoBracket();
-	}
-
-	NodesBeingInteracted.Reset();
-	NodePositionsDuringStart.Reset();
-
-#endif
-}
-
-TArray<UEdGraphNode*> UControlRigGraphSchema::GetNodesToMoveForNode(UEdGraphNode* InNode)
-{
-	TArray<UEdGraphNode*> NodesToMove;
-
-#if WITH_EDITOR
-
-	UControlRigGraph* Graph = Cast<UControlRigGraph>(InNode->GetOuter());
-	if (Graph == nullptr)
-	{
-		return NodesToMove;
-	}
-
-	NodesToMove.Add(InNode);
-
-	for (UEdGraphNode* SelectedGraphNode : Graph->Nodes)
-	{
-		if (SelectedGraphNode->IsSelected())
-		{
-			NodesToMove.AddUnique(SelectedGraphNode);
-		}
-	}
-
-	for (int32 NodeIndex = 0; NodeIndex < NodesToMove.Num(); NodeIndex++)
-	{
-		if (UEdGraphNode_Comment* CommentNode = Cast<UEdGraphNode_Comment>(NodesToMove[NodeIndex]))
-		{
-			if (CommentNode->MoveMode == ECommentBoxMode::GroupMovement)
-			{
-				for (FCommentNodeSet::TConstIterator NodeIt(CommentNode->GetNodesUnderComment()); NodeIt; ++NodeIt)
-				{
-					if (UEdGraphNode* NodeUnderComment = Cast<UEdGraphNode>(*NodeIt))
-					{
-						NodesToMove.AddUnique(NodeUnderComment);
-					}
-				}
-			}
-		}
-	}
-
-#endif
-
-	return NodesToMove;
-}
-
-FVector2D UControlRigGraphSchema::GetNodePositionAtStartOfInteraction(UEdGraphNode* InNode) const
-{
-#if WITH_EDITOR
-	if(InNode)
-	{
-		if(const FVector2D* Position = NodePositionsDuringStart.Find(InNode->GetFName()))
-		{
-			return *Position;
-		}
-
-		return FVector2D(InNode->NodePosX, InNode->NodePosY);
-	}
-#endif
-
-	return FVector2D::ZeroVector;
-}
-
-bool UControlRigGraphSchema::AutowireNewNode(UControlRigGraphNode* NewNode, UEdGraphPin* FromPin) const
-{
-	// copying high level information into a local array since the try create connection below
-	// may cause the pin array to be destroyed / changed
-	TArray<TPair<FName, EEdGraphPinDirection>> PinsToVisit;
-	for(UEdGraphPin* Pin : NewNode->Pins)
-	{
-		PinsToVisit.Emplace(Pin->GetFName(), Pin->Direction);
-	}
-
-	TArray<UEdGraphPin*> OldLinkedTo = FromPin->LinkedTo;
-	for(const TPair<FName, EEdGraphPinDirection>& PinToVisit : PinsToVisit)
-	{
-		UEdGraphPin* Pin = NewNode->FindPin(PinToVisit.Key, PinToVisit.Value);
-		if(Pin == nullptr)
-		{
-			continue;
-		}
-		
-		if (Pin->ParentPin != nullptr)
-		{
-			continue;
-		}
-
-		FPinConnectionResponse ConnectResponse = CanCreateConnection(FromPin, Pin);
-		if(ConnectResponse.Response != ECanCreateConnectionResponse::CONNECT_RESPONSE_DISALLOW)
-		{
-			if (TryCreateConnection(FromPin, Pin))
-			{
-				// It might have been linked in a different direction. Find pin with the correct direction
-				if (Pin->LinkedTo.IsEmpty())
-				{
-					for (UEdGraphPin* Linked : FromPin->LinkedTo)
-					{
-						if (Linked->GetOwningNode() == Pin->GetOwningNode())
-						{
-							Pin = Linked;
-							break;
-						}
-					}
-				}
-
-				// If the pin is an execute context, try to create a link to the pin which was previously connected to FromPin
-				if (URigVMPin* ModelPin = NewNode->FindModelPinFromGraphPin(Pin))
-				{
-					if (!OldLinkedTo.IsEmpty() && ModelPin->IsExecuteContext())
-					{
-						bool bIsInput = Pin->Direction == EEdGraphPinDirection::EGPD_Input;
-
-						// If the pin is an input to a control flow node, the output used should be the completed pin
-						if (ModelPin->GetNode()->IsControlFlowNode())
-						{
-							if (bIsInput)
-							{
-								ModelPin = ModelPin->GetNode()->FindPin(FRigVMStruct::ControlFlowCompletedName.ToString());
-							}
-						}
-
-						// Try to create the second connection
-						if (UEdGraphPin* OppositePin = NewNode->FindGraphPinFromModelPin(ModelPin, !bIsInput))
-						{
-							if (CanCreateConnection(OppositePin, OldLinkedTo[0]).Response != ECanCreateConnectionResponse::CONNECT_RESPONSE_DISALLOW)
-							{
-								TryCreateConnection(OppositePin, OldLinkedTo[0]);
-							}
-						}
-					}
-
-					// copy the default value over if the node is a reroute,
-					// a make array, make struct or make constant
-					if(!ModelPin->IsExecuteContext() && FromPin->Direction == EGPD_Input)
-					{
-						if(const URigVMPin* OtherModelPin = ModelPin->GetGraph()->FindPin(FromPin->GetName()))
-						{
-							FString PinNameToSet;
-							if(ModelPin->GetNode()->IsA<URigVMRerouteNode>())
-							{
-								PinNameToSet = TEXT("Value");
-							}
-							else if(const URigVMDispatchNode* DispatchNode = Cast<URigVMDispatchNode>(ModelPin->GetNode()))
-							{
-								if(const FRigVMDispatchFactory* Factory = DispatchNode->GetFactory())
-								{
-									if(Factory->GetFactoryName() == FRigVMDispatch_Constant().GetFactoryName())
-									{
-										PinNameToSet = TEXT("Value");
-									}
-									else if(Factory->GetFactoryName() == FRigVMDispatch_MakeStruct().GetFactoryName())
-									{
-										PinNameToSet = TEXT("Elements");
-									}
-									else if(Factory->GetFactoryName() == FRigVMDispatch_ArrayMake().GetFactoryName())
-									{
-										PinNameToSet = TEXT("Values");
-									}
-								}
-							}
-
-							if(!PinNameToSet.IsEmpty())
-							{
-								const FString DefaultValue = OtherModelPin->GetDefaultValue();
-								if(!DefaultValue.IsEmpty())
-								{
-									NewNode->GetController()->SetPinDefaultValue(RigVMStringUtils::JoinPinPath(NewNode->GetName(), PinNameToSet), DefaultValue, true, true, false, true);
-								}
-							}
-						}
-					}
-				}
-				return true;
-			}
-		}
-	}
-	return false;
-}
-
-void UControlRigGraphSchema::HandleModifiedEvent(ERigVMGraphNotifType InNotifType, URigVMGraph* InGraph,
-                                                 UObject* InSubject)
-{
-	switch(InNotifType)
-	{
-		case ERigVMGraphNotifType::NodeAdded:
-		case ERigVMGraphNotifType::NodeRemoved:
-		case ERigVMGraphNotifType::PinAdded:
-		case ERigVMGraphNotifType::PinRemoved:
-		case ERigVMGraphNotifType::PinRenamed:
-		case ERigVMGraphNotifType::PinArraySizeChanged:
-		case ERigVMGraphNotifType::PinTypeChanged:
-		case ERigVMGraphNotifType::LinkAdded:
-		case ERigVMGraphNotifType::LinkRemoved:
-		{
-			LastPinForCompatibleCheck = nullptr;
-			break;
-		}
-		default:
-		{
-			break;
-		}
-	}
-}
-
-bool UControlRigGraphSchema::IsControlRigDefaultEvent(const FName& InEventName)
-{
 	return InEventName == FRigUnit_BeginExecution::EventName ||
 		InEventName == FRigUnit_InverseExecution::EventName ||
 		InEventName == FRigUnit_PrepareForExecution::EventName ||
