RigVM: First steps towards rebasing ControlRigBlueprint

diff --git a/Engine/Plugins/Animation/ControlRig/Source/ControlRigDeveloper/Public/ControlRigBlueprint.h b/Engine/Plugins/Animation/ControlRig/Source/ControlRigDeveloper/Public/ControlRigBlueprint.h
--- a/Engine/Plugins/Animation/ControlRig/Source/ControlRigDeveloper/Public/ControlRigBlueprint.h
+++ b/Engine/Plugins/Animation/ControlRig/Source/ControlRigDeveloper/Public/ControlRigBlueprint.h
@@ -4,24 +4,19 @@
 
 #include "CoreMinimal.h"
 #include "UObject/ObjectMacros.h"
-#include "Misc/Guid.h"
 #include "Engine/Blueprint.h"
-#include "Misc/Crc.h"
 #include "ControlRigDefines.h"
 #include "Rigs/RigHierarchyContainer.h"
 #include "Rigs/RigHierarchy.h"
-#include "Rigs/RigHierarchyController.h"
 #include "Interfaces/Interface_PreviewMeshProvider.h"
 #include "ControlRigGizmoLibrary.h"
-#include "RigVMCore/RigVM.h"
+#include "ControlRigSchema.h"
 #include "RigVMCore/RigVMStatistics.h"
 #include "RigVMModel/RigVMClient.h"
-#include "RigVMCompiler/RigVMCompiler.h"
 #include "ControlRigValidationPass.h"
-#include "RigVMCore/RigVMGraphFunctionDefinition.h"
+#include "RigVMBlueprint.h"
 
 #if WITH_EDITOR
-#include "Kismet2/Kismet2NameValidators.h"
 #include "Kismet2/CompilerResultsLog.h"
 #endif
 
@@ -32,22 +27,6 @@ class USkeletalMesh;
 class UControlRigGraph;
 struct FEndLoadPackageContext;
 
-DECLARE_EVENT_TwoParams(UControlRigBlueprint, FOnVMCompiledEvent, UObject*, URigVM*);
-DECLARE_EVENT_OneParam(UControlRigBlueprint, FOnRefreshEditorEvent, UControlRigBlueprint*);
-DECLARE_EVENT_FourParams(UControlRigBlueprint, FOnVariableDroppedEvent, UObject*, FProperty*, const FVector2D&, const FVector2D&);
-DECLARE_EVENT_OneParam(UControlRigBlueprint, FOnExternalVariablesChanged, const TArray<FRigVMExternalVariable>&);
-DECLARE_EVENT_TwoParams(UControlRigBlueprint, FOnNodeDoubleClicked, UControlRigBlueprint*, URigVMNode*);
-DECLARE_EVENT_OneParam(UControlRigBlueprint, FOnGraphImported, UEdGraph*);
-DECLARE_EVENT_OneParam(UControlRigBlueprint, FOnPostEditChangeChainProperty, FPropertyChangedChainEvent&);
-DECLARE_EVENT_ThreeParams(UControlRigBlueprint, FOnLocalizeFunctionDialogRequested, FRigVMGraphFunctionIdentifier&, UControlRigBlueprint*, bool);
-DECLARE_EVENT_ThreeParams(UControlRigBlueprint, FOnReportCompilerMessage, EMessageSeverity::Type, UObject*, const FString&);
-DECLARE_DELEGATE_RetVal_FourParams(FRigVMController_BulkEditResult, FControlRigOnBulkEditDialogRequestedDelegate, UControlRigBlueprint*, URigVMController*, URigVMLibraryNode*, ERigVMControllerBulkEditType);
-DECLARE_DELEGATE_RetVal_OneParam(bool, FControlRigOnBreakLinksDialogRequestedDelegate, TArray<URigVMLink*>);
-DECLARE_DELEGATE_RetVal_OneParam(TRigVMTypeIndex, FControlRigOnPinTypeSelectionRequestedDelegate, const TArray<TRigVMTypeIndex>&);
-DECLARE_EVENT(UControlRigBlueprint, FOnBreakpointAdded);
-DECLARE_EVENT_OneParam(UControlRigBlueprint, FOnRequestInspectObject, const TArray<UObject*>& );
-DECLARE_DELEGATE_RetVal(URigVMGraph*, FControlRigBlueprintGetFocusedGraph);
-
 USTRUCT(meta = (Deprecated = "5.2"))
 struct CONTROLRIGDEVELOPER_API FControlRigPublicFunctionArg
 {
@@ -109,182 +88,44 @@ struct CONTROLRIGDEVELOPER_API FControlRigPublicFunctionData
 	bool IsMutable() const;
 };
 
-USTRUCT()
-struct CONTROLRIGDEVELOPER_API FRigGraphDisplaySettings
-{
-	GENERATED_BODY();
-
-	FRigGraphDisplaySettings()
-		: bShowNodeInstructionIndex(false)
-		, bShowNodeRunCounts(false)
-		, NodeRunLowerBound(1)
-		, NodeRunLimit(256)
-		, MinMicroSeconds(0.0)
-		, MaxMicroSeconds(1.0)
-		, TotalMicroSeconds(0.0)
-		, bAutoDetermineRange(true)
-		, LastMinMicroSeconds(0.0)
-		, LastMaxMicroSeconds(1.0)
-		, MinDurationColor(FLinearColor::Green)
-		, MaxDurationColor(FLinearColor::Red)
-	{
-	}
-
-	// When enabled shows the first node instruction index
-	// matching the execution stack window.
-	UPROPERTY(EditAnywhere, Category = "Graph Display Settings")
-	bool bShowNodeInstructionIndex;
-
-	// When enabled shows the node counts both in the graph view as
-	// we as in the execution stack window.
-	// The number on each node represents how often the node has been run.
-	// Keep in mind when looking at nodes in a function the count
-	// represents the sum of all counts for each node based on all
-	// references of the function currently running.
-	UPROPERTY(EditAnywhere, Category = "Graph Display Settings")
-	bool bShowNodeRunCounts;
-
-	// A lower limit for counts for nodes used for debugging.
-	// Any node lower than this count won't show the run count.
-	UPROPERTY(EditAnywhere, Category = "Graph Display Settings")
-	int32 NodeRunLowerBound;
-
-	// A upper limit for counts for nodes used for debugging.
-	// If a node reaches this count a warning will be issued for the
-	// node and displayed both in the execution stack as well as in the
-	// graph. Setting this to <= 1 disables the warning.
-	// Note: The count limit doesn't apply to functions / collapse nodes.
-	UPROPERTY(EditAnywhere, Category = "Graph Display Settings")
-	int32 NodeRunLimit;
-
-	// The duration in microseconds of the fastest instruction / node
-	UPROPERTY(EditAnywhere, Category = "Graph Display Settings", transient, meta = (EditCondition = "!bAutoDetermineRange"))
-	double MinMicroSeconds;
-
-	// The duration in microseconds of the slowest instruction / node
-	UPROPERTY(EditAnywhere, Category = "Graph Display Settings", transient, meta = (EditCondition = "!bAutoDetermineRange"))
-	double MaxMicroSeconds;
-
-	// The total duration of the last execution of the rig
-	UPROPERTY(VisibleAnywhere, Category = "Graph Display Settings", transient)
-	double TotalMicroSeconds;
-
-	UPROPERTY(EditAnywhere, Category = "Graph Display Settings")
-	bool bAutoDetermineRange;
-
-	UPROPERTY(transient)
-	double LastMinMicroSeconds;
-
-	UPROPERTY(transient)
-	double LastMaxMicroSeconds;
-
-	// The color of the fastest instruction / node
-	UPROPERTY(EditAnywhere, Category = "Graph Display Settings")
-	FLinearColor MinDurationColor;
-
-	// The color of the slowest instruction / node
-	UPROPERTY(EditAnywhere, Category = "Graph Display Settings")
-	FLinearColor MaxDurationColor;
-};
-
-USTRUCT()
-struct CONTROLRIGDEVELOPER_API FControlRigPythonSettings
-{
-	GENERATED_BODY();
-
-	FControlRigPythonSettings()
-	{
-	}
-};
-
-enum class EControlRigBlueprintLoadType : uint8
-{
-	PostLoad,
-	CheckUserDefinedStructs
-};
-
 UCLASS(BlueprintType, meta=(IgnoreClassThumbnail))
-class CONTROLRIGDEVELOPER_API UControlRigBlueprint : public UBlueprint, public IInterface_PreviewMeshProvider, public IRigVMClientHost
+class CONTROLRIGDEVELOPER_API UControlRigBlueprint : public URigVMBlueprint, public IInterface_PreviewMeshProvider
 {
 	GENERATED_UCLASS_BODY()
 
 public:
 	UControlRigBlueprint();
 
-	void InitializeModelIfRequired(bool bRecompileVM = true);
-
-	/** Get the (full) generated class for this control rig blueprint */
-	URigVMBlueprintGeneratedClass* GetControlRigBlueprintGeneratedClass() const;
-
-	/** Get the (skeleton) generated class for this control rig blueprint */
-	URigVMBlueprintGeneratedClass* GetControlRigBlueprintSkeletonClass() const;
-
+	// URigVMBlueprint interface
+	virtual UClass* GetRigVMSchemaClass() const override { return UControlRigSchema::StaticClass(); }
+	virtual UScriptStruct* GetRigVMExecuteContextStruct() const override { return FControlRigExecuteContext::StaticStruct(); }
+	virtual UClass* GetRigVMEdGraphClass() const override;
+	virtual UClass* GetRigVMEdGraphNodeClass() const override;
+	virtual UClass* GetRigVMEdGraphSchemaClass() const override;
+	virtual TArray<FString> GeneratePythonCommands(const FString InNewBlueprintName) override;
 	virtual void Serialize(FArchive& Ar) override;
 
 #if WITH_EDITOR
 
 	// UBlueprint interface
-	virtual UClass* GetBlueprintClass() const override;
 	virtual UClass* RegenerateClass(UClass* ClassToRegenerate, UObject* PreviousCDO) override;
 	virtual bool SupportedByDefaultBlueprintFactory() const override { return false; }
 	virtual bool IsValidForBytecodeOnlyRecompile() const override { return false; }
-	virtual void LoadModulesRequiredForCompilation() override;
 	virtual void GetTypeActions(FBlueprintActionDatabaseRegistrar& ActionRegistrar) const override;
 	virtual void GetInstanceActions(FBlueprintActionDatabaseRegistrar& ActionRegistrar) const override;
-	virtual void SetObjectBeingDebugged(UObject* NewObject) override;
-	PRAGMA_DISABLE_DEPRECATION_WARNINGS // Suppress compiler warning on override of deprecated function
-	UE_DEPRECATED(5.0, "Use version that takes FObjectPreSaveContext instead.")
-	virtual void PreSave(const class ITargetPlatform* TargetPlatform) override;
-	PRAGMA_ENABLE_DEPRECATION_WARNINGS
 	virtual void PreSave(FObjectPreSaveContext ObjectSaveContext) override;
 	virtual void PostLoad() override;
-#if WITH_EDITORONLY_DATA
-	static void DeclareConstructClasses(TArray<FTopLevelAssetPath>& OutConstructClasses, const UClass* SpecificSubclass);
-#endif
-	virtual bool IsPostLoadThreadSafe() const override { return false; }
 	virtual void PostTransacted(const FTransactionObjectEvent& TransactionEvent) override;
-	virtual void ReplaceDeprecatedNodes() override;
-	virtual void PreDuplicate(FObjectDuplicationParameters& DupParams) override;
 	virtual void PostDuplicate(bool bDuplicateForPIE) override;
-	virtual void GetAssetRegistryTags(TArray<FAssetRegistryTag>& OutTags) const override;
-
-	virtual bool SupportsGlobalVariables() const override { return true; }
-	virtual bool SupportsLocalVariables() const override { return true; }
-	virtual bool SupportsFunctions() const override { return true; }
-	virtual bool SupportsMacros() const override { return false; }
-	virtual bool SupportsDelegates() const override { return false; }
-	virtual bool SupportsEventGraphs() const override { return true; }
-	virtual bool SupportsAnimLayers() const override { return false; }
-	virtual bool ExportGraphToText(UEdGraph* InEdGraph, FString& OutText) override;
-	virtual bool TryImportGraphFromText(const FString& InClipboardText, UEdGraph** OutGraphPtr = nullptr) override;
-	virtual bool CanImportGraphFromText(const FString& InClipboardText) override;
 	virtual bool RequiresForceLoadMembers(UObject* InObject) const override;
 
 	// UObject interface
 	virtual void PostEditChangeChainProperty(struct FPropertyChangedChainEvent& PropertyChangedEvent) override;
-	virtual void PostRename(UObject* OldOuter, const FName OldName) override;
-	/** Called during cooking. Must return all objects that will be Preload()ed when this is serialized at load time. */
-	void GetPreloadDependencies(TArray<UObject*>& OutDeps) override;
-
-	// IRigVMClientHost interface
-	virtual FRigVMClient* GetRigVMClient() override;
-	virtual const FRigVMClient* GetRigVMClient() const override;
-	virtual IRigVMGraphFunctionHost* GetRigVMGraphFunctionHost() override;
-	virtual const IRigVMGraphFunctionHost* GetRigVMGraphFunctionHost() const override;
-	virtual UObject* GetEditorObjectForRigVMGraph(URigVMGraph* InVMGraph) const override;
-	virtual void HandleRigVMGraphAdded(const FRigVMClient* InClient, const FString& InNodePath) override;
-	virtual void HandleRigVMGraphRemoved(const FRigVMClient* InClient, const FString& InNodePath) override;
-	virtual void HandleRigVMGraphRenamed(const FRigVMClient* InClient, const FString& InOldNodePath, const FString& InNewNodePath) override;
-	virtual void HandleConfigureRigVMController(const FRigVMClient* InClient, URigVMController* InControllerToConfigure) override;
-
-	FOnRequestInspectObject& OnRequestInspectObject() { return OnRequestInspectObjectEvent; }
-	void RequestInspectObject(const TArray<UObject*>& InObjects) { OnRequestInspectObjectEvent.Broadcast(InObjects); }
 
 #endif	// #if WITH_EDITOR
 
-	FControlRigBlueprintGetFocusedGraph& OnGetFocusedGraph() { return OnGetFocusedGraphDelegate; }
-
-	virtual bool ShouldBeMarkedDirtyUponTransaction() const override { return false; }
+	UFUNCTION(BlueprintCallable, Category = "Control Rig Blueprint")
+	UControlRig* CreateControlRig() { return Cast<UControlRig>(CreateRigVMHost()); }
 
 	/** IInterface_PreviewMeshProvider interface */
 	UFUNCTION(BlueprintCallable, Category = "Control Rig Blueprint")
@@ -293,210 +134,24 @@ public:
 	UFUNCTION(BlueprintCallable, Category = "Control Rig Blueprint")
 	virtual USkeletalMesh* GetPreviewMesh() const override;
 
-	UFUNCTION(BlueprintCallable, Category = "Control Rig Blueprint")
-	void RecompileVM();
-
-	UFUNCTION(BlueprintCallable, Category = "Control Rig Blueprint")
-	void RecompileVMIfRequired();
-	
-	UFUNCTION(BlueprintCallable, Category = "Control Rig Blueprint")
-	void RequestAutoVMRecompilation();
-
-	UFUNCTION(BlueprintCallable, Category = "Control Rig Blueprint")
-	void SetAutoVMRecompile(bool bAutoRecompile) { bAutoRecompileVM = bAutoRecompile; }
-
-	void IncrementVMRecompileBracket();
-	void DecrementVMRecompileBracket();
-
-	// this is needed since even after load
-	// model data can change while the Control Rig BP is not opened
-	// for example, if a user defined struct changed after BP load,
-	// any pin that references the struct needs to be regenerated
-	void RefreshAllModels(EControlRigBlueprintLoadType InLoadType = EControlRigBlueprintLoadType::PostLoad);
-
-	// RigVMRegistry changes can be triggered when new user defined types(structs/enums) are added/removed
-	// in which case we have to refresh the model
-	void OnRigVMRegistryChanged();
-	
-	UFUNCTION(BlueprintCallable, Category = "Control Rig Blueprint")
-	void RequestControlRigInit();
-
-	UFUNCTION(BlueprintCallable, Category = "Control Rig Blueprint")
-	URigVMGraph* GetModel(const UEdGraph* InEdGraph = nullptr) const;
-	URigVMGraph* GetModel(const FString& InNodePath) const;
-
-	UFUNCTION(BlueprintCallable, Category = "Control Rig Blueprint")
-	URigVMGraph* GetDefaultModel() const;
-
-	UFUNCTION(BlueprintCallable, Category = "Control Rig Blueprint")
-	TArray<URigVMGraph*> GetAllModels() const;
-
-	UFUNCTION(BlueprintCallable, Category = "Control Rig Blueprint")
-	URigVMFunctionLibrary* GetLocalFunctionLibrary() const;
-
-	UFUNCTION(BlueprintCallable, Category = "Control Rig Blueprint")
-	URigVMGraph* AddModel(FString InName = TEXT("Rig Graph"), bool bSetupUndoRedo = true, bool bPrintPythonCommand = true);
-
-	UFUNCTION(BlueprintCallable, Category = "Control Rig Blueprint")
-	bool RemoveModel(FString InName = TEXT("Rig Graph"), bool bSetupUndoRedo = true, bool bPrintPythonCommand = true);
-
-	UFUNCTION(BlueprintCallable, Category = "Control Rig Blueprint")
-	URigVMGraph* GetFocusedModel() const;
-
-	UFUNCTION(BlueprintCallable, Category = "Control Rig Blueprint")
-	URigVMController* GetController(const URigVMGraph* InGraph = nullptr) const;
-
-	UFUNCTION(BlueprintCallable, Category = "Control Rig Blueprint")
-	URigVMController* GetControllerByName(const FString InGraphName = TEXT("")) const;
-
-	UFUNCTION(BlueprintCallable, Category = "Control Rig Blueprint")
-	URigVMController* GetOrCreateController(URigVMGraph* InGraph = nullptr);
-
-	URigVMController* GetController(const UEdGraph* InEdGraph) const;
-	URigVMController* GetOrCreateController(const UEdGraph* InGraph);
-
-	UFUNCTION(BlueprintCallable, Category = "Control Rig Blueprint")
-	TArray<FString> GeneratePythonCommands(const FString InNewBlueprintName);
-
-	URigVMGraph* GetTemplateModel(bool bIsFunctionLibrary = false);
-	URigVMController* GetTemplateController(bool bIsFunctionLibrary = false);
-
-#if WITH_EDITOR
-	UEdGraph* GetEdGraph(URigVMGraph* InModel) const;
-	UEdGraph* GetEdGraph(const FString& InNodePath) const;
-#endif
-
-#if WITH_EDITORONLY_DATA
-	UPROPERTY()
-	TObjectPtr<UControlRigGraph> FunctionLibraryEdGraph;
-#endif
-
-	bool IsFunctionPublic(const FName& InFunctionName) const;
-	void MarkFunctionPublic(const FName& InFunctionName, bool bIsPublic = true);
-
-	// Returns a list of dependencies of this blueprint.
-	// Dependencies are blueprints that contain functions used in this blueprint
-	TArray<UControlRigBlueprint*> GetDependencies(bool bRecursive = false) const;
-
-	// Returns a list of dependents as unresolved soft object pointers.
-	// A dependent is a blueprint which uses a function defined in this blueprint.
-	// This function is not recursive, since it avoids opening the asset.
-	// Use GetDependentBlueprints as an alternative.
-	TArray<FAssetData> GetDependentAssets() const;
-
-	// Returns a list of dependents as resolved blueprints.
-	// A dependent is a blueprint which uses a function defined in this blueprint.
-	// If bOnlyLoaded is false, this function loads the dependent assets and can introduce a large cost
-	// depending on the size / count of assets in the project.
-	TArray<UControlRigBlueprint*> GetDependentBlueprints(bool bRecursive = false, bool bOnlyLoaded = false) const;
-
-	UPROPERTY(EditAnywhere, Category = "User Interface")
-	FRigGraphDisplaySettings RigGraphDisplaySettings;
-
 	UPROPERTY(EditAnywhere, Category = "Hierarchy")
 	FRigHierarchySettings HierarchySettings;
 
-	UPROPERTY(EditAnywhere, Category = "VM")
-	FRigVMRuntimeSettings VMRuntimeSettings;
-
-	UPROPERTY(EditAnywhere, BlueprintReadWrite, Category = "VM")
-	FRigVMCompileSettings VMCompileSettings;
-
-	UPROPERTY(EditAnywhere, Category = "Python Log Settings")
-	FControlRigPythonSettings PythonLogSettings;
-
 protected:
 
-#if WITH_EDITORONLY_DATA
-	UPROPERTY()
-	TObjectPtr<URigVMGraph> Model_DEPRECATED;
-
-	UPROPERTY()
-	TObjectPtr<URigVMFunctionLibrary> FunctionLibrary_DEPRECATED;
-#endif
-
-	UPROPERTY()
-	FRigVMClient RigVMClient;
-
-#if WITH_EDITORONLY_DATA
-
-	UPROPERTY()
-	bool ReferencedObjectPathsStored;
-
-	UPROPERTY()
-	TArray<FSoftObjectPath> ReferencedObjectPaths;
-
-#endif
-
 	/** Asset searchable information about exposed public functions on this rig */
 	UPROPERTY(AssetRegistrySearchable)
 	TArray<FControlRigPublicFunctionData> PublicFunctions_DEPRECATED;
 
-	/** Asset searchable information about exposed public functions on this rig */
-	UPROPERTY(AssetRegistrySearchable)
-	TArray<FRigVMGraphFunctionHeader> PublicGraphFunctions;
-
-	/** Asset searchable information function references in this rig */
-	UPROPERTY(AssetRegistrySearchable)
-	TArray<FRigVMReferenceNodeData> FunctionReferenceNodeData;
-
-#if WITH_EDITORONLY_DATA
-
-	UPROPERTY(transient, DuplicateTransient)
-	TObjectPtr<URigVMGraph> TemplateModel;
-
-	UPROPERTY(transient, DuplicateTransient)
-	TObjectPtr<URigVMController> TemplateController;
-
-	mutable TArray<FAssetRegistryTag> CachedAssetTags;
-
-#endif
+	virtual void SetupDefaultObjectDuringCompilation(URigVMHost* InCDO) override;
 
 public:
 
-	UPROPERTY(transient, DuplicateTransient)
-	TMap<FString, FRigVMOperand> PinToOperandMap;
-
-	bool bSuspendModelNotificationsForSelf;
-	bool bSuspendModelNotificationsForOthers;
-	bool bSuspendAllNotifications;
-
-	void PopulateModelFromGraphForBackwardsCompatibility(UControlRigGraph* InGraph);
-	void SetupPinRedirectorsForBackwardsCompatibility();
-	void RebuildGraphFromModel();
-
-	FRigVMGraphModifiedEvent& OnModified();
-	FOnVMCompiledEvent& OnVMCompiled();
+	virtual void SetupPinRedirectorsForBackwardsCompatibility() override;
 
 	UFUNCTION(BlueprintCallable, Category = "VM")
 	static TArray<UControlRigBlueprint*> GetCurrentlyOpenRigBlueprints();
 
-	UFUNCTION(BlueprintCallable, Category = "VM")
-	UClass* GetControlRigClass();
-
-	UFUNCTION(BlueprintCallable, Category = "VM")
-	UControlRig* CreateControlRig();
-
-	UFUNCTION(BlueprintCallable, Category = "VM")
-	static TArray<UStruct*> GetAvailableRigUnits();
-
-#if WITH_EDITOR
-	UFUNCTION(BlueprintCallable, Category = "Variables")
-	TArray<FRigVMGraphVariableDescription> GetMemberVariables() const;
-	
-	UFUNCTION(BlueprintCallable, Category = "Variables")
-	FName AddMemberVariable(const FName& InName, const FString& InCPPType, bool bIsPublic = false, bool bIsReadOnly = false, FString InDefaultValue = TEXT(""));
-
-	UFUNCTION(BlueprintCallable, Category = "Variables")
-	bool RemoveMemberVariable(const FName& InName);
-
-	UFUNCTION(BlueprintCallable, Category = "Variables")
-	bool RenameMemberVariable(const FName& InOldName, const FName& InNewName);
-
-	UFUNCTION(BlueprintCallable, Category = "Variables")
-	bool ChangeMemberVariableType(const FName& InName, const FString& InCPPType, bool bIsPublic = false, bool bIsReadOnly = false, FString InDefaultValue = TEXT(""));
-#endif
-
 #if WITH_EDITORONLY_DATA
 	UPROPERTY()
 	TSoftObjectPtr<UControlRigShapeLibrary> GizmoLibrary_DEPRECATED;
@@ -571,85 +226,30 @@ private:
 	UPROPERTY(DuplicateTransient, AssetRegistrySearchable)
 	TSoftObjectPtr<UObject> SourceCurveImport;
 
-	/** The event names this control rig blueprint contains */
-	UPROPERTY(AssetRegistrySearchable)
-	TArray<FName> SupportedEventNames;
-
 	/** If set to true, this control rig has animatable controls */
 	UPROPERTY(AssetRegistrySearchable)
 	bool bExposesAnimatableControls;
 
-	UPROPERTY(transient, DuplicateTransient)
-	bool bAutoRecompileVM;
-
-	UPROPERTY(transient, DuplicateTransient)
-	bool bVMRecompilationRequired;
-
-	UPROPERTY(transient, DuplicateTransient)
-	bool bIsCompiling;
-
-	UPROPERTY(transient, DuplicateTransient)
-	int32 VMRecompilationBracket;
-
-	FRigVMGraphModifiedEvent ModifiedEvent;
-	void Notify(ERigVMGraphNotifType InNotifType, UObject* InSubject);
-	void HandleModifiedEvent(ERigVMGraphNotifType InNotifType, URigVMGraph* InGraph, UObject* InSubject);
-
-#if WITH_EDITOR
-
-	UFUNCTION(BlueprintCallable, Category = "Control Rig Blueprint")
-	void SuspendNotifications(bool bSuspendNotifs);
-
-	FOnRefreshEditorEvent RefreshEditorEvent;
-	FOnVariableDroppedEvent VariableDroppedEvent;
-	FOnBreakpointAdded BreakpointAddedEvent;
-
-public:
-
-	void BroadcastRefreshEditor() { return RefreshEditorEvent.Broadcast(this); }
-	FOnRefreshEditorEvent& OnRefreshEditor() { return RefreshEditorEvent; }
-	FOnVariableDroppedEvent& OnVariableDropped() { return VariableDroppedEvent; }
-	FOnBreakpointAdded& OnBreakpointAdded() { return BreakpointAddedEvent; }
-
 private:
 
-#endif
-
-	FOnVMCompiledEvent VMCompiledEvent;
-
 	static TArray<UControlRigBlueprint*> sCurrentlyOpenedRigBlueprints;
 
-	void CreateMemberVariablesOnLoad();
-#if WITH_EDITOR
-	static FName FindCRMemberVariableUniqueName(TSharedPtr<FKismetNameValidator> InNameValidator, const FString& InBaseName);
-	static int32 AddCRMemberVariable(UControlRigBlueprint* InBlueprint, const FName& InVarName, FEdGraphPinType InVarType, bool bIsPublic, bool bIsReadOnly, FString InDefaultValue);
-	FName AddCRMemberVariableFromExternal(FRigVMExternalVariable InVariableToCreate, FString InDefaultValue = FString());
-#endif
-	void PatchFunctionReferencesOnLoad();
-	void PatchVariableNodesOnLoad();
+	virtual void PathDomainSpecificContentOnLoad() override;
+	virtual void PatchFunctionsOnLoad() override;
 	void PatchRigElementKeyCacheOnLoad();
-	void PatchBoundVariables();
-	void PatchVariableNodesWithIncorrectType();
 	void PatchPropagateToChildren();
-	void PatchParameterNodesOnLoad();
-	void PatchLinksWithCast();
-	void PatchFunctionsOnLoad();
-
-	TMap<FName, int32> AddedMemberVariableMap;
-	TArray<FBPVariableDescription> LastNewVariables;
 
 public:
-	void PropagatePoseFromInstanceToBP(UControlRig* InControlRig);
-	void PropagatePoseFromBPToInstances();
-	void PropagateHierarchyFromBPToInstances();
-	void PropagateDrawInstructionsFromBPToInstances();
-	void PropagateRuntimeSettingsFromBPToInstances();
-	void PropagatePropertyFromBPToInstances(FRigElementKey InRigElement, const FProperty* InProperty);
-	void PropagatePropertyFromInstanceToBP(FRigElementKey InRigElement, const FProperty* InProperty, UControlRig* InInstance);
+	void PropagatePoseFromInstanceToBP(UControlRig* InControlRig) const;
+	void PropagatePoseFromBPToInstances() const;
+	void PropagateHierarchyFromBPToInstances() const;
+	void PropagateDrawInstructionsFromBPToInstances() const;
+	void PropagatePropertyFromBPToInstances(FRigElementKey InRigElement, const FProperty* InProperty) const;
+	void PropagatePropertyFromInstanceToBP(FRigElementKey InRigElement, const FProperty* InProperty, UControlRig* InInstance) const;
 
 	/**
 	* Returns the modified event, which can be used to 
-	* subscribe to topological changes happening within the hierarchy. The event is broadcasted only after all hierarchy instances are up to date
+	* subscribe to topological changes happening within the hierarchy. The event is broadcast only after all hierarchy instances are up to date
 	* @return The event used for subscription.
 	*/
 	FRigHierarchyModifiedEvent& OnHierarchyModified() { return HierarchyModifiedEvent; }
@@ -664,18 +264,7 @@ private:
 	void HandleHierarchyModified(ERigHierarchyNotification InNotification, URigHierarchy* InHierarchy, const FRigBaseElement* InElement);
 
 #if WITH_EDITOR
-	void HandlePackageDone(const FEndLoadPackageContext& Context);
-	void HandlePackageDone();
-	// ControlRigBP, once end-loaded, will inform other ControlRig-Dependent systems that ControlRig instances are ready.
-	void BroadcastControlRigPackageDone();
-
-	// Previously some memory classes were parented to the asset object
-	// however it is no longer supported since classes are now identified 
-	// with only package name + class name, see FTopLevelAssetPath
-	// this function removes those deprecated class.
-	// new classes should be created by RecompileVM and parented to the Package
-	// during PostLoad
-	void RemoveDeprecatedVMMemoryClass() const;
+	virtual void HandlePackageDone() override;
 #endif
 
 	// Class used to temporarily cache all 
@@ -697,124 +286,15 @@ private:
 	float DebugBoneRadius;
 
 #if WITH_EDITOR
-
-public:
-
-	FOnExternalVariablesChanged& OnExternalVariablesChanged() { return ExternalVariablesChangedEvent; }
-
-	virtual void OnPreVariableChange(UObject* InObject);
-	virtual void OnPostVariableChange(UBlueprint* InBlueprint);
-	virtual void OnVariableAdded(const FName& InVarName);
-	virtual void OnVariableRemoved(const FName& InVarName);
-	virtual void OnVariableRenamed(const FName& InOldVarName, const FName& InNewVarName);
-	virtual void OnVariableTypeChanged(const FName& InVarName, FEdGraphPinType InOldPinType, FEdGraphPinType InNewPinType);
-
-	FOnNodeDoubleClicked& OnNodeDoubleClicked() { return NodeDoubleClickedEvent; }
-	void BroadcastNodeDoubleClicked(URigVMNode* InNode);
-
-	FOnGraphImported& OnGraphImported() { return GraphImportedEvent; }
-	void BroadcastGraphImported(UEdGraph* InGraph);
-
-	FOnPostEditChangeChainProperty& OnPostEditChangeChainProperty() { return PostEditChangeChainPropertyEvent; }
-	void BroadcastPostEditChangeChainProperty(FPropertyChangedChainEvent& PropertyChangedChainEvent);
-
-	FOnLocalizeFunctionDialogRequested& OnRequestLocalizeFunctionDialog() { return RequestLocalizeFunctionDialog; }
-	void BroadcastRequestLocalizeFunctionDialog(FRigVMGraphFunctionIdentifier InFunction, bool bForce = false);
-
-	FControlRigOnBulkEditDialogRequestedDelegate& OnRequestBulkEditDialog() { return RequestBulkEditDialog; }
-
-	FControlRigOnBreakLinksDialogRequestedDelegate& OnRequestBreakLinksDialog() { return RequestBreakLinksDialog; }
-
-	FControlRigOnPinTypeSelectionRequestedDelegate& OnRequestPinTypeSelectionDialog() { return RequestPinTypeSelectionDialog; }
-
-	FRigVMController_RequestJumpToHyperlinkDelegate& OnRequestJumpToHyperlink() { return RequestJumpToHyperlink; };
-
-	FOnReportCompilerMessage& OnReportCompilerMessage() { return ReportCompilerMessageEvent; }
-	void BroadCastReportCompilerMessage(EMessageSeverity::Type InSeverity, UObject* InSubject, const FString& InMessage);
-
-private:
-
-	FOnExternalVariablesChanged ExternalVariablesChangedEvent;
-	bool bUpdatingExternalVariables;
-	void BroadcastExternalVariablesChangedEvent();
-	FCompilerResultsLog CompileLog;
-
-	FOnNodeDoubleClicked NodeDoubleClickedEvent;
-	FOnGraphImported GraphImportedEvent;
-	FOnPostEditChangeChainProperty PostEditChangeChainPropertyEvent;
-	FOnLocalizeFunctionDialogRequested RequestLocalizeFunctionDialog;
-	FOnReportCompilerMessage ReportCompilerMessageEvent;
-	FControlRigOnBulkEditDialogRequestedDelegate RequestBulkEditDialog;
-	FControlRigOnBreakLinksDialogRequestedDelegate RequestBreakLinksDialog;
-	FControlRigOnPinTypeSelectionRequestedDelegate RequestPinTypeSelectionDialog;
-	FRigVMController_RequestJumpToHyperlinkDelegate RequestJumpToHyperlink;
-
-#endif
-
-	UEdGraph* CreateEdGraph(URigVMGraph* InModel, bool bForce = false);
-	bool RemoveEdGraph(URigVMGraph* InModel);
-	void DestroyObject(UObject* InObject);
-	void RenameGraph(const FString& InNodePath, const FName& InNewName);
-	void CreateEdGraphForCollapseNodeIfNeeded(URigVMCollapseNode* InNode, bool bForce = false);
-	bool RemoveEdGraphForCollapseNode(URigVMCollapseNode* InNode, bool bNotify = false);
-	void HandleReportFromCompiler(EMessageSeverity::Type InSeverity, UObject* InSubject, const FString& InMessage);
-
-	TArray<IRigVMGraphFunctionHost*> GetReferencedFunctionHosts(bool bForceLoad);
-	
-#if WITH_EDITOR
-private:
-	bool bCompileInDebugMode;
-	
-	TArray<URigVMNode*> RigVMBreakpointNodes;
-
-	FOnRequestInspectObject OnRequestInspectObjectEvent;
-	FControlRigBlueprintGetFocusedGraph OnGetFocusedGraphDelegate;
 		
 public:
 
-	/** Sets the execution mode. In Release mode the rig will ignore all breakpoints. */
-	void SetDebugMode(const bool bValue) { bCompileInDebugMode = bValue; }
-
-	/** Removes all the breakpoints from the blueprint and the VM */
-	void ClearBreakpoints();
-
-	/** Adds a breakpoint to all loaded blueprints which use the node indicated by InBreakpointNodePath 
-	  * If the node is inside a public function, it will add a breakpoint to all blueprints calling this function. */
-	bool AddBreakpoint(const FString& InBreakpointNodePath);
-
-	/** Adds a breakpoint to all loaded blueprints which use the InBreakpointNode. 
-	  * If LibraryNode is not null, it indicates that the library uses the InBreapointNode, and the function will add
-	  * breakpoints to any other loaded blueprint that references this library. */
-	bool AddBreakpoint(URigVMNode* InBreakpointNode, URigVMLibraryNode* LibraryNode = nullptr);
-
-	/** Adds a breakpoint to the first instruction of each callpath related to the InBreakpointNode */
-	bool AddBreakpointToControlRig(URigVMNode* InBreakpointNode);
-
-	/** Removes the given breakpoint from all the loaded blueprints that use this node, and recomputes all breakpoints
-	  * in the VM  */
-	bool RemoveBreakpoint(const FString& InBreakpointNodePath);
-	bool RemoveBreakpoint(URigVMNode* InBreakpointNode);
-
-	/** Recomputes the instruction breakpoints given the node breakpoints in the blueprint */
-	void RefreshControlRigBreakpoints();
-
 	/** Shape libraries to load during package load completed */ 
 	TArray<FString> ShapeLibrariesToLoadOnPackageLoaded;
 
-	TArray<FRigVMReferenceNodeData> GetReferenceNodeData() const;
-
 #endif
 
-protected:
-
-	static FSoftObjectPath PreDuplicateAssetPath;
-	static FSoftObjectPath PreDuplicateHostPath;
-
 private:
-	bool bDirtyDuringLoad;
-	bool bErrorsDuringCompilation;
-	bool bSuspendPythonMessagesForRigVMClient;
-	bool bMarkBlueprintAsStructurallyModifiedPending;
 
 	friend class FControlRigBlueprintCompilerContext;
 	friend class SRigHierarchy;
@@ -834,24 +314,3 @@ private:
 	friend struct FControlRigGraphSchemaAction_PromoteToVariable;
 	friend class UControlRigGraphSchema;
 };
-
-class CONTROLRIGDEVELOPER_API FControlRigBlueprintVMCompileScope
-{
-public:
-   
-	FControlRigBlueprintVMCompileScope(UControlRigBlueprint *InBlueprint)
-	: Blueprint(InBlueprint)
-	{
-		check(Blueprint);
-		Blueprint->IncrementVMRecompileBracket();
-	}
-
-	~FControlRigBlueprintVMCompileScope()
-	{
-		Blueprint->DecrementVMRecompileBracket();
-	}
-   
-private:
-
-	UControlRigBlueprint* Blueprint;
-};
\ No newline at end of file
diff --git a/Engine/Plugins/Animation/ControlRig/Source/ControlRigDeveloper/Public/Graph/ControlRigGraph.h b/Engine/Plugins/Animation/ControlRig/Source/ControlRigDeveloper/Public/Graph/ControlRigGraph.h
--- a/Engine/Plugins/Animation/ControlRig/Source/ControlRigDeveloper/Public/Graph/ControlRigGraph.h
+++ b/Engine/Plugins/Animation/ControlRig/Source/ControlRigDeveloper/Public/Graph/ControlRigGraph.h
@@ -9,6 +9,7 @@
 #include "Rigs/RigHierarchyController.h"
 #include "RigVMModel/RigVMGraph.h"
 #include "RigVMCore/RigVM.h"
+#include "EdGraph/RigVMEdGraph.h"
 #include "ControlRigGraph.generated.h"
 
 class UControlRigBlueprint;
@@ -17,31 +18,19 @@ class UControlRig;
 class URigVMController;
 struct FRigCurveContainer;
 
-DECLARE_MULTICAST_DELEGATE_OneParam(FControlRigGraphNodeClicked, UControlRigGraphNode*);
-
 UCLASS()
-class CONTROLRIGDEVELOPER_API UControlRigGraph : public UEdGraph, public IRigVMEditorSideObject
+class CONTROLRIGDEVELOPER_API UControlRigGraph : public URigVMEdGraph
 {
 	GENERATED_BODY()
 
 public:
 	UControlRigGraph();
 
-	/** IRigVMEditorSideObject interface */
-	virtual FRigVMClient* GetRigVMClient() const override;
-	virtual FString GetRigVMNodePath() const override;
-	virtual void HandleRigVMGraphRenamed(const FString& InOldNodePath, const FString& InNewNodePath) override;
-
 	/** Set up this graph */
-	void Initialize(UControlRigBlueprint* InBlueprint);
+	virtual void InitializeFromBlueprint(URigVMBlueprint* InBlueprint) override;
 
-	/** Get the skeleton graph schema */
-	const UControlRigGraphSchema* GetControlRigGraphSchema();
+	virtual bool HandleModifiedEvent_Internal(ERigVMGraphNotifType InNotifType, URigVMGraph* InGraph, UObject* InSubject) override;
 
-#if WITH_EDITORONLY_DATA
-	/** Customize blueprint changes based on backwards compatibility */
-	virtual void Serialize(FArchive& Ar) override;
-#endif
 #if WITH_EDITOR
 
 	const TArray<TSharedPtr<FString>>* GetBoneNameList(URigVMPin* InPin = nullptr) const
@@ -73,31 +62,6 @@ public:
 	const TArray<TSharedPtr<FString>>* GetEntryNameList(URigVMPin* InPin = nullptr) const;
 	const TArray<TSharedPtr<FString>>* GetShapeNameList(URigVMPin* InPin = nullptr) const;
 
-	bool bSuspendModelNotifications;
-	bool bIsTemporaryGraphForCopyPaste;
-
-	UEdGraphNode* FindNodeForModelNodeName(const FName& InModelNodeName, const bool bCacheIfRequired = true);
-
-	UControlRigBlueprint* GetBlueprint() const;
-	URigVMGraph* GetModel() const;
-	URigVMController* GetController() const;
-	bool IsRootGraph() const { return GetRootGraph() == this; }
-	const UControlRigGraph* GetRootGraph() const;
-
-	void HandleModifiedEvent(ERigVMGraphNotifType InNotifType, URigVMGraph* InGraph, UObject* InSubject);
-	void ConsumeQueuedNotifications();
-
-	int32 GetInstructionIndex(const UControlRigGraphNode* InNode, bool bAsInput);
-
-	UPROPERTY()
-	FString ModelNodePath;
-
-	UPROPERTY()
-	bool bIsFunctionDefinition;
-
-protected:
-	using Super::AddNode;
-	virtual void AddNode(UEdGraphNode* NodeToAdd, bool bUserAction = false, bool bSelectNewNode = true) override;
 
 private:
 
@@ -156,32 +120,9 @@ private:
 	TArray<TSharedPtr<FString>> ShapeNameList;
 	int32 LastHierarchyTopologyVersion;
 
-	bool bIsSelecting;
-
-	FControlRigGraphNodeClicked OnGraphNodeClicked;
-
-	TMap<URigVMNode*, TPair<int32, int32>> CachedInstructionIndices;
-
-	void RemoveNode(UEdGraphNode* InNode);
-
-#endif
-#if WITH_EDITORONLY_DATA
-
-	UPROPERTY(transient)
-	TObjectPtr<URigVMController> TemplateController;
+	static TArray<TSharedPtr<FString>> EmptyElementNameList;
 
 #endif
-#if WITH_EDITOR
-
-	URigVMController* GetTemplateController();
-
-protected:
-	void HandleVMCompiledEvent(UObject* InCompiledObject, URigVM* InVM);
-
-private:
-	static TArray<TSharedPtr<FString>> EmptyElementNameList;
-	TMap<FName, UEdGraphNode*> ModelNodePathToEdNode;
-	mutable TWeakObjectPtr<URigVMGraph> CachedModelGraph;
 
 	friend class UControlRigUnitNodeSpawner;
 	friend class UControlRigVariableNodeSpawner;
@@ -195,7 +136,6 @@ private:
 	friend class UControlRigArrayNodeSpawner;
 	friend class UControlRigInvokeEntryNodeSpawner;
 
-#endif
 	friend class UControlRigGraphNode;
 	friend class FControlRigEditor;
 	friend class SControlRigGraphNode;
diff --git a/Engine/Plugins/Animation/ControlRig/Source/ControlRigDeveloper/Public/Graph/ControlRigGraphNode.h b/Engine/Plugins/Animation/ControlRig/Source/ControlRigDeveloper/Public/Graph/ControlRigGraphNode.h
--- a/Engine/Plugins/Animation/ControlRig/Source/ControlRigDeveloper/Public/Graph/ControlRigGraphNode.h
+++ b/Engine/Plugins/Animation/ControlRig/Source/ControlRigDeveloper/Public/Graph/ControlRigGraphNode.h
@@ -11,6 +11,7 @@
 #include "UObject/SoftObjectPath.h"
 #include "EdGraphSchema_K2.h"
 #include "RigVMModel/RigVMGraph.h"
+#include "EdGraph/RigVMEdGraphNode.h"
 #include "ControlRigGraphNode.generated.h"
 
 class FBlueprintActionDatabaseRegistrar;
@@ -20,7 +21,7 @@ class UControlRigBlueprint;
 
 /** Base class for animation ControlRig-related nodes */
 UCLASS()
-class CONTROLRIGDEVELOPER_API UControlRigGraphNode : public UEdGraphNode
+class CONTROLRIGDEVELOPER_API UControlRigGraphNode : public URigVMEdGraphNode
 {
 	GENERATED_BODY()
 
@@ -34,246 +35,15 @@ class CONTROLRIGDEVELOPER_API UControlRigGraphNode : public UEdGraphNode
 	friend class FControlRigEditor;
 	friend class SControlRigGraphPinCurveFloat;
 
-private:
-
-	UPROPERTY()
-	FString ModelNodePath;
-
-	UPROPERTY(transient)
-	TWeakObjectPtr<URigVMNode> CachedModelNode;
-
-	UPROPERTY(transient)
-	TMap<FString, TWeakObjectPtr<URigVMPin>> PinPathToModelPin;
-
-#if WITH_EDITORONLY_DATA
-	/** The property we represent. For template nodes this represents the struct/property type name. */
-	UPROPERTY()
-	FName PropertyName_DEPRECATED;
-
-	UPROPERTY()
-	FString StructPath_DEPRECATED;
-
-	/** Pin Type for property */
-	UPROPERTY()
-	FEdGraphPinType PinType_DEPRECATED;
-
-	/** The type of parameter */
-	UPROPERTY()
-	int32 ParameterType_DEPRECATED;
-
-	/** Expanded pins */
-	UPROPERTY()
-	TArray<FString> ExpandedPins_DEPRECATED;
-#endif
-
-	/** Cached dimensions of this node (used for auto-layout) */
-	FVector2D Dimensions;
-
-	/** The cached node titles */
-	mutable FText NodeTitle;
-
-	/** The cached fulol node title */
-	mutable FText FullNodeTitle;
-
 public:
 
-	DECLARE_MULTICAST_DELEGATE(FNodeTitleDirtied);
-	DECLARE_MULTICAST_DELEGATE(FNodePinsChanged);
-	DECLARE_MULTICAST_DELEGATE(FNodePinExpansionChanged);
-	DECLARE_MULTICAST_DELEGATE(FNodeBeginRemoval);
-
-	struct FPinPair
-	{
-		FPinPair()
-			: InputPin(nullptr)
-			, OutputPin(nullptr)
-		{}
-		UEdGraphPin* InputPin;
-		UEdGraphPin* OutputPin;
-
-		bool IsValid() const { return InputPin != nullptr || OutputPin != nullptr; }
-	};
-
 	UControlRigGraphNode();
 
-	// UObject Interface.
-#if WITH_EDITOR
-	virtual bool Modify( bool bAlwaysMarkDirty=true ) { return false; }
-#endif
-	
-	// UEdGraphNode Interface.
-	virtual FText GetNodeTitle(ENodeTitleType::Type TitleType) const override;
-	virtual FLinearColor GetNodeTitleColor() const override;
-	virtual FLinearColor GetNodeBodyTintColor() const;
-	virtual bool ShowPaletteIconOnNode() const override;
-	virtual FSlateIcon GetIconAndTint(FLinearColor& OutColor) const override;
-	virtual void AllocateDefaultPins() override;
-	virtual void ReconstructNode() override;
-	virtual void ReconstructNode_Internal(bool bForce = false);
-	virtual void GetNodeContextMenuActions(class UToolMenu* Menu, class UGraphNodeContextMenuContext* Context) const override;
-	virtual void DestroyNode() override;
-	virtual void PinDefaultValueChanged(UEdGraphPin* Pin) override;
-	virtual FText GetTooltipText() const override;
-	virtual bool CanCreateUnderSpecifiedSchema(const UEdGraphSchema* InSchema) const override;
-	virtual void AutowireNewNode(UEdGraphPin* FromPin) override;	
-	virtual bool SupportsCommentBubble() const override { return false; }
-	virtual bool IsSelectedInEditor() const;
-	virtual bool ShouldDrawNodeAsControlPointOnly(int32& OutInputPinIndex, int32& OutOutputPinIndex) const override;
-	virtual void BeginDestroy() override;
-
-	virtual bool IsDeprecated() const override;
-	virtual FEdGraphNodeDeprecationResponse GetDeprecationResponse(EEdGraphNodeDeprecationType DeprecationType) const override;
-
-	/** Set the cached dimensions of this node */
-	void SetDimensions(const FVector2D& InDimensions) { Dimensions = InDimensions; }
-
-	/** Get the cached dimensions of this node */
-	const FVector2D& GetDimensions() const { return Dimensions; }
-
-	/** Check a pin's expansion state */
-	bool IsPinExpanded(const FString& InPinPath);
-
-	/** Propagate pin defaults to underlying properties if they have changed */
-	void CopyPinDefaultsToModel(UEdGraphPin* Pin, bool bUndo = false, bool bPrintPythonCommand = false);
-
-	/** Get the blueprint that this node is contained within */
-	UControlRigBlueprint* GetBlueprint() const;
-
-	/** Get the VM model this node lives within */
-	URigVMGraph* GetModel() const;
-
-	/** Get the blueprint that this node is contained within */
-	URigVMController* GetController() const;
-
-	/** Get the VM node this is node is wrapping */
-	URigVMNode* GetModelNode() const;
-
-	/** Get the VM node name this node is wrapping */
-	FName GetModelNodeName() const;
-
-	URigVMPin* GetModelPinFromPinPath(const FString& InPinPath) const;
-
-	/** Add a new element to the aggregate node referred to by the property path */
-	void HandleAddAggregateElement(const FString& InNodePath);
-
-	/** Add a new array element to the array referred to by the property path */
-	void HandleAddArrayElement(FString InPinPath);
-	
-	/** Clear the array referred to by the property path */
-	void HandleClearArray(FString InPinPath);
-
-	/** Remove the array element referred to by the property path */
-	void HandleRemoveArrayElement(FString InPinPath);
-
-	/** Insert a new array element after the element referred to by the property path */
-	void HandleInsertArrayElement(FString InPinPath);
-
-	int32 GetInstructionIndex(bool bAsInput) const;
-
-	const FRigVMTemplate* GetTemplate() const;
-
-	void ClearErrorInfo();
-
-	URigVMPin* FindModelPinFromGraphPin(const UEdGraphPin* InGraphPin) const;
-	UEdGraphPin* FindGraphPinFromModelPin(const URigVMPin* InModelPin, bool bAsInput) const;
-
-	/// Synchronize the stored name/value/type on the graph pin with the value stored on the node. 
-	/// If the pin has sub-pins, the value update is done recursively.
-	void SynchronizeGraphPinNameWithModelPin(const URigVMPin* InModelPin, bool bNotify = true);
-	void SynchronizeGraphPinValueWithModelPin(const URigVMPin* InModelPin);
-	void SynchronizeGraphPinTypeWithModelPin(const URigVMPin* InModelPin);
-	void SynchronizeGraphPinExpansionWithModelPin(const URigVMPin* InModelPin);
-	
-	void SyncGraphNodeTitleWithModelNodeTitle();
-	void SyncGraphNodeNameWithModelNodeName(const URigVMNode* InModelNode);
-
-	FNodeTitleDirtied& OnNodeTitleDirtied() { return NodeTitleDirtied; }
-	FNodePinsChanged& OnNodePinsChanged() { return NodePinsChanged; }
-	FNodePinExpansionChanged& OnNodePinExpansionChanged() { return NodePinExpansionChanged; }
-	FNodeBeginRemoval& OnNodeBeginRemoval() { return NodeBeginRemoval; }
-
-	/** Called when there's a drastic change in the pins */
-	bool ModelPinsChanged(bool bForce = false);
-
-	/** Called when a model pin is added after the node creation */
-	bool ModelPinAdded(const URigVMPin* InModelPin);
-
-	/** Called when a model pin is being removed */
-	bool ModelPinRemoved(const URigVMPin* InModelPin);
-
-	/** Returns true if this node is relying on the cast template */
-	bool DrawAsCompactNode() const;
-
-protected:
-
-	FLinearColor GetNodeProfilingColor() const;
-	FLinearColor GetNodeOpacityColor() const;
-
-	/** Helper function for AllocateDefaultPins */
-	void UpdatePinLists();
-	bool CreateGraphPinFromModelPin(const URigVMPin* InModelPin, EEdGraphPinDirection InDirection,  UEdGraphPin* InParentPin = nullptr);
-	void RemoveGraphSubPins(UEdGraphPin *InParentPin, const TArray<UEdGraphPin*>& InPinsToKeep = TArray<UEdGraphPin*>());
-	bool ModelPinAdded_Internal(const URigVMPin* InModelPin);
-	bool ModelPinRemoved_Internal(const URigVMPin* InModelPin);
-
-	/** Copies default values from underlying properties into pin defaults, for editing */
-	void SetupPinDefaultsFromModel(UEdGraphPin* Pin, const URigVMPin* InModelPin = nullptr);
-
-	/** Recreate pins when we reconstruct this node */
-	virtual void ReallocatePinsDuringReconstruction(const TArray<UEdGraphPin*>& OldPins);
-
-	/** Wire-up new pins given old pin wiring */
-	virtual void RewireOldPinsToNewPins(TArray<UEdGraphPin*>& InOldPins, TArray<UEdGraphPin*>& InNewPins);
-
-	/** Handle anything post-reconstruction */
-	virtual void PostReconstructNode();
-
-	/** Something that could change our title has changed */
-	void InvalidateNodeTitle() const;
-
-	/** Destroy all pins in an array */
-	void DestroyPinList(TArray<UEdGraphPin*>& InPins);
-
-	/** Sets the body + title color from a color provided by the model */
-	void SetColorFromModel(const FLinearColor& InColor);
-
-	UClass* GetControlRigGeneratedClass() const;
-	UClass* GetControlRigSkeletonGeneratedClass() const;
-
-	static FEdGraphPinType GetPinTypeForModelPin(const URigVMPin* InModelPin);
+	// URigVMEdGraphNode
+	virtual void GetNodeContextMenuActions(UToolMenu* Menu, UGraphNodeContextMenuContext* Context) const override;
 
 private:
 
-	int32 GetNodeTopologyVersion() const { return NodeTopologyVersion; }
-	int32 NodeTopologyVersion;
-
-	static void ConfigurePin(UEdGraphPin* EdGraphPin, const URigVMPin* ModelPin);
-	TArray<URigVMPin*>& PinListForPin(const URigVMPin* InModelPin);
-
-	FLinearColor CachedTitleColor;
-	FLinearColor CachedNodeColor;
-
-#if WITH_EDITOR
-	bool bEnableProfiling;
-#endif
-
-	TArray<URigVMPin*> ExecutePins;
-	TArray<URigVMPin*> InputOutputPins;
-	TArray<URigVMPin*> InputPins;
-	TArray<URigVMPin*> OutputPins;
-	TArray<TSharedPtr<FRigVMExternalVariable>> ExternalVariables;
-	TArray<UEdGraphPin*> LastEdGraphPins;
-	
-	TMap<URigVMPin*, FPinPair> CachedPins;
-
-	FNodeTitleDirtied NodeTitleDirtied;
-	FNodePinsChanged NodePinsChanged;
-	FNodePinExpansionChanged NodePinExpansionChanged;
-	FNodeBeginRemoval NodeBeginRemoval;
-
-	mutable const FRigVMTemplate* CachedTemplate;
-	mutable TOptional<bool> DrawAsCompactNodeCache;
-
 	friend class SControlRigGraphNode;
 	friend class FControlRigArgumentLayout;
 	friend class FControlRigGraphDetails;
diff --git a/Engine/Plugins/Animation/ControlRig/Source/ControlRigDeveloper/Public/Graph/ControlRigGraphSchema.h b/Engine/Plugins/Animation/ControlRig/Source/ControlRigDeveloper/Public/Graph/ControlRigGraphSchema.h
--- a/Engine/Plugins/Animation/ControlRig/Source/ControlRigDeveloper/Public/Graph/ControlRigGraphSchema.h
+++ b/Engine/Plugins/Animation/ControlRig/Source/ControlRigDeveloper/Public/Graph/ControlRigGraphSchema.h
@@ -3,11 +3,10 @@
 #pragma once
 
 #include "ControlRigGraphNode.h"
-#include "EdGraph/EdGraphSchema.h"
 #include "GraphEditorDragDropAction.h"
 #include "RigVMModel/RigVMGraph.h"
 #include "EdGraphSchema_K2_Actions.h"
-#include "Kismet2/Kismet2NameValidators.h"
+#include "EdGraph/RigVMEdGraphSchema.h"
 
 #include "ControlRigGraphSchema.generated.h"
 
@@ -17,230 +16,8 @@ class UControlRigGraphNode;
 class UControlRigGraphNode_Unit;
 class UControlRigGraphNode_Property;
 
-/** Extra operations that can be performed on pin connection */
-enum class ECanCreateConnectionResponse_Extended
-{
-	None,
-
-	BreakChildren,
-
-	BreakParent,
-};
-
-/** Struct used to extend the response to a conneciton request to include breaking parents/children */
-struct FControlRigPinConnectionResponse
-{
-	FControlRigPinConnectionResponse(const ECanCreateConnectionResponse InResponse, FText InMessage, ECanCreateConnectionResponse_Extended InExtendedResponse = ECanCreateConnectionResponse_Extended::None)
-		: Response(InResponse, MoveTemp(InMessage))
-		, ExtendedResponse(InExtendedResponse)
-	{
-	}
-
-	friend bool operator==(const FControlRigPinConnectionResponse& A, const FControlRigPinConnectionResponse& B)
-	{
-		return (A.Response == B.Response) && (A.ExtendedResponse == B.ExtendedResponse);
-	}	
-
-	FPinConnectionResponse Response;
-	ECanCreateConnectionResponse_Extended ExtendedResponse;
-};
-
-/////////////////////////////////////////////////////
-// FControlRigLocalVariableNameValidator
-class CONTROLRIGDEVELOPER_API FControlRigLocalVariableNameValidator : public FStringSetNameValidator
-{
-
-public:
-	FControlRigLocalVariableNameValidator(const class UBlueprint* Blueprint, const URigVMGraph* Graph, FName InExistingName = NAME_None);
-
-	// Begin FNameValidatorInterface
-	virtual EValidatorResult IsValid(const FString& Name, bool bOriginal) override;
-	virtual EValidatorResult IsValid(const FName& Name, bool bOriginal) override;
-	// End FNameValidatorInterface
-};
-
-/////////////////////////////////////////////////////
-// FControlRigNameValidator
-class CONTROLRIGDEVELOPER_API FControlRigNameValidator : public FStringSetNameValidator
-{
-
-public:
-	FControlRigNameValidator(const class UBlueprint* Blueprint, const UStruct* ValidationScope, FName InExistingName = NAME_None);
-
-	// Begin FNameValidatorInterface
-	virtual EValidatorResult IsValid(const FString& Name, bool bOriginal) override;
-	virtual EValidatorResult IsValid(const FName& Name, bool bOriginal) override;
-	// End FNameValidatorInterface
-};
-
-USTRUCT()
-struct CONTROLRIGDEVELOPER_API FControlRigGraphSchemaAction_LocalVar : public FEdGraphSchemaAction_BlueprintVariableBase
-{
-	GENERATED_BODY()
-
-	
-public:
-
-	// Simple type info
-	static FName StaticGetTypeId() {static FName Type("FControlRigGraphSchemaAction_LocalVar"); return Type;}
-	virtual FName GetTypeId() const override { return StaticGetTypeId(); } 
-
-	FControlRigGraphSchemaAction_LocalVar()
-		: FEdGraphSchemaAction_BlueprintVariableBase()
-	{}
-
-	FControlRigGraphSchemaAction_LocalVar(FText InNodeCategory, FText InMenuDesc, FText InToolTip, const int32 InGrouping, const int32 InSectionID)
-		: FEdGraphSchemaAction_BlueprintVariableBase(MoveTemp(InNodeCategory), MoveTemp(InMenuDesc), MoveTemp(InToolTip), InGrouping, InSectionID)
-	{}
-
-	virtual bool IsA(const FName& InType) const override
-	{
-		return InType == GetTypeId() || InType == FEdGraphSchemaAction_BlueprintVariableBase::StaticGetTypeId();
-	}
-
-	virtual FEdGraphPinType GetPinType() const override;
-
-	virtual void ChangeVariableType(const FEdGraphPinType& NewPinType) override;
-
-	virtual void RenameVariable(const FName& NewName) override;
-
-	virtual bool IsValidName(const FName& NewName, FText& OutErrorMessage) const override;
-
-	virtual void DeleteVariable() override;
-
-	virtual bool IsVariableUsed() override;
-};
-
-USTRUCT()
-struct CONTROLRIGDEVELOPER_API FControlRigGraphSchemaAction_PromoteToVariable : public FEdGraphSchemaAction
-{
-	GENERATED_BODY()
-
-public:
-
-	// Simple type info
-	static FName StaticGetTypeId() {static FName Type("FControlRigGraphSchemaAction_PromoteToVariable"); return Type;}
-	virtual FName GetTypeId() const override { return StaticGetTypeId(); } 
-
-	FControlRigGraphSchemaAction_PromoteToVariable()
-		: FEdGraphSchemaAction()
-	{}
-
-	FControlRigGraphSchemaAction_PromoteToVariable(UEdGraphPin* InEdGraphPin, bool InLocalVariable);
-
-	virtual bool IsA(const FName& InType) const override
-	{
-		return InType == GetTypeId();
-	}
-
-	virtual UEdGraphNode* PerformAction(UEdGraph* ParentGraph, UEdGraphPin* FromPin, const FVector2D Location, bool bSelectNewNode) override;
-
-private:
-
-	UEdGraphPin* EdGraphPin;
-	bool bLocalVariable;
-};
-
-USTRUCT()
-struct CONTROLRIGDEVELOPER_API FControlRigGraphSchemaAction_PromoteToExposedPin : public FEdGraphSchemaAction
-{
-	GENERATED_BODY()
-
-public:
-
-	// Simple type info
-	static FName StaticGetTypeId() {static FName Type("FControlRigGraphSchemaAction_PromoteToExposedPin"); return Type;}
-	virtual FName GetTypeId() const override { return StaticGetTypeId(); } 
-
-	FControlRigGraphSchemaAction_PromoteToExposedPin()
-		: FEdGraphSchemaAction()
-	{}
-
-	FControlRigGraphSchemaAction_PromoteToExposedPin(UEdGraphPin* InEdGraphPin);
-
-	virtual bool IsA(const FName& InType) const override
-	{
-		return InType == GetTypeId();
-	}
-
-	virtual UEdGraphNode* PerformAction(UEdGraph* ParentGraph, UEdGraphPin* FromPin, const FVector2D Location, bool bSelectNewNode) override;
-
-private:
-
-	UEdGraphPin* EdGraphPin;
-};
-
-USTRUCT()
-struct CONTROLRIGDEVELOPER_API FControlRigGraphSchemaAction_Event : public FEdGraphSchemaAction
-{
-	GENERATED_BODY()
-
-public:
-
-	// Simple type info
-	static FName StaticGetTypeId() {static FName Type("FControlRigGraphSchemaAction_Event"); return Type;}
-	virtual FName GetTypeId() const override { return StaticGetTypeId(); } 
-
-	FControlRigGraphSchemaAction_Event()
-		: FEdGraphSchemaAction()
-	{}
-
-	FControlRigGraphSchemaAction_Event(const FName& InEventName, const FString& InNodePath, const FText& InNodeCategory);
-
-	virtual bool IsA(const FName& InType) const override
-	{
-		return InType == GetTypeId() || Super::IsA(InType);
-	}
-	virtual bool IsParentable() const override { return true; }
-	virtual FReply OnDoubleClick(UBlueprint* InBlueprint) override;
-	virtual bool CanBeRenamed() const override { return false; }
-	virtual FSlateBrush const* GetPaletteIcon() const override;
-
-private:
-
-	FString NodePath;
-};
-
-/** DragDropAction class for drag and dropping an item from the My Blueprints tree (e.g., variable or function) */
-class CONTROLRIGDEVELOPER_API FControlRigFunctionDragDropAction : public FGraphSchemaActionDragDropAction
-{
-public:
-
-	DRAG_DROP_OPERATOR_TYPE(FControlRigFunctionDragDropAction, FGraphSchemaActionDragDropAction)
-
-	// FGraphEditorDragDropAction interface
-	virtual FReply DroppedOnPanel(const TSharedRef< class SWidget >& Panel, FVector2D ScreenPosition, FVector2D GraphPosition, UEdGraph& Graph) override;
-	virtual FReply DroppedOnPin(FVector2D ScreenPosition, FVector2D GraphPosition) override;
-	virtual FReply DroppedOnAction(TSharedRef<FEdGraphSchemaAction> Action) override;
-	virtual FReply DroppedOnCategory(FText Category) override;
-	virtual void HoverTargetChanged() override;
-	// End of FGraphEditorDragDropAction
-
-	/** Set if operation is modified by alt */
-	void SetAltDrag(bool InIsAltDrag) { bAltDrag = InIsAltDrag; }
-
-	/** Set if operation is modified by the ctrl key */
-	void SetCtrlDrag(bool InIsCtrlDrag) { bControlDrag = InIsCtrlDrag; }
-
-protected:
-
-	/** Constructor */
-	FControlRigFunctionDragDropAction();
-
-	static TSharedRef<FControlRigFunctionDragDropAction> New(TSharedPtr<FEdGraphSchemaAction> InAction, UControlRigBlueprint* InRigBlueprint, UControlRigGraph* InRigGraph);
-
-protected:
-
-	UControlRigBlueprint* SourceRigBlueprint;
-	UControlRigGraph* SourceRigGraph;
-	bool bControlDrag;
-	bool bAltDrag;
-
-	friend class UControlRigGraphSchema;
-};
-
 UCLASS()
-class CONTROLRIGDEVELOPER_API UControlRigGraphSchema : public UEdGraphSchema
+class CONTROLRIGDEVELOPER_API UControlRigGraphSchema : public URigVMEdGraphSchema
 {
 	GENERATED_BODY()
 
@@ -252,98 +29,9 @@ public:
 	UControlRigGraphSchema();
 
 	// UEdGraphSchema interface
-	virtual void GetGraphContextActions(FGraphContextMenuBuilder& ContextMenuBuilder) const override;
-	virtual void GetContextMenuActions(class UToolMenu* Menu, class UGraphNodeContextMenuContext* Context) const override;
-	virtual bool TryCreateConnection(UEdGraphPin* PinA, UEdGraphPin* PinB) const override;
-	virtual const FPinConnectionResponse CanCreateConnection(const UEdGraphPin* A, const UEdGraphPin* B) const override;
 	virtual FLinearColor GetPinTypeColor(const FEdGraphPinType& PinType) const override;
-	virtual bool SupportsPinType(TWeakPtr<const FEdGraphSchemaAction> SchemaAction, const FEdGraphPinType& PinType) const override;
-	virtual bool SupportsPinTypeContainer(TWeakPtr<const FEdGraphSchemaAction> SchemaAction, const FEdGraphPinType& PinType, const EPinContainerType& ContainerType) const override;
-	virtual void BreakPinLinks(UEdGraphPin& TargetPin, bool bSendsNodeNotifcation) const override;
-	virtual void BreakSinglePinLink(UEdGraphPin* SourcePin, UEdGraphPin* TargetPin) const override;
-	virtual bool CanGraphBeDropped(TSharedPtr<FEdGraphSchemaAction> InAction) const override;
-	virtual FString GetFindReferenceSearchTerm(const FEdGraphSchemaAction* InGraphAction) const override;
-	virtual FReply BeginGraphDragAction(TSharedPtr<FEdGraphSchemaAction> InAction, const FPointerEvent& MouseEvent = FPointerEvent() ) const override;
-	virtual class FConnectionDrawingPolicy* CreateConnectionDrawingPolicy(int32 InBackLayerID, int32 InFrontLayerID, float InZoomFactor, const FSlateRect& InClippingRect, class FSlateWindowElementList& InDrawElements, class UEdGraph* InGraphObj) const override;
-	virtual bool ShouldHidePinDefaultValue(UEdGraphPin* Pin) const override;
-	virtual void TrySetDefaultValue(UEdGraphPin& InPin, const FString& InNewDefaultValue, bool bMarkAsModified = true) const override;
-	virtual void TrySetDefaultObject(UEdGraphPin& InPin, UObject* InNewDefaultObject, bool bMarkAsModified) const override;
-	virtual void TrySetDefaultText(UEdGraphPin& InPin, const FText& InNewDefaultText, bool bMarkAsModified) const override;
-	virtual bool ShouldAlwaysPurgeOnModification() const override { return false; }
-	virtual bool ArePinsCompatible(const UEdGraphPin* PinA, const UEdGraphPin* PinB, const UClass* CallingContext, bool bIgnoreArray /*= false*/) const override;
-	virtual bool DoesSupportPinWatching() const	override { return true; }
-	virtual bool IsPinBeingWatched(UEdGraphPin const* Pin) const override;
-	virtual void ClearPinWatch(UEdGraphPin const* Pin) const override;
-	virtual void OnPinConnectionDoubleCicked(UEdGraphPin* PinA, UEdGraphPin* PinB, const FVector2D& GraphPosition) const override;
-	virtual bool MarkBlueprintDirtyFromNewNode(UBlueprint* InBlueprint, UEdGraphNode* InEdGraphNode) const override;
-	virtual bool SafeDeleteNodeFromGraph(UEdGraph* Graph, UEdGraphNode* Node) const override;
-	virtual bool CanVariableBeDropped(UEdGraph* InGraph, FProperty* InVariableToDrop) const override;
-	virtual bool RequestVariableDropOnPanel(UEdGraph* InGraph, FProperty* InVariableToDrop, const FVector2D& InDropPosition, const FVector2D& InScreenPosition) override;
-	virtual bool RequestVariableDropOnPin(UEdGraph* InGraph, FProperty* InVariableToDrop, UEdGraphPin* InPin, const FVector2D& InDropPosition, const FVector2D& InScreenPosition) override;
-	virtual bool IsStructEditable(UStruct* InStruct) const;
-	virtual void SetNodePosition(UEdGraphNode* Node, const FVector2D& Position) const override;
-	void SetNodePosition(UEdGraphNode* Node, const FVector2D& Position, bool bSetupUndo) const;
-	virtual void GetGraphDisplayInformation(const UEdGraph& Graph, /*out*/ FGraphDisplayInfo& DisplayInfo) const override;
-	virtual bool GetLocalVariables(const UEdGraph* InGraph, TArray<FBPVariableDescription>& OutLocalVariables) const override;
-	virtual TSharedPtr<FEdGraphSchemaAction> MakeActionFromVariableDescription(const UEdGraph* InEdGraph, const FBPVariableDescription& Variable) const override;
-	virtual FText GetGraphCategory(const UEdGraph* InGraph) const override;
-	virtual EGraphType GetGraphType(const UEdGraph* TestEdGraph) const override;
-	virtual FReply TrySetGraphCategory(const UEdGraph* InGraph, const FText& InCategory) override;
-	virtual bool TryDeleteGraph(UEdGraph* GraphToDelete) const override;
-	virtual bool TryRenameGraph(UEdGraph* GraphToRename, const FName& InNewName) const override;
-	virtual bool TryToGetChildEvents(const UEdGraph* Graph, const int32 SectionId, TArray<TSharedPtr<FEdGraphSchemaAction>>& Actions, const FText& ParentCategory) const override;
-	virtual bool CanDuplicateGraph(UEdGraph* InSourceGraph) const { return false; }
-	virtual UEdGraphPin* DropPinOnNode(UEdGraphNode* InTargetNode, const FName& InSourcePinName, const FEdGraphPinType& InSourcePinType, EEdGraphPinDirection InSourcePinDirection) const override;
-	virtual bool SupportsDropPinOnNode(UEdGraphNode* InTargetNode, const FEdGraphPinType& InSourcePinType, EEdGraphPinDirection InSourcePinDirection, FText& OutErrorMessage) const override;
-	virtual void SetPinBeingDroppedOnNode(UEdGraphPin* InSourcePin) const override { PinBeingDropped = InSourcePin; }
-	virtual void InsertAdditionalActions(TArray<UBlueprint*> InBlueprints, TArray<UEdGraph*> InGraphs, TArray<UEdGraphPin*> InPins, FGraphActionListBuilderBase& OutAllActions) const override;
-	virtual TSharedPtr<INameValidatorInterface> GetNameValidator(const UBlueprint* BlueprintObj, const FName& OriginalName, const UStruct* ValidationScope, const FName& ActionTypeId) const override;
-
-	/** Returns true if the schema supports the script type */
-	bool SupportsPinType(UScriptStruct* ScriptStruct) const;
-
-	/** Create a graph node for a rig */
-	UControlRigGraphNode* CreateGraphNode(UControlRigGraph* InGraph, const FName& InPropertyName) const;
-
-	/** Automatically layout the passed-in nodes */
-	void LayoutNodes(UControlRigGraph* InGraph, const TArray<UControlRigGraphNode*>& InNodes) const;
-
-	/** Helper function to rename a node */
-	void RenameNode(UControlRigGraphNode* Node, const FName& InNewNodeName) const;
-
-	/** Helper function to recursively reset the pin defaults */
-	virtual void ResetPinDefaultsRecursive(UEdGraphPin* InPin) const;
-
-	/** Returns all of the applicable pin types for variables within a control rig */
-	virtual void GetVariablePinTypes(TArray<FEdGraphPinType>& PinTypes) const;
-
-	void StartGraphNodeInteraction(UEdGraphNode* InNode) const;
-	void EndGraphNodeInteraction(UEdGraphNode* InNode) const;
-	static TArray<UEdGraphNode*> GetNodesToMoveForNode(UEdGraphNode* InNode);
-	FVector2D GetNodePositionAtStartOfInteraction(UEdGraphNode* InNode) const;
-
-	bool AutowireNewNode(UControlRigGraphNode* NewNode, UEdGraphPin* FromPin) const;
-
-	void HandleModifiedEvent(ERigVMGraphNotifType InNotifType, URigVMGraph* InGraph, UObject* InSubject);
-
-	// Allow derived classes to spawn derived node classes
-	virtual TSubclassOf<UControlRigGraphNode> GetGraphNodeClass() const { return UControlRigGraphNode::StaticClass(); }
-
-	static bool IsControlRigDefaultEvent(const FName& InEventName);
-
-private:
-
-	const UEdGraphPin* LastPinForCompatibleCheck = nullptr;
-	bool bLastPinWasInput;
-	mutable UEdGraphPin* PinBeingDropped = nullptr;
-
-#if WITH_EDITOR
-	mutable TArray<UEdGraphNode*> NodesBeingInteracted;
-	mutable TMap<FName, FVector2D> NodePositionsDuringStart;
-#endif
 
-	friend class UControlRigIfNodeSpawner;
-	friend class UControlRigSelectNodeSpawner;
-	friend class UControlRigUnitNodeSpawner;
-	friend class UControlRigArrayNodeSpawner;
+	// URigVMEdGraphSchema interface
+	virtual const FName& GetGraphNameForUI() const override { return GraphName_ControlRig; }
+	virtual bool IsRigVMDefaultEvent(const FName& InEventName) const override;
 };
diff --git a/Engine/Plugins/Animation/ControlRig/Source/ControlRigEditor/Private/ControlRigBlueprintActions.cpp b/Engine/Plugins/Animation/ControlRig/Source/ControlRigEditor/Private/ControlRigBlueprintActions.cpp
--- a/Engine/Plugins/Animation/ControlRig/Source/ControlRigEditor/Private/ControlRigBlueprintActions.cpp
+++ b/Engine/Plugins/Animation/ControlRig/Source/ControlRigEditor/Private/ControlRigBlueprintActions.cpp
@@ -18,15 +18,10 @@
 #include "Engine/SkeletalMesh.h"
 #include "Animation/SkeletalMeshActor.h"
 #include "Components/SkeletalMeshComponent.h"
-#include "ScopedTransaction.h"
 #include "LevelSequenceActor.h"
-#include "LevelSequence.h"
 #include "LevelSequencePlayer.h"
 #include "EditorDirectories.h"
-#include "EngineUtils.h"
-#include "ISequencer.h"
 #include "ILevelSequenceEditorToolkit.h"
-#include "Sequencer/MovieSceneControlRigParameterTrack.h"
 #include "ControlRigObjectBinding.h"
 #include "EditMode/ControlRigEditMode.h"
 #include "Editor.h"
@@ -37,6 +32,7 @@
 #include "Misc/PackageName.h"
 #include "LevelSequenceEditorBlueprintLibrary.h"
 #include "Sequencer/ControlRigParameterTrackEditor.h"
+#include "Rigs/RigHierarchyController.h"
 
 #define LOCTEXT_NAMESPACE "ControlRigBlueprintActions"
 
diff --git a/Engine/Plugins/Animation/ControlRig/Source/ControlRigEditor/Private/ControlRigBlueprintEditorLibrary.cpp b/Engine/Plugins/Animation/ControlRig/Source/ControlRigEditor/Private/ControlRigBlueprintEditorLibrary.cpp
--- a/Engine/Plugins/Animation/ControlRig/Source/ControlRigEditor/Private/ControlRigBlueprintEditorLibrary.cpp
+++ b/Engine/Plugins/Animation/ControlRig/Source/ControlRigEditor/Private/ControlRigBlueprintEditorLibrary.cpp
@@ -66,7 +66,7 @@ void UControlRigBlueprintEditorLibrary::RequestControlRigInit(UControlRigBluepri
 	{
 		return;
 	}
-	InRigBlueprint->RequestControlRigInit();
+	InRigBlueprint->RequestRigVMInit();
 }
 
 URigVMGraph* UControlRigBlueprintEditorLibrary::GetModel(UControlRigBlueprint* InRigBlueprint)
@@ -94,7 +94,8 @@ TArray<UControlRigBlueprint*> UControlRigBlueprintEditorLibrary::GetCurrentlyOpe
 
 TArray<UStruct*> UControlRigBlueprintEditorLibrary::GetAvailableRigUnits()
 {
-	return UControlRigBlueprint::GetAvailableRigUnits();
+	UControlRigBlueprint* CDO = CastChecked<UControlRigBlueprint>(UControlRigBlueprint::StaticClass()->GetDefaultObject());
+	return CDO->GetAvailableRigVMStructs();
 }
 
 URigHierarchy* UControlRigBlueprintEditorLibrary::GetHierarchy(UControlRigBlueprint* InRigBlueprint)
diff --git a/Engine/Plugins/Animation/ControlRig/Source/ControlRigEditor/Private/ControlRigEditorModule.cpp b/Engine/Plugins/Animation/ControlRig/Source/ControlRigEditor/Private/ControlRigEditorModule.cpp
--- a/Engine/Plugins/Animation/ControlRig/Source/ControlRigEditor/Private/ControlRigEditorModule.cpp
+++ b/Engine/Plugins/Animation/ControlRig/Source/ControlRigEditor/Private/ControlRigEditorModule.cpp
@@ -174,7 +174,7 @@ void FControlRigEditorModule::StartupModule()
 	PropertiesToUnregisterOnShutdown.Add(FRigVMCompileSettings::StaticStruct()->GetFName());
 	PropertyEditorModule.RegisterCustomPropertyTypeLayout(PropertiesToUnregisterOnShutdown.Last(), FOnGetPropertyTypeCustomizationInstance::CreateStatic(&FRigVMCompileSettingsDetails::MakeInstance));
 
-	PropertiesToUnregisterOnShutdown.Add(FControlRigPythonSettings::StaticStruct()->GetFName());
+	PropertiesToUnregisterOnShutdown.Add(FRigVMPythonSettings::StaticStruct()->GetFName());
 	PropertyEditorModule.RegisterCustomPropertyTypeLayout(PropertiesToUnregisterOnShutdown.Last(), FOnGetPropertyTypeCustomizationInstance::CreateStatic(&FControlRigPythonLogDetails::MakeInstance));
 
 	PropertiesToUnregisterOnShutdown.Add(FRigVMDrawContainer::StaticStruct()->GetFName());
@@ -1207,7 +1207,7 @@ void FControlRigEditorModule::GetTypeActions(UControlRigBlueprint* CRB, FBluepri
 
 void FControlRigEditorModule::GetInstanceActions(UControlRigBlueprint* CRB, FBlueprintActionDatabaseRegistrar& ActionRegistrar)
 {
-	if (URigVMBlueprintGeneratedClass* GeneratedClass = CRB->GetControlRigBlueprintGeneratedClass())
+	if (URigVMBlueprintGeneratedClass* GeneratedClass = CRB->GetRigVMBlueprintGeneratedClass())
 	{
 		if (UControlRig* CDO = Cast<UControlRig>(GeneratedClass->GetDefaultObject()))
 		{
@@ -2203,7 +2203,7 @@ void FControlRigEditorModule::GetContextMenuActions(const UControlRigGraphSchema
 						if(ModelNode->IsEvent())
 						{
 							const FName& EventName = ModelNode->GetEventName();
-							const bool bCanRunOnce = !UControlRigGraphSchema::IsControlRigDefaultEvent(EventName);
+							const bool bCanRunOnce = !CastChecked<URigVMEdGraphSchema>(RigNode->GetSchema())->IsRigVMDefaultEvent(EventName);
 
 							FToolMenuSection& EventsSection = Menu->AddSection("EdGraphSchemaEvents", LOCTEXT("EventsHeader", "Events"));
 
@@ -2537,7 +2537,7 @@ void FControlRigEditorModule::GetContextMenuActions(const UControlRigGraphSchema
 													   }
 												   }
                                     			
-												   FName NewVariableName = RigBlueprint->AddCRMemberVariableFromExternal(ExternalVariable, DefaultValue);
+												   FName NewVariableName = RigBlueprint->AddHostMemberVariableFromExternal(ExternalVariable, DefaultValue);
 												   if(!NewVariableName.IsNone())
 												   {
 													   Controller->SetRemappedVariable(FunctionReferenceNode, ExternalVariable.Name, NewVariableName);
@@ -2796,7 +2796,7 @@ void FControlRigEditorModule::PostChange(const UUserDefinedStruct* Changed,
 		// so to simplify things, here we just reset all rigs upon error
 		if (ResultsLog.NumErrors > 0)
 		{
-			URigVMBlueprintGeneratedClass* RigClass = RigBlueprint->GetControlRigBlueprintGeneratedClass();
+			URigVMBlueprintGeneratedClass* RigClass = RigBlueprint->GetRigVMBlueprintGeneratedClass();
 			UControlRig* CDO = Cast<UControlRig>(RigClass->GetDefaultObject(true /* create if needed */));
 			if (CDO->GetVM() != nullptr)
 			{
diff --git a/Engine/Plugins/Animation/ControlRig/Source/ControlRigEditor/Private/ControlRigGraphDetails.h b/Engine/Plugins/Animation/ControlRig/Source/ControlRigEditor/Private/ControlRigGraphDetails.h
--- a/Engine/Plugins/Animation/ControlRig/Source/ControlRigEditor/Private/ControlRigGraphDetails.h
+++ b/Engine/Plugins/Animation/ControlRig/Source/ControlRigEditor/Private/ControlRigGraphDetails.h
@@ -118,7 +118,7 @@ private:
 	TWeakPtr<SEditableTextBox> ArgumentNameWidget;
 
 	/** The validator to check if a name for an argument is valid */
-	FControlRigLocalVariableNameValidator NameValidator;
+	FRigVMLocalVariableNameValidator NameValidator;
 };
 
 class FControlRigArgumentDefaultNode : public IDetailCustomNodeBuilder, public TSharedFromThis<FControlRigArgumentDefaultNode>
diff --git a/Engine/Plugins/Animation/ControlRig/Source/ControlRigEditor/Private/ControlRigLocalVariableDetails.cpp b/Engine/Plugins/Animation/ControlRig/Source/ControlRigEditor/Private/ControlRigLocalVariableDetails.cpp
--- a/Engine/Plugins/Animation/ControlRig/Source/ControlRigEditor/Private/ControlRigLocalVariableDetails.cpp
+++ b/Engine/Plugins/Animation/ControlRig/Source/ControlRigEditor/Private/ControlRigLocalVariableDetails.cpp
@@ -40,7 +40,7 @@ void FRigVMLocalVariableDetails::CustomizeDetails(IDetailLayoutBuilder& DetailBu
 			BlueprintBeingCustomized = GraphBeingCustomized->GetTypedOuter<UControlRigBlueprint>();
 		}
 
-		NameValidator = FControlRigLocalVariableNameValidator(BlueprintBeingCustomized, GraphBeingCustomized, VariableDescription.Name);
+		NameValidator = FRigVMLocalVariableNameValidator(BlueprintBeingCustomized, GraphBeingCustomized, VariableDescription.Name);
 	}
 
 	DetailBuilder.HideCategory(TEXT("RigVMGraphVariableDescription"));
@@ -92,7 +92,7 @@ void FRigVMLocalVariableDetails::CustomizeDetails(IDetailLayoutBuilder& DetailBu
 
 	if (BlueprintBeingCustomized)
 	{
-		URigVMBlueprintGeneratedClass* RigClass = BlueprintBeingCustomized->GetControlRigBlueprintGeneratedClass();
+		URigVMBlueprintGeneratedClass* RigClass = BlueprintBeingCustomized->GetRigVMBlueprintGeneratedClass();
 		UControlRig* CDO = Cast<UControlRig>(RigClass->GetDefaultObject(true /* create if needed */));
 		if (CDO->GetVM() != nullptr)
 		{
@@ -164,7 +164,7 @@ FEdGraphPinType FRigVMLocalVariableDetails::OnGetPinInfo() const
 void FRigVMLocalVariableDetails::HandlePinInfoChanged(const FEdGraphPinType& PinType)
 {
 	VariableDescription.ChangeType(PinType);
-	FControlRigBlueprintVMCompileScope CompileScope(BlueprintBeingCustomized);
+	FRigVMBlueprintCompileScope CompileScope(BlueprintBeingCustomized);
 	TypeHandle->SetValue(VariableDescription.CPPType);
 	TypeObjectHandle->SetValue(VariableDescription.CPPTypeObject);	
 }
diff --git a/Engine/Plugins/Animation/ControlRig/Source/ControlRigEditor/Private/ControlRigLocalVariableDetails.h b/Engine/Plugins/Animation/ControlRig/Source/ControlRigEditor/Private/ControlRigLocalVariableDetails.h
--- a/Engine/Plugins/Animation/ControlRig/Source/ControlRigEditor/Private/ControlRigLocalVariableDetails.h
+++ b/Engine/Plugins/Animation/ControlRig/Source/ControlRigEditor/Private/ControlRigLocalVariableDetails.h
@@ -48,7 +48,7 @@ private:
 	TSharedPtr<IPropertyHandle> TypeObjectHandle;
 	TSharedPtr<IPropertyHandle> DefaultValueHandle;
 
-	FControlRigLocalVariableNameValidator NameValidator;
+	FRigVMLocalVariableNameValidator NameValidator;
 	TArray<TSharedPtr<FString>> EnumOptions;
 
 	FText GetName() const;
diff --git a/Engine/Plugins/Animation/ControlRig/Source/ControlRigEditor/Private/Editor/ControlRigEditor.cpp b/Engine/Plugins/Animation/ControlRig/Source/ControlRigEditor/Private/Editor/ControlRigEditor.cpp
--- a/Engine/Plugins/Animation/ControlRig/Source/ControlRigEditor/Private/Editor/ControlRigEditor.cpp
+++ b/Engine/Plugins/Animation/ControlRig/Source/ControlRigEditor/Private/Editor/ControlRigEditor.cpp
@@ -383,7 +383,7 @@ void FControlRigEditor::InitControlRigEditor(const EToolkitMode::Type Mode, cons
 	// user-defined-struct can change even after load
 	// refresh the models such that pins are updated to match
 	// the latest struct member layout
-	InControlRigBlueprint->RefreshAllModels(EControlRigBlueprintLoadType::CheckUserDefinedStructs);
+	InControlRigBlueprint->RefreshAllModels(ERigVMBlueprintLoadType::CheckUserDefinedStructs);
 
 	{
 		TArray<UEdGraph*> EdGraphs;
@@ -397,7 +397,7 @@ void FControlRigEditor::InitControlRigEditor(const EToolkitMode::Type Mode, cons
 				continue;
 			}
 
-			RigGraph->Initialize(InControlRigBlueprint);
+			RigGraph->InitializeFromBlueprint(InControlRigBlueprint);
 		}
 
 	}
@@ -510,7 +510,7 @@ void FControlRigEditor::InitControlRigEditor(const EToolkitMode::Type Mode, cons
 
 				if (UPackage* Package = InControlRigBlueprint->GetOutermost())
 				{
-					Package->SetDirtyFlag(InControlRigBlueprint->bDirtyDuringLoad);
+					Package->SetDirtyFlag(InControlRigBlueprint->IsMarkedDirtyDuringLoad());
 				}
 			}
 		}
@@ -741,8 +741,8 @@ void FControlRigEditor::ToggleAutoCompileGraph()
 {
 	if (UControlRigBlueprint* RigBlueprint = Cast<UControlRigBlueprint>(GetBlueprintObj()))
 	{
-		RigBlueprint->bAutoRecompileVM = !RigBlueprint->bAutoRecompileVM;
-		if (RigBlueprint->bAutoRecompileVM)
+		RigBlueprint->SetAutoVMRecompile(!RigBlueprint->GetAutoVMRecompile());
+		if (RigBlueprint->GetAutoVMRecompile())
 		{
 			RigBlueprint->RequestAutoVMRecompilation();
 		}
@@ -753,7 +753,7 @@ bool FControlRigEditor::IsAutoCompileGraphOn() const
 {
 	if (UControlRigBlueprint* RigBlueprint = Cast<UControlRigBlueprint>(GetBlueprintObj()))
 	{
-		return RigBlueprint->bAutoRecompileVM;
+		return RigBlueprint->GetAutoVMRecompile();
 	}
 	return false;
 }
@@ -786,11 +786,13 @@ TSharedRef<SWidget> FControlRigEditor::GenerateEventQueueMenuContent()
 
 	if (const UControlRigBlueprint* RigBlueprint = Cast<UControlRigBlueprint>(GetBlueprintObj()))
 	{
+		URigVMEdGraphSchema* Schema = CastChecked<URigVMEdGraphSchema>(RigBlueprint->GetRigVMEdGraphSchemaClass()->GetDefaultObject());
+		
 		bool bFoundUserDefinedEvent = false;
 		const TArray<FName> EntryNames = RigBlueprint->GetRigVMClient()->GetEntryNames();
 		for(const FName& EntryName : EntryNames)
 		{
-			if(UControlRigGraphSchema::IsControlRigDefaultEvent(EntryName))
+			if(Schema->IsRigVMDefaultEvent(EntryName))
 			{
 				continue;
 			}
@@ -1306,7 +1308,7 @@ void FControlRigEditor::SetEventQueue(TArray<FName> InEventQueue, bool bCompile)
 	{
 		if(bCompile)
 		{
-			if (RigBlueprint->bAutoRecompileVM)
+			if (RigBlueprint->GetAutoVMRecompile())
 			{
 				RigBlueprint->RequestAutoVMRecompilation();
 			}
@@ -1515,7 +1517,7 @@ void FControlRigEditor::GetCustomDebugObjects(TArray<FCustomDebugObject>& DebugL
 		DebugList.Add(DebugObject);
 	}
 
-	URigVMBlueprintGeneratedClass* GeneratedClass = RigBlueprint->GetControlRigBlueprintGeneratedClass();
+	URigVMBlueprintGeneratedClass* GeneratedClass = RigBlueprint->GetRigVMBlueprintGeneratedClass();
 	if (GeneratedClass)
 	{
 		struct Local
@@ -1627,13 +1629,13 @@ void FControlRigEditor::HandleSetObjectBeingDebugged(UObject* InObject)
 
 	if (UControlRigBlueprint* RigBlueprint = Cast<UControlRigBlueprint>(GetBlueprintObj()))
 	{
-		if (URigVMBlueprintGeneratedClass* GeneratedClass = RigBlueprint->GetControlRigBlueprintGeneratedClass())
+		if (URigVMBlueprintGeneratedClass* GeneratedClass = RigBlueprint->GetRigVMBlueprintGeneratedClass())
 		{
 			UControlRig* CDO = Cast<UControlRig>(GeneratedClass->GetDefaultObject(true /* create if needed */));
 			if (CDO->VM->GetInstructions().Num() <= 1 /* only exit */)
 			{
 				RigBlueprint->RecompileVM();
-				RigBlueprint->RequestControlRigInit();
+				RigBlueprint->RequestRigVMInit();
 			}
 		}
 
@@ -2248,7 +2250,7 @@ void FControlRigEditor::Compile()
 			return;
 		}
 
-		RigBlueprint->CompileLog.Messages.Reset();
+		RigBlueprint->GetCompileLog().Messages.Reset();
 
 		FString LastDebuggedObjectName = GetCustomDebugObjectLabel(RigBlueprint->GetObjectBeingDebugged());
 		RigBlueprint->SetObjectBeingDebugged(nullptr);
@@ -3258,9 +3260,9 @@ void FControlRigEditor::HandleVMCompiledEvent(UObject* InCompiledObject, URigVM*
 	if(UControlRigBlueprint* RigBlueprint = Cast<UControlRigBlueprint>(InCompiledObject))
 	{
 		CompilerResultsListing->ClearMessages();
-		CompilerResultsListing->AddMessages(RigBlueprint->CompileLog.Messages);
-		RigBlueprint->CompileLog.Messages.Reset();
-		RigBlueprint->CompileLog.NumErrors = RigBlueprint->CompileLog.NumWarnings = 0;
+		CompilerResultsListing->AddMessages(RigBlueprint->GetCompileLog().Messages);
+		RigBlueprint->GetCompileLog().Messages.Reset();
+		RigBlueprint->GetCompileLog().NumErrors = RigBlueprint->GetCompileLog().NumWarnings = 0;
 	}
 
 	RefreshDetailView();
@@ -5070,7 +5072,7 @@ void FControlRigEditor::OnWrappedPropertyChangedChainEvent(UDetailsViewWrapperOb
 	}
 }
 
-void FControlRigEditor::OnRequestLocalizeFunctionDialog(FRigVMGraphFunctionIdentifier& InFunction, UControlRigBlueprint* InTargetBlueprint, bool bForce)
+void FControlRigEditor::OnRequestLocalizeFunctionDialog(FRigVMGraphFunctionIdentifier& InFunction, URigVMBlueprint* InTargetBlueprint, bool bForce)
 {
 	check(InTargetBlueprint);
 
@@ -5099,7 +5101,7 @@ void FControlRigEditor::OnRequestLocalizeFunctionDialog(FRigVMGraphFunctionIdent
 	}
 }
 
-FRigVMController_BulkEditResult FControlRigEditor::OnRequestBulkEditDialog(UControlRigBlueprint* InBlueprint, URigVMController* InController,
+FRigVMController_BulkEditResult FControlRigEditor::OnRequestBulkEditDialog(URigVMBlueprint* InBlueprint, URigVMController* InController,
 	URigVMLibraryNode* InFunction, ERigVMControllerBulkEditType InEditType)
 {
 	if (bAllowBulkEdits)
@@ -5557,7 +5559,7 @@ void FControlRigEditor::OnHierarchyModified(ERigHierarchyNotification InNotif, U
 				}
 
 				{
-					FControlRigBlueprintVMCompileScope CompileScope(RigBlueprint);
+					FRigVMBlueprintCompileScope CompileScope(RigBlueprint);
 					for (UEdGraphNode* Node : RigGraph->Nodes)
 					{
 						if (UControlRigGraphNode* RigNode = Cast<UControlRigGraphNode>(Node))
@@ -6855,7 +6857,7 @@ void FControlRigEditor::HandleOnControlModified(UControlRig* Subject, FRigContro
 	}
 }
 
-void FControlRigEditor::HandleRefreshEditorFromBlueprint(UControlRigBlueprint* InBlueprint)
+void FControlRigEditor::HandleRefreshEditorFromBlueprint(URigVMBlueprint* InBlueprint)
 {
 	OnHierarchyChanged();
 	Compile();
@@ -6930,7 +6932,7 @@ void FControlRigEditor::HandleBreakpointAdded()
 	SetExecutionMode(EControlRigExecutionModeType_Debug);
 }
 
-void FControlRigEditor::OnGraphNodeClicked(UControlRigGraphNode* InNode)
+void FControlRigEditor::OnGraphNodeClicked(URigVMEdGraphNode* InNode)
 {
 	if (InNode)
 	{
@@ -6941,7 +6943,7 @@ void FControlRigEditor::OnGraphNodeClicked(UControlRigGraphNode* InNode)
 	}
 }
 
-void FControlRigEditor::OnNodeDoubleClicked(UControlRigBlueprint* InBlueprint, URigVMNode* InNode)
+void FControlRigEditor::OnNodeDoubleClicked(URigVMBlueprint* InBlueprint, URigVMNode* InNode)
 {
 	ensure(GetControlRigBlueprint() == InBlueprint);
 
diff --git a/Engine/Plugins/Animation/ControlRig/Source/ControlRigEditor/Private/Editor/ControlRigEditor.h b/Engine/Plugins/Animation/ControlRig/Source/ControlRigEditor/Private/Editor/ControlRigEditor.h
--- a/Engine/Plugins/Animation/ControlRig/Source/ControlRigEditor/Private/Editor/ControlRigEditor.h
+++ b/Engine/Plugins/Animation/ControlRig/Source/ControlRigEditor/Private/Editor/ControlRigEditor.h
@@ -32,7 +32,7 @@ class FStructOnScope;
 class UToolMenu;
 class FControlRigEditor;
 
-DECLARE_MULTICAST_DELEGATE_TwoParams(FControlRigEditorClosed, const FControlRigEditor*, UControlRigBlueprint*);
+DECLARE_MULTICAST_DELEGATE_TwoParams(FControlRigEditorClosed, const FControlRigEditor*, URigVMBlueprint*);
 
 struct FControlRigEditorModes
 {
@@ -360,15 +360,15 @@ private:
 
 	 void HandleOnControlModified(UControlRig* Subject, FRigControlElement* ControlElement, const FRigControlModifiedContext& Context);
 
-	 void HandleRefreshEditorFromBlueprint(UControlRigBlueprint* InBlueprint);
+	 void HandleRefreshEditorFromBlueprint(URigVMBlueprint* InBlueprint);
 
 	 void HandleVariableDroppedFromBlueprint(UObject* InSubject, FProperty* InVariableToDrop, const FVector2D& InDropPosition, const FVector2D& InScreenPosition);
 
 	 void HandleBreakpointAdded();
 
-	 void OnGraphNodeClicked(UControlRigGraphNode* InNode);
+	 void OnGraphNodeClicked(URigVMEdGraphNode* InNode);
 
-	 void OnNodeDoubleClicked(UControlRigBlueprint* InBlueprint, URigVMNode* InNode);
+	 void OnNodeDoubleClicked(URigVMBlueprint* InBlueprint, URigVMNode* InNode);
 
 	 void OnGraphImported(UEdGraph* InEdGraph);
 
@@ -436,8 +436,8 @@ protected:
 	virtual void OnFinishedChangingProperties(const FPropertyChangedEvent& PropertyChangedEvent) override;
 	void OnPropertyChanged(UObject* InObject, FPropertyChangedEvent& InEvent);
 	void OnWrappedPropertyChangedChainEvent(UDetailsViewWrapperObject* InWrapperObject, const FString& InPropertyPath, FPropertyChangedChainEvent& InPropertyChangedChainEvent);
-	void OnRequestLocalizeFunctionDialog(FRigVMGraphFunctionIdentifier& InFunction, UControlRigBlueprint* InTargetBlueprint, bool bForce);
-	FRigVMController_BulkEditResult OnRequestBulkEditDialog(UControlRigBlueprint* InBlueprint, URigVMController* InController, URigVMLibraryNode* InFunction, ERigVMControllerBulkEditType InEditType);
+	void OnRequestLocalizeFunctionDialog(FRigVMGraphFunctionIdentifier& InFunction, URigVMBlueprint* InTargetBlueprint, bool bForce);
+	FRigVMController_BulkEditResult OnRequestBulkEditDialog(URigVMBlueprint* InBlueprint, URigVMController* InController, URigVMLibraryNode* InFunction, ERigVMControllerBulkEditType InEditType);
 	bool OnRequestBreakLinksDialog(TArray<URigVMLink*> InLinks);
 	TRigVMTypeIndex OnRequestPinTypeSelectionDialog(const TArray<TRigVMTypeIndex>& InTypes);
 	void HandleJumpToHyperlink(const UObject* InSubject);
diff --git a/Engine/Plugins/Animation/ControlRig/Source/ControlRigEditor/Private/Editor/SControlRigFunctionBulkEditWidget.cpp b/Engine/Plugins/Animation/ControlRig/Source/ControlRigEditor/Private/Editor/SControlRigFunctionBulkEditWidget.cpp
--- a/Engine/Plugins/Animation/ControlRig/Source/ControlRigEditor/Private/Editor/SControlRigFunctionBulkEditWidget.cpp
+++ b/Engine/Plugins/Animation/ControlRig/Source/ControlRigEditor/Private/Editor/SControlRigFunctionBulkEditWidget.cpp
@@ -23,7 +23,7 @@
 /// SControlRigFunctionBulkEditWidget
 ///////////////////////////////////////////////////////////
 
-void SControlRigFunctionBulkEditWidget::Construct(const FArguments& InArgs, UControlRigBlueprint* InBlueprint, URigVMController* InController, URigVMLibraryNode* InFunction, ERigVMControllerBulkEditType InEditType)
+void SControlRigFunctionBulkEditWidget::Construct(const FArguments& InArgs, URigVMBlueprint* InBlueprint, URigVMController* InController, URigVMLibraryNode* InFunction, ERigVMControllerBulkEditType InEditType)
 {
 	Blueprint = InBlueprint;
 	Controller = InController;
diff --git a/Engine/Plugins/Animation/ControlRig/Source/ControlRigEditor/Private/Editor/SControlRigFunctionBulkEditWidget.h b/Engine/Plugins/Animation/ControlRig/Source/ControlRigEditor/Private/Editor/SControlRigFunctionBulkEditWidget.h
--- a/Engine/Plugins/Animation/ControlRig/Source/ControlRigEditor/Private/Editor/SControlRigFunctionBulkEditWidget.h
+++ b/Engine/Plugins/Animation/ControlRig/Source/ControlRigEditor/Private/Editor/SControlRigFunctionBulkEditWidget.h
@@ -16,7 +16,7 @@ public:
 	SLATE_BEGIN_ARGS(SControlRigFunctionBulkEditWidget) {}
 	SLATE_END_ARGS()
 
-	void Construct(const FArguments& InArgs, UControlRigBlueprint* InBlueprint, URigVMController* InController, URigVMLibraryNode* InFunction, ERigVMControllerBulkEditType InEditType);
+	void Construct(const FArguments& InArgs, URigVMBlueprint* InBlueprint, URigVMController* InController, URigVMLibraryNode* InFunction, ERigVMControllerBulkEditType InEditType);
 
 private:
 
@@ -29,7 +29,7 @@ private:
 	/** Show a progress bar and load all of the assets */
 	void LoadAffectedAssets();
 
-	UControlRigBlueprint* Blueprint;
+	URigVMBlueprint* Blueprint;
 	URigVMController* Controller;
 	URigVMLibraryNode* Function;
 	ERigVMControllerBulkEditType EditType;
@@ -47,7 +47,7 @@ public:
 	{
 	}
 
-	SLATE_ARGUMENT(UControlRigBlueprint*, Blueprint)
+	SLATE_ARGUMENT(URigVMBlueprint*, Blueprint)
 	SLATE_ARGUMENT(URigVMController*, Controller)
 	SLATE_ARGUMENT(URigVMLibraryNode*, Function)
 	SLATE_ARGUMENT(ERigVMControllerBulkEditType, EditType)
diff --git a/Engine/Plugins/Animation/ControlRig/Source/ControlRigEditor/Private/Editor/SControlRigFunctionLocalizationWidget.cpp b/Engine/Plugins/Animation/ControlRig/Source/ControlRigEditor/Private/Editor/SControlRigFunctionLocalizationWidget.cpp
--- a/Engine/Plugins/Animation/ControlRig/Source/ControlRigEditor/Private/Editor/SControlRigFunctionLocalizationWidget.cpp
+++ b/Engine/Plugins/Animation/ControlRig/Source/ControlRigEditor/Private/Editor/SControlRigFunctionLocalizationWidget.cpp
@@ -86,7 +86,7 @@ void SControlRigFunctionLocalizationTableRow::Construct(const FArguments& InArgs
 /// SControlRigFunctionLocalizationWidget
 ///////////////////////////////////////////////////////////
 
-void SControlRigFunctionLocalizationWidget::Construct(const FArguments& InArgs, const FRigVMGraphFunctionIdentifier& InFunctionToLocalize, UControlRigBlueprint* InTargetBlueprint)
+void SControlRigFunctionLocalizationWidget::Construct(const FArguments& InArgs, const FRigVMGraphFunctionIdentifier& InFunctionToLocalize, URigVMBlueprint* InTargetBlueprint)
 {
 	FunctionsToLocalize.Add(InFunctionToLocalize);
 	FunctionItems.Reset();
diff --git a/Engine/Plugins/Animation/ControlRig/Source/ControlRigEditor/Private/Editor/SControlRigFunctionLocalizationWidget.h b/Engine/Plugins/Animation/ControlRig/Source/ControlRigEditor/Private/Editor/SControlRigFunctionLocalizationWidget.h
--- a/Engine/Plugins/Animation/ControlRig/Source/ControlRigEditor/Private/Editor/SControlRigFunctionLocalizationWidget.h
+++ b/Engine/Plugins/Animation/ControlRig/Source/ControlRigEditor/Private/Editor/SControlRigFunctionLocalizationWidget.h
@@ -34,7 +34,7 @@ public:
 	SLATE_BEGIN_ARGS(SControlRigFunctionLocalizationWidget) {}
 	SLATE_END_ARGS()
 
-	void Construct(const FArguments& InArgs, const FRigVMGraphFunctionIdentifier& InFunctionToLocalize, UControlRigBlueprint* InTargetBlueprint);
+	void Construct(const FArguments& InArgs, const FRigVMGraphFunctionIdentifier& InFunctionToLocalize, URigVMBlueprint* InTargetBlueprint);
 
 	TSharedRef<ITableRow> GenerateFunctionListRow(TSharedPtr<SControlRigFunctionLocalizationItem> InItem, const TSharedRef<STableViewBase>& InOwningTable);
 	ECheckBoxState IsFunctionEnabled(const FRigVMGraphFunctionIdentifier InFunction) const;
@@ -60,7 +60,7 @@ public:
 	}
 
 	SLATE_ARGUMENT(FRigVMGraphFunctionIdentifier, Function)
-	SLATE_ARGUMENT(UControlRigBlueprint*, TargetBlueprint)
+	SLATE_ARGUMENT(URigVMBlueprint*, TargetBlueprint)
 	SLATE_END_ARGS()
 
 	void Construct(const FArguments& InArgs);
diff --git a/Engine/Plugins/Animation/ControlRig/Source/ControlRigEditor/Private/Editor/SControlRigStackView.cpp b/Engine/Plugins/Animation/ControlRig/Source/ControlRigEditor/Private/Editor/SControlRigStackView.cpp
--- a/Engine/Plugins/Animation/ControlRig/Source/ControlRigEditor/Private/Editor/SControlRigStackView.cpp
+++ b/Engine/Plugins/Animation/ControlRig/Source/ControlRigEditor/Private/Editor/SControlRigStackView.cpp
@@ -378,7 +378,7 @@ void SControlRigStackView::OnSelectionChanged(TSharedPtr<FRigStackEntry> Selecti
 			return;
 		}
 
-		URigVMBlueprintGeneratedClass* GeneratedClass = ControlRigBlueprint->GetControlRigBlueprintGeneratedClass();
+		URigVMBlueprintGeneratedClass* GeneratedClass = ControlRigBlueprint->GetRigVMBlueprintGeneratedClass();
 		if (GeneratedClass == nullptr)
 		{
 			return;
diff --git a/Engine/Plugins/Animation/ControlRig/Source/ControlRigEditor/Private/Editor/SRigCurveContainer.cpp b/Engine/Plugins/Animation/ControlRig/Source/ControlRigEditor/Private/Editor/SRigCurveContainer.cpp
--- a/Engine/Plugins/Animation/ControlRig/Source/ControlRigEditor/Private/Editor/SRigCurveContainer.cpp
+++ b/Engine/Plugins/Animation/ControlRig/Source/ControlRigEditor/Private/Editor/SRigCurveContainer.cpp
@@ -13,11 +13,10 @@
 #include "Editor/ControlRigEditor.h"
 #include "Engine/SkeletalMesh.h"
 #include "Animation/Skeleton.h"
-#include "ControlRig.h"
 #include "ControlRigBlueprint.h"
 #include "ScopedTransaction.h"
 #include "Kismet2/BlueprintEditorUtils.h"
-#include "ScopedTransaction.h"
+#include "Rigs/RigHierarchyController.h"
 
 #define LOCTEXT_NAMESPACE "SRigCurveContainer"
 
@@ -305,12 +304,12 @@ TSharedPtr<SWidget> SRigCurveContainer::OnGetContextMenuContent() const
 	return MenuBuilder.MakeWidget();
 }
 
-void SRigCurveContainer::OnEditorClose(const FControlRigEditor* InEditor, UControlRigBlueprint* InBlueprint)
+void SRigCurveContainer::OnEditorClose(const FControlRigEditor* InEditor, URigVMBlueprint* InBlueprint)
 {
-	if (InBlueprint)
+	if (UControlRigBlueprint* BP = Cast<UControlRigBlueprint>(InBlueprint))
 	{
-		InBlueprint->Hierarchy->OnModified().RemoveAll(this);
-		InBlueprint->OnRefreshEditor().RemoveAll(this);
+		BP->Hierarchy->OnModified().RemoveAll(this);
+		BP->OnRefreshEditor().RemoveAll(this);
 	}
 
 	ControlRigBlueprint.Reset();
@@ -618,7 +617,7 @@ void SRigCurveContainer::OnHierarchyModified(ERigHierarchyNotification InNotif,
     }
 }
 
-void SRigCurveContainer::HandleRefreshEditorFromBlueprint(UControlRigBlueprint* InBlueprint)
+void SRigCurveContainer::HandleRefreshEditorFromBlueprint(URigVMBlueprint* InBlueprint)
 {
 	RefreshCurveList();
 }
diff --git a/Engine/Plugins/Animation/ControlRig/Source/ControlRigEditor/Private/Editor/SRigCurveContainer.h b/Engine/Plugins/Animation/ControlRig/Source/ControlRigEditor/Private/Editor/SRigCurveContainer.h
--- a/Engine/Plugins/Animation/ControlRig/Source/ControlRigEditor/Private/Editor/SRigCurveContainer.h
+++ b/Engine/Plugins/Animation/ControlRig/Source/ControlRigEditor/Private/Editor/SRigCurveContainer.h
@@ -15,7 +15,7 @@
 class SInlineEditableTextBlock;
 class FControlRigEditor;
 struct FAssetData;
-class UControlRigBlueprint;
+class URigVMBlueprint;
 
 class FDisplayedRigCurveInfo
 {
@@ -206,7 +206,7 @@ private:
 	/** Handler for context menus */
 	TSharedPtr<SWidget> OnGetContextMenuContent() const;
 
-	void OnEditorClose(const FControlRigEditor* InEditor, UControlRigBlueprint* InBlueprint);
+	void OnEditorClose(const FControlRigEditor* InEditor, URigVMBlueprint* InBlueprint);
 
 	/**
 	* Clears and rebuilds the table, according to an optional search string
@@ -236,7 +236,7 @@ private:
 
 	bool bIsChangingRigHierarchy;
 	void OnHierarchyModified(ERigHierarchyNotification InNotif, URigHierarchy* InHierarchy, const FRigBaseElement* InElement);
-	void HandleRefreshEditorFromBlueprint(UControlRigBlueprint* InBlueprint);
+	void HandleRefreshEditorFromBlueprint(URigVMBlueprint* InBlueprint);
 
 	// import curve part
 	void ImportCurve(const FAssetData& InAssetData);
diff --git a/Engine/Plugins/Animation/ControlRig/Source/ControlRigEditor/Private/Editor/SRigHierarchy.cpp b/Engine/Plugins/Animation/ControlRig/Source/ControlRigEditor/Private/Editor/SRigHierarchy.cpp
--- a/Engine/Plugins/Animation/ControlRig/Source/ControlRigEditor/Private/Editor/SRigHierarchy.cpp
+++ b/Engine/Plugins/Animation/ControlRig/Source/ControlRigEditor/Private/Editor/SRigHierarchy.cpp
@@ -272,7 +272,7 @@ void SRigHierarchy::Construct(const FArguments& InArgs, TSharedRef<FControlRigEd
 	CreateDragDropMenu();
 }
 
-void SRigHierarchy::OnEditorClose(const FControlRigEditor* InEditor, UControlRigBlueprint* InBlueprint)
+void SRigHierarchy::OnEditorClose(const FControlRigEditor* InEditor, URigVMBlueprint* InBlueprint)
 {
 	if (InEditor)
 	{
@@ -282,9 +282,9 @@ void SRigHierarchy::OnEditorClose(const FControlRigEditor* InEditor, UControlRig
 		Editor->OnViewportContextMenuCommands().Unbind();
 	}
 
-	if (InBlueprint)
+	if (UControlRigBlueprint* BP = Cast<UControlRigBlueprint>(InBlueprint))
 	{
-		InBlueprint->Hierarchy->OnModified().RemoveAll(this);
+		BP->Hierarchy->OnModified().RemoveAll(this);
 		InBlueprint->OnRefreshEditor().RemoveAll(this);
 		InBlueprint->OnSetObjectBeingDebugged().RemoveAll(this);
 	}
@@ -820,7 +820,7 @@ void SRigHierarchy::OnHierarchyModified_AnyThread(ERigHierarchyNotification InNo
 	}
 }
 
-void SRigHierarchy::HandleRefreshEditorFromBlueprint(UControlRigBlueprint* InBlueprint)
+void SRigHierarchy::HandleRefreshEditorFromBlueprint(URigVMBlueprint* InBlueprint)
 {
 	if (bIsChangingRigHierarchy)
 	{
diff --git a/Engine/Plugins/Animation/ControlRig/Source/ControlRigEditor/Private/Editor/SRigHierarchy.h b/Engine/Plugins/Animation/ControlRig/Source/ControlRigEditor/Private/Editor/SRigHierarchy.h
--- a/Engine/Plugins/Animation/ControlRig/Source/ControlRigEditor/Private/Editor/SRigHierarchy.h
+++ b/Engine/Plugins/Animation/ControlRig/Source/ControlRigEditor/Private/Editor/SRigHierarchy.h
@@ -8,13 +8,14 @@
 #include "Engine/SkeletalMesh.h"
 #include "Editor/SRigHierarchyTreeView.h"
 #include "Units/RigUnitContext.h"
+#include "ControlRigBlueprint.h"
 #include "SRigHierarchy.generated.h"
 
 class SRigHierarchy;
 class FControlRigEditor;
 class SSearchBox;
 class FUICommandList;
-class UControlRigBlueprint;
+class URigVMBlueprint;
 class UControlRig;
 struct FAssetData;
 class FMenuBuilder;
@@ -85,7 +86,7 @@ public:
 
 private:
 
-	void OnEditorClose(const FControlRigEditor* InEditor, UControlRigBlueprint* InBlueprint);
+	void OnEditorClose(const FControlRigEditor* InEditor, URigVMBlueprint* InBlueprint);
 
 	/** Bind commands that this widget handles */
 	void BindCommands();
@@ -223,7 +224,7 @@ private:
 	bool bIsChangingRigHierarchy;
 	void OnHierarchyModified(ERigHierarchyNotification InNotif, URigHierarchy* InHierarchy, const FRigBaseElement* InElement);
 	void OnHierarchyModified_AnyThread(ERigHierarchyNotification InNotif, URigHierarchy* InHierarchy, const FRigBaseElement* InElement);
-	void HandleRefreshEditorFromBlueprint(UControlRigBlueprint* InBlueprint);
+	void HandleRefreshEditorFromBlueprint(URigVMBlueprint* InBlueprint);
 	void HandleSetObjectBeingDebugged(UObject* InObject);
 	void OnPreConstruction_AnyThread(UControlRig* InRig, const FName& InEventName);
 	void OnPostConstruction_AnyThread(UControlRig* InRig, const FName& InEventName);
diff --git a/Engine/Plugins/Animation/ControlRig/Source/ControlRigEditor/Private/Graph/SControlRigGraphPinNameList.cpp b/Engine/Plugins/Animation/ControlRig/Source/ControlRigEditor/Private/Graph/SControlRigGraphPinNameList.cpp
--- a/Engine/Plugins/Animation/ControlRig/Source/ControlRigEditor/Private/Graph/SControlRigGraphPinNameList.cpp
+++ b/Engine/Plugins/Animation/ControlRig/Source/ControlRigEditor/Private/Graph/SControlRigGraphPinNameList.cpp
@@ -259,6 +259,8 @@ FReply SControlRigGraphPinNameList::OnGetSelectedClicked()
 {
 	if (UControlRigGraph* Graph = Cast<UControlRigGraph>(GetPinObj()->GetOwningNode()->GetGraph()))
 	{
+		UControlRigBlueprint* ControlRigBlueprint = CastChecked<UControlRigBlueprint>(Graph->GetBlueprint());
+
 		if (OnGetNameFromSelection.IsBound())
 		{
 			const TArray<TSharedPtr<FString>> Result = OnGetNameFromSelection.Execute();
@@ -279,7 +281,7 @@ FReply SControlRigGraphPinNameList::OnGetSelectedClicked()
 
 					else if (ModelPin->GetCustomWidgetName() == TEXT("BoneName"))
 					{
-						URigHierarchy* Hierarchy = Graph->GetBlueprint()->Hierarchy;
+						URigHierarchy* Hierarchy = ControlRigBlueprint->Hierarchy;
 						TArray<FRigElementKey> Keys = Hierarchy->GetSelectedKeys();
 						FRigBaseElement* Element = Hierarchy->FindChecked(Keys[0]);
 						if (Element->GetType() == ERigElementType::Bone)
@@ -331,6 +333,7 @@ FReply SControlRigGraphPinNameList::OnBrowseClicked()
 		TSharedPtr<FString> Selected = NameListComboBox->GetSelectedItem();
 		if (Selected.IsValid() && ModelPin)
 		{
+			UControlRigBlueprint* ControlRigBlueprint = CastChecked<UControlRigBlueprint>(Graph->GetBlueprint());
 			URigVMPin* KeyPin = ModelPin->GetParentPin();
 			if(KeyPin && KeyPin->GetCPPTypeObject() == FRigElementKey::StaticStruct())
 			{
@@ -342,7 +345,7 @@ FReply SControlRigGraphPinNameList::OnBrowseClicked()
 					FRigElementKey::StaticStruct()->ImportText(*DefaultValue, &Key, nullptr, EPropertyPortFlags::PPF_None, nullptr, FRigElementKey::StaticStruct()->GetName(), true);
 					if (Key.IsValid())
 					{
-						Graph->GetBlueprint()->GetHierarchyController()->SetSelection({Key});
+						ControlRigBlueprint->GetHierarchyController()->SetSelection({Key});
 					}
 				}
 			}
@@ -351,7 +354,7 @@ FReply SControlRigGraphPinNameList::OnBrowseClicked()
 				// browse to named bone
 				const FString DefaultValue = ModelPin->GetDefaultValue();
 				FRigElementKey Key(*DefaultValue, ERigElementType::Bone);
-				Graph->GetBlueprint()->GetHierarchyController()->SetSelection({Key});
+				ControlRigBlueprint->GetHierarchyController()->SetSelection({Key});
 			}
 			else
 			{
@@ -370,7 +373,7 @@ FReply SControlRigGraphPinNameList::OnBrowseClicked()
 							const FString DefaultValue = ModelPin->GetDefaultValue();
 							const ERigElementType ElementType = (ERigElementType)StaticEnum<ERigElementType>()->GetValueByIndex(EnumIndex);
 							FRigElementKey Key(*DefaultValue, ElementType);
-							Graph->GetBlueprint()->GetHierarchyController()->SetSelection({Key});
+							ControlRigBlueprint->GetHierarchyController()->SetSelection({Key});
 							break;
 						}
 					}
diff --git a/Engine/Plugins/Animation/ControlRig/Source/ControlRigEditor/Private/Sequencer/ControlRigParameterTrackEditor.cpp b/Engine/Plugins/Animation/ControlRig/Source/ControlRigEditor/Private/Sequencer/ControlRigParameterTrackEditor.cpp
--- a/Engine/Plugins/Animation/ControlRig/Source/ControlRigEditor/Private/Sequencer/ControlRigParameterTrackEditor.cpp
+++ b/Engine/Plugins/Animation/ControlRig/Source/ControlRigEditor/Private/Sequencer/ControlRigParameterTrackEditor.cpp
@@ -1920,7 +1920,7 @@ void FControlRigParameterTrackEditor::AddTrackForComponent(USceneComponent* InCo
 					{
 						if (UControlRigBlueprint* BPControlRig = Cast<UControlRigBlueprint>(Object))
 						{
-							if (URigVMBlueprintGeneratedClass* RigClass = BPControlRig->GetControlRigBlueprintGeneratedClass())
+							if (URigVMBlueprintGeneratedClass* RigClass = BPControlRig->GetRigVMBlueprintGeneratedClass())
 							{
 								if (UControlRig* CDO = Cast<UControlRig>(RigClass->GetDefaultObject(true /* create if needed */)))
 								{
@@ -3913,7 +3913,7 @@ bool FControlRigParameterTrackEditor::HandleAssetAdded(UObject* Asset, const FGu
 	}
 
 	UControlRigBlueprint* ControlRigBlueprint = Cast<UControlRigBlueprint>(Asset);
-	URigVMBlueprintGeneratedClass* RigClass = ControlRigBlueprint->GetControlRigBlueprintGeneratedClass();
+	URigVMBlueprintGeneratedClass* RigClass = ControlRigBlueprint->GetRigVMBlueprintGeneratedClass();
 	if (!RigClass)
 	{
 		return false;
