Control Rig: FIx compilation error
#rnx

diff --git a/Engine/Plugins/Animation/ControlRig/Source/ControlRigEditor/Private/Editor/SControlRigValidationWidget.cpp b/Engine/Plugins/Animation/ControlRig/Source/ControlRigEditor/Private/Editor/SControlRigValidationWidget.cpp
--- a/Engine/Plugins/Animation/ControlRig/Source/ControlRigEditor/Private/Editor/SControlRigValidationWidget.cpp
+++ b/Engine/Plugins/Animation/ControlRig/Source/ControlRigEditor/Private/Editor/SControlRigValidationWidget.cpp
@@ -18,6 +18,7 @@
 #include "Framework/Application/SlateApplication.h"
 #include "ControlRigBlueprint.h"
 #include "UObject/UObjectIterator.h"
+#include "Rigs/RigHierarchyController.h"
 
 #define LOCTEXT_NAMESPACE "SControlRigValidationWidget"
 
