Sequencer: Control Rig: Remove unneeded Reconstruct call in track editor that was added a while back for constraints but is now handled correctly by the track itself which checks for changing controls. Also noticed that we copy pastad GetBindableObject as a hotfix so now just exposed that as a static.

diff --git a/Engine/Plugins/Animation/ControlRig/Source/ControlRig/Private/ControlRigObjectBinding.cpp b/Engine/Plugins/Animation/ControlRig/Source/ControlRig/Private/ControlRigObjectBinding.cpp
--- a/Engine/Plugins/Animation/ControlRig/Source/ControlRig/Private/ControlRigObjectBinding.cpp
+++ b/Engine/Plugins/Animation/ControlRig/Source/ControlRig/Private/ControlRigObjectBinding.cpp
@@ -43,7 +43,7 @@ AActor* FControlRigObjectBinding::GetHostingActor() const
 	return nullptr;
 }
 
-UObject* FControlRigObjectBinding::GetBindableObject(UObject* InObject) const
+UObject* FControlRigObjectBinding::GetBindableObject(UObject* InObject) 
 {
 	// If we are binding to an actor, find the first skeletal mesh component
 	if (AActor* Actor = Cast<AActor>(InObject))
diff --git a/Engine/Plugins/Animation/ControlRig/Source/ControlRig/Private/Sequencer/MovieSceneControlRigParameterTemplate.cpp b/Engine/Plugins/Animation/ControlRig/Source/ControlRig/Private/Sequencer/MovieSceneControlRigParameterTemplate.cpp
--- a/Engine/Plugins/Animation/ControlRig/Source/ControlRig/Private/Sequencer/MovieSceneControlRigParameterTemplate.cpp
+++ b/Engine/Plugins/Animation/ControlRig/Source/ControlRig/Private/Sequencer/MovieSceneControlRigParameterTemplate.cpp
@@ -1124,35 +1124,6 @@ struct FControlRigParameterExecutionToken : IMovieSceneExecutionToken
 	FControlRigParameterExecutionToken(const FControlRigParameterExecutionToken&) = delete;
 	FControlRigParameterExecutionToken& operator=(const FControlRigParameterExecutionToken&) = delete;
 
-	UObject* GetBindableObject(UObject* InObject) const
-	{
-		// If we are binding to an actor, find the first skeletal mesh component
-		if (AActor* Actor = Cast<AActor>(InObject))
-		{
-			if (UControlRigComponent* ControlRigComponent = Actor->FindComponentByClass<UControlRigComponent>())
-			{
-				return ControlRigComponent;
-			}
-			else if (USkeletalMeshComponent* SkeletalMeshComponent = Actor->FindComponentByClass<USkeletalMeshComponent>())
-			{
-				return SkeletalMeshComponent;
-			}
-		}
-		else if (UControlRigComponent* ControlRigComponent = Cast<UControlRigComponent>(InObject))
-		{
-			return ControlRigComponent;
-		}
-		else if (USkeletalMeshComponent* SkeletalMeshComponent = Cast<USkeletalMeshComponent>(InObject))
-		{
-			return SkeletalMeshComponent;
-		}
-		else if (USkeleton* Skeleton = Cast<USkeleton>(InObject))
-		{
-			return Skeleton;
-		}
-		return nullptr;
-	}
-
 	virtual void Execute(const FMovieSceneContext& Context, const FMovieSceneEvaluationOperand& Operand, FPersistentEvaluationData& PersistentData, IMovieScenePlayer& Player)
 	{
 		MOVIESCENE_DETAILED_SCOPE_CYCLE_COUNTER(MovieSceneEval_ControlRigParameterTrack_TokenExecute)
@@ -1175,12 +1146,12 @@ struct FControlRigParameterExecutionToken : IMovieSceneExecutionToken
 					ControlRig->SetObjectBinding(MakeShared<FControlRigObjectBinding>());
 				}
 
-				if (ControlRig->GetObjectBinding()->GetBoundObject() != GetBindableObject(BoundObject))
+				if (ControlRig->GetObjectBinding()->GetBoundObject() != FControlRigObjectBinding::GetBindableObject(BoundObject))
 				{
 					ControlRig->GetObjectBinding()->BindToObject(BoundObject);
 					TArray<FName> SelectedControls = ControlRig->CurrentControlSelection();
 					ControlRig->Initialize();
-					if (USkeletalMeshComponent* SkeletalMeshComponent = Cast<USkeletalMeshComponent>(GetBindableObject(BoundObject)))
+					if (USkeletalMeshComponent* SkeletalMeshComponent = Cast<USkeletalMeshComponent>(FControlRigObjectBinding::GetBindableObject(BoundObject)))
 					{
 						ControlRig->RequestInit();
 						ControlRig->SetBoneInitialTransformsFromSkeletalMeshComponent(SkeletalMeshComponent, true);
diff --git a/Engine/Plugins/Animation/ControlRig/Source/ControlRig/Public/ControlRigObjectBinding.h b/Engine/Plugins/Animation/ControlRig/Source/ControlRig/Public/ControlRigObjectBinding.h
--- a/Engine/Plugins/Animation/ControlRig/Source/ControlRig/Public/ControlRigObjectBinding.h
+++ b/Engine/Plugins/Animation/ControlRig/Source/ControlRig/Public/ControlRigObjectBinding.h
@@ -22,8 +22,8 @@ public:
 	virtual bool IsBoundToObject(UObject* InObject) const override;
 	virtual UObject* GetBoundObject() const override;
 	virtual AActor* GetHostingActor() const override;
-protected:
-	UObject* GetBindableObject(UObject* InObject) const;
+
+	static UObject* GetBindableObject(UObject* InObject);
 private:
 	/** The scene component or USkeleton we are bound to */
 	TWeakObjectPtr<UObject> BoundObject;
diff --git a/Engine/Plugins/Animation/ControlRig/Source/ControlRigEditor/Private/Sequencer/ControlRigParameterTrackEditor.cpp b/Engine/Plugins/Animation/ControlRig/Source/ControlRigEditor/Private/Sequencer/ControlRigParameterTrackEditor.cpp
--- a/Engine/Plugins/Animation/ControlRig/Source/ControlRigEditor/Private/Sequencer/ControlRigParameterTrackEditor.cpp
+++ b/Engine/Plugins/Animation/ControlRig/Source/ControlRigEditor/Private/Sequencer/ControlRigParameterTrackEditor.cpp
@@ -3208,15 +3208,6 @@ void FControlRigParameterTrackEditor::HandleOnObjectBoundToControlRig(UObject* I
 		}
 	}
 
-	// reconstruct proxies
-	if (!SectionsToUpdate.IsEmpty())
-	{
-		for (UMovieSceneControlRigParameterSection* Section : SectionsToUpdate)
-		{
-			Section->ReconstructChannelProxy();
-			Section->MarkAsChanged();
-		}
-	}
 	if (ReselectIfNeeded.Num() > 0)
 	{
 		GEditor->GetTimerManager()->SetTimerForNextTick([ReselectIfNeeded]()
