Control Rig: Editing animation channels within parent details doesn't work

diff --git a/Engine/Plugins/Animation/ControlRig/Source/ControlRigEditor/Private/ControlRigElementDetails.cpp b/Engine/Plugins/Animation/ControlRig/Source/ControlRigEditor/Private/ControlRigElementDetails.cpp
--- a/Engine/Plugins/Animation/ControlRig/Source/ControlRigEditor/Private/ControlRigElementDetails.cpp
+++ b/Engine/Plugins/Animation/ControlRig/Source/ControlRigEditor/Private/ControlRigElementDetails.cpp
@@ -3274,7 +3274,10 @@ void FRigControlElementDetails::CustomizeAnimationChannels(IDetailLayoutBuilder&
 				[
 					SNew(SInlineEditableTextBlock)
 					.Font(IDetailLayoutBuilder::GetDetailFont())
-					.Text(FText::FromName(ChildControlElement->GetDisplayName()))
+					.Text_Lambda([this, ChildElementKey]() -> FText
+					{
+						return GetDisplayNameForElement(ChildElementKey);
+					})
 					.OnTextCommitted_Lambda([this, ChildElementKey](const FText& InNewText, ETextCommit::Type InCommitType)
 					{
 						SetDisplayNameForElement(InNewText, InCommitType, ChildElementKey);
@@ -3549,13 +3552,26 @@ void FRigControlElementDetails::HandleControlTypeChanged(TSharedPtr<ERigControlT
 
 void FRigControlElementDetails::HandleControlTypeChanged(ERigControlType ControlType, TArray<FRigElementKey> ControlKeys, const TSharedRef<IPropertyUtilities> PropertyUtilities)
 {
-	for(FPerElementInfo& Info : PerElementInfos)
+	if(PerElementInfos.IsEmpty())
+	{
+		return;
+	}
+	
+	if(ControlKeys.IsEmpty())
 	{
-		URigHierarchy* Hierarchy = Info.GetHierarchy();
-		URigHierarchy* HierarchyToChange = Info.GetDefaultHierarchy();
+		for(const FPerElementInfo& Info : PerElementInfos)
+		{
+			ControlKeys.Add(Info.GetDefaultElement<FRigControlElement>()->GetKey());
+		}
+	}
+
+	for(const FRigElementKey& ControlKey : ControlKeys)
+	{
+		URigHierarchy* Hierarchy = PerElementInfos[0].GetHierarchy();
+		URigHierarchy* HierarchyToChange = PerElementInfos[0].GetDefaultHierarchy();
 		HierarchyToChange->Modify();
 		
-		FRigControlElement* ControlElement = Info.GetDefaultElement<FRigControlElement>();
+		FRigControlElement* ControlElement = HierarchyToChange->FindChecked<FRigControlElement>(ControlKey);
 		
 		FRigControlValue ValueToSet;
 
@@ -3661,11 +3677,17 @@ void FRigControlElementDetails::HandleControlTypeChanged(ERigControlType Control
 		HierarchyToChange->SetControlValue(ControlElement, ValueToSet, ERigControlValueType::Initial, true, false, true);
 		HierarchyToChange->SetControlValue(ControlElement, ValueToSet, ERigControlValueType::Current, true, false, true);
 
-		Info.WrapperObject->SetContent<FRigControlElement>(*ControlElement);
+		for(const FPerElementInfo& Info : PerElementInfos)
+		{
+			if(Info.Element.Get()->GetKey() == ControlKey)
+			{
+				Info.WrapperObject->SetContent<FRigControlElement>(*ControlElement);
+			}
+		}
 
 		if (HierarchyToChange != Hierarchy)
 		{
-			if(FRigControlElement* OtherControlElement = Info.GetElement<FRigControlElement>())
+			if(FRigControlElement* OtherControlElement = Hierarchy->Find<FRigControlElement>(ControlKey))
 			{
 				OtherControlElement->Settings = ControlElement->Settings;
 				Hierarchy->SetControlSettings(OtherControlElement, OtherControlElement->Settings, true, true, true);
@@ -3675,7 +3697,7 @@ void FRigControlElementDetails::HandleControlTypeChanged(ERigControlType Control
 		}
 		else
 		{
-			Info.GetBlueprint()->PropagateHierarchyFromBPToInstances();
+			PerElementInfos[0].GetBlueprint()->PropagateHierarchyFromBPToInstances();
 		}
 	}
 	
@@ -4248,6 +4270,23 @@ void FRigControlElementDetails::SetDisplayName(const FText& InNewText, ETextComm
 	}
 }
 
+FText FRigControlElementDetails::GetDisplayNameForElement(const FRigElementKey& InKey) const
+{
+	if(PerElementInfos.IsEmpty())
+	{
+		return FText();
+	}
+
+	URigHierarchy* Hierarchy = PerElementInfos[0].GetDefaultHierarchy();
+	const FRigControlElement* ControlElement = Hierarchy->Find<FRigControlElement>(InKey);
+	if(ControlElement == nullptr)
+	{
+		return FText();
+	}
+
+	return FText::FromName(ControlElement->GetDisplayName());
+}
+
 void FRigControlElementDetails::SetDisplayNameForElement(const FText& InNewText, ETextCommit::Type InCommitType, const FRigElementKey& InKeyToRename)
 {
 	if(InCommitType == ETextCommit::OnCleared)
@@ -4255,20 +4294,25 @@ void FRigControlElementDetails::SetDisplayNameForElement(const FText& InNewText,
 		return;
 	}
 
-	const FPerElementInfo& Info = FindElement(InKeyToRename);
-	if(!Info.IsValid())
+	if(PerElementInfos.IsEmpty())
 	{
 		return;
 	}
 
-	if(Info.IsProcedural())
+	URigHierarchy* Hierarchy = PerElementInfos[0].GetDefaultHierarchy();
+	const FRigControlElement* ControlElement = Hierarchy->Find<FRigControlElement>(InKeyToRename);
+	if(ControlElement == nullptr)
+	{
+		return;
+	}
+	if(ControlElement->IsProcedural())
 	{
 		return;
 	}
 
 	const FName DisplayName = InNewText.IsEmpty() ? FName(NAME_None) : FName(*InNewText.ToString());
 	const bool bRename = IsAnyControlOfAnimationType(ERigControlAnimationType::AnimationChannel);
-	Info.GetDefaultHierarchy()->GetController(true)->SetDisplayName(InKeyToRename, DisplayName, bRename, true, true);
+	Hierarchy->GetController(true)->SetDisplayName(InKeyToRename, DisplayName, bRename, true, true);
 }
 
 bool FRigControlElementDetails::OnVerifyDisplayNameChanged(const FText& InText, FText& OutErrorMessage, const FRigElementKey& InKeyToRename)
@@ -4280,31 +4324,30 @@ bool FRigControlElementDetails::OnVerifyDisplayNameChanged(const FText& InText,
 		return false;
 	}
 
-	const FPerElementInfo& Info = FindElement(InKeyToRename);
-	if(!Info.IsValid())
+	if(PerElementInfos.IsEmpty())
 	{
 		return false;
 	}
 
-	if(Info.IsProcedural())
+	const URigHierarchy* Hierarchy = PerElementInfos[0].GetDefaultHierarchy();
+	const FRigControlElement* ControlElement = Hierarchy->Find<FRigControlElement>(InKeyToRename);
+	if(ControlElement == nullptr)
+	{
+		return false;
+	}
+	if(ControlElement->IsProcedural())
 	{
 		return false;
 	}
 
 	// make sure there is no duplicate
-	if(const URigHierarchy* Hierarchy = Info.GetDefaultHierarchy())
+	if(const FRigBaseElement* ParentElement = Hierarchy->GetFirstParent(ControlElement))
 	{
-		if(const FRigControlElement* ControlElement = Info.GetDefaultElement<FRigControlElement>())
+		FString OutErrorString;
+		if (!Hierarchy->IsDisplayNameAvailable(ParentElement->GetKey(), NewName, &OutErrorString))
 		{
-			if(const FRigBaseElement* ParentElement = Hierarchy->GetFirstParent(ControlElement))
-			{
-				FString OutErrorString;
-				if (!Hierarchy->IsDisplayNameAvailable(ParentElement->GetKey(), NewName, &OutErrorString))
-				{
-					OutErrorMessage = FText::FromString(OutErrorString);
-					return false;
-				}
-			}
+			OutErrorMessage = FText::FromString(OutErrorString);
+			return false;
 		}
 	}
 	return true;
diff --git a/Engine/Plugins/Animation/ControlRig/Source/ControlRigEditor/Private/ControlRigElementDetails.h b/Engine/Plugins/Animation/ControlRig/Source/ControlRigEditor/Private/ControlRigElementDetails.h
--- a/Engine/Plugins/Animation/ControlRig/Source/ControlRigEditor/Private/ControlRigElementDetails.h
+++ b/Engine/Plugins/Animation/ControlRig/Source/ControlRigEditor/Private/ControlRigElementDetails.h
@@ -361,6 +361,7 @@ public:
 
 	FText GetDisplayName() const;
 	void SetDisplayName(const FText& InNewText, ETextCommit::Type InCommitType);
+	FText GetDisplayNameForElement(const FRigElementKey& InKey) const;
 	void SetDisplayNameForElement(const FText& InNewText, ETextCommit::Type InCommitType, const FRigElementKey& InKeyToRename);
 	bool OnVerifyDisplayNameChanged(const FText& InText, FText& OutErrorMessage, const FRigElementKey& InKeyToRename);
 
diff --git a/Engine/Plugins/Animation/ControlRig/Source/ControlRigEditor/Private/Editor/SRigHierarchyTreeView.cpp b/Engine/Plugins/Animation/ControlRig/Source/ControlRigEditor/Private/Editor/SRigHierarchyTreeView.cpp
--- a/Engine/Plugins/Animation/ControlRig/Source/ControlRigEditor/Private/Editor/SRigHierarchyTreeView.cpp
+++ b/Engine/Plugins/Animation/ControlRig/Source/ControlRigEditor/Private/Editor/SRigHierarchyTreeView.cpp
@@ -43,21 +43,6 @@ FRigTreeElement::FRigTreeElement(const FRigElementKey& InKey, TWeakPtr<SRigHiera
 	{
 		if(const URigHierarchy* Hierarchy = InTreeView.Pin()->GetRigTreeDelegates().GetHierarchy())
 		{
-			if(const FRigBaseElement* Element = Hierarchy->Find(Key))
-			{
-				bIsProcedural = Element->IsProcedural();
-				
-				if(const FRigControlElement* ControlElement = Cast<FRigControlElement>(Element))
-				{
-					bIsTransient = ControlElement->Settings.bIsTransientControl;
-					bIsAnimationChannel = ControlElement->IsAnimationChannel();
-					if(bIsAnimationChannel)
-					{
-						ChannelName = ControlElement->GetDisplayName();
-					}
-				}
-			}
-
 			const FRigTreeDisplaySettings& Settings = InTreeView.Pin()->GetRigTreeDelegates().GetDisplaySettings();
 			RefreshDisplaySettings(Hierarchy, Settings);
 		}
@@ -82,6 +67,21 @@ void FRigTreeElement::RefreshDisplaySettings(const URigHierarchy* InHierarchy, c
 {
 	const TPair<const FSlateBrush*, FSlateColor> Result = SRigHierarchyItem::GetBrushForElementType(InHierarchy, Key);
 
+	if(const FRigBaseElement* Element = InHierarchy->Find(Key))
+	{
+		bIsProcedural = Element->IsProcedural();
+				
+		if(const FRigControlElement* ControlElement = Cast<FRigControlElement>(Element))
+		{
+			bIsTransient = ControlElement->Settings.bIsTransientControl;
+			bIsAnimationChannel = ControlElement->IsAnimationChannel();
+			if(bIsAnimationChannel)
+			{
+				ChannelName = ControlElement->GetDisplayName();
+			}
+		}
+	}
+
 	IconBrush = Result.Key;
 	IconColor = Result.Value;
 	if(IconColor.IsColorSpecified() && InSettings.bShowIconColors)
