Fix for Control Rig Gizmos not working, after a changing the refs to TSoftObjectPtr

diff --git a/Engine/Plugins/Animation/ControlRig/Source/ControlRig/Private/ControlRigGizmoActor.cpp b/Engine/Plugins/Animation/ControlRig/Source/ControlRig/Private/ControlRigGizmoActor.cpp
--- a/Engine/Plugins/Animation/ControlRig/Source/ControlRig/Private/ControlRigGizmoActor.cpp
+++ b/Engine/Plugins/Animation/ControlRig/Source/ControlRig/Private/ControlRigGizmoActor.cpp
@@ -143,6 +143,11 @@ bool AControlRigShapeActor::UpdateControlSettings(
 
 			if(bShapeNameUpdated)
 			{
+				if (!ShapeDef->StaticMesh.IsValid())
+				{
+					ShapeDef->StaticMesh.LoadSynchronous();
+				}
+
 				if(ShapeDef->StaticMesh.IsValid())
 				{
 					if(UStaticMesh* StaticMesh = ShapeDef->StaticMesh.Get())
