ControlRig: SRigCurveContainer fix non-unity build

diff --git a/Engine/Plugins/Animation/ControlRig/Source/ControlRigEditor/Private/Editor/SRigCurveContainer.h b/Engine/Plugins/Animation/ControlRig/Source/ControlRigEditor/Private/Editor/SRigCurveContainer.h
--- a/Engine/Plugins/Animation/ControlRig/Source/ControlRigEditor/Private/Editor/SRigCurveContainer.h
+++ b/Engine/Plugins/Animation/ControlRig/Source/ControlRigEditor/Private/Editor/SRigCurveContainer.h
@@ -8,6 +8,7 @@
 #include "Widgets/Views/STableRow.h"
 #include "Widgets/Views/SListView.h"
 #include "Rigs/RigHierarchy.h"
+#include "ControlRigBlueprint.h"
 
 //////////////////////////////////////////////////////////////////////////
 // FDisplayedRigCurveInfo
