RigVM: Move Control Rig Editor settings base class to RigVM

diff --git a/Engine/Plugins/Animation/ControlRig/Source/ControlRig/Private/Settings/ControlRigSettings.cpp b/Engine/Plugins/Animation/ControlRig/Source/ControlRig/Private/Settings/ControlRigSettings.cpp
--- a/Engine/Plugins/Animation/ControlRig/Source/ControlRig/Private/Settings/ControlRigSettings.cpp
+++ b/Engine/Plugins/Animation/ControlRig/Source/ControlRig/Private/Settings/ControlRigSettings.cpp
@@ -4,11 +4,6 @@
 
 #include UE_INLINE_GENERATED_CPP_BY_NAME(ControlRigSettings)
 
-#if WITH_EDITOR
-#include "RigVMModel/RigVMController.h"
-#include "RigVMModel/RigVMNode.h"
-#endif
-
 UControlRigSettings::UControlRigSettings(const FObjectInitializer& ObjectInitializer)
 	: Super(ObjectInitializer)
 {
@@ -24,7 +19,6 @@ UControlRigEditorSettings::UControlRigEditorSettings(const FObjectInitializer& O
 #endif
 {
 #if WITH_EDITORONLY_DATA
-	bAutoLinkMutableNodes = false;
 	bResetControlsOnCompile = true;
 	bResetControlsOnPinValueInteraction = false;
 	bEnableUndoForPoseInteraction = true;
diff --git a/Engine/Plugins/Animation/ControlRig/Source/ControlRig/Public/Settings/ControlRigSettings.h b/Engine/Plugins/Animation/ControlRig/Source/ControlRig/Public/Settings/ControlRigSettings.h
--- a/Engine/Plugins/Animation/ControlRig/Source/ControlRig/Public/Settings/ControlRigSettings.h
+++ b/Engine/Plugins/Animation/ControlRig/Source/ControlRig/Public/Settings/ControlRigSettings.h
@@ -7,7 +7,7 @@
 #pragma once
 
 #include "CoreMinimal.h"
-#include "Engine/DeveloperSettings.h"
+#include "RigVMSettings.h"
 #include "ControlRigGizmoLibrary.h"
 
 #if WITH_EDITOR
@@ -49,20 +49,13 @@ public:
  * Customize Control Rig Editor.
  */
 UCLASS(config = EditorPerProjectUserSettings, meta=(DisplayName="Control Rig Editor"))
-class CONTROLRIG_API UControlRigEditorSettings : public UDeveloperSettings
+class CONTROLRIG_API UControlRigEditorSettings : public URigVMEditorSettings
 {
 	GENERATED_UCLASS_BODY()
 
 public:
-	/** Gets the category for the settings, some high level grouping like, Editor, Engine, Game...etc. */
-	virtual FName GetCategoryName() const override { return TEXT("ContentEditors"); }
-	
-#if WITH_EDITORONLY_DATA
 
-	// When this is checked mutable nodes (nodes with an execute pin)
-	// will be hooked up automatically.
-	UPROPERTY(EditAnywhere, config, Category = Interaction)
-	bool bAutoLinkMutableNodes;
+#if WITH_EDITORONLY_DATA
 
 	// When this is checked all controls will return to their initial
 	// value as the user hits the Compile button.
@@ -124,66 +117,6 @@ public:
 	UPROPERTY(EditAnywhere, config, Category = Hierarchy, meta = (EditCondition = "bShowStackedHierarchy"))
 	int32 MaxStackSize;
 
-	/**
-	 * The default node snippet to create when pressing 1 + Left Mouse Button
-	 */
-	UPROPERTY(EditAnywhere, config, Category = Snippets, meta = (DisplayName = "1"))
-	FString NodeSnippet_1;
-
-	/**
-	* The default node snippet to create when pressing 2 + Left Mouse Button
-	*/
-	UPROPERTY(EditAnywhere, config, Category = Snippets, meta = (DisplayName = "2"))
-	FString NodeSnippet_2;
-
-	/**
-	* The default node snippet to create when pressing 3 + Left Mouse Button
-	*/
-	UPROPERTY(EditAnywhere, config, Category = Snippets, meta = (DisplayName = "3"))
-	FString NodeSnippet_3;
-
-	/**
-	* The default node snippet to create when pressing 4 + Left Mouse Button
-	*/
-	UPROPERTY(EditAnywhere, config, Category = Snippets, meta = (DisplayName = "4"))
-	FString NodeSnippet_4;
-
-	/**
-	* The default node snippet to create when pressing 5 + Left Mouse Button
-	*/
-	UPROPERTY(EditAnywhere, config, Category = Snippets, meta = (DisplayName = "5"))
-	FString NodeSnippet_5;
-
-	/**
-	* The default node snippet to create when pressing 6 + Left Mouse Button
-	*/
-	UPROPERTY(EditAnywhere, config, Category = Snippets, meta = (DisplayName = "6"))
-	FString NodeSnippet_6;
-
-	/**
-	* The default node snippet to create when pressing 7 + Left Mouse Button
-	*/
-	UPROPERTY(EditAnywhere, config, Category = Snippets, meta = (DisplayName = "7"))
-	FString NodeSnippet_7;
-
-	/**
-	* The default node snippet to create when pressing 8 + Left Mouse Button
-	*/
-	UPROPERTY(EditAnywhere, config, Category = Snippets, meta = (DisplayName = "8"))
-	FString NodeSnippet_8;
-
-	/**
-	* The default node snippet to create when pressing 9 + Left Mouse Button
-	*/
-	UPROPERTY(EditAnywhere, config, Category = Snippets, meta = (DisplayName = "9"))
-	FString NodeSnippet_9;
-
-	/**
-	* The default node snippet to create when pressing 0 + Left Mouse Button
-	*/
-	UPROPERTY(EditAnywhere, config, Category = Snippets, meta = (DisplayName = "0"))
-	FString NodeSnippet_0;
-
 #endif
 
 	static UControlRigEditorSettings * Get() { return GetMutableDefault<UControlRigEditorSettings>(); }
diff --git a/Engine/Plugins/Animation/ControlRig/Source/ControlRigDeveloper/Private/ControlRigBlueprint.cpp b/Engine/Plugins/Animation/ControlRig/Source/ControlRigDeveloper/Private/ControlRigBlueprint.cpp
--- a/Engine/Plugins/Animation/ControlRig/Source/ControlRigDeveloper/Private/ControlRigBlueprint.cpp
+++ b/Engine/Plugins/Animation/ControlRig/Source/ControlRigDeveloper/Private/ControlRigBlueprint.cpp
@@ -23,6 +23,7 @@
 #include "RigVMTypeUtils.h"
 #include "RigVMModel/Nodes/RigVMAggregateNode.h"
 #include "Rigs/RigControlHierarchy.h"
+#include "Settings/ControlRigSettings.h"
 
 #include UE_INLINE_GENERATED_CPP_BY_NAME(ControlRigBlueprint)
 
@@ -30,7 +31,6 @@
 #include "IControlRigEditorModule.h"
 #include "Kismet2/WatchedPin.h"
 #include "Kismet2/BlueprintEditorUtils.h"
-#include "Settings/ControlRigSettings.h"
 #include "Editor/UnrealEdEngine.h"
 #include "Editor/Transactor.h"
 #include "CookOnTheSide/CookOnTheFlyServer.h"
@@ -376,6 +376,11 @@ UClass* UControlRigBlueprint::GetRigVMEdGraphClass() const
 	return UControlRigGraph::StaticClass();
 }
 
+UClass* UControlRigBlueprint::GetRigVMEditorSettingsClass() const
+{
+	return UControlRigEditorSettings::StaticClass();
+}
+
 TArray<FString> UControlRigBlueprint::GeneratePythonCommands(const FString InNewBlueprintName)
 {
 	TArray<FString> InternalCommands;
diff --git a/Engine/Plugins/Animation/ControlRig/Source/ControlRigDeveloper/Public/ControlRigBlueprint.h b/Engine/Plugins/Animation/ControlRig/Source/ControlRigDeveloper/Public/ControlRigBlueprint.h
--- a/Engine/Plugins/Animation/ControlRig/Source/ControlRigDeveloper/Public/ControlRigBlueprint.h
+++ b/Engine/Plugins/Animation/ControlRig/Source/ControlRigDeveloper/Public/ControlRigBlueprint.h
@@ -43,6 +43,7 @@ public:
 	virtual UClass* GetRigVMEdGraphNodeClass() const override;
 	virtual UClass* GetRigVMEdGraphSchemaClass() const override;
 	virtual TArray<FString> GeneratePythonCommands(const FString InNewBlueprintName) override;
+	virtual UClass* GetRigVMEditorSettingsClass() const override;
 	virtual void Serialize(FArchive& Ar) override;
 
 #if WITH_EDITOR
diff --git a/Engine/Plugins/Animation/ControlRig/Source/ControlRigEditor/Private/ControlRigBlueprintCommands.cpp b/Engine/Plugins/Animation/ControlRig/Source/ControlRigEditor/Private/ControlRigBlueprintCommands.cpp
--- a/Engine/Plugins/Animation/ControlRig/Source/ControlRigEditor/Private/ControlRigBlueprintCommands.cpp
+++ b/Engine/Plugins/Animation/ControlRig/Source/ControlRigEditor/Private/ControlRigBlueprintCommands.cpp
@@ -22,16 +22,6 @@ void FControlRigBlueprintCommands::RegisterCommands()
 	UI_COMMAND(StepInto, "Step Into", "Steps into the collapsed/function node, when halted at a breakpoint.", EUserInterfaceActionType::Button, FInputChord(EKeys::F11));
 	UI_COMMAND(StepOut, "Step Out", "Steps out of the collapsed/function node, when halted at a breakpoint.", EUserInterfaceActionType::Button, FInputChord(EKeys::F11, EModifierKey::Shift));
 	UI_COMMAND(StepOver, "Step Over", "Steps over the node, when halted at a breakpoint.", EUserInterfaceActionType::Button, FInputChord(EKeys::F10));
-	UI_COMMAND(StoreNodeSnippet1, "Store Node Snippet 1", "Stores the selected node(s) into snippet 1", EUserInterfaceActionType::Button, FInputChord(EKeys::One, EModifierKey::Alt | EModifierKey::Shift));
-	UI_COMMAND(StoreNodeSnippet2, "Store Node Snippet 2", "Stores the selected node(s) into snippet 2", EUserInterfaceActionType::Button, FInputChord(EKeys::Two, EModifierKey::Alt | EModifierKey::Shift));
-	UI_COMMAND(StoreNodeSnippet3, "Store Node Snippet 3", "Stores the selected node(s) into snippet 3", EUserInterfaceActionType::Button, FInputChord(EKeys::Three, EModifierKey::Alt | EModifierKey::Shift));
-	UI_COMMAND(StoreNodeSnippet4, "Store Node Snippet 4", "Stores the selected node(s) into snippet 4", EUserInterfaceActionType::Button, FInputChord(EKeys::Four, EModifierKey::Alt | EModifierKey::Shift));
-	UI_COMMAND(StoreNodeSnippet5, "Store Node Snippet 5", "Stores the selected node(s) into snippet 5", EUserInterfaceActionType::Button, FInputChord(EKeys::Five, EModifierKey::Alt | EModifierKey::Shift));
-	UI_COMMAND(StoreNodeSnippet6, "Store Node Snippet 6", "Stores the selected node(s) into snippet 6", EUserInterfaceActionType::Button, FInputChord(EKeys::Six, EModifierKey::Alt | EModifierKey::Shift));
-	UI_COMMAND(StoreNodeSnippet7, "Store Node Snippet 7", "Stores the selected node(s) into snippet 7", EUserInterfaceActionType::Button, FInputChord(EKeys::Seven, EModifierKey::Alt | EModifierKey::Shift));
-	UI_COMMAND(StoreNodeSnippet8, "Store Node Snippet 8", "Stores the selected node(s) into snippet 8", EUserInterfaceActionType::Button, FInputChord(EKeys::Eight, EModifierKey::Alt | EModifierKey::Shift));
-	UI_COMMAND(StoreNodeSnippet9, "Store Node Snippet 9", "Stores the selected node(s) into snippet 9", EUserInterfaceActionType::Button, FInputChord(EKeys::Nine, EModifierKey::Alt | EModifierKey::Shift));
-	UI_COMMAND(StoreNodeSnippet0, "Store Node Snippet 0", "Stores the selected node(s) into snippet 0", EUserInterfaceActionType::Button, FInputChord(EKeys::Zero, EModifierKey::Alt | EModifierKey::Shift));
 	UI_COMMAND(FrameSelection, "Frame Selection", "Frames the selected nodes in the Graph View.", EUserInterfaceActionType::Button, FInputChord(EKeys::F));
 }
 
diff --git a/Engine/Plugins/Animation/ControlRig/Source/ControlRigEditor/Private/ControlRigBlueprintCommands.h b/Engine/Plugins/Animation/ControlRig/Source/ControlRigEditor/Private/ControlRigBlueprintCommands.h
--- a/Engine/Plugins/Animation/ControlRig/Source/ControlRigEditor/Private/ControlRigBlueprintCommands.h
+++ b/Engine/Plugins/Animation/ControlRig/Source/ControlRigEditor/Private/ControlRigBlueprintCommands.h
@@ -66,36 +66,6 @@ public:
 	/** Steps out of the collapsed/function node, when halted at a breakpoint */
 	TSharedPtr< FUICommandInfo > StepOut;
 
-	/** Stores the selected node(s) into snippet 1. */
-	TSharedPtr< FUICommandInfo > StoreNodeSnippet1;
-
-	/** Stores the selected node(s) into snippet 2. */
-	TSharedPtr< FUICommandInfo > StoreNodeSnippet2;
-
-	/** Stores the selected node(s) into snippet 3. */
-	TSharedPtr< FUICommandInfo > StoreNodeSnippet3;
-
-	/** Stores the selected node(s) into snippet 4. */
-	TSharedPtr< FUICommandInfo > StoreNodeSnippet4;
-
-	/** Stores the selected node(s) into snippet 5. */
-	TSharedPtr< FUICommandInfo > StoreNodeSnippet5;
-
-	/** Stores the selected node(s) into snippet 6. */
-	TSharedPtr< FUICommandInfo > StoreNodeSnippet6;
-
-	/** Stores the selected node(s) into snippet 7. */
-	TSharedPtr< FUICommandInfo > StoreNodeSnippet7;
-
-	/** Stores the selected node(s) into snippet 8. */
-	TSharedPtr< FUICommandInfo > StoreNodeSnippet8;
-
-	/** Stores the selected node(s) into snippet 9. */
-	TSharedPtr< FUICommandInfo > StoreNodeSnippet9;
-
-	/** Stores the selected node(s) into snippet 0. */
-	TSharedPtr< FUICommandInfo > StoreNodeSnippet0;
-
 	/** Frames the selected nodes */
 	TSharedPtr< FUICommandInfo > FrameSelection;
 
diff --git a/Engine/Plugins/Animation/ControlRig/Source/ControlRigEditor/Private/Editor/ControlRigEditor.cpp b/Engine/Plugins/Animation/ControlRig/Source/ControlRigEditor/Private/Editor/ControlRigEditor.cpp
--- a/Engine/Plugins/Animation/ControlRig/Source/ControlRigEditor/Private/Editor/ControlRigEditor.cpp
+++ b/Engine/Plugins/Animation/ControlRig/Source/ControlRigEditor/Private/Editor/ControlRigEditor.cpp
@@ -2167,56 +2167,6 @@ void FControlRigEditor::CreateDefaultCommands()
 
 void FControlRigEditor::OnCreateGraphEditorCommands(TSharedPtr<FUICommandList> GraphEditorCommandsList)
 {
-	GraphEditorCommandsList->MapAction(
-		FControlRigBlueprintCommands::Get().StoreNodeSnippet1,
-		FExecuteAction::CreateSP(this, &FControlRigEditor::StoreNodeSnippet, 1),
-		FCanExecuteAction());
-
-	GraphEditorCommandsList->MapAction(
-		FControlRigBlueprintCommands::Get().StoreNodeSnippet2,
-		FExecuteAction::CreateSP(this, &FControlRigEditor::StoreNodeSnippet, 2),
-		FCanExecuteAction());
-
-	GraphEditorCommandsList->MapAction(
-		FControlRigBlueprintCommands::Get().StoreNodeSnippet3,
-		FExecuteAction::CreateSP(this, &FControlRigEditor::StoreNodeSnippet, 3),
-		FCanExecuteAction());
-
-	GraphEditorCommandsList->MapAction(
-		FControlRigBlueprintCommands::Get().StoreNodeSnippet4,
-		FExecuteAction::CreateSP(this, &FControlRigEditor::StoreNodeSnippet, 4),
-		FCanExecuteAction());
-
-	GraphEditorCommandsList->MapAction(
-		FControlRigBlueprintCommands::Get().StoreNodeSnippet5,
-		FExecuteAction::CreateSP(this, &FControlRigEditor::StoreNodeSnippet, 5),
-		FCanExecuteAction());
-
-	GraphEditorCommandsList->MapAction(
-		FControlRigBlueprintCommands::Get().StoreNodeSnippet6,
-		FExecuteAction::CreateSP(this, &FControlRigEditor::StoreNodeSnippet, 6),
-		FCanExecuteAction());
-
-	GraphEditorCommandsList->MapAction(
-		FControlRigBlueprintCommands::Get().StoreNodeSnippet7,
-		FExecuteAction::CreateSP(this, &FControlRigEditor::StoreNodeSnippet, 7),
-		FCanExecuteAction());
-
-	GraphEditorCommandsList->MapAction(
-		FControlRigBlueprintCommands::Get().StoreNodeSnippet8,
-		FExecuteAction::CreateSP(this, &FControlRigEditor::StoreNodeSnippet, 8),
-		FCanExecuteAction());
-
-	GraphEditorCommandsList->MapAction(
-		FControlRigBlueprintCommands::Get().StoreNodeSnippet9,
-		FExecuteAction::CreateSP(this, &FControlRigEditor::StoreNodeSnippet, 9),
-		FCanExecuteAction());
-
-	GraphEditorCommandsList->MapAction(
-		FControlRigBlueprintCommands::Get().StoreNodeSnippet0,
-		FExecuteAction::CreateSP(this, &FControlRigEditor::StoreNodeSnippet, 0),
-		FCanExecuteAction());
-
 }
 
 TSharedRef<SGraphEditor> FControlRigEditor::CreateGraphEditorWidget(TSharedRef<FTabInfo> InTabInfo, UEdGraph* InGraph)
@@ -2698,59 +2648,6 @@ FReply FControlRigEditor::OnSpawnGraphNodeByShortcut(FInputChord InChord, const
 			}
 		}
 	}
-	else if(InChord.NeedsAlt() && !InChord.NeedsControl() && !InChord.NeedsShift())
-	{
-		if(InChord.Key == EKeys::One)
-		{
-			RestoreNodeSnippet(1);
-			return FReply::Handled();
-		}
-		else if(InChord.Key == EKeys::Two)
-		{
-			RestoreNodeSnippet(2);
-			return FReply::Handled();
-		}
-		else if(InChord.Key == EKeys::Three)
-		{
-			RestoreNodeSnippet(3);
-			return FReply::Handled();
-		}
-		else if(InChord.Key == EKeys::Four)
-		{
-			RestoreNodeSnippet(4);
-			return FReply::Handled();
-		}
-		else if(InChord.Key == EKeys::Five)
-		{
-			RestoreNodeSnippet(5);
-			return FReply::Handled();
-		}
-		else if(InChord.Key == EKeys::Six)
-		{
-			RestoreNodeSnippet(6);
-			return FReply::Handled();
-		}
-		else if(InChord.Key == EKeys::Seven)
-		{
-			RestoreNodeSnippet(7);
-			return FReply::Handled();
-		}
-		else if(InChord.Key == EKeys::Eight)
-		{
-			RestoreNodeSnippet(8);
-			return FReply::Handled();
-		}
-		else if(InChord.Key == EKeys::Nine)
-		{
-			RestoreNodeSnippet(9);
-			return FReply::Handled();
-		}
-		else if(InChord.Key == EKeys::Zero)
-		{
-			RestoreNodeSnippet(0);
-			return FReply::Handled();
-		}
-	}
 
 	return FReply::Unhandled();
 }
@@ -7201,150 +7098,6 @@ TSharedPtr<SGraphEditor> FControlRigEditor::GetGraphEditor(UEdGraph* InEdGraph)
 	return TSharedPtr<SGraphEditor>();
 }
 
-void FControlRigEditor::StoreNodeSnippet(int32 InSnippetIndex)
-{
-	check((InSnippetIndex >= 0) && (InSnippetIndex < 10));
-
-	URigVMController* Controller = GetFocusedController();
-	if(Controller == nullptr)
-	{
-		return;
-	}
-
-	const TArray<FName> SelectedNodeNames = Controller->GetGraph()->GetSelectNodes();
-	if(SelectedNodeNames.Num() == 0)
-	{
-		return;
-	}
-	
-	FString Snippet = Controller->ExportNodesToText(SelectedNodeNames);
-	if(Snippet.IsEmpty())
-	{
-		return;
-	}
-
-	TArray<FString> NodeNameStrings;
-	for(const FName& SelectedNodeName : SelectedNodeNames)
-	{
-		NodeNameStrings.Add(SelectedNodeName.ToString());
-	}
-	const FString NodeNamesJoined = FString::Join(NodeNameStrings, TEXT(", "));
-
-	FString* Setting = GetSnippetStorage(InSnippetIndex);
-	check(Setting != nullptr);
-
-	UControlRigEditorSettings* Settings = UControlRigEditorSettings::Get();
-	Settings->Modify();
-	(*Setting) = Snippet;
-
-	const FString PropertyName = FString::Printf(TEXT("NodeSnippet_%d"), InSnippetIndex);
-	FProperty* Property = UControlRigEditorSettings::StaticClass()->FindPropertyByName(*PropertyName);
-	check(Property);
-	
-	Settings->UpdateSinglePropertyInConfigFile(Property, Settings->GetDefaultConfigFilename());
-
-	FNotificationInfo Info(FText::FromString(FString::Printf(
-		TEXT("A snippet has been stored to the Project Settings.\n")
-		TEXT("Nodes %s are now stored as snippet %d.\n")
-		TEXT("You can restore the snippet by pressing Alt-%d and left clicking into the graph."),
-		*NodeNamesJoined, InSnippetIndex, InSnippetIndex)));
-	Info.bFireAndForget = true;
-	Info.FadeOutDuration = 2.0f;
-	Info.ExpireDuration = 8.0f;
-
-	TSharedPtr<SNotificationItem> NotificationPtr = FSlateNotificationManager::Get().AddNotification(Info);
-	NotificationPtr->SetCompletionState(SNotificationItem::CS_Success);
-}
-
-void FControlRigEditor::RestoreNodeSnippet(int32 InSnippetIndex)
-{
-	FString* Setting = GetSnippetStorage(InSnippetIndex);
-	check(Setting != nullptr);
-
-	if(Setting->IsEmpty())
-	{
-		return;
-	}
-
-	FPlatformApplicationMisc::ClipboardCopy(*(*Setting));
-	PasteNodes();
-	
-	FNotificationInfo Info(FText::FromString(FString::Printf(
-		TEXT("Snippet %d has been restored."),
-		InSnippetIndex)));
-	Info.bFireAndForget = true;
-	Info.FadeOutDuration = 0.5f;
-	Info.ExpireDuration = 3.0f;
-
-	TSharedPtr<SNotificationItem> NotificationPtr = FSlateNotificationManager::Get().AddNotification(Info);
-	NotificationPtr->SetCompletionState(SNotificationItem::CS_Success);
-}
-
-FString* FControlRigEditor::GetSnippetStorage(int32 InSnippetIndex)
-{
-	UControlRigEditorSettings* Settings = UControlRigEditorSettings::Get();
-	switch(InSnippetIndex)
-	{
-		case 1:
-		{
-			return &Settings->NodeSnippet_1;
-			break;
-		}
-		case 2:
-		{
-			return &Settings->NodeSnippet_2;
-			break;
-		}
-		case 3:
-		{
-			return &Settings->NodeSnippet_3;
-			break;
-		}
-		case 4:
-		{
-			return &Settings->NodeSnippet_4;
-			break;
-		}
-		case 5:
-		{
-			return &Settings->NodeSnippet_5;
-			break;
-		}
-		case 6:
-		{
-			return &Settings->NodeSnippet_6;
-			break;
-		}
-		case 7:
-		{
-			return &Settings->NodeSnippet_7;
-			break;
-		}
-		case 8:
-		{
-			return &Settings->NodeSnippet_8;
-			break;
-		}
-		case 9:
-		{
-			return &Settings->NodeSnippet_9;
-			break;
-		}
-		case 0:
-		{
-			return &Settings->NodeSnippet_0;
-			break;
-		}
-		default:
-		{
-			checkNoEntry();
-			break;
-		}
-	}
-
-	return nullptr;
-}
-
 void FControlRigEditor::OnPreConstruction_AnyThread(UControlRig* InRig, const FName& InEventName)
 {
 	bIsConstructionEventRunning = true;
diff --git a/Engine/Plugins/Animation/ControlRig/Source/ControlRigEditor/Private/Editor/ControlRigEditor.h b/Engine/Plugins/Animation/ControlRig/Source/ControlRigEditor/Private/Editor/ControlRigEditor.h
--- a/Engine/Plugins/Animation/ControlRig/Source/ControlRigEditor/Private/Editor/ControlRigEditor.h
+++ b/Engine/Plugins/Animation/ControlRig/Source/ControlRigEditor/Private/Editor/ControlRigEditor.h
@@ -447,13 +447,6 @@ protected:
 	URigVMController* GetFocusedController() const;
 	TSharedPtr<SGraphEditor> GetGraphEditor(UEdGraph* InEdGraph) const;
 	
-	// stores a node snippet into the setting
-	void StoreNodeSnippet(int32 InSnippetIndex);
-	// restores a node snippet from the setting
-	void RestoreNodeSnippet(int32 InSnippetIndex);
-	// returns the right setting storage for a given snippet index
-	static FString* GetSnippetStorage(int32 InSnippetIndex);
-
 	URigVMController* ActiveController;
 	bool bControlRigEditorInitialized;
 	bool bIsSettingObjectBeingDebugged;
diff --git a/Engine/Plugins/Runtime/RigVM/Source/RigVM/Private/RigVMSettings.cpp b/Engine/Plugins/Runtime/RigVM/Source/RigVM/Private/RigVMSettings.cpp
new file mode 100644
--- /dev/null
+++ b/Engine/Plugins/Runtime/RigVM/Source/RigVM/Private/RigVMSettings.cpp
@@ -0,0 +1,14 @@
+// Copyright Epic Games, Inc. All Rights Reserved.
+
+#include "RigVMSettings.h"
+
+#include UE_INLINE_GENERATED_CPP_BY_NAME(RigVMSettings)
+
+URigVMEditorSettings::URigVMEditorSettings(const FObjectInitializer& ObjectInitializer)
+    : Super(ObjectInitializer)
+{
+#if WITH_EDITORONLY_DATA
+	bAutoLinkMutableNodes = false;
+#endif
+}
+
diff --git a/Engine/Plugins/Runtime/RigVM/Source/RigVM/Public/RigVMSettings.h b/Engine/Plugins/Runtime/RigVM/Source/RigVM/Public/RigVMSettings.h
new file mode 100644
--- /dev/null
+++ b/Engine/Plugins/Runtime/RigVM/Source/RigVM/Public/RigVMSettings.h
@@ -0,0 +1,38 @@
+// Copyright Epic Games, Inc. All Rights Reserved.
+
+/*=============================================================================
+	RigVMSettings.h: Declares the RigVMEditorSettings class.
+=============================================================================*/
+
+#pragma once
+
+#include "CoreMinimal.h"
+#include "Engine/DeveloperSettings.h"
+
+#include "RigVMSettings.generated.h"
+
+class UStaticMesh;
+
+/**
+ * Customize RigVM Rig Editor.
+ */
+UCLASS(config = EditorPerProjectUserSettings, meta=(DisplayName="RigVM Editor"))
+class RIGVM_API URigVMEditorSettings : public UDeveloperSettings
+{
+	GENERATED_UCLASS_BODY()
+
+public:
+	/** Gets the category for the settings, some high level grouping like, Editor, Engine, Game...etc. */
+	virtual FName GetCategoryName() const override { return TEXT("ContentEditors"); }
+	
+#if WITH_EDITORONLY_DATA
+
+	// When this is checked mutable nodes (nodes with an execute pin)
+	// will be hooked up automatically.
+	UPROPERTY(EditAnywhere, config, Category = Interaction)
+	bool bAutoLinkMutableNodes;
+
+#endif
+};
+
+
diff --git a/Engine/Plugins/Runtime/RigVM/Source/RigVM/RigVM.Build.cs b/Engine/Plugins/Runtime/RigVM/Source/RigVM/RigVM.Build.cs
--- a/Engine/Plugins/Runtime/RigVM/Source/RigVM/RigVM.Build.cs
+++ b/Engine/Plugins/Runtime/RigVM/Source/RigVM/RigVM.Build.cs
@@ -16,6 +16,7 @@ public class RigVM : ModuleRules
                 "Engine",
                 "AnimationCore",
                 "AnimGraphRuntime",
+                "DeveloperSettings",
             }
         );
 
diff --git a/Engine/Plugins/Runtime/RigVM/Source/RigVMDeveloper/Private/RigVMBlueprint.cpp b/Engine/Plugins/Runtime/RigVM/Source/RigVMDeveloper/Private/RigVMBlueprint.cpp
--- a/Engine/Plugins/Runtime/RigVM/Source/RigVMDeveloper/Private/RigVMBlueprint.cpp
+++ b/Engine/Plugins/Runtime/RigVM/Source/RigVMDeveloper/Private/RigVMBlueprint.cpp
@@ -673,6 +673,11 @@ bool URigVMBlueprint::TryImportGraphFromText(const FString& InClipboardText, UEd
 	return true;
 }
 
+URigVMEditorSettings* URigVMBlueprint::GetRigVMEditorSettings() const
+{
+	return GetMutableDefault<URigVMEditorSettings>(GetRigVMEditorSettingsClass());
+}
+
 void URigVMBlueprint::Serialize(FArchive& Ar)
 {
 	RigVMClient.SetOuterClientHost(this, GET_MEMBER_NAME_CHECKED(URigVMBlueprint, RigVMClient));
diff --git a/Engine/Plugins/Runtime/RigVM/Source/RigVMDeveloper/Public/RigVMBlueprint.h b/Engine/Plugins/Runtime/RigVM/Source/RigVMDeveloper/Public/RigVMBlueprint.h
--- a/Engine/Plugins/Runtime/RigVM/Source/RigVMDeveloper/Public/RigVMBlueprint.h
+++ b/Engine/Plugins/Runtime/RigVM/Source/RigVMDeveloper/Public/RigVMBlueprint.h
@@ -11,6 +11,7 @@
 #include "RigVMCore/RigVMGraphFunctionDefinition.h"
 #include "EdGraph/RigVMEdGraph.h"
 #include "EdGraph/RigVMEdGraphSchema.h"
+#include "RigVMSettings.h"
 
 #include "RigVMBlueprint.generated.h"
 
@@ -222,6 +223,12 @@ public:
 	/** Returns the expected ed graph schema class to use for this blueprint */
 	virtual UClass* GetRigVMEdGraphSchemaClass() const { return URigVMEdGraphSchema::StaticClass(); }
 
+	/** Returns the class of the settings to use */
+	virtual UClass* GetRigVMEditorSettingsClass() const { return URigVMEditorSettings::StaticClass(); }
+
+	/** Returns the settings defaults for this blueprint */
+	URigVMEditorSettings* GetRigVMEditorSettings() const;
+
 	virtual void Serialize(FArchive& Ar) override;
 
 #if WITH_EDITOR
diff --git a/Engine/Plugins/Runtime/RigVM/Source/RigVMEditor/Private/EdGraph/NodeSpawners/RigVMEdGraphUnitNodeSpawner.cpp b/Engine/Plugins/Runtime/RigVM/Source/RigVMEditor/Private/EdGraph/NodeSpawners/RigVMEdGraphUnitNodeSpawner.cpp
--- a/Engine/Plugins/Runtime/RigVM/Source/RigVMEditor/Private/EdGraph/NodeSpawners/RigVMEdGraphUnitNodeSpawner.cpp
+++ b/Engine/Plugins/Runtime/RigVM/Source/RigVMEditor/Private/EdGraph/NodeSpawners/RigVMEdGraphUnitNodeSpawner.cpp
@@ -337,13 +337,10 @@ URigVMEdGraphNode* URigVMEdGraphUnitNodeSpawner::SpawnNode(UEdGraph* ParentGraph
 
 void URigVMEdGraphUnitNodeSpawner::HookupMutableNode(URigVMNode* InModelNode, URigVMBlueprint* InRigBlueprint)
 {
-	/*
-	// Todo in a follow up submit
-	if(!UControlRigEditorSettings::Get()->bAutoLinkMutableNodes)
+	if(!InRigBlueprint->GetRigVMEditorSettings()->bAutoLinkMutableNodes)
 	{
 		return;
 	}
-	*/
 	
 	URigVMController* Controller = InRigBlueprint->GetController(InModelNode->GetGraph());
 
