Control Rig: Added a cvar: ControlRig.Sequencer.AutoGenerateTrack to toggle whether control rig tracks should automatically be created for components

diff --git a/Engine/Plugins/Animation/ControlRig/Source/ControlRigEditor/Private/Sequencer/ControlRigParameterTrackEditor.cpp b/Engine/Plugins/Animation/ControlRig/Source/ControlRigEditor/Private/Sequencer/ControlRigParameterTrackEditor.cpp
--- a/Engine/Plugins/Animation/ControlRig/Source/ControlRigEditor/Private/Sequencer/ControlRigParameterTrackEditor.cpp
+++ b/Engine/Plugins/Animation/ControlRig/Source/ControlRigEditor/Private/Sequencer/ControlRigParameterTrackEditor.cpp
@@ -92,6 +92,8 @@
 
 #define LOCTEXT_NAMESPACE "FControlRigParameterTrackEditor"
 
+TAutoConsoleVariable<bool> CVarAutoGenerateControlRigTrack(TEXT("ControlRig.Sequencer.AutoGenerateTrack"), true, TEXT("When true automatically create control rig tracks in Sequencer when a control rig is added to a level."));
+
 TAutoConsoleVariable<bool> CVarSelectedKeysSelectControls(TEXT("ControlRig.Sequencer.SelectedKeysSelectControls"), false, TEXT("When true when we select a key in Sequencer it will select the Control, by default false."));
 
 TAutoConsoleVariable<bool> CVarSelectedSectionSetsSectionToKey(TEXT("ControlRig.Sequencer.SelectedSectionSetsSectionToKey"), false, TEXT("When true when we select a channel in a section, if it's the only section selected we set it as the Section To Key, by default false."));
@@ -1905,7 +1907,7 @@ void FControlRigParameterTrackEditor::AddTrackForComponent(USceneComponent* InCo
 {
 	if (USkeletalMeshComponent* SkelMeshComp = Cast<USkeletalMeshComponent>(InComponent))
 	{
-		if(bAutoGenerateControlRigTrack && !SkelMeshComp->GetDefaultAnimatingRig().IsNull())
+		if(bAutoGenerateControlRigTrack && CVarAutoGenerateControlRigTrack.GetValueOnGameThread() && !SkelMeshComp->GetDefaultAnimatingRig().IsNull())
 		{
 			UObject* Object = SkelMeshComp->GetDefaultAnimatingRig().LoadSynchronous();
 			if (Object != nullptr && (Object->IsA<UControlRigBlueprint>() || Object->IsA<UControlRigComponent>()))
