Control Rig: Added a cvar: ControlRig.Sequencer.AutoGenerateTrack to toggle whether control rig tracks should automatically be created for components

diff --git a/Engine/Plugins/Animation/ControlRig/Source/ControlRigEditor/Private/Sequencer/ControlRigParameterTrackEditor.cpp b/Engine/Plugins/Animation/ControlRig/Source/ControlRigEditor/Private/Sequencer/ControlRigParameterTrackEditor.cpp
--- a/Engine/Plugins/Animation/ControlRig/Source/ControlRigEditor/Private/Sequencer/ControlRigParameterTrackEditor.cpp
+++ b/Engine/Plugins/Animation/ControlRig/Source/ControlRigEditor/Private/Sequencer/ControlRigParameterTrackEditor.cpp
@@ -1907,7 +1907,7 @@ void FControlRigParameterTrackEditor::AddTrackForComponent(USceneComponent* InCo
 {
 	if (USkeletalMeshComponent* SkelMeshComp = Cast<USkeletalMeshComponent>(InComponent))
 	{
-		if(bAutoGenerateControlRigTrack && CVarAutoGenerateControlRigTrack.GetValueOnGameThread() && !SkelMeshComp->GetDefaultAnimatingRig().IsNull())
+		if(bAutoGenerateControlRigTrack && !SkelMeshComp->GetDefaultAnimatingRig().IsNull())
 		{
 			UObject* Object = SkelMeshComp->GetDefaultAnimatingRig().LoadSynchronous();
 			if (Object != nullptr && (Object->IsA<UControlRigBlueprint>() || Object->IsA<UControlRigComponent>()))
@@ -1945,6 +1945,11 @@ void FControlRigParameterTrackEditor::AddTrackForComponent(USceneComponent* InCo
 
 void FControlRigParameterTrackEditor::HandleActorAdded(AActor* Actor, FGuid TargetObjectGuid)
 {
+	if (!CVarAutoGenerateControlRigTrack.GetValueOnGameThread())
+	{
+		return;
+	}
+
 	if (Actor)
 	{
 		if (UControlRigComponent* ControlRigComponent = Actor->FindComponentByClass<UControlRigComponent>())
