Control Rig: Restore GetControlRigClass / CreateControlRig

diff --git a/Engine/Plugins/Animation/ControlRig/Source/ControlRigDeveloper/Private/ControlRigBlueprint.cpp b/Engine/Plugins/Animation/ControlRig/Source/ControlRigDeveloper/Private/ControlRigBlueprint.cpp
--- a/Engine/Plugins/Animation/ControlRig/Source/ControlRigDeveloper/Private/ControlRigBlueprint.cpp
+++ b/Engine/Plugins/Animation/ControlRig/Source/ControlRigDeveloper/Private/ControlRigBlueprint.cpp
@@ -112,6 +112,11 @@ void UControlRigBlueprint::PostEditChangeChainProperty(FPropertyChangedChainEven
 	}
 }
 
+UClass* UControlRigBlueprint::GetControlRigClass()
+{
+	return GetRigVMHostClass();
+}
+
 USkeletalMesh* UControlRigBlueprint::GetPreviewMesh() const
 {
 	DECLARE_SCOPE_HIERARCHICAL_COUNTER_FUNC()
diff --git a/Engine/Plugins/Animation/ControlRig/Source/ControlRigDeveloper/Public/ControlRigBlueprint.h b/Engine/Plugins/Animation/ControlRig/Source/ControlRigDeveloper/Public/ControlRigBlueprint.h
--- a/Engine/Plugins/Animation/ControlRig/Source/ControlRigDeveloper/Public/ControlRigBlueprint.h
+++ b/Engine/Plugins/Animation/ControlRig/Source/ControlRigDeveloper/Public/ControlRigBlueprint.h
@@ -65,6 +65,9 @@ public:
 
 #endif	// #if WITH_EDITOR
 
+	UFUNCTION(BlueprintCallable, Category = "VM")
+	UClass* GetControlRigClass();
+
 	UFUNCTION(BlueprintCallable, Category = "Control Rig Blueprint")
 	UControlRig* CreateControlRig() { return Cast<UControlRig>(CreateRigVMHost()); }
 
