Control Rig / RigVM: Moving graph related widgets to RigVMEditor

diff --git a/Engine/Plugins/Animation/ControlRig/Config/BaseControlRig.ini b/Engine/Plugins/Animation/ControlRig/Config/BaseControlRig.ini
--- a/Engine/Plugins/Animation/ControlRig/Config/BaseControlRig.ini
+++ b/Engine/Plugins/Animation/ControlRig/Config/BaseControlRig.ini
@@ -436,3 +436,5 @@
 +StructRedirects=(OldName="/Script/ControlRig.ControlRigPythonSettings",NewName="/Script/RigVMDeveloper.RigVMPythonSettings")
 +StructRedirects=(OldName="/Script/ControlRig.ControlRigPublicFunctionArg",NewName="/Script/RigVMDeveloper.RigVMOldPublicFunctionArg")
 +StructRedirects=(OldName="/Script/ControlRig.ControlRigPublicFunctionData",NewName="/Script/RigVMDeveloper.RigVMOldPublicFunctionData")
++ClassRedirects=(OldName="/Script/ControlRig.NameSpacedUserData", NewName="/Script/RigVM.NameSpacedUserData")
++ClassRedirects=(OldName="/Script/ControlRig.DataAssetLink", NewName="/Script/RigVM.DataAssetLink")
diff --git a/Engine/Plugins/Animation/ControlRig/Source/ControlRig/Private/ControlRigAssetUserData.cpp b/Engine/Plugins/Animation/ControlRig/Source/ControlRig/Private/ControlRigAssetUserData.cpp
--- a/Engine/Plugins/Animation/ControlRig/Source/ControlRig/Private/ControlRigAssetUserData.cpp
+++ b/Engine/Plugins/Animation/ControlRig/Source/ControlRig/Private/ControlRigAssetUserData.cpp
@@ -1,475 +1,6 @@
 // Copyright Epic Games, Inc. All Rights Reserved.
 
 #include "ControlRigAssetUserData.h"
-#include "Engine/UserDefinedStruct.h"
-
-#include UE_INLINE_GENERATED_CPP_BY_NAME(ControlRigAssetUserData)
-
-const TArray<const UNameSpacedUserData::FUserData*> UNameSpacedUserData::EmptyUserDatas;
-
-FString UNameSpacedUserData::FUserData::GetName() const
-{
-	FString Left, Right;
-	if(SplitAtEnd(Path, Left, Right))
-	{
-		return Right;
-	}
-	return Path;
-}
-
-#if WITH_EDITOR
-
-FString UNameSpacedUserData::FUserData::GetDisplayName() const
-{
-	if(Property)
-	{
-		if(Property->GetTypedOwner<UUserDefinedStruct>())
-		{
-			return Property->GetDisplayNameText().ToString();
-		}
-	}
-	return GetName();
-}
-
-bool UNameSpacedUserData::FUserData::IsArray() const
-{
-	return CastField<FArrayProperty>(Property) != nullptr;
-}
-
-bool UNameSpacedUserData::FUserData::IsArrayElement() const
-{
-	if(Property)
-	{
-		return CastField<FArrayProperty>(Property->GetOwnerProperty()) != nullptr;
-	}
-	return false;
-}
-
-bool UNameSpacedUserData::FUserData::IsUObject() const
-{
-	return CastField<FObjectProperty>(Property) != nullptr;
-}
-
-#endif
-
-TArray<UStruct*> UNameSpacedUserData::FUserData::GetSuperStructs(UStruct* InStruct)
-{
-	// Create an array of structs, ordered super -> child struct
-	TArray<UStruct*> SuperStructs = {InStruct};
-	while(true)
-	{
-		if(UStruct* SuperStruct = SuperStructs[0]->GetSuperStruct())
-		{
-			SuperStructs.Insert(SuperStruct, 0);
-		}
-		else
-		{
-			break;
-		}
-	}
-	return SuperStructs;
-}
-
-void UNameSpacedUserData::Serialize(FArchive& Ar)
-{
-	Super::Serialize(Ar);
-	InvalidateCache();
-}
-
-const UNameSpacedUserData::FUserData* UNameSpacedUserData::GetUserData(const FString& InPath, FString* OutErrorMessage) const
-{
-	FScopeLock Lock(&CacheLock);
-	if(FUserData** UserDataInCache = CachedUserData.Find(InPath))
-	{
-		if(*UserDataInCache)
-		{
-			return *UserDataInCache;
-		}
-	}
-	return nullptr;
-}
-
-const TArray<const UNameSpacedUserData::FUserData*>& UNameSpacedUserData::GetUserDataArray(const FString& InParentPath, FString* OutErrorMessage) const
-{
-	FScopeLock Lock(&CacheLock);
-	if(const TArray<const FUserData*>* UserDataArrayInCache = CachedUserDataArray.Find(InParentPath))
-	{
-		return *UserDataArrayInCache;
-	}
-	return EmptyUserDatas;
-}
-
-const UNameSpacedUserData::FUserData* UNameSpacedUserData::GetUserDataWithinStruct(const UStruct* InStruct, const uint8* InMemory, const FString& InPath, const FString& InPropertyName, FString* OutErrorMessage) const
-{
-	check(InStruct);
-	check(InMemory);
-
-	FString Left, Right;
-	if(!InPropertyName.IsEmpty() || !FUserData::SplitAtEnd(InPath, Left, Right))
-	{
-		const FString& PropertyName = InPropertyName.IsEmpty() ? InPath : InPropertyName; 
-		if(const FProperty* Property = FindPropertyByName(InStruct, *PropertyName))
-		{
-			if(!IsPropertySupported(Property, InPath, true, OutErrorMessage))
-			{
-				return nullptr;
-			}
-
-			const uint8* PropertyMemory = Property->ContainerPtrToValuePtr<uint8>(InMemory);
-			const FUserData UserData(InPath, Property, PropertyMemory);
-			return StoreCacheForUserData(UserData);
-		}
-	}
-	else
-	{
-		if(const FUserData* ParentUserData = GetUserData(Left, OutErrorMessage))
-		{
-			check(ParentUserData->GetProperty());
-
-			const uint8* Memory = ParentUserData->GetMemory();
-			if(Memory == nullptr)
-			{
-				if(OutErrorMessage && OutErrorMessage->IsEmpty())
-				{
-					(*OutErrorMessage) = FString::Printf(InvalidMemoryFormat, *InPath);
-				}
-				return nullptr;
-			}
-			
-			if(const FStructProperty* StructProperty = CastField<FStructProperty>(ParentUserData->GetProperty()))
-			{
-				return GetUserDataWithinStruct(StructProperty->Struct, Memory, InPath, Right, OutErrorMessage);
-			}
-
-			if(const FArrayProperty* ArrayProperty = CastField<FArrayProperty>(ParentUserData->GetProperty()))
-			{
-				const FProperty* Property = ArrayProperty->Inner;
-				if(!IsPropertySupported(Property, InPath, false, OutErrorMessage))
-				{
-					return nullptr;
-				}
-
-				// make sure the next bit is provided as a number and is within bounds
-				if(!Right.IsNumeric())
-				{
-					if(OutErrorMessage && OutErrorMessage->IsEmpty())
-					{
-						(*OutErrorMessage) = FString::Printf(InvalidArrayIndexFormat, *Right, *InPath);
-					}
-					return nullptr;
-				}
-
-				FScriptArrayHelper ArrayHelper(ArrayProperty, Memory);
-
-				const int32 Index = FCString::Atoi(*Right);
-				if(!ArrayHelper.IsValidIndex(Index))
-				{
-					if(OutErrorMessage && OutErrorMessage->IsEmpty())
-					{
-						(*OutErrorMessage) = FString::Printf(OutOfBoundArrayIndexFormat, *InPath, ArrayHelper.Num());
-					}
-					return nullptr;
-				}
-
-				Memory = ArrayHelper.GetRawPtr(Index);
-				const FUserData UserData(InPath, Property, Memory);
-				return StoreCacheForUserData(UserData);
-			}
-
-			if(OutErrorMessage && OutErrorMessage->IsEmpty())
-			{
-				(*OutErrorMessage) = FString::Printf(UnSupportedSubPathsFormat, *ParentUserData->GetCPPType().ToString(), *Left);
-			}
-		}
-	}
-
-	if(OutErrorMessage && OutErrorMessage->IsEmpty())
-	{
-		(*OutErrorMessage) = FString::Printf(PathNotFoundFormat, *InPath);
-	}
-	return nullptr;
-}
-
-const TArray<const UNameSpacedUserData::FUserData*>& UNameSpacedUserData::GetUserDataArrayWithinStruct(UStruct* InStruct, const uint8* InMemory, const FString& InPath, FString* OutErrorMessage) const
-{
-	FScopeLock Lock(&CacheLock);
-
-	TArray<const FUserData*> UserDataArray;
-
-	const TArray<UStruct*> Structs = FUserData::GetSuperStructs(InStruct);
-	for(const UStruct* Struct : Structs)
-	{
-		//const bool bIsUserDefinedStruct = Struct->IsA<UUserDefinedStruct>();
-		for (TFieldIterator<FProperty> It(Struct, EFieldIterationFlags::None); It; ++It)
-		{
-			const FProperty* Property = *It;
-
-			const FString& PropertyName = Property->GetNameCPP();
-			const FString PropertyPath = InPath.IsEmpty() ? PropertyName : FUserData::Join(InPath, PropertyName);
-					
-			if(!IsPropertySupported(Property, PropertyPath, true, OutErrorMessage))
-			{
-				continue;
-			}
-					
-			FUserData*& UserDataInCache = CachedUserData.FindOrAdd(PropertyPath);
-			if(UserDataInCache == nullptr)
-			{
-				const uint8* PropertyMemory = Property->ContainerPtrToValuePtr<uint8>(InMemory);
-				UserDataInCache = new FUserData(PropertyPath, Property, PropertyMemory);
-			}
-
-			UserDataArray.Add(UserDataInCache);
-		}
-	}
-
-	if(!UserDataArray.IsEmpty())
-	{
-		TArray<const FUserData*>& UserDataArrayInCache = CachedUserDataArray.FindOrAdd(InPath);
-		if(UserDataArrayInCache.IsEmpty())
-		{
-			UserDataArrayInCache = UserDataArray;
-		}
-		return UserDataArrayInCache;
-	}
-
-	return EmptyUserDatas;
-}
-
-const TArray<const UNameSpacedUserData::FUserData*>& UNameSpacedUserData::GetUserDataArrayWithinArray(const FArrayProperty* InArrayProperty, const uint8* InMemory, const FString& InPath, FString* OutErrorMessage) const
-{
-	check(!InPath.IsEmpty());
-	
-	FScopeLock Lock(&CacheLock);
-
-	TArray<const FUserData*> UserDataArray;
-
-	const FProperty* Property = InArrayProperty->Inner;
-	if(!IsPropertySupported(Property, InPath, false, OutErrorMessage))
-	{
-		return EmptyUserDatas;
-	}
-
-	FScriptArrayHelper ArrayHelper(InArrayProperty, InMemory);
-	for(int32 Index = 0; Index < ArrayHelper.Num(); Index++)
-	{
-		const FString PropertyPath = FUserData::Join(InPath, FString::FromInt(Index));
-
-		FUserData*& UserDataInCache = CachedUserData.FindOrAdd(PropertyPath);
-		if(UserDataInCache == nullptr)
-		{
-			const uint8* PropertyMemory = ArrayHelper.GetRawPtr(Index);
-			UserDataInCache = new FUserData(PropertyPath, Property, PropertyMemory);
-		}
-
-		UserDataArray.Add(UserDataInCache);
-	}
-
-	if(!UserDataArray.IsEmpty())
-	{
-		TArray<const FUserData*>& UserDataArrayInCache = CachedUserDataArray.FindOrAdd(InPath);
-		if(UserDataArrayInCache.IsEmpty())
-		{
-			UserDataArrayInCache = UserDataArray;
-		}
-		return UserDataArrayInCache;
-	}
-
-	return EmptyUserDatas;
-}
-
-const FProperty* UNameSpacedUserData::FindPropertyByName(const UStruct* InStruct, const FName& InName)
-{
-	check(InStruct);
-
-	if(const FProperty* Property = InStruct->FindPropertyByName(InName))
-	{
-		return Property;
-	}
-
-	// user defined properties may actually be called quite differently
-	return nullptr;
-}
-
-void UNameSpacedUserData::InvalidateCache() const
-{
-	FScopeLock Lock(&CacheLock);
-	for(const TPair<FString, FUserData*>& Pair : CachedUserData)
-	{
-		delete(Pair.Value);
-	}
-	CachedUserData.Reset();
-	CachedUserDataArray.Reset();
-}
-
-const UNameSpacedUserData::FUserData* UNameSpacedUserData::StoreCacheForUserData(const FUserData& InUserData) const
-{
-	check(InUserData.IsValid());
-	FScopeLock Lock(&CacheLock);
-
-	if(const FUserData* const* UserDataInCache = CachedUserData.Find(InUserData.GetPath()))
-	{
-		if(*UserDataInCache)
-		{
-			check((*UserDataInCache)->GetCPPType() == InUserData.GetCPPType());
-			return *UserDataInCache;
-		}
-	}
-
-	FUserData* UserDataPtr = new FUserData(InUserData);
-	CachedUserData.Add(InUserData.GetPath(), UserDataPtr);
-	return UserDataPtr;
-}
-
-const TArray<const UNameSpacedUserData::FUserData*>& UNameSpacedUserData::StoreCacheForUserDataArray(const FString& InPath, const TArray<const FUserData*>& InUserDataArray) const
-{
-	if(InUserDataArray.IsEmpty())
-	{
-		return EmptyUserDatas;
-	}
-	
-	FScopeLock Lock(&CacheLock);
-
-	TArray<const FUserData*>& UserDataArrayInCache = CachedUserDataArray.FindOrAdd(InPath);
-	if(UserDataArrayInCache.IsEmpty())
-	{
-		UserDataArrayInCache = InUserDataArray;
-	}
-	return UserDataArrayInCache;
-}
-
-bool UNameSpacedUserData::IsPropertySupported(const FProperty* InProperty, const FString& InPath, bool bCheckPropertyFlags, FString* OutErrorMessage) const
-{
-	check(InProperty);
-
-	if(bCheckPropertyFlags)
-	{
-		if(!InProperty->HasAnyPropertyFlags(CPF_Edit | CPF_EditConst | CPF_BlueprintVisible))
-		{
-			return false;
-		}
-		if(InProperty->HasAnyPropertyFlags(CPF_EditorOnly))
-		{
-			return false;
-		}
-	}
-
-	if(InProperty->IsA<FSetProperty>() ||
-		InProperty->IsA<FMapProperty>())
-	{
-		if(OutErrorMessage && OutErrorMessage->IsEmpty())
-		{
-			FString ExtendedCppType;
-			FString CPPTypeString = InProperty->GetCPPType(&ExtendedCppType);
-			CPPTypeString += ExtendedCppType;
-
-			static constexpr TCHAR UnsupportedCPPTypeFormat[] = TEXT("Unsupported user data type '%s' at path '%s'.");
-			(*OutErrorMessage) = FString::Printf(UnsupportedCPPTypeFormat, *CPPTypeString, *InPath);
-		}
-		return false;
-	}
-
-	return true;
-}
-
-void UDataAssetLink::SetDataAsset(UDataAsset* InDataAsset)
-{
-	InvalidateCache();
-	DataAsset = InDataAsset;
-
-	if(NameSpace.IsEmpty() && DataAsset)
-	{
-		NameSpace = DataAsset->GetName();
-	}
-}
-
-const UNameSpacedUserData::FUserData* UDataAssetLink::GetUserData(const FString& InPath, FString* OutErrorMessage) const
-{
-	// rely on super to look it up from the cache
-	if(const FUserData* ResultFromSuper = Super::GetUserData(InPath))
-	{
-		return ResultFromSuper;
-	}
-
-	if(DataAsset)
-	{
-		// this method caches as well - so the next time around Super::GetUserData should return the cache 
-		return GetUserDataWithinStruct(DataAsset->GetClass(), (const uint8*)DataAsset, InPath, FString(), OutErrorMessage);
-	}
-
-	if(OutErrorMessage && OutErrorMessage->IsEmpty())
-	{
-		(*OutErrorMessage) = FString::Printf(DataAssetNullFormat, *InPath);
-	}
-	return nullptr;
-}
-
-const TArray<const UNameSpacedUserData::FUserData*>& UDataAssetLink::GetUserDataArray(const FString& InParentPath, FString* OutErrorMessage) const
-{
-	// rely on super to look it up from the cache
-	const TArray<const FUserData*>& ResultFromSuper = Super::GetUserDataArray(InParentPath);
-	if(!ResultFromSuper.IsEmpty())
-	{
-		return ResultFromSuper;
-	}
-
-	if(DataAsset)
-	{
-		// we should only get here if we haven't cached this user data array before.
-		if(InParentPath.IsEmpty())
-		{
-			// this method caches as well - so the next time around Super::GetUserDataArray should return the cache 
-			return GetUserDataArrayWithinStruct(DataAsset->GetClass(), (const uint8*)DataAsset, InParentPath, OutErrorMessage);
-		}
-
-		if(const FUserData* ParentUserData = GetUserData(InParentPath, OutErrorMessage))
-		{
-			check(ParentUserData->GetProperty());
-
-			const uint8* Memory = ParentUserData->GetMemory();
-			if(Memory == nullptr)
-			{
-				return EmptyUserDatas;
-			}
-
-			if(const FStructProperty* StructProperty = CastField<FStructProperty>(ParentUserData->GetProperty()))
-			{
-				// this method caches as well - so the next time around Super::GetUserDataArray should return the cache 
-				return GetUserDataArrayWithinStruct(StructProperty->Struct, Memory, InParentPath, OutErrorMessage);
-			}
-			
-			if(const FArrayProperty* ArrayProperty = CastField<FArrayProperty>(ParentUserData->GetProperty()))
-			{
-				// this method caches as well - so the next time around Super::GetUserDataArray should return the cache 
-				return GetUserDataArrayWithinArray(ArrayProperty, Memory, InParentPath, OutErrorMessage);
-			}
-
-			// other types of containers are not expected here due to UNameSpacedUserData::IsPropertySupported
-		}
-	}
-
-	if(OutErrorMessage && OutErrorMessage->IsEmpty())
-	{
-		(*OutErrorMessage) = FString::Printf(DataAssetNullFormat, *InParentPath);
-	}
-	return EmptyUserDatas;
-}
-
-#if WITH_EDITOR
-
-void UDataAssetLink::PostEditChangeProperty(FPropertyChangedEvent& PropertyChangedEvent)
-{
-	Super::PostEditChangeProperty(PropertyChangedEvent);
-
-	static const FProperty* DataAssetProperty= FindPropertyByName(GetClass(), GET_MEMBER_NAME_CHECKED(UDataAssetLink, DataAsset));
-	if(PropertyChangedEvent.Property == DataAssetProperty ||
-		PropertyChangedEvent.MemberProperty == DataAssetProperty)
-	{
-		SetDataAsset(DataAsset);
-	}
-}
-
-#endif
 
 void UControlRigShapeLibraryLink::SetShapeLibrary(UControlRigShapeLibrary* InShapeLibrary)
 {
diff --git a/Engine/Plugins/Animation/ControlRig/Source/ControlRig/Public/ControlRigAssetUserData.h b/Engine/Plugins/Animation/ControlRig/Source/ControlRig/Public/ControlRigAssetUserData.h
--- a/Engine/Plugins/Animation/ControlRig/Source/ControlRig/Public/ControlRigAssetUserData.h
+++ b/Engine/Plugins/Animation/ControlRig/Source/ControlRig/Public/ControlRigAssetUserData.h
@@ -2,216 +2,10 @@
 
 #pragma once
 
+#include "RigVMCore/RigVMAssetUserData.h"
 #include "ControlRigGizmoLibrary.h"
-#include "Engine/AssetUserData.h"
-#include "Engine/DataAsset.h"
 #include "ControlRigAssetUserData.generated.h"
 
-/**
-* User data that can be attached to assets to provide namespaced data access.
-*/
-UCLASS(Abstract)
-class CONTROLRIG_API UNameSpacedUserData : public UAssetUserData
-{
-	GENERATED_BODY()
-
-public:
-
-	/** A helper struct to represent a single user data */
-	struct CONTROLRIG_API FUserData
-	{
-	public:
-
-		FUserData()
-		: Path()
-		, CPPType(NAME_None)
-		, Property(nullptr)
-		, Memory(nullptr)
-		{
-		}
-
-		FUserData(const FString& InPath, const FProperty* InProperty, const uint8* InMemory)
-		: Path(InPath)
-		, CPPType(NAME_None)
-		, Property(InProperty)
-		, Memory(InMemory)
-		{
-		}
-
-		FUserData(const FString& InPath, const FName& InCPPType, const FProperty* InProperty, const uint8* InMemory)
-		: Path(InPath)
-		, CPPType(InCPPType)
-		, Property(InProperty)
-		, Memory(InMemory)
-		{
-		}
-
-		// Returns true if this user data is mapped correctly
-		bool IsValid() const
-		{
-			return (!Path.IsEmpty()) && (Memory != nullptr);
-		}
-
-		operator bool() const { return IsValid(); }
-
-		operator FString() const { return Path; }
-
-		// returns the string path inside of the asset user data
-		const FString& GetPath() const { return Path; }
-
-		// returns the name of the last segment of the path
-		FString GetName() const;
-
-#if WITH_EDITOR
-
-		// returns the display name of the last segment of the path
-		FString GetDisplayName() const;
-
-#endif
-
-		// returns the CPP type of this user data (computes it from the property as needed)
-		const FName& GetCPPType() const
-		{
-			if(CPPType.IsNone() && Property != nullptr)
-			{
-				FString ExtendedCppType;
-  				FString CPPTypeString = Property->GetCPPType(&ExtendedCppType);
-  				CPPTypeString += ExtendedCppType;
-				CPPType = *CPPTypeString;
-			}
-			return CPPType;
-		}
-
-		bool IsArray() const;
-
-		bool IsArrayElement() const;
-
-		bool IsUObject() const;
-
-		// returns the property of this user data or nullptr
-		const FProperty* GetProperty() const { return Property; }
-
-		// returns the mapped memory of this user data
-		const uint8* GetMemory() const { return Memory; }
-
-		// joints two user data path strings 
-		static FString Join(const FString& InLeft, const FString& InRight)
-		{
-			check(!InLeft.IsEmpty());
-			check(!InRight.IsEmpty());
-			return InLeft + TEXT(".") + InRight;
-		}
-
-		// splits a user data string path at the first segment
-		static bool SplitAtStart(const FString& InPath, FString& OutLeft, FString& OutRight)
-		{
-			return InPath.Split(TEXT("."), &OutLeft, &OutRight, ESearchCase::IgnoreCase, ESearchDir::FromStart);
-		}
-
-		// splits a user data string path at the last segment
-		static bool SplitAtEnd(const FString& InPath, FString& OutLeft, FString& OutRight)
-		{
-			return InPath.Split(TEXT("."), &OutLeft, &OutRight, ESearchCase::IgnoreCase, ESearchDir::FromEnd);
-		}
-
-		static TArray<UStruct*> GetSuperStructs(UStruct* InStruct);
-
-	private:
-
-		FString Path;
-		mutable FName CPPType;
-		const FProperty* Property;
-		const uint8* Memory;
-	};
-
-	virtual ~UNameSpacedUserData() override
-	{
-		InvalidateCache();
-	}
-
-	virtual void Serialize(FArchive& Ar) override;
-
-	/** The namespace to use when looking up values inside of the user data. */
-	UPROPERTY(EditAnywhere, Category = General)
-	FString NameSpace;
-
-	/** Returns true if userdata exists for a given path */
-	bool ContainsUserData(const FString& InPath) const { return GetUserData(InPath) != nullptr; }
-
-	/** Returns a user data path given its string path */
-	virtual const FUserData* GetUserData(const FString& InPath, FString* OutErrorMessage = nullptr) const;
-
-	/** Retrieves the user data paths given a (optional) parent path */
-	virtual const TArray<const FUserData*>& GetUserDataArray(const FString& InParentPath = FString(), FString* OutErrorMessage = nullptr) const;
-
-	/** Retrieves the user data path given the path string */
-	const FUserData* GetParentUserData(const FString& InPath, FString* OutErrorMessage = nullptr) const
-	{
-		FString Left, Right;
-		if(FUserData::SplitAtEnd(InPath, Left, Right))
-		{
-			return GetUserData(Left, OutErrorMessage);
-		}
-		return nullptr;
-	}
-
-protected:
-
-	const FUserData* GetUserDataWithinStruct(const UStruct* InStruct, const uint8* InMemory, const FString& InPath, const FString& InPropertyName, FString* OutErrorMessage = nullptr) const;
-	const TArray<const FUserData*>& GetUserDataArrayWithinStruct(UStruct* InStruct, const uint8* InMemory, const FString& InPath, FString* OutErrorMessage = nullptr) const;
-	const TArray<const FUserData*>& GetUserDataArrayWithinArray(const FArrayProperty* InArrayProperty, const uint8* InMemory, const FString& InPath, FString* OutErrorMessage = nullptr) const;
-	static const FProperty* FindPropertyByName(const UStruct* InStruct, const FName& InName);
-
-	void InvalidateCache() const;
-	const FUserData* StoreCacheForUserData(const FUserData& InUserData) const; 
-	const TArray<const FUserData*>& StoreCacheForUserDataArray(const FString& InPath, const TArray<const FUserData*>& InUserDataArray) const;
-	virtual bool IsPropertySupported(const FProperty* InProperty, const FString& InPath, bool bCheckPropertyFlags, FString* OutErrorMessage = nullptr) const;
-	
-	static const TArray<const FUserData*> EmptyUserDatas;
-
-	mutable TMap<FString, FUserData*> CachedUserData;
-	mutable TMap<FString, TArray<const FUserData*>> CachedUserDataArray;
-	mutable FCriticalSection CacheLock;
-
-	static inline constexpr TCHAR PathNotFoundFormat[] = TEXT("UserData path '%s' could not be found.");
-	static inline constexpr TCHAR InvalidMemoryFormat[] = TEXT("UserData at path '%s' does not have any data backing it up.");
-	static inline constexpr TCHAR InvalidArrayIndexFormat[] = TEXT("Provided array index '%s' for path '%s' is invalid.");
-	static inline constexpr TCHAR OutOfBoundArrayIndexFormat[] = TEXT("Provided array index for path '%s' is out of bounds (num == %d).");
-	static inline constexpr TCHAR UnSupportedSubPathsFormat[] = TEXT("Cannot access sub paths of user data type '%s' for path '%s'.");
-};
-
-/**
-* Namespaced user data which provides access to a linked data asset
-*/
-UCLASS(BlueprintType)
-class CONTROLRIG_API UDataAssetLink : public UNameSpacedUserData
-{
-	GENERATED_BODY()
-
-public:
-
-	/** If assigned, the data asset link will provide access to the data asset's content. */
-	UPROPERTY(EditAnywhere, BlueprintReadWrite, Category = General, BlueprintGetter = GetDataAsset, BlueprintSetter = SetDataAsset, Meta = (DisplayAfter="NameSpace"))
-	TObjectPtr<UDataAsset> DataAsset;
-
-	UFUNCTION(BlueprintGetter)
-	UDataAsset* GetDataAsset() const { return DataAsset; }
-
-	UFUNCTION(BlueprintSetter)
-	void SetDataAsset(UDataAsset* InDataAsset);
-
-	virtual const FUserData* GetUserData(const FString& InPath, FString* OutErrorMessage = nullptr) const override;
-	virtual const TArray<const FUserData*>& GetUserDataArray(const FString& InParentPath = FString(), FString* OutErrorMessage = nullptr) const override;
-
-#if WITH_EDITOR
-	virtual void PostEditChangeProperty(FPropertyChangedEvent& PropertyChangedEvent) override;
-#endif
-
-protected:
-
-	static inline constexpr TCHAR DataAssetNullFormat[] = TEXT("User data path '%s' could not be found (DataAsset not provided)");
-};
-
 /**
 * Namespaced user data which provides access to a linked shape library
 */
diff --git a/Engine/Plugins/Animation/ControlRig/Source/ControlRigDeveloper/Private/Graph/ControlRigGraph.cpp b/Engine/Plugins/Animation/ControlRig/Source/ControlRigDeveloper/Private/Graph/ControlRigGraph.cpp
--- a/Engine/Plugins/Animation/ControlRig/Source/ControlRigDeveloper/Private/Graph/ControlRigGraph.cpp
+++ b/Engine/Plugins/Animation/ControlRig/Source/ControlRigDeveloper/Private/Graph/ControlRigGraph.cpp
@@ -83,18 +83,6 @@ void UControlRigGraph::CacheNameLists(URigHierarchy* InHierarchy, const FRigVMDr
 	}
 	CacheNameList<FRigVMDrawContainer>(*DrawContainer, DrawingNameList);
 
-	EntryNameList.Reset();
-	EntryNameList.Add(MakeShared<FString>(FName(NAME_None).ToString()));
-
-	if(const UControlRigBlueprint* Blueprint = CastChecked<UControlRigBlueprint>(GetBlueprint()))
-	{
-		const TArray<FName> EntryNames = Blueprint->GetRigVMClient()->GetEntryNames();
-		for (const FName& EntryName : EntryNames)
-		{
-			EntryNameList.Add(MakeShared<FString>(EntryName.ToString()));
-		}
-	}
-
 	ShapeNameList.Reset();
 	ShapeNameList.Add(MakeShared<FString>(FName(NAME_None).ToString()));
 
@@ -222,22 +210,119 @@ const TArray<TSharedPtr<FString>>* UControlRigGraph::GetDrawingNameList(URigVMPi
 	return &DrawingNameList;
 }
 
-const TArray<TSharedPtr<FString>>* UControlRigGraph::GetEntryNameList(URigVMPin* InPin) const
+const TArray<TSharedPtr<FString>>* UControlRigGraph::GetShapeNameList(URigVMPin* InPin) const
 {
 	if (UControlRigGraph* OuterGraph = Cast<UControlRigGraph>(GetOuter()))
 	{
-		return OuterGraph->GetEntryNameList(InPin);
+		return OuterGraph->GetShapeNameList(InPin);
 	}
-	return &EntryNameList;
+	return &ShapeNameList;
 }
 
-const TArray<TSharedPtr<FString>>* UControlRigGraph::GetShapeNameList(URigVMPin* InPin) const
+FReply UControlRigGraph::HandleGetSelectedClicked(URigVMEdGraph* InEdGraph, URigVMPin* InPin, FString InDefaultValue)
 {
-	if (UControlRigGraph* OuterGraph = Cast<UControlRigGraph>(GetOuter()))
+	UControlRigBlueprint* RigVMBlueprint = CastChecked<UControlRigBlueprint>(InEdGraph->GetBlueprint());
+	if(InPin->GetCustomWidgetName() == TEXT("ElementName"))
 	{
-		return OuterGraph->GetShapeNameList(InPin);
+		if (URigVMPin* ParentPin = InPin->GetParentPin())
+		{
+			InEdGraph->GetController()->SetPinDefaultValue(ParentPin->GetPinPath(), InDefaultValue, true, true, false, true);
+			return FReply::Handled();
+		}
 	}
-	return &ShapeNameList;
+
+	else if (InPin->GetCustomWidgetName() == TEXT("BoneName"))
+	{
+		URigHierarchy* Hierarchy = RigVMBlueprint->Hierarchy;
+		TArray<FRigElementKey> Keys = Hierarchy->GetSelectedKeys();
+		FRigBaseElement* Element = Hierarchy->FindChecked(Keys[0]);
+		if (Element->GetType() == ERigElementType::Bone)
+		{
+			InEdGraph->GetController()->SetPinDefaultValue(InPin->GetPinPath(), Keys[0].Name.ToString(), true, true, false, true);
+			return FReply::Handled();
+		}
+	}
+
+	// if we don't have a key pin - this is just a plain name.
+	// let's derive the type of element this node deals with from its name.
+	// there's nothing better in place for now.
+	else if(const URigVMUnitNode* UnitNode = Cast<URigVMUnitNode>(InPin->GetNode()))
+	{
+		const int32 LastIndex = StaticEnum<ERigElementType>()->GetIndexByName(TEXT("Last")); 
+		const FString UnitName = UnitNode->GetScriptStruct()->GetStructCPPName();
+		for(int32 EnumIndex = 0; EnumIndex < LastIndex; EnumIndex++)
+		{
+			const FString EnumDisplayName = StaticEnum<ERigElementType>()->GetDisplayNameTextByIndex(EnumIndex).ToString();
+			if(UnitName.Contains(EnumDisplayName))
+			{
+				const ERigElementType ElementType = (ERigElementType)StaticEnum<ERigElementType>()->GetValueByIndex(EnumIndex);
+
+				FRigElementKey Key;
+				FRigElementKey::StaticStruct()->ImportText(*InDefaultValue, &Key, nullptr, EPropertyPortFlags::PPF_None, nullptr, FRigElementKey::StaticStruct()->GetName(), true);
+				if (Key.IsValid())
+				{
+					if(Key.Type == ElementType)
+					{
+						InEdGraph->GetController()->SetPinDefaultValue(InPin->GetPinPath(), Key.Name.ToString(), true, true, false, true);
+						return FReply::Handled();
+					}
+				}
+				break;
+			}
+		}
+	}
+	return FReply::Unhandled();
+}
+
+FReply UControlRigGraph::HandleBrowseClicked(URigVMEdGraph* InEdGraph, URigVMPin* InPin, FString InDefaultValue)
+{
+	UControlRigBlueprint* RigVMBlueprint = CastChecked<UControlRigBlueprint>(InEdGraph->GetBlueprint());
+	URigVMPin* KeyPin = InPin->GetParentPin();
+	if(KeyPin && KeyPin->GetCPPTypeObject() == FRigElementKey::StaticStruct())
+	{
+		// browse to rig element key
+		FString DefaultValue = InPin->GetParentPin()->GetDefaultValue();
+		if (!DefaultValue.IsEmpty())
+		{
+			FRigElementKey Key;
+			FRigElementKey::StaticStruct()->ImportText(*DefaultValue, &Key, nullptr, EPropertyPortFlags::PPF_None, nullptr, FRigElementKey::StaticStruct()->GetName(), true);
+			if (Key.IsValid())
+			{
+				RigVMBlueprint->GetHierarchyController()->SetSelection({Key});
+			}
+		}
+	}
+	else if (InPin->GetCustomWidgetName() == TEXT("BoneName"))
+	{
+		// browse to named bone
+		const FString DefaultValue = InPin->GetDefaultValue();
+		FRigElementKey Key(*DefaultValue, ERigElementType::Bone);
+		RigVMBlueprint->GetHierarchyController()->SetSelection({Key});
+	}
+	else
+	{
+		// if we don't have a key pin - this is just a plain name.
+		// let's derive the type of element this node deals with from its name.
+		// there's nothing better in place for now.
+		if(const URigVMUnitNode* UnitNode = Cast<URigVMUnitNode>(InPin->GetNode()))
+		{
+			const int32 LastIndex = StaticEnum<ERigElementType>()->GetIndexByName(TEXT("Last")); 
+			const FString UnitName = UnitNode->GetScriptStruct()->GetStructCPPName();
+			for(int32 EnumIndex = 0; EnumIndex < LastIndex; EnumIndex++)
+			{
+				const FString EnumDisplayName = StaticEnum<ERigElementType>()->GetDisplayNameTextByIndex(EnumIndex).ToString();
+				if(UnitName.Contains(EnumDisplayName))
+				{
+					const FString DefaultValue = InPin->GetDefaultValue();
+					const ERigElementType ElementType = (ERigElementType)StaticEnum<ERigElementType>()->GetValueByIndex(EnumIndex);
+					FRigElementKey Key(*DefaultValue, ElementType);
+					RigVMBlueprint->GetHierarchyController()->SetSelection({Key});
+					break;
+				}
+			}
+		}
+	}
+	return FReply::Handled();
 }
 
 #endif
diff --git a/Engine/Plugins/Animation/ControlRig/Source/ControlRigDeveloper/Public/Graph/ControlRigGraph.h b/Engine/Plugins/Animation/ControlRig/Source/ControlRigDeveloper/Public/Graph/ControlRigGraph.h
--- a/Engine/Plugins/Animation/ControlRig/Source/ControlRigDeveloper/Public/Graph/ControlRigGraph.h
+++ b/Engine/Plugins/Animation/ControlRig/Source/ControlRigDeveloper/Public/Graph/ControlRigGraph.h
@@ -59,9 +59,10 @@ public:
 	const TArray<TSharedPtr<FString>>* GetElementNameList(URigVMPin* InPin = nullptr) const;
 	const TArray<TSharedPtr<FString>> GetSelectedElementsNameList() const;
 	const TArray<TSharedPtr<FString>>* GetDrawingNameList(URigVMPin* InPin = nullptr) const;
-	const TArray<TSharedPtr<FString>>* GetEntryNameList(URigVMPin* InPin = nullptr) const;
 	const TArray<TSharedPtr<FString>>* GetShapeNameList(URigVMPin* InPin = nullptr) const;
 
+	FReply HandleGetSelectedClicked(URigVMEdGraph* InEdGraph, URigVMPin* InPin, FString InDefaultValue);
+	FReply HandleBrowseClicked(URigVMEdGraph* InEdGraph, URigVMPin* InPin, FString InDefaultValue);
 
 private:
 
@@ -116,7 +117,6 @@ private:
 	TMap<ERigElementType, TArray<TSharedPtr<FString>>> ElementNameLists;
 	TArray<TSharedPtr<FString>>	ControlNameListWithoutAnimationChannels;
 	TArray<TSharedPtr<FString>> DrawingNameList;
-	TArray<TSharedPtr<FString>> EntryNameList;
 	TArray<TSharedPtr<FString>> ShapeNameList;
 	int32 LastHierarchyTopologyVersion;
 
@@ -126,7 +126,7 @@ private:
 
 	friend class UControlRigGraphNode;
 	friend class FControlRigEditor;
-	friend class SControlRigGraphNode;
+	friend class SRigVMGraphNode;
 	friend class UControlRigBlueprint;
 };
 
diff --git a/Engine/Plugins/Animation/ControlRig/Source/ControlRigDeveloper/Public/Graph/ControlRigGraphNode.h b/Engine/Plugins/Animation/ControlRig/Source/ControlRigDeveloper/Public/Graph/ControlRigGraphNode.h
--- a/Engine/Plugins/Animation/ControlRig/Source/ControlRigDeveloper/Public/Graph/ControlRigGraphNode.h
+++ b/Engine/Plugins/Animation/ControlRig/Source/ControlRigDeveloper/Public/Graph/ControlRigGraphNode.h
@@ -44,7 +44,7 @@ public:
 
 private:
 
-	friend class SControlRigGraphNode;
+	friend class SRigVMGraphNode;
 	friend class FControlRigArgumentLayout;
 	friend class FControlRigGraphDetails;
 	friend class UControlRigTemplateNodeSpawner;
diff --git a/Engine/Plugins/Animation/ControlRig/Source/ControlRigEditor/Private/ControlRigEditorModule.cpp b/Engine/Plugins/Animation/ControlRig/Source/ControlRigEditor/Private/ControlRigEditorModule.cpp
--- a/Engine/Plugins/Animation/ControlRig/Source/ControlRigEditor/Private/ControlRigEditorModule.cpp
+++ b/Engine/Plugins/Animation/ControlRig/Source/ControlRigEditor/Private/ControlRigEditorModule.cpp
@@ -8,7 +8,7 @@
 #include "ControlRig.h"
 #include "ControlRigComponent.h"
 #include "ControlRigVariableDetailsCustomization.h"
-#include "Graph/ControlRigConnectionDrawingPolicy.h"
+#include "EdGraph/RigVMEdGraphConnectionDrawingPolicy.h"
 #include "GraphEditorActions.h"
 #include "Kismet2/KismetEditorUtilities.h"
 #include "ISequencerModule.h"
@@ -56,7 +56,7 @@
 #include "Kismet2/KismetDebugUtilities.h"
 #include "Graph/ControlRigGraphNode.h"
 #include "EdGraphUtilities.h"
-#include "Graph/ControlRigGraphPanelNodeFactory.h"
+#include "EdGraph/RigVMEdGraphPanelNodeFactory.h"
 #include "Graph/ControlRigGraphPanelPinFactory.h"
 #include "RigVMBlueprintUtils.h"
 #include "ControlRigBlueprintCommands.h"
@@ -106,8 +106,8 @@
 #include "Rigs/FKControlRig.h"
 #include "SBakeToControlRigDialog.h"
 #include "Units/Execution/RigUnit_InverseExecution.h"
-#include "Graph/SControlRigGraphPinVariableBinding.h"
-#include "Graph/SControlRigGraphChangePinType.h"
+#include "Widgets/SRigVMGraphPinVariableBinding.h"
+#include "Widgets/SRigVMGraphChangePinType.h"
 #include "ControlRigBlueprintFactory.h"
 #include "ControlRigPythonLogDetails.h"
 #include "Dialog/SCustomDialog.h"
@@ -135,6 +135,7 @@
 #include "RigVMFunctions/RigVMDispatch_Array.h"
 #include "RigVMFunctions/RigVMDispatch_MakeStruct.h"
 #include "RigVMFunctions/RigVMDispatch_Constant.h"
+#include "EdGraph/RigVMEdGraphPanelNodeFactory.h"
 
 #define LOCTEXT_NAMESPACE "ControlRigEditorModule"
 
@@ -234,7 +235,7 @@ void FControlRigEditorModule::StartupModule()
 		8500);
 
 
-	ControlRigGraphPanelNodeFactory = MakeShared<FControlRigGraphPanelNodeFactory>();
+	ControlRigGraphPanelNodeFactory = MakeShared<FRigVMEdGraphPanelNodeFactory>();
 	FEdGraphUtilities::RegisterVisualNodeFactory(ControlRigGraphPanelNodeFactory);
 
 	ControlRigGraphPanelPinFactory = MakeShared<FControlRigGraphPanelPinFactory>();
@@ -1283,7 +1284,7 @@ void FControlRigEditorModule::GetInstanceActions(UControlRigBlueprint* CRB, FBlu
 
 FConnectionDrawingPolicy* FControlRigEditorModule::CreateConnectionDrawingPolicy(int32 InBackLayerID, int32 InFrontLayerID, float InZoomFactor, const FSlateRect& InClippingRect, class FSlateWindowElementList& InDrawElements, class UEdGraph* InGraphObj)
 {
-	return new FControlRigConnectionDrawingPolicy(InBackLayerID, InFrontLayerID, InZoomFactor, InClippingRect, InDrawElements, InGraphObj);
+	return new FRigVMEdGraphConnectionDrawingPolicy(InBackLayerID, InFrontLayerID, InZoomFactor, InClippingRect, InDrawElements, InGraphObj);
 }
 
 void FControlRigEditorModule::GetContextMenuActions(const UControlRigGraphSchema* Schema, UToolMenu* Menu, UGraphNodeContextMenuContext* Context) const
@@ -1462,8 +1463,8 @@ void FControlRigEditorModule::GetContextMenuActions(const UControlRigGraphSchema
 										if(!Argument->IsSingleton())
 										{
 											TArray<TRigVMTypeIndex> ResolvedTypeIndices = Argument->GetSupportedTypeIndices(TemplateNode->GetResolvedPermutationIndices(true));
-											TSharedRef<SControlRigChangePinType> ChangePinTypeWidget =
-											SNew(SControlRigChangePinType)
+											TSharedRef<SRigVMGraphChangePinType> ChangePinTypeWidget =
+											SNew(SRigVMGraphChangePinType)
 											.Blueprint(RigBlueprint)
 											.Types(ResolvedTypeIndices)
 											.OnTypeSelected_Lambda([RigBlueprint, ModelPin](const TRigVMTypeIndex& TypeSelected)
@@ -1599,8 +1600,8 @@ void FControlRigEditorModule::GetContextMenuActions(const UControlRigGraphSchema
 							{
 								FToolMenuSection& VariablesSection = Menu->FindOrAddSection(TEXT("Variables"));
 
-								TSharedRef<SControlRigVariableBinding> VariableBindingWidget =
-									SNew(SControlRigVariableBinding)
+								TSharedRef<SRigVMGraphVariableBinding> VariableBindingWidget =
+									SNew(SRigVMGraphVariableBinding)
 									.Blueprint(RigBlueprint)
 									.ModelPins({ModelPin})
 									.CanRemoveBinding(false);
diff --git a/Engine/Plugins/Animation/ControlRig/Source/ControlRigEditor/Private/ControlRigEditorModule.h b/Engine/Plugins/Animation/ControlRig/Source/ControlRigEditor/Private/ControlRigEditorModule.h
--- a/Engine/Plugins/Animation/ControlRig/Source/ControlRigEditor/Private/ControlRigEditorModule.h
+++ b/Engine/Plugins/Animation/ControlRig/Source/ControlRigEditor/Private/ControlRigEditorModule.h
@@ -19,7 +19,7 @@ class FToolBarBuilder;
 class FExtender;
 class FUICommandList;
 class UMovieSceneTrack;
-class FControlRigGraphPanelNodeFactory;
+class FRigVMEdGraphPanelNodeFactory;
 class FControlRigGraphPanelPinFactory;
 
 class FControlRigEditorModule : public IControlRigEditorModule
@@ -77,7 +77,7 @@ private:
 	TArray<FControlRigEditorToolbarExtender> ControlRigEditorToolbarExtenders;
 
 	/** Node factory for the control rig graph */
-	TSharedPtr<FControlRigGraphPanelNodeFactory> ControlRigGraphPanelNodeFactory;
+	TSharedPtr<FRigVMEdGraphPanelNodeFactory> ControlRigGraphPanelNodeFactory;
 
 	/** Pin factory for the control rig graph */
 	TSharedPtr<FControlRigGraphPanelPinFactory> ControlRigGraphPanelPinFactory;
diff --git a/Engine/Plugins/Animation/ControlRig/Source/ControlRigEditor/Private/ControlRigElementDetails.cpp b/Engine/Plugins/Animation/ControlRig/Source/ControlRigEditor/Private/ControlRigElementDetails.cpp
--- a/Engine/Plugins/Animation/ControlRig/Source/ControlRigEditor/Private/ControlRigElementDetails.cpp
+++ b/Engine/Plugins/Animation/ControlRig/Source/ControlRigEditor/Private/ControlRigElementDetails.cpp
@@ -15,7 +15,7 @@
 #include "SEnumCombo.h"
 #include "Units/Execution/RigUnit_BeginExecution.h"
 #include "Units/Execution/RigUnit_DynamicHierarchy.h"
-#include "Graph/SControlRigGraphPinVariableBinding.h"
+#include "Widgets/SRigVMGraphPinVariableBinding.h"
 #include "HAL/PlatformApplicationMisc.h"
 #include "Styling/AppStyle.h"
 #include "Editor/SRigHierarchyTreeView.h"
diff --git a/Engine/Plugins/Animation/ControlRig/Source/ControlRigEditor/Private/ControlRigElementDetails.h b/Engine/Plugins/Animation/ControlRig/Source/ControlRigEditor/Private/ControlRigElementDetails.h
--- a/Engine/Plugins/Animation/ControlRig/Source/ControlRigEditor/Private/ControlRigElementDetails.h
+++ b/Engine/Plugins/Animation/ControlRig/Source/ControlRigEditor/Private/ControlRigElementDetails.h
@@ -10,7 +10,7 @@
 #include "ControlRig.h"
 #include "ControlRigBlueprint.h"
 #include "Editor/DetailsViewWrapperObject.h"
-#include "Graph/SControlRigGraphPinNameListValueWidget.h"
+#include "Widgets/SRigVMGraphPinNameListValueWidget.h"
 #include "Editor/SControlRigGizmoNameList.h"
 #include "Styling/SlateTypes.h"
 #include "IPropertyUtilities.h"
diff --git a/Engine/Plugins/Animation/ControlRig/Source/ControlRigEditor/Private/ControlRigGraphDetails.cpp b/Engine/Plugins/Animation/ControlRig/Source/ControlRigEditor/Private/ControlRigGraphDetails.cpp
--- a/Engine/Plugins/Animation/ControlRig/Source/ControlRigEditor/Private/ControlRigGraphDetails.cpp
+++ b/Engine/Plugins/Animation/ControlRig/Source/ControlRigEditor/Private/ControlRigGraphDetails.cpp
@@ -23,7 +23,7 @@
 #include "Graph/ControlRigGraphSchema.h"
 #include "EditorCategoryUtils.h"
 #include "IPropertyUtilities.h"
-#include "Graph/SControlRigGraphPinVariableBinding.h"
+#include "Widgets/SRigVMGraphPinVariableBinding.h"
 
 #define LOCTEXT_NAMESPACE "ControlRigGraphDetails"
 
@@ -310,7 +310,7 @@ void FControlRigArgumentLayout::GenerateHeaderRowContent(FDetailWidgetRow& NodeR
 
 void FControlRigArgumentLayout::GenerateChildContent(IDetailChildrenBuilder& ChildrenBuilder)
 {
-	// we don't show defaults here - we rely on a SControlRigGraphNode widget in the top of the details
+	// we don't show defaults here - we rely on a SRigVMGraphNode widget in the top of the details
 }
 
 void FControlRigArgumentLayout::OnRemoveClicked()
@@ -622,7 +622,7 @@ void FControlRigArgumentDefaultNode::GenerateChildContent(IDetailChildrenBuilder
 	.WholeRowContent()
 	.MaxDesiredWidth(980.f)
 	[
-		SAssignNew(OwnedNodeWidget, SControlRigGraphNode).GraphNodeObj(ControlRigGraphNode)
+		SAssignNew(OwnedNodeWidget, SRigVMGraphNode).GraphNodeObj(ControlRigGraphNode)
 	];
 
 	OwnedNodeWidget->SetIsEditable(true);
@@ -1351,7 +1351,7 @@ void FControlRigWrappedNodeDetails::CustomizeDetails(IDetailLayoutBuilder& Detai
 					]
 					.ValueContent()
 					[
-					SNew(SControlRigVariableBinding)
+					SNew(SRigVMGraphVariableBinding)
 						.ModelPins(ModelPins)
 						.Blueprint(BlueprintBeingCustomized)
 					];
@@ -1401,7 +1401,7 @@ void FControlRigWrappedNodeDetails::CustomizeDetails(IDetailLayoutBuilder& Detai
 
         		if (NameList)
         		{
-        			TSharedPtr<SControlRigGraphPinNameListValueWidget> NameListWidget;
+        			TSharedPtr<SRigVMGraphPinNameListValueWidget> NameListWidget;
 
         			if(IDetailPropertyRow* Row = DetailLayout.EditDefaultProperty(PropertyHandle))
         			{
@@ -1412,7 +1412,7 @@ void FControlRigWrappedNodeDetails::CustomizeDetails(IDetailLayoutBuilder& Detai
 						]
 						.ValueContent()
 						[
-							SAssignNew(NameListWidget, SControlRigGraphPinNameListValueWidget)
+							SAssignNew(NameListWidget, SRigVMGraphPinNameListValueWidget)
 							.OptionsSource(NameList)
 							.OnGenerateWidget(this, &FControlRigWrappedNodeDetails::MakeNameListItemWidget)
 							.OnSelectionChanged(this, &FControlRigWrappedNodeDetails::OnNameListChanged, NameProperty, DetailLayout.GetPropertyUtilities())
@@ -1574,7 +1574,7 @@ void FControlRigWrappedNodeDetails::OnNameListChanged(TSharedPtr<FString> NewSel
 
 void FControlRigWrappedNodeDetails::OnNameListComboBox(FNameProperty* InProperty, const TArray<TSharedPtr<FString>>* InNameList)
 {
-	TSharedPtr<SControlRigGraphPinNameListValueWidget> Widget = NameListWidgets.FindChecked(InProperty->GetFName());
+	TSharedPtr<SRigVMGraphPinNameListValueWidget> Widget = NameListWidgets.FindChecked(InProperty->GetFName());
 	const TSharedPtr<FString> CurrentlySelected = GetCurrentlySelectedItem(InProperty, InNameList);
 	Widget->SetSelectedItem(CurrentlySelected);
 }
diff --git a/Engine/Plugins/Animation/ControlRig/Source/ControlRigEditor/Private/ControlRigGraphDetails.h b/Engine/Plugins/Animation/ControlRig/Source/ControlRigEditor/Private/ControlRigGraphDetails.h
--- a/Engine/Plugins/Animation/ControlRig/Source/ControlRigEditor/Private/ControlRigGraphDetails.h
+++ b/Engine/Plugins/Animation/ControlRig/Source/ControlRigEditor/Private/ControlRigGraphDetails.h
@@ -10,7 +10,7 @@
 #include "ControlRigBlueprint.h"
 #include "Editor/ControlRigEditor.h"
 #include "SGraphPin.h"
-#include "Graph/SControlRigGraphNode.h"
+#include "Widgets/SRigVMGraphNode.h"
 #include "Widgets/Colors/SColorBlock.h"
 #include "Editor/DetailsViewWrapperObject.h"
 #include "IDetailPropertyExtensionHandler.h"
@@ -149,7 +149,7 @@ private:
 	TWeakObjectPtr<UControlRigGraph> EdGraphOuterPtr;
 	TWeakObjectPtr<UControlRigBlueprint> ControlRigBlueprintPtr;
 	FSimpleDelegate OnRebuildChildren;
-	TSharedPtr<SControlRigGraphNode> OwnedNodeWidget;
+	TSharedPtr<SRigVMGraphNode> OwnedNodeWidget;
 	FDelegateHandle GraphChangedDelegateHandle;
 };
 
@@ -234,7 +234,7 @@ public:
 	UControlRigBlueprint* BlueprintBeingCustomized;
 	TArray<TWeakObjectPtr<UDetailsViewWrapperObject>> ObjectsBeingCustomized;
 	TArray<TWeakObjectPtr<URigVMNode>> NodesBeingCustomized;
-	TMap<FName, TSharedPtr<SControlRigGraphPinNameListValueWidget>> NameListWidgets;
+	TMap<FName, TSharedPtr<SRigVMGraphPinNameListValueWidget>> NameListWidgets;
 };
 
 /** Customization for editing a Control Rig node */
diff --git a/Engine/Plugins/Animation/ControlRig/Source/ControlRigEditor/Private/ControlRigLocalVariableDetails.h b/Engine/Plugins/Animation/ControlRig/Source/ControlRigEditor/Private/ControlRigLocalVariableDetails.h
--- a/Engine/Plugins/Animation/ControlRig/Source/ControlRigEditor/Private/ControlRigLocalVariableDetails.h
+++ b/Engine/Plugins/Animation/ControlRig/Source/ControlRigEditor/Private/ControlRigLocalVariableDetails.h
@@ -9,7 +9,7 @@
 #include "ControlRigBlueprint.h"
 #include "Editor/DetailsViewWrapperObject.h"
 #include "Graph/ControlRigGraph.h"
-#include "Graph/SControlRigGraphPinNameListValueWidget.h"
+#include "Widgets/SRigVMGraphPinNameListValueWidget.h"
 #include "Styling/SlateTypes.h"
 #include "IPropertyUtilities.h"
 #include "Graph/ControlRigGraphSchema.h"
diff --git a/Engine/Plugins/Animation/ControlRig/Source/ControlRigEditor/Private/Editor/ControlRigEditor.cpp b/Engine/Plugins/Animation/ControlRig/Source/ControlRigEditor/Private/Editor/ControlRigEditor.cpp
--- a/Engine/Plugins/Animation/ControlRig/Source/ControlRigEditor/Private/Editor/ControlRigEditor.cpp
+++ b/Engine/Plugins/Animation/ControlRig/Source/ControlRigEditor/Private/Editor/ControlRigEditor.cpp
@@ -87,7 +87,7 @@
 #include "Editor/SControlRigFunctionLocalizationWidget.h"
 #include "Editor/SControlRigFunctionBulkEditWidget.h"
 #include "Editor/SControlRigBreakLinksWidget.h"
-#include "Graph/SControlRigGraphChangePinType.h"
+#include "Widgets/SRigVMGraphChangePinType.h"
 #include "SGraphPanel.h"
 #include "Engine/StaticMesh.h"
 #include "Engine/StaticMeshActor.h"
@@ -3916,7 +3916,7 @@ void FControlRigEditor::HandleViewportCreated(const TSharedRef<class IPersonaVie
 									.WidthOverride(100.0f)
 									.IsEnabled(this, &FControlRigEditor::IsPinControlNameListEnabled)
 									[
-										SAssignNew(PinControlNameList, SControlRigGraphPinNameListValueWidget)
+										SAssignNew(PinControlNameList, SRigVMGraphPinNameListValueWidget)
 										.OptionsSource(BoneNameList)
 										.OnGenerateWidget(this, &FControlRigEditor::MakePinControlNameListItemWidget)
 										.OnSelectionChanged(this, &FControlRigEditor::OnPinControlNameListChanged)
@@ -4384,6 +4384,7 @@ void FControlRigEditor::CacheNameLists()
 			{
 				ShapeLibraries = &DebuggedControlRig->GetShapeLibraries();
 			}
+			RigGraph->CacheEntryNameList();
 			RigGraph->CacheNameLists(Hierarchy, &ControlRigBP->DrawContainer, *ShapeLibraries);
 		}
 	}
diff --git a/Engine/Plugins/Animation/ControlRig/Source/ControlRigEditor/Private/Editor/ControlRigEditor.h b/Engine/Plugins/Animation/ControlRig/Source/ControlRigEditor/Private/Editor/ControlRigEditor.h
--- a/Engine/Plugins/Animation/ControlRig/Source/ControlRigEditor/Private/Editor/ControlRigEditor.h
+++ b/Engine/Plugins/Animation/ControlRig/Source/ControlRigEditor/Private/Editor/ControlRigEditor.h
@@ -13,7 +13,7 @@
 #include "RigVMModel/RigVMGraph.h"
 #include "RigVMModel/Nodes/RigVMLibraryNode.h"
 #include "RigVMCore/RigVM.h"
-#include "Graph/SControlRigGraphPinNameListValueWidget.h"
+#include "Widgets/SRigVMGraphPinNameListValueWidget.h"
 #include "Styling/SlateTypes.h"
 #include "AnimPreviewInstance.h"
 #include "ScopedTransaction.h"
@@ -466,7 +466,7 @@ protected:
 
 	FPreviewControlRigUpdated PreviewControlRigUpdated;
 
-	TSharedPtr<SControlRigGraphPinNameListValueWidget> PinControlNameList;
+	TSharedPtr<SRigVMGraphPinNameListValueWidget> PinControlNameList;
 	bool IsPinControlNameListEnabled() const;
 	TSharedRef<SWidget> MakePinControlNameListItemWidget(TSharedPtr<FString> InItem);
 	FText GetPinControlNameListText() const;
diff --git a/Engine/Plugins/Animation/ControlRig/Source/ControlRigEditor/Private/Editor/SControlRigGizmoNameList.cpp b/Engine/Plugins/Animation/ControlRig/Source/ControlRigEditor/Private/Editor/SControlRigGizmoNameList.cpp
--- a/Engine/Plugins/Animation/ControlRig/Source/ControlRigEditor/Private/Editor/SControlRigGizmoNameList.cpp
+++ b/Engine/Plugins/Animation/ControlRig/Source/ControlRigEditor/Private/Editor/SControlRigGizmoNameList.cpp
@@ -2,7 +2,7 @@
 
 
 #include "Editor/SControlRigGizmoNameList.h"
-#include "Graph/SControlRigGraphPinNameListValueWidget.h"
+#include "Widgets/SRigVMGraphPinNameListValueWidget.h"
 #include "Widgets/Input/SEditableTextBox.h"
 #include "EdGraphSchema_K2.h"
 #include "ScopedTransaction.h"
@@ -71,7 +71,7 @@ void SControlRigShapeNameList::ConstructCommon()
 		.MinDesiredWidth(150)
 		.MaxDesiredWidth(400)
 		[
-			SAssignNew(NameListComboBox, SControlRigGraphPinNameListValueWidget)
+			SAssignNew(NameListComboBox, SRigVMGraphPinNameListValueWidget)
 			.OptionsSource(&GetNameList())
 			.OnGenerateWidget(this, &SControlRigShapeNameList::MakeNameListItemWidget)
 			.OnSelectionChanged(this, &SControlRigShapeNameList::OnNameListChanged)
diff --git a/Engine/Plugins/Animation/ControlRig/Source/ControlRigEditor/Private/Editor/SControlRigGizmoNameList.h b/Engine/Plugins/Animation/ControlRig/Source/ControlRigEditor/Private/Editor/SControlRigGizmoNameList.h
--- a/Engine/Plugins/Animation/ControlRig/Source/ControlRigEditor/Private/Editor/SControlRigGizmoNameList.h
+++ b/Engine/Plugins/Animation/ControlRig/Source/ControlRigEditor/Private/Editor/SControlRigGizmoNameList.h
@@ -6,7 +6,7 @@
 #include "Widgets/DeclarativeSyntaxSupport.h"
 #include "Widgets/SWidget.h"
 #include "Widgets/Layout/SBox.h"
-#include "Graph/SControlRigGraphPinNameListValueWidget.h"
+#include "Widgets/SRigVMGraphPinNameListValueWidget.h"
 #include "Rigs/RigHierarchyContainer.h"
 #include "ControlRigBlueprint.h"
 
@@ -46,7 +46,7 @@ protected:
 	}
 
 	FOnGetNameListContent OnGetNameListContent;
-	TSharedPtr<SControlRigGraphPinNameListValueWidget> NameListComboBox;
+	TSharedPtr<SRigVMGraphPinNameListValueWidget> NameListComboBox;
 
 	TArray<FRigElementKey> ControlKeys;
 	UControlRigBlueprint* Blueprint;
diff --git a/Engine/Plugins/Animation/ControlRig/Source/ControlRigEditor/Private/Graph/ControlRigGraphPanelPinFactory.cpp b/Engine/Plugins/Animation/ControlRig/Source/ControlRigEditor/Private/Graph/ControlRigGraphPanelPinFactory.cpp
--- a/Engine/Plugins/Animation/ControlRig/Source/ControlRigEditor/Private/Graph/ControlRigGraphPanelPinFactory.cpp
+++ b/Engine/Plugins/Animation/ControlRig/Source/ControlRigEditor/Private/Graph/ControlRigGraphPanelPinFactory.cpp
@@ -4,12 +4,12 @@
 #include "Graph/ControlRigGraphSchema.h"
 #include "Graph/ControlRigGraphNode.h"
 #include "Graph/ControlRigGraph.h"
-#include "Graph/SControlRigGraphPinNameList.h"
-#include "Graph/SControlRigGraphPinCurveFloat.h"
-#include "Graph/SControlRigGraphPinVariableName.h"
-#include "Graph/SControlRigGraphPinVariableBinding.h"
-#include "Graph/SGraphPinUserDataNameSpace.h"
-#include "Graph/SGraphPinUserDataPath.h"
+#include "Widgets/SRigVMGraphPinNameList.h"
+#include "Widgets/SRigVMGraphPinCurveFloat.h"
+#include "Widgets/SRigVMGraphPinVariableName.h"
+#include "Widgets/SRigVMGraphPinVariableBinding.h"
+#include "Widgets/SRigVMGraphPinUserDataNameSpace.h"
+#include "Widgets/SRigVMGraphPinUserDataPath.h"
 #include "KismetPins/SGraphPinExec.h"
 #include "SGraphPinComboBox.h"
 #include "ControlRig.h"
@@ -24,6 +24,97 @@
 
 TSharedPtr<SGraphPin> FControlRigGraphPanelPinFactory::CreatePin(UEdGraphPin* InPin) const
 {
+	TSharedPtr<SGraphPin> InternalResult = CreatePin_Internal(InPin);
+	if(InternalResult.IsValid())
+	{
+		return InternalResult;
+	}
+	
+	TSharedPtr<SGraphPin> K2PinWidget = FNodeFactory::CreateK2PinWidget(InPin);
+	if(K2PinWidget.IsValid())
+	{
+		if(InPin->Direction == EEdGraphPinDirection::EGPD_Input)
+		{
+			// if we are an enum pin - and we are inside a RigElementKey,
+			// let's remove the "all" entry.
+			if(InPin->PinType.PinSubCategoryObject == StaticEnum<ERigElementType>())
+			{
+				if(InPin->ParentPin)
+				{
+					if(InPin->ParentPin->PinType.PinSubCategoryObject == FRigElementKey::StaticStruct())
+					{
+						TSharedPtr<SWidget> ValueWidget = K2PinWidget->GetValueWidget();
+						if(ValueWidget.IsValid())
+						{
+							if(TSharedPtr<SPinComboBox> EnumCombo = StaticCastSharedPtr<SPinComboBox>(ValueWidget))
+							{
+								if(EnumCombo.IsValid())
+								{
+									EnumCombo->RemoveItemByIndex(StaticEnum<ERigElementType>()->GetIndexByValue((int64)ERigElementType::All));
+								}
+							}
+						}
+					}
+				}
+			}
+
+			const UEnum* RigControlTransformChannelEnum = StaticEnum<ERigControlTransformChannel>();
+			if (InPin->PinType.PinSubCategoryObject == RigControlTransformChannelEnum)
+			{
+				TSharedPtr<SWidget> ValueWidget = K2PinWidget->GetValueWidget();
+				if(ValueWidget.IsValid())
+				{
+					if(TSharedPtr<SPinComboBox> EnumCombo = StaticCastSharedPtr<SPinComboBox>(ValueWidget))
+					{
+						if(EnumCombo.IsValid())
+						{
+							if (const UControlRigGraphNode* RigNode = Cast<UControlRigGraphNode>(InPin->GetOwningNode()))
+							{
+								if (const URigVMPin* ModelPin = RigNode->GetModelPinFromPinPath(InPin->GetName()))
+								{
+									if(const URigVMUnitNode* UnitNode = Cast<URigVMUnitNode>(ModelPin->GetNode()))
+									{
+										if(UnitNode->GetScriptStruct() &&
+											UnitNode->GetScriptStruct()->IsChildOf(FRigUnit_HierarchyAddControlElement::StaticStruct()))
+										{
+											const TSharedPtr<FStructOnScope> StructInstanceScope = UnitNode->ConstructStructInstance();
+											const FRigUnit_HierarchyAddControlElement* StructInstance = 
+												(const FRigUnit_HierarchyAddControlElement*)StructInstanceScope->GetStructMemory();
+
+											if(const TArray<ERigControlTransformChannel>* VisibleChannels =
+												FRigControlTransformChannelDetails::GetVisibleChannelsForControlType(StructInstance->GetControlTypeToSpawn()))
+											{
+												for(int32 Index = 0; Index < RigControlTransformChannelEnum->NumEnums(); Index++)
+												{
+													const ERigControlTransformChannel Value =
+														(ERigControlTransformChannel)RigControlTransformChannelEnum->GetValueByIndex(Index);
+													if(!VisibleChannels->Contains(Value))
+													{
+														EnumCombo->RemoveItemByIndex(Index);
+													}
+												}
+											}
+										}
+									}
+								}
+							}
+						}
+					}
+				}
+			}
+		}
+	}
+	return K2PinWidget;
+}
+
+TSharedPtr<SGraphPin> FControlRigGraphPanelPinFactory::CreatePin_Internal(UEdGraphPin* InPin) const
+{
+	TSharedPtr<SGraphPin> SuperResult = FRigVMEdGraphPanelPinFactory::CreatePin_Internal(InPin);
+	if(SuperResult.IsValid())
+	{
+		return SuperResult;
+	}
+
 	if (InPin)
 	{
 		if (const UEdGraphNode* OwningNode = InPin->GetOwningNode())
@@ -41,76 +132,64 @@ TSharedPtr<SGraphPin> FControlRigGraphPanelPinFactory::CreatePin(UEdGraphPin* In
 
 			if (URigVMPin* ModelPin = RigNode->GetModelPinFromPinPath(InPin->GetName()))
 			{
-				if (ModelPin->IsBoundToVariable())
-				{
-					if (UControlRigBlueprint* Blueprint = RigGraph->GetTypedOuter<UControlRigBlueprint>())
-					{
-						return SNew(SControlRigGraphPinVariableBinding, InPin)
-							.ModelPins({ModelPin})
-							.Blueprint(Blueprint);
-					}
-				}
-
 				FName CustomWidgetName = ModelPin->GetCustomWidgetName();
 				if (CustomWidgetName == TEXT("BoneName"))
 				{
-					return SNew(SControlRigGraphPinNameList, InPin)
+					return SNew(SRigVMGraphPinNameList, InPin)
 						.ModelPin(ModelPin)
 						.OnGetNameFromSelection_UObject(RigGraph, &UControlRigGraph::GetSelectedElementsNameList)
-						.OnGetNameListContent_UObject(RigGraph, &UControlRigGraph::GetBoneNameList);
-				}
+						.OnGetNameListContent_UObject(RigGraph, &UControlRigGraph::GetBoneNameList)
+						.OnGetSelectedClicked_UObject(RigGraph, &UControlRigGraph::HandleGetSelectedClicked)
+						.OnBrowseClicked_UObject(RigGraph, &UControlRigGraph::HandleBrowseClicked);
+			}
 				else if (CustomWidgetName == TEXT("ControlName"))
 				{
-					return SNew(SControlRigGraphPinNameList, InPin)
+					return SNew(SRigVMGraphPinNameList, InPin)
 						.ModelPin(ModelPin)
 						.OnGetNameFromSelection_UObject(RigGraph, &UControlRigGraph::GetSelectedElementsNameList)
 						.OnGetNameListContent_UObject(RigGraph, &UControlRigGraph::GetControlNameListWithoutAnimationChannels)
-						.OnGetNameListContentForValidation_UObject(RigGraph, &UControlRigGraph::GetControlNameList);
+						.OnGetNameListContentForValidation_UObject(RigGraph, &UControlRigGraph::GetControlNameList)
+						.OnGetSelectedClicked_UObject(RigGraph, &UControlRigGraph::HandleGetSelectedClicked)
+						.OnBrowseClicked_UObject(RigGraph, &UControlRigGraph::HandleBrowseClicked);
 				}
-				else if (CustomWidgetName == TEXT("SpaceName"))
+				else if (CustomWidgetName == TEXT("SpaceName") || CustomWidgetName == TEXT("NullName"))
 				{
-					return SNew(SControlRigGraphPinNameList, InPin)
+					return SNew(SRigVMGraphPinNameList, InPin)
 						.ModelPin(ModelPin)
 						.OnGetNameFromSelection_UObject(RigGraph, &UControlRigGraph::GetSelectedElementsNameList)
-						.OnGetNameListContent_UObject(RigGraph, &UControlRigGraph::GetNullNameList);
+						.OnGetNameListContent_UObject(RigGraph, &UControlRigGraph::GetNullNameList)
+						.OnGetSelectedClicked_UObject(RigGraph, &UControlRigGraph::HandleGetSelectedClicked)
+						.OnBrowseClicked_UObject(RigGraph, &UControlRigGraph::HandleBrowseClicked);
 				}
 				else if (CustomWidgetName == TEXT("CurveName"))
 				{
-					return SNew(SControlRigGraphPinNameList, InPin)
+					return SNew(SRigVMGraphPinNameList, InPin)
 						.ModelPin(ModelPin)
 						.OnGetNameFromSelection_UObject(RigGraph, &UControlRigGraph::GetSelectedElementsNameList)
 						.OnGetNameListContent_UObject(RigGraph, &UControlRigGraph::GetCurveNameList);
 				}
 				else if (CustomWidgetName == TEXT("ElementName"))
 				{
-					return SNew(SControlRigGraphPinNameList, InPin)
+					return SNew(SRigVMGraphPinNameList, InPin)
 						.ModelPin(ModelPin)
 						.OnGetNameFromSelection_UObject(RigGraph, &UControlRigGraph::GetSelectedElementsNameList)
-						.OnGetNameListContent_UObject(RigGraph, &UControlRigGraph::GetElementNameList);
-				}
-				else if (CustomWidgetName == TEXT("EntryName"))
-				{
-					return SNew(SControlRigGraphPinNameList, InPin)
-						.ModelPin(ModelPin)
-						.OnGetNameListContent_UObject(RigGraph, &UControlRigGraph::GetEntryNameList);
+						.OnGetNameListContent_UObject(RigGraph, &UControlRigGraph::GetElementNameList)
+						.OnGetSelectedClicked_UObject(RigGraph, &UControlRigGraph::HandleGetSelectedClicked)
+						.OnBrowseClicked_UObject(RigGraph, &UControlRigGraph::HandleBrowseClicked);
 				}
 				else if (CustomWidgetName == TEXT("DrawingName"))
 				{
-					return SNew(SControlRigGraphPinNameList, InPin)
+					return SNew(SRigVMGraphPinNameList, InPin)
 						.ModelPin(ModelPin)
 						.OnGetNameFromSelection_UObject(RigGraph, &UControlRigGraph::GetSelectedElementsNameList)
 						.OnGetNameListContent_UObject(RigGraph, &UControlRigGraph::GetDrawingNameList);
 				}
 				else if (CustomWidgetName == TEXT("ShapeName"))
 				{
-					return SNew(SControlRigGraphPinNameList, InPin)
+					return SNew(SRigVMGraphPinNameList, InPin)
 						.ModelPin(ModelPin)
 						.OnGetNameListContent_UObject(RigGraph, &UControlRigGraph::GetShapeNameList);
 				}
-				else if (CustomWidgetName == TEXT("VariableName"))
-				{
-					return SNew(SControlRigGraphPinVariableName, InPin);
-				}
 				else if (CustomWidgetName == TEXT("AnimationChannelName"))
 				{
 					struct FCachedAnimationChannelNames
@@ -123,7 +202,7 @@ TSharedPtr<SGraphPin> FControlRigGraphPanelPinFactory::CreatePin(UEdGraphPin* In
 						{}
 					};
 					
-					return SNew(SControlRigGraphPinNameList, InPin)
+					return SNew(SRigVMGraphPinNameList, InPin)
 						.ModelPin(ModelPin)
 						.OnGetNameListContent_Lambda([RigGraph](const URigVMPin* InPin)
 						{
@@ -212,7 +291,9 @@ TSharedPtr<SGraphPin> FControlRigGraphPanelPinFactory::CreatePin(UEdGraphPin* In
 
 							static TArray<TSharedPtr<FString>> EmptyNameList;
 							return &EmptyNameList;
-						});
+						})
+						.OnGetSelectedClicked_UObject(RigGraph, &UControlRigGraph::HandleGetSelectedClicked)
+						.OnBrowseClicked_UObject(RigGraph, &UControlRigGraph::HandleBrowseClicked);
 				}
 				else if (CustomWidgetName == TEXT("MetadataName"))
 				{
@@ -226,7 +307,7 @@ TSharedPtr<SGraphPin> FControlRigGraphPanelPinFactory::CreatePin(UEdGraphPin* In
 						{}
 					};
 
-					return SNew(SControlRigGraphPinNameList, InPin)
+					return SNew(SRigVMGraphPinNameList, InPin)
 						.ModelPin(ModelPin)
 						.SearchHintText(NSLOCTEXT("FControlRigGraphPanelPinFactory", "MetadataName", "Metadata Name"))
 						.AllowUserProvidedText(true)
@@ -293,7 +374,7 @@ TSharedPtr<SGraphPin> FControlRigGraphPanelPinFactory::CreatePin(UEdGraphPin* In
 						{}
 					};
 
-					return SNew(SControlRigGraphPinNameList, InPin)
+					return SNew(SRigVMGraphPinNameList, InPin)
 						.ModelPin(ModelPin)
 						.SearchHintText(NSLOCTEXT("FControlRigGraphPanelPinFactory", "TagName", "Tag Name"))
 						.AllowUserProvidedText(true)
@@ -352,106 +433,7 @@ TSharedPtr<SGraphPin> FControlRigGraphPanelPinFactory::CreatePin(UEdGraphPin* In
 							return &EmptyNameList;
 						});
 				}
-				else if (CustomWidgetName == TEXT("UserDataNameSpace"))
-				{
-					return SNew(SGraphPinUserDataNameSpace, InPin);
-				}
-				else if (CustomWidgetName == TEXT("UserDataPath"))
-				{
-					return SNew(SGraphPinUserDataPath, InPin)
-						.ModelPins({ModelPin});
-				}
-			}
-
-			if (InPin->PinType.PinCategory == UEdGraphSchema_K2::PC_Struct)
-			{
-				const UStruct* Struct = Cast<UStruct>(InPin->PinType.PinSubCategoryObject);
-				if (Struct && Struct->IsChildOf(FRigVMExecuteContext::StaticStruct()))
-				{
-					return SNew(SGraphPinExec, InPin);
-				}
-				if (InPin->PinType.PinSubCategoryObject == FRuntimeFloatCurve::StaticStruct())
-				{
-					return SNew(SControlRigGraphPinCurveFloat, InPin);
-				}
-			}
-		}
-
-		TSharedPtr<SGraphPin> K2PinWidget = FNodeFactory::CreateK2PinWidget(InPin);
-		if(K2PinWidget.IsValid())
-		{
-			if(InPin->Direction == EEdGraphPinDirection::EGPD_Input)
-			{
-				// if we are an enum pin - and we are inside a RigElementKey,
-				// let's remove the "all" entry.
-				if(InPin->PinType.PinSubCategoryObject == StaticEnum<ERigElementType>())
-				{
-					if(InPin->ParentPin)
-					{
-						if(InPin->ParentPin->PinType.PinSubCategoryObject == FRigElementKey::StaticStruct())
-						{
-							TSharedPtr<SWidget> ValueWidget = K2PinWidget->GetValueWidget();
-							if(ValueWidget.IsValid())
-							{
-								if(TSharedPtr<SPinComboBox> EnumCombo = StaticCastSharedPtr<SPinComboBox>(ValueWidget))
-								{
-									if(EnumCombo.IsValid())
-									{
-										EnumCombo->RemoveItemByIndex(StaticEnum<ERigElementType>()->GetIndexByValue((int64)ERigElementType::All));
-									}
-								}
-							}
-						}
-					}
-				}
-
-				const UEnum* RigControlTransformChannelEnum = StaticEnum<ERigControlTransformChannel>();
-				if (InPin->PinType.PinSubCategoryObject == RigControlTransformChannelEnum)
-				{
-					TSharedPtr<SWidget> ValueWidget = K2PinWidget->GetValueWidget();
-					if(ValueWidget.IsValid())
-					{
-						if(TSharedPtr<SPinComboBox> EnumCombo = StaticCastSharedPtr<SPinComboBox>(ValueWidget))
-						{
-							if(EnumCombo.IsValid())
-							{
-								if (const UControlRigGraphNode* RigNode = Cast<UControlRigGraphNode>(InPin->GetOwningNode()))
-								{
-									if (const URigVMPin* ModelPin = RigNode->GetModelPinFromPinPath(InPin->GetName()))
-									{
-										if(const URigVMUnitNode* UnitNode = Cast<URigVMUnitNode>(ModelPin->GetNode()))
-										{
-											if(UnitNode->GetScriptStruct() &&
-												UnitNode->GetScriptStruct()->IsChildOf(FRigUnit_HierarchyAddControlElement::StaticStruct()))
-											{
-												const TSharedPtr<FStructOnScope> StructInstanceScope = UnitNode->ConstructStructInstance();
-												const FRigUnit_HierarchyAddControlElement* StructInstance = 
-													(const FRigUnit_HierarchyAddControlElement*)StructInstanceScope->GetStructMemory();
-
-												if(const TArray<ERigControlTransformChannel>* VisibleChannels =
-													FRigControlTransformChannelDetails::GetVisibleChannelsForControlType(StructInstance->GetControlTypeToSpawn()))
-												{
-													for(int32 Index = 0; Index < RigControlTransformChannelEnum->NumEnums(); Index++)
-													{
-														const ERigControlTransformChannel Value =
-															(ERigControlTransformChannel)RigControlTransformChannelEnum->GetValueByIndex(Index);
-														if(!VisibleChannels->Contains(Value))
-														{
-															EnumCombo->RemoveItemByIndex(Index);
-														}
-													}
-												}
-											}
-										}
-									}
-								}
-							}
-						}
-					}
-				}
 			}
-
-			return K2PinWidget;
 		}
 	}
 
diff --git a/Engine/Plugins/Animation/ControlRig/Source/ControlRigEditor/Private/Graph/ControlRigGraphPanelPinFactory.h b/Engine/Plugins/Animation/ControlRig/Source/ControlRigEditor/Private/Graph/ControlRigGraphPanelPinFactory.h
--- a/Engine/Plugins/Animation/ControlRig/Source/ControlRigEditor/Private/Graph/ControlRigGraphPanelPinFactory.h
+++ b/Engine/Plugins/Animation/ControlRig/Source/ControlRigEditor/Private/Graph/ControlRigGraphPanelPinFactory.h
@@ -2,12 +2,14 @@
 
 #pragma once
 
-#include "CoreMinimal.h"
-#include "EdGraphUtilities.h"
+#include "EdGraph/RigVMEdGraphPanelPinFactory.h"
 
-class FControlRigGraphPanelPinFactory : public FGraphPanelPinFactory
+class FControlRigGraphPanelPinFactory : public FRigVMEdGraphPanelPinFactory
 {
 public:
 	// FGraphPanelPinFactory interface
 	virtual TSharedPtr<class SGraphPin> CreatePin(class UEdGraphPin* InPin) const override;
+
+	// FRigVMEdGraphPanelPinFactory interface
+	virtual TSharedPtr<class SGraphPin> CreatePin_Internal(class UEdGraphPin* InPin) const override;
 };
diff --git a/Engine/Plugins/Runtime/RigVM/Source/RigVM/Public/RigVMCore/RigVMAssetUserData.h b/Engine/Plugins/Runtime/RigVM/Source/RigVM/Public/RigVMCore/RigVMAssetUserData.h
new file mode 100644
--- /dev/null
+++ b/Engine/Plugins/Runtime/RigVM/Source/RigVM/Public/RigVMCore/RigVMAssetUserData.h
@@ -0,0 +1,213 @@
+// Copyright Epic Games, Inc. All Rights Reserved.
+
+#pragma once
+
+#include "CoreMinimal.h"
+#include "Engine/AssetUserData.h"
+#include "Engine/DataAsset.h"
+#include "RigVMAssetUserData.generated.h"
+
+/**
+* User data that can be attached to assets to provide namespaced data access.
+*/
+UCLASS(Abstract)
+class RIGVM_API UNameSpacedUserData : public UAssetUserData
+{
+	GENERATED_BODY()
+
+public:
+
+	/** A helper struct to represent a single user data */
+	struct RIGVM_API FUserData
+	{
+	public:
+
+		FUserData()
+		: Path()
+		, CPPType(NAME_None)
+		, Property(nullptr)
+		, Memory(nullptr)
+		{
+		}
+
+		FUserData(const FString& InPath, const FProperty* InProperty, const uint8* InMemory)
+		: Path(InPath)
+		, CPPType(NAME_None)
+		, Property(InProperty)
+		, Memory(InMemory)
+		{
+		}
+
+		FUserData(const FString& InPath, const FName& InCPPType, const FProperty* InProperty, const uint8* InMemory)
+		: Path(InPath)
+		, CPPType(InCPPType)
+		, Property(InProperty)
+		, Memory(InMemory)
+		{
+		}
+
+		// Returns true if this user data is mapped correctly
+		bool IsValid() const
+		{
+			return (!Path.IsEmpty()) && (Memory != nullptr);
+		}
+
+		operator bool() const { return IsValid(); }
+
+		operator FString() const { return Path; }
+
+		// returns the string path inside of the asset user data
+		const FString& GetPath() const { return Path; }
+
+		// returns the name of the last segment of the path
+		FString GetName() const;
+
+#if WITH_EDITOR
+
+		// returns the display name of the last segment of the path
+		FString GetDisplayName() const;
+
+#endif
+
+		// returns the CPP type of this user data (computes it from the property as needed)
+		const FName& GetCPPType() const
+		{
+			if(CPPType.IsNone() && Property != nullptr)
+			{
+				FString ExtendedCppType;
+  				FString CPPTypeString = Property->GetCPPType(&ExtendedCppType);
+  				CPPTypeString += ExtendedCppType;
+				CPPType = *CPPTypeString;
+			}
+			return CPPType;
+		}
+
+		bool IsArray() const;
+
+		bool IsArrayElement() const;
+
+		bool IsUObject() const;
+
+		// returns the property of this user data or nullptr
+		const FProperty* GetProperty() const { return Property; }
+
+		// returns the mapped memory of this user data
+		const uint8* GetMemory() const { return Memory; }
+
+		// joints two user data path strings 
+		static FString Join(const FString& InLeft, const FString& InRight)
+		{
+			check(!InLeft.IsEmpty());
+			check(!InRight.IsEmpty());
+			return InLeft + TEXT(".") + InRight;
+		}
+
+		// splits a user data string path at the first segment
+		static bool SplitAtStart(const FString& InPath, FString& OutLeft, FString& OutRight)
+		{
+			return InPath.Split(TEXT("."), &OutLeft, &OutRight, ESearchCase::IgnoreCase, ESearchDir::FromStart);
+		}
+
+		// splits a user data string path at the last segment
+		static bool SplitAtEnd(const FString& InPath, FString& OutLeft, FString& OutRight)
+		{
+			return InPath.Split(TEXT("."), &OutLeft, &OutRight, ESearchCase::IgnoreCase, ESearchDir::FromEnd);
+		}
+
+		static TArray<UStruct*> GetSuperStructs(UStruct* InStruct);
+
+	private:
+
+		FString Path;
+		mutable FName CPPType;
+		const FProperty* Property;
+		const uint8* Memory;
+	};
+
+	virtual ~UNameSpacedUserData() override
+	{
+		InvalidateCache();
+	}
+
+	virtual void Serialize(FArchive& Ar) override;
+
+	/** The namespace to use when looking up values inside of the user data. */
+	UPROPERTY(EditAnywhere, Category = General)
+	FString NameSpace;
+
+	/** Returns true if userdata exists for a given path */
+	bool ContainsUserData(const FString& InPath) const { return GetUserData(InPath) != nullptr; }
+
+	/** Returns a user data path given its string path */
+	virtual const FUserData* GetUserData(const FString& InPath, FString* OutErrorMessage = nullptr) const;
+
+	/** Retrieves the user data paths given a (optional) parent path */
+	virtual const TArray<const FUserData*>& GetUserDataArray(const FString& InParentPath = FString(), FString* OutErrorMessage = nullptr) const;
+
+	/** Retrieves the user data path given the path string */
+	const FUserData* GetParentUserData(const FString& InPath, FString* OutErrorMessage = nullptr) const
+	{
+		FString Left, Right;
+		if(FUserData::SplitAtEnd(InPath, Left, Right))
+		{
+			return GetUserData(Left, OutErrorMessage);
+		}
+		return nullptr;
+	}
+
+protected:
+
+	const FUserData* GetUserDataWithinStruct(const UStruct* InStruct, const uint8* InMemory, const FString& InPath, const FString& InPropertyName, FString* OutErrorMessage = nullptr) const;
+	const TArray<const FUserData*>& GetUserDataArrayWithinStruct(UStruct* InStruct, const uint8* InMemory, const FString& InPath, FString* OutErrorMessage = nullptr) const;
+	const TArray<const FUserData*>& GetUserDataArrayWithinArray(const FArrayProperty* InArrayProperty, const uint8* InMemory, const FString& InPath, FString* OutErrorMessage = nullptr) const;
+	static const FProperty* FindPropertyByName(const UStruct* InStruct, const FName& InName);
+
+	void InvalidateCache() const;
+	const FUserData* StoreCacheForUserData(const FUserData& InUserData) const; 
+	const TArray<const FUserData*>& StoreCacheForUserDataArray(const FString& InPath, const TArray<const FUserData*>& InUserDataArray) const;
+	virtual bool IsPropertySupported(const FProperty* InProperty, const FString& InPath, bool bCheckPropertyFlags, FString* OutErrorMessage = nullptr) const;
+	
+	static const TArray<const FUserData*> EmptyUserDatas;
+
+	mutable TMap<FString, FUserData*> CachedUserData;
+	mutable TMap<FString, TArray<const FUserData*>> CachedUserDataArray;
+	mutable FCriticalSection CacheLock;
+
+	static inline constexpr TCHAR PathNotFoundFormat[] = TEXT("UserData path '%s' could not be found.");
+	static inline constexpr TCHAR InvalidMemoryFormat[] = TEXT("UserData at path '%s' does not have any data backing it up.");
+	static inline constexpr TCHAR InvalidArrayIndexFormat[] = TEXT("Provided array index '%s' for path '%s' is invalid.");
+	static inline constexpr TCHAR OutOfBoundArrayIndexFormat[] = TEXT("Provided array index for path '%s' is out of bounds (num == %d).");
+	static inline constexpr TCHAR UnSupportedSubPathsFormat[] = TEXT("Cannot access sub paths of user data type '%s' for path '%s'.");
+};
+
+/**
+* Namespaced user data which provides access to a linked data asset
+*/
+UCLASS(BlueprintType)
+class RIGVM_API UDataAssetLink : public UNameSpacedUserData
+{
+	GENERATED_BODY()
+
+public:
+
+	/** If assigned, the data asset link will provide access to the data asset's content. */
+	UPROPERTY(EditAnywhere, BlueprintReadWrite, Category = General, BlueprintGetter = GetDataAsset, BlueprintSetter = SetDataAsset, Meta = (DisplayAfter="NameSpace"))
+	TObjectPtr<UDataAsset> DataAsset;
+
+	UFUNCTION(BlueprintGetter)
+	UDataAsset* GetDataAsset() const { return DataAsset; }
+
+	UFUNCTION(BlueprintSetter)
+	void SetDataAsset(UDataAsset* InDataAsset);
+
+	virtual const FUserData* GetUserData(const FString& InPath, FString* OutErrorMessage = nullptr) const override;
+	virtual const TArray<const FUserData*>& GetUserDataArray(const FString& InParentPath = FString(), FString* OutErrorMessage = nullptr) const override;
+
+#if WITH_EDITOR
+	virtual void PostEditChangeProperty(FPropertyChangedEvent& PropertyChangedEvent) override;
+#endif
+
+protected:
+
+	static inline constexpr TCHAR DataAssetNullFormat[] = TEXT("User data path '%s' could not be found (DataAsset not provided)");
+};
diff --git a/Engine/Plugins/Runtime/RigVM/Source/RigVMDeveloper/Private/EdGraph/RigVMEdGraph.cpp b/Engine/Plugins/Runtime/RigVM/Source/RigVMDeveloper/Private/EdGraph/RigVMEdGraph.cpp
--- a/Engine/Plugins/Runtime/RigVM/Source/RigVMDeveloper/Private/EdGraph/RigVMEdGraph.cpp
+++ b/Engine/Plugins/Runtime/RigVM/Source/RigVMDeveloper/Private/EdGraph/RigVMEdGraph.cpp
@@ -837,6 +837,34 @@ int32 URigVMEdGraph::GetInstructionIndex(const URigVMEdGraphNode* InNode, bool b
 	return INDEX_NONE;
 }
 
+#if WITH_EDITOR
+
+void URigVMEdGraph::CacheEntryNameList()
+{
+	EntryNameList.Reset();
+	EntryNameList.Add(MakeShared<FString>(FName(NAME_None).ToString()));
+
+	if(const URigVMBlueprint* Blueprint = CastChecked<URigVMBlueprint>(GetBlueprint()))
+	{
+		const TArray<FName> EntryNames = Blueprint->GetRigVMClient()->GetEntryNames();
+		for (const FName& EntryName : EntryNames)
+		{
+			EntryNameList.Add(MakeShared<FString>(EntryName.ToString()));
+		}
+	}
+}
+
+const TArray<TSharedPtr<FString>>* URigVMEdGraph::GetEntryNameList(URigVMPin* InPin) const
+{
+	if (const URigVMEdGraph* OuterGraph = Cast<URigVMEdGraph>(GetOuter()))
+	{
+		return OuterGraph->GetEntryNameList(InPin);
+	}
+	return &EntryNameList;
+}
+
+#endif
+
 UEdGraphNode* URigVMEdGraph::FindNodeForModelNodeName(const FName& InModelNodeName, const bool bCacheIfRequired)
 {
 	DECLARE_SCOPE_HIERARCHICAL_COUNTER_FUNC()
diff --git a/Engine/Plugins/Runtime/RigVM/Source/RigVMDeveloper/Private/EdGraph/RigVMEdGraphNode.cpp b/Engine/Plugins/Runtime/RigVM/Source/RigVMDeveloper/Private/EdGraph/RigVMEdGraphNode.cpp
--- a/Engine/Plugins/Runtime/RigVM/Source/RigVMDeveloper/Private/EdGraph/RigVMEdGraphNode.cpp
+++ b/Engine/Plugins/Runtime/RigVM/Source/RigVMDeveloper/Private/EdGraph/RigVMEdGraphNode.cpp
@@ -469,8 +469,8 @@ const FRigVMTemplate* URigVMEdGraphNode::GetTemplate() const
 void URigVMEdGraphNode::ClearErrorInfo()
 {
 	bHasCompilerMessage = false;
-	// SRigVMEdGraphNode only updates if the error types do not match so we have
-	// clear the error type as well, see SRigVMEdGraphNode::RefreshErrorInfo()
+	// SRigVMGraphNode only updates if the error types do not match so we have
+	// clear the error type as well, see SRigVMGraphNode::RefreshErrorInfo()
 	ErrorType = (int32)EMessageSeverity::Info + 1;
 	ErrorMsg = FString();	
 }
diff --git a/Engine/Plugins/Runtime/RigVM/Source/RigVMDeveloper/Public/EdGraph/RigVMEdGraph.h b/Engine/Plugins/Runtime/RigVM/Source/RigVMDeveloper/Public/EdGraph/RigVMEdGraph.h
--- a/Engine/Plugins/Runtime/RigVM/Source/RigVMDeveloper/Public/EdGraph/RigVMEdGraph.h
+++ b/Engine/Plugins/Runtime/RigVM/Source/RigVMDeveloper/Public/EdGraph/RigVMEdGraph.h
@@ -57,6 +57,9 @@ public:
 
 	int32 GetInstructionIndex(const URigVMEdGraphNode* InNode, bool bAsInput);
 
+	void CacheEntryNameList();
+	const TArray<TSharedPtr<FString>>* GetEntryNameList(URigVMPin* InPin = nullptr) const;
+
 	UPROPERTY()
 	FString ModelNodePath;
 
@@ -94,13 +97,14 @@ protected:
 private:
 	TMap<FName, UEdGraphNode*> ModelNodePathToEdNode;
 	mutable TWeakObjectPtr<URigVMGraph> CachedModelGraph;
+	TArray<TSharedPtr<FString>> EntryNameList;
 
 #endif
 	friend class URigVMEdGraphNode;
 	friend class URigVMBlueprint;
 	friend class UControlRigGraphNode;
 	friend class FControlRigEditor;
-	friend class SControlRigGraphNode;
+	friend class SRigVMGraphNode;
 	friend class URigVMBlueprint;
 
 	friend class URigVMEdGraphUnitNodeSpawner;
diff --git a/Engine/Plugins/Runtime/RigVM/Source/RigVMDeveloper/Public/EdGraph/RigVMEdGraphNode.h b/Engine/Plugins/Runtime/RigVM/Source/RigVMDeveloper/Public/EdGraph/RigVMEdGraphNode.h
--- a/Engine/Plugins/Runtime/RigVM/Source/RigVMDeveloper/Public/EdGraph/RigVMEdGraphNode.h
+++ b/Engine/Plugins/Runtime/RigVM/Source/RigVMDeveloper/Public/EdGraph/RigVMEdGraphNode.h
@@ -270,7 +270,7 @@ private:
 	mutable const FRigVMTemplate* CachedTemplate;
 	mutable TOptional<bool> DrawAsCompactNodeCache;
 
-	friend class SRigVMEdGraphNode;
+	friend class SRigVMGraphNode;
 	friend class FRigVMArgumentLayout;
 	friend class FRigVMEdGraphDetails;
 	friend class URigVMEdGraphTemplateNodeSpawner;
diff --git a/Engine/Plugins/Runtime/RigVM/Source/RigVMEditor/Public/EdGraph/RigVMEdGraphConnectionDrawingPolicy.h b/Engine/Plugins/Runtime/RigVM/Source/RigVMEditor/Public/EdGraph/RigVMEdGraphConnectionDrawingPolicy.h
new file mode 100644
--- /dev/null
+++ b/Engine/Plugins/Runtime/RigVM/Source/RigVMEditor/Public/EdGraph/RigVMEdGraphConnectionDrawingPolicy.h
@@ -0,0 +1,47 @@
+// Copyright Epic Games, Inc. All Rights Reserved.
+
+#pragma once
+
+#include "BlueprintConnectionDrawingPolicy.h"
+
+class URigVMEdGraphNode;
+
+class RIGVMEDITOR_API FRigVMEdGraphConnectionDrawingPolicy : public FKismetConnectionDrawingPolicy
+{
+public:
+	FRigVMEdGraphConnectionDrawingPolicy(int32 InBackLayerID, int32 InFrontLayerID, float InZoomFactor, const FSlateRect& InClippingRect, FSlateWindowElementList& InDrawElements, UEdGraph* InGraphObj)
+		: FKismetConnectionDrawingPolicy(InBackLayerID, InFrontLayerID, InZoomFactor, InClippingRect, InDrawElements, InGraphObj)
+	{
+		CastImage = FAppStyle::GetBrush( TEXT("GraphEditor.Cast_16x") );
+		ArrowImage = nullptr;
+		ArrowRadius = CastImage->ImageSize * InZoomFactor * 0.5f;
+		MidpointImage = nullptr;
+		MidpointRadius = CastImage->ImageSize * InZoomFactor * 0.5f;
+	}
+
+	virtual void SetIncompatiblePinDrawState(const TSharedPtr<SGraphPin>& StartPin, const TSet< TSharedRef<SWidget> >& VisiblePins) override;
+	virtual void ResetIncompatiblePinDrawState(const TSet< TSharedRef<SWidget> >& VisiblePins) override;
+	virtual void BuildPinToPinWidgetMap(TMap<TSharedRef<SWidget>, FArrangedWidget>& InPinGeometries) override;
+	virtual void DrawPinGeometries(TMap<TSharedRef<SWidget>, FArrangedWidget>& InPinGeometries, FArrangedChildren& ArrangedNodes) override;
+	virtual void DetermineLinkGeometry(
+	FArrangedChildren& ArrangedNodes, 
+	TSharedRef<SWidget>& OutputPinWidget,
+	UEdGraphPin* OutputPin,
+	UEdGraphPin* InputPin,
+	/*out*/ FArrangedWidget*& StartWidgetGeometry,
+	/*out*/ FArrangedWidget*& EndWidgetGeometry) override;
+	virtual void DetermineWiringStyle(UEdGraphPin* OutputPin, UEdGraphPin* InputPin, /*inout*/ FConnectionParams& Params) override;
+	virtual bool UseDrawStateCaching() const override { return true; }
+
+private:
+	// Each time a reroute node is encountered, input geometry is compared to output geometry to see if the pins on the reroute node need to be reversed
+	TMap<URigVMEdGraphNode*, bool> RerouteNodeToReversedDirectionMap;
+
+	bool UseLowDetailConnections() const { return ZoomFactor <= 0.175f; /* zoom level -9 */ }
+
+	bool ShouldChangeTangentForReouteControlPoint(URigVMEdGraphNode* Node);
+	// Average of the positions of all pins connected to InPin
+	bool GetAverageConnectedPositionForPin(UEdGraphPin* InPin, FVector2D& OutPos) const;
+
+	const FSlateBrush* CastImage;
+};
diff --git a/Engine/Plugins/Runtime/RigVM/Source/RigVMEditor/Public/EdGraph/RigVMEdGraphPanelNodeFactory.h b/Engine/Plugins/Runtime/RigVM/Source/RigVMEditor/Public/EdGraph/RigVMEdGraphPanelNodeFactory.h
new file mode 100644
--- /dev/null
+++ b/Engine/Plugins/Runtime/RigVM/Source/RigVMEditor/Public/EdGraph/RigVMEdGraphPanelNodeFactory.h
@@ -0,0 +1,10 @@
+// Copyright Epic Games, Inc. All Rights Reserved.
+
+#pragma once
+
+#include "EdGraphUtilities.h"
+
+class RIGVMEDITOR_API FRigVMEdGraphPanelNodeFactory : public FGraphPanelNodeFactory
+{
+	virtual TSharedPtr<class SGraphNode> CreateNode(UEdGraphNode* Node) const override;
+};
diff --git a/Engine/Plugins/Runtime/RigVM/Source/RigVMEditor/Public/EdGraph/RigVMEdGraphPanelPinFactory.h b/Engine/Plugins/Runtime/RigVM/Source/RigVMEditor/Public/EdGraph/RigVMEdGraphPanelPinFactory.h
new file mode 100644
--- /dev/null
+++ b/Engine/Plugins/Runtime/RigVM/Source/RigVMEditor/Public/EdGraph/RigVMEdGraphPanelPinFactory.h
@@ -0,0 +1,16 @@
+// Copyright Epic Games, Inc. All Rights Reserved.
+
+#pragma once
+
+#include "CoreMinimal.h"
+#include "EdGraphUtilities.h"
+
+class RIGVMEDITOR_API FRigVMEdGraphPanelPinFactory : public FGraphPanelPinFactory
+{
+public:
+	// FGraphPanelPinFactory interface
+	virtual TSharedPtr<class SGraphPin> CreatePin(class UEdGraphPin* InPin) const override;
+
+
+	virtual TSharedPtr<class SGraphPin> CreatePin_Internal(class UEdGraphPin* InPin) const;
+};
diff --git a/Engine/Plugins/Runtime/RigVM/Source/RigVMEditor/Public/Widgets/SRigVMGraphChangePinType.h b/Engine/Plugins/Runtime/RigVM/Source/RigVMEditor/Public/Widgets/SRigVMGraphChangePinType.h
new file mode 100644
--- /dev/null
+++ b/Engine/Plugins/Runtime/RigVM/Source/RigVMEditor/Public/Widgets/SRigVMGraphChangePinType.h
@@ -0,0 +1,50 @@
+// Copyright Epic Games, Inc. All Rights Reserved.
+
+#pragma once
+
+#include "CoreMinimal.h"
+#include "Widgets/DeclarativeSyntaxSupport.h"
+#include "Widgets/SWidget.h"
+#include "SGraphPin.h"
+#include "RigVMModel/RigVMPin.h"
+#include "IPropertyAccessEditor.h"
+#include "RigVMBlueprint.h"
+
+DECLARE_DELEGATE_OneParam(FOnTypeSelected, TRigVMTypeIndex);
+
+class RIGVMEDITOR_API SRigVMGraphChangePinType : public SCompoundWidget
+{
+public:
+
+	SLATE_BEGIN_ARGS(SRigVMGraphChangePinType)
+	: _Types()
+    , _Blueprint(nullptr)
+	, _OnTypeSelected(nullptr)
+	{}
+
+		SLATE_ARGUMENT(TArray<TRigVMTypeIndex>, Types)
+		SLATE_ARGUMENT(URigVMBlueprint*, Blueprint)
+		SLATE_EVENT(FOnTypeSelected, OnTypeSelected)
+
+	SLATE_END_ARGS()
+
+	void Construct(const FArguments& InArgs);
+
+protected:
+
+	static FText GetBindingText(const FRigVMTemplateArgumentType& InType);
+	FText GetBindingText(URigVMPin* ModelPin) const;
+	FText GetBindingText() const;
+	const FSlateBrush* GetBindingImage() const;
+	FLinearColor GetBindingColor() const;
+	bool OnCanBindProperty(FProperty* InProperty) const;
+	bool OnCanBindToClass(UClass* InClass) const;
+	void OnAddBinding(FName InPropertyName, const TArray<FBindingChainElement>& InBindingChain);
+	void FillPinTypeMenu( FMenuBuilder& MenuBuilder );
+	void HandlePinTypeChanged(FRigVMTemplateArgumentType InType);
+
+	TArray<TRigVMTypeIndex> Types;
+	URigVMBlueprint* Blueprint;
+	FOnTypeSelected OnTypeSelected;
+	FPropertyBindingWidgetArgs BindingArgs;
+};
diff --git a/Engine/Plugins/Runtime/RigVM/Source/RigVMEditor/Public/Widgets/SRigVMGraphNode.h b/Engine/Plugins/Runtime/RigVM/Source/RigVMEditor/Public/Widgets/SRigVMGraphNode.h
new file mode 100644
--- /dev/null
+++ b/Engine/Plugins/Runtime/RigVM/Source/RigVMEditor/Public/Widgets/SRigVMGraphNode.h
@@ -0,0 +1,162 @@
+// Copyright Epic Games, Inc. All Rights Reserved.
+
+#pragma once
+
+#include "CoreMinimal.h"
+#include "Widgets/DeclarativeSyntaxSupport.h"
+#include "SGraphNode.h"
+#include "Widgets/Images/SImage.h"
+#include "Widgets/Text/STextBlock.h"
+#include "RigVMModel/RigVMPin.h"
+#include "RigVMBlueprint.h"
+
+class URigVMEdGraphNode;
+class STableViewBase;
+class SOverlay;
+class SGraphPin;
+class UEdGraphPin;
+
+class RIGVMEDITOR_API SRigVMGraphNode : public SGraphNode
+{
+public:
+	SLATE_BEGIN_ARGS(SRigVMGraphNode)
+		: _GraphNodeObj(nullptr)
+		{}
+
+	SLATE_ARGUMENT(URigVMEdGraphNode*, GraphNodeObj)
+
+	SLATE_END_ARGS()
+
+	void Construct(const FArguments& InArgs);
+
+	// SGraphNode interface
+	virtual TSharedRef<SWidget> CreateTitleWidget(TSharedPtr<SNodeTitle> InNodeTitle) override;
+
+	virtual void EndUserInteraction() const override;
+	virtual void MoveTo(const FVector2D& NewPosition, FNodeSet& NodeFilter, bool bMarkDirty = true) override;
+	virtual void AddPin( const TSharedRef<SGraphPin>& PinToAdd ) override;
+	virtual void CreateStandardPinWidget(UEdGraphPin* CurPin) override;
+	virtual void SetDefaultTitleAreaWidget(TSharedRef<SOverlay> DefaultTitleAreaWidget) override
+	{
+		TitleAreaWidget = DefaultTitleAreaWidget;
+	}
+	virtual const FSlateBrush * GetNodeBodyBrush() const override;
+	virtual FReply OnMouseButtonDown(const FGeometry& MyGeometry, const FPointerEvent& MouseEvent) override;
+	virtual FReply OnMouseButtonDoubleClick(const FGeometry& InMyGeometry, const FPointerEvent& InMouseEvent) override;
+
+	virtual TSharedRef<SWidget> CreateNodeContentArea() override;
+	virtual void GetOverlayBrushes(bool bSelected, const FVector2D WidgetSize, TArray<FOverlayBrushInfo>& Brushes) const override;
+	virtual void GetNodeInfoPopups(FNodeInfoContext* Context, TArray<FGraphInformationPopupInfo>& Popups) const override;
+	virtual TArray<FOverlayWidgetInfo> GetOverlayWidgets(bool bSelected, const FVector2D& WidgetSize) const override;
+
+	virtual void RefreshErrorInfo() override;
+	virtual void Tick(const FGeometry& AllottedGeometry, const double InCurrentTime, const float InDeltaTime) override;
+	virtual bool IsHidingPinWidgets() const override { return UseLowDetailNodeContent(); }
+	virtual bool UseLowDetailPinNames() const override;
+	virtual void UpdateGraphNode() override;
+	void UpdateStandardNode();
+	void UpdateCompactNode();
+
+	void CreateAddPinButton();
+
+	/** Callback function executed when Add pin button is clicked */
+	virtual FReply OnAddPin() override;
+
+private:
+
+	bool UseLowDetailNodeContent() const;
+
+	FVector2D GetLowDetailDesiredSize() const;
+
+	EVisibility GetTitleVisibility() const;
+	EVisibility GetArrayPlusButtonVisibility(URigVMPin* InModelPin) const;
+
+	FText GetPinLabel(TWeakPtr<SGraphPin> GraphPin) const;
+
+	FSlateColor GetPinTextColor(TWeakPtr<SGraphPin> GraphPin) const;
+	FSlateColor GetVariableLabelTextColor(TWeakObjectPtr<URigVMFunctionReferenceNode> FunctionReferenceNode, FName InVariableName) const;
+	FText GetVariableLabelTooltipText(TWeakObjectPtr<URigVMBlueprint> InBlueprint, FName InVariableName) const;
+
+	FReply HandleAddArrayElement(FString InModelPinPath);
+
+	void HandleNodeTitleDirtied();
+	void HandleNodePinsChanged();
+	void HandleNodeBeginRemoval();
+
+	FText GetInstructionCountText() const;
+	FText GetInstructionDurationText() const;
+
+private:
+
+	int32 GetNodeTopologyVersion() const;
+	EVisibility GetPinVisibility(int32 InPinInfoIndex, bool bAskingForSubPin) const;
+	const FSlateBrush * GetExpanderImage(int32 InPinInfoIndex, bool bLeft, bool bHovered) const;
+	FReply OnExpanderArrowClicked(int32 InPinInfoIndex);
+	void HandleModifiedEvent(ERigVMGraphNotifType InNotifType, URigVMGraph* InGraph, UObject* InSubject);
+	void UpdatePinTreeView();
+
+	/** Cached widget title area */
+	TSharedPtr<SOverlay> TitleAreaWidget;
+
+	int32 NodeErrorType;
+
+	TSharedPtr<SImage> VisualDebugIndicatorWidget;
+	TSharedPtr<STextBlock> InstructionCountTextBlockWidget;
+	TSharedPtr<STextBlock> InstructionDurationTextBlockWidget;
+
+	static const FSlateBrush* CachedImg_CR_Pin_Connected;
+	static const FSlateBrush* CachedImg_CR_Pin_Disconnected;
+
+	/** Cache the node title so we can invalidate it */
+	TSharedPtr<SNodeTitle> NodeTitle;
+
+	TWeakObjectPtr<URigVMBlueprint> Blueprint;
+
+	FVector2D LastHighDetailSize;
+
+	struct FPinInfo
+	{
+		int32 Index;
+		int32 ParentIndex;
+		bool bHasChildren;
+		bool bHideInputWidget;
+		bool bIsContainer;
+		int32 Depth;
+		FString ModelPinPath;
+		TSharedPtr<SGraphPin> InputPinWidget;
+		TSharedPtr<SGraphPin> OutputPinWidget;
+		bool bExpanded;
+		bool bAutoHeight;
+		bool bShowOnlySubPins;
+	};
+
+	
+	/**
+	 * Simple tagging metadata
+	 */
+	class FPinInfoMetaData : public ISlateMetaData
+	{
+	public:
+		SLATE_METADATA_TYPE(FPinInfoMetaData, ISlateMetaData)
+
+		FPinInfoMetaData(const FString& InCPPType, const FString& InBoundVariableName)
+		: CPPType(InCPPType)
+		, BoundVariableName(InBoundVariableName)
+		{}
+
+		FString CPPType;
+		FString BoundVariableName;
+	};
+	
+	TArray<FPinInfo> PinInfos;
+	TWeakObjectPtr<URigVMNode> ModelNode;
+
+	// Pins to keep after calling HandleNodePinsChanged. We recycle these pins in
+	// CreateStandardPinWidget.
+	TMap<const UEdGraphPin *, TSharedRef<SGraphPin>> PinsToKeep;
+
+	// Delayed pin deletion. To deal with the fact that pin deletion cannot occur until we
+	// have re-generated the pin list. SRigVMGraphNode has already relinquished them
+	// but we still have a pointer to them in our pin widget.
+	TSet<UEdGraphPin *> PinsToDelete;
+};
diff --git a/Engine/Plugins/Runtime/RigVM/Source/RigVMEditor/Public/Widgets/SRigVMGraphNodeComment.h b/Engine/Plugins/Runtime/RigVM/Source/RigVMEditor/Public/Widgets/SRigVMGraphNodeComment.h
new file mode 100644
--- /dev/null
+++ b/Engine/Plugins/Runtime/RigVM/Source/RigVMEditor/Public/Widgets/SRigVMGraphNodeComment.h
@@ -0,0 +1,33 @@
+// Copyright Epic Games, Inc. All Rights Reserved.
+
+#pragma once
+
+#include "CoreMinimal.h"
+#include "SGraphNodeComment.h"
+
+class UEdGraphNode_Comment;
+
+class RIGVMEDITOR_API SRigVMGraphNodeComment : public SGraphNodeComment
+{
+public:
+
+	SRigVMGraphNodeComment();
+
+	virtual FReply OnMouseButtonUp(const FGeometry& MyGeometry, const FPointerEvent& MouseEvent) override;
+	virtual void EndUserInteraction() const override;
+	virtual void Tick(const FGeometry& AllottedGeometry, const double InCurrentTime, const float InDeltaTime) override;
+	virtual void MoveTo(const FVector2D& NewPosition, FNodeSet& NodeFilter, bool bMarkDirty = true) override;
+
+protected:
+
+	virtual bool IsNodeUnderComment(UEdGraphNode_Comment* InCommentNode, const TSharedRef<SGraphNode> InNodeWidget) const override;
+	bool IsNodeUnderComment(UEdGraphNode* InNode) const;
+	bool IsNodeUnderComment(const FVector2D& InNodePosition) const;
+	bool IsNodeUnderComment(const FVector2D& InCommentPosition, const FVector2D& InNodePosition) const;
+
+private:
+
+	FLinearColor CachedNodeCommentColor;
+
+	int8 CachedColorBubble;
+};
diff --git a/Engine/Plugins/Runtime/RigVM/Source/RigVMEditor/Public/Widgets/SRigVMGraphNodeKnot.h b/Engine/Plugins/Runtime/RigVM/Source/RigVMEditor/Public/Widgets/SRigVMGraphNodeKnot.h
new file mode 100644
--- /dev/null
+++ b/Engine/Plugins/Runtime/RigVM/Source/RigVMEditor/Public/Widgets/SRigVMGraphNodeKnot.h
@@ -0,0 +1,21 @@
+// Copyright Epic Games, Inc. All Rights Reserved.
+
+#pragma once
+
+#include "CoreMinimal.h"
+#include "Widgets/DeclarativeSyntaxSupport.h"
+#include "SGraphNodeKnot.h"
+
+class RIGVMEDITOR_API SRigVMGraphNodeKnot : public SGraphNodeKnot
+{
+	SLATE_BEGIN_ARGS(SRigVMGraphNodeKnot) {} 
+	SLATE_END_ARGS()
+
+	void Construct(const FArguments& InArgs, UEdGraphNode* InKnot);
+
+	// SGraphNode interface
+	virtual void EndUserInteraction() const override;
+	virtual void MoveTo(const FVector2D& NewPosition, FNodeSet& NodeFilter, bool bMarkDirty = true) override;
+
+	void HandleNodeBeginRemoval();
+};
\ No newline at end of file
diff --git a/Engine/Plugins/Runtime/RigVM/Source/RigVMEditor/Public/Widgets/SRigVMGraphPinCurveFloat.h b/Engine/Plugins/Runtime/RigVM/Source/RigVMEditor/Public/Widgets/SRigVMGraphPinCurveFloat.h
new file mode 100644
--- /dev/null
+++ b/Engine/Plugins/Runtime/RigVM/Source/RigVMEditor/Public/Widgets/SRigVMGraphPinCurveFloat.h
@@ -0,0 +1,39 @@
+// Copyright Epic Games, Inc. All Rights Reserved.
+
+#pragma once
+
+#include "CoreMinimal.h"
+#include "Widgets/DeclarativeSyntaxSupport.h"
+#include "Widgets/SWidget.h"
+#include "SCurveEditor.h"
+#include "Curves/CurveFloat.h"
+#include "SGraphPin.h"
+
+class RIGVMEDITOR_API SRigVMGraphPinCurveFloat : public SGraphPin, public FCurveOwnerInterface
+{
+public:
+	SLATE_BEGIN_ARGS(SRigVMGraphPinCurveFloat) {}
+	SLATE_END_ARGS()
+
+	void Construct(const FArguments& InArgs, UEdGraphPin* InGraphPinObj);
+
+	/** FCurveOwnerInterface interface */
+	virtual TArray<FRichCurveEditInfoConst> GetCurves() const override;
+	virtual TArray<FRichCurveEditInfo> GetCurves() override;
+	virtual void ModifyOwner() override;
+	virtual TArray<const UObject*> GetOwners() const override;
+	virtual void MakeTransactional() override;
+	virtual void OnCurveChanged(const TArray<FRichCurveEditInfo>& ChangedCurveEditInfos) override;
+	virtual bool IsValidCurve(FRichCurveEditInfo CurveInfo) override;
+
+protected:
+
+	//~ Begin SGraphPin Interface
+	virtual TSharedRef<SWidget>	GetDefaultValueWidget() override;
+	//~ End SGraphPin Interface
+
+	FRuntimeFloatCurve& UpdateAndGetCurve();
+
+	TSharedPtr<SCurveEditor> CurveEditor;
+	FRuntimeFloatCurve Curve;
+};
diff --git a/Engine/Plugins/Runtime/RigVM/Source/RigVMEditor/Public/Widgets/SRigVMGraphPinNameList.h b/Engine/Plugins/Runtime/RigVM/Source/RigVMEditor/Public/Widgets/SRigVMGraphPinNameList.h
new file mode 100644
--- /dev/null
+++ b/Engine/Plugins/Runtime/RigVM/Source/RigVMEditor/Public/Widgets/SRigVMGraphPinNameList.h
@@ -0,0 +1,80 @@
+// Copyright Epic Games, Inc. All Rights Reserved.
+
+#pragma once
+
+#include "CoreMinimal.h"
+#include "Widgets/DeclarativeSyntaxSupport.h"
+#include "Widgets/SWidget.h"
+#include "SGraphPin.h"
+#include "RigVMModel/RigVMPin.h"
+#include "Widgets/SRigVMGraphPinNameListValueWidget.h"
+
+class URigVMEdGraph;
+
+class RIGVMEDITOR_API SRigVMGraphPinNameList : public SGraphPin
+{
+public:
+
+	DECLARE_DELEGATE_RetVal_OneParam( const TArray<TSharedPtr<FString>>*, FOnGetNameListContent, URigVMPin*);
+	DECLARE_DELEGATE_RetVal( const TArray<TSharedPtr<FString>>, FOnGetNameFromSelection);
+	DECLARE_DELEGATE_RetVal_ThreeParams( FReply, FOnGetSelectedClicked, URigVMEdGraph*, URigVMPin*, FString);
+	DECLARE_DELEGATE_RetVal_ThreeParams( FReply, FOnBrowseClicked, URigVMEdGraph*, URigVMPin*, FString);
+
+	SLATE_BEGIN_ARGS(SRigVMGraphPinNameList)
+		: _MarkupInvalidItems(true)
+		, _EnableNameListCache(true)
+		, _SearchHintText(NSLOCTEXT("SRigVMGraphPinNameList", "Search", "Search"))
+		, _AllowUserProvidedText(false)
+	{}
+
+		SLATE_ARGUMENT(URigVMPin*, ModelPin)
+		SLATE_EVENT(FOnGetNameListContent, OnGetNameListContent)
+		SLATE_EVENT(FOnGetNameListContent, OnGetNameListContentForValidation)
+		SLATE_EVENT(FOnGetNameFromSelection, OnGetNameFromSelection)
+		SLATE_EVENT(FOnGetSelectedClicked, OnGetSelectedClicked)
+		SLATE_EVENT(FOnBrowseClicked, OnBrowseClicked)
+		SLATE_ARGUMENT(bool, MarkupInvalidItems)
+		SLATE_ARGUMENT(bool, EnableNameListCache)
+		SLATE_ARGUMENT(FText, SearchHintText)
+		SLATE_ARGUMENT(bool, AllowUserProvidedText)
+
+	SLATE_END_ARGS()
+
+	void Construct(const FArguments& InArgs, UEdGraphPin* InGraphPinObj);
+
+protected:
+	//~ Begin SGraphPin Interface
+	virtual TSharedRef<SWidget>	GetDefaultValueWidget() override;
+	//~ End SGraphPin Interface
+
+	const TArray<TSharedPtr<FString>>* GetNameList(bool bForContent = true) const;
+	FText GetNameListText() const;
+	FSlateColor GetNameColor() const;
+	virtual void SetNameListText(const FText& NewTypeInValue, ETextCommit::Type CommitInfo);
+	void UpdateNameLists(bool bUpdateCurrent = true, bool bUpdateValidation = true);
+
+	TSharedRef<SWidget> MakeNameListItemWidget(TSharedPtr<FString> InItem);
+	void OnNameListChanged(TSharedPtr<FString> NewSelection, ESelectInfo::Type SelectInfo);
+	void OnNameListComboBox();
+
+	FOnGetNameListContent OnGetNameListContent;
+	FOnGetNameListContent OnGetNameListContentForValidation;
+	URigVMPin* ModelPin;
+	TSharedPtr<SRigVMGraphPinNameListValueWidget> NameListComboBox;
+	TArray<TSharedPtr<FString>> EmptyList;
+	const TArray<TSharedPtr<FString>>* CurrentList;
+	const TArray<TSharedPtr<FString>>* ValidationList;
+	bool bMarkupInvalidItems;
+	bool EnableNameListCache;
+	FText SearchHintText;
+	bool AllowUserProvidedText;
+
+	/** Helper buttons. */
+	FSlateColor OnGetWidgetForeground() const;
+	FSlateColor OnGetWidgetBackground() const;
+	FReply HandleGetSelectedClicked();
+	FReply HandleBrowseClicked();
+	FOnGetNameFromSelection OnGetNameFromSelection;
+	FOnGetSelectedClicked OnGetSelectedClicked;
+	FOnBrowseClicked OnBrowseClicked;
+};
