Addresses regression in curve entry passthrough for ControlRig anim node
#fix reset all rig-hierarchy curve entry value-set flags before propagating incoming curve values. Only transfer out curve entries which have been set during control-rig eval

diff --git a/Engine/Plugins/Animation/ControlRig/Source/ControlRig/Private/AnimNode_ControlRigBase.cpp b/Engine/Plugins/Animation/ControlRig/Source/ControlRig/Private/AnimNode_ControlRigBase.cpp
--- a/Engine/Plugins/Animation/ControlRig/Source/ControlRig/Private/AnimNode_ControlRigBase.cpp
+++ b/Engine/Plugins/Animation/ControlRig/Source/ControlRig/Private/AnimNode_ControlRigBase.cpp
@@ -83,13 +83,6 @@ void FAnimNode_ControlRigBase::Initialize_AnyThread(const FAnimationInitializeCo
 		ControlRig->RequestInit();
 		bControlRigRequiresInitialization = true;
 		LastBonesSerialNumberForCacheBones = 0;
-
-		// Cache curves
-		CachedBulkCurves.Empty();
-		for(FRigCurveElement* CurveElement : ControlRig->GetHierarchy()->GetCurves())
-		{
-			CachedBulkCurves.AddIndexed(CurveElement->GetName());
-		}
 	}
 	else
 	{
@@ -255,12 +248,11 @@ void FAnimNode_ControlRigBase::UpdateInput(UControlRig* ControlRig, const FPoseC
 
 	if (InputSettings.bUpdateCurves && bTransferInputCurves)
 	{
-		// TODO: Do we need to 'unset' old values manually here? The previous code could do this because it knows the
-		// global set of curves, but we dont know that any more
-		UE::Anim::FCurveUtils::BulkGet(InOutput.Curve, CachedBulkCurves, [Hierarchy](const UE::Anim::FNamedIndexElement& InBulkElement, float InValue)
+		Hierarchy->UnsetCurveValues();
+		InOutput.Curve.ForEachElement([Hierarchy](const UE::Anim::FCurveElement& InCurveElement)
 		{
-			const FRigElementKey Key(InBulkElement.Name, ERigElementType::Curve);
-			Hierarchy->SetCurveValue(Key, InValue);
+			const FRigElementKey Key(InCurveElement.Name, ERigElementType::Curve);
+			Hierarchy->SetCurveValue(Key, InCurveElement.Value);
 		});
 	}
 
@@ -376,16 +368,25 @@ void FAnimNode_ControlRigBase::UpdateOutput(UControlRig* ControlRig, FPoseContex
 
 	if (OutputSettings.bUpdateCurves)
 	{
-		TArray<FRigCurveElement*> Curves = Hierarchy->GetCurves();
-		UE::Anim::FCurveUtils::BulkSet(InOutput.Curve, CachedBulkCurves, [Hierarchy, &Curves](const UE::Anim::FNamedIndexElement& InBulkElement)
+		const TArray<FRigCurveElement*> Curves = Hierarchy->GetCurves();
+		auto GetNameFromIndex = [&Curves](int32 InCurveIndex)
 		{
-			if (Curves.IsValidIndex(InBulkElement.Index))
-			{
-				FRigCurveElement* CurveElement = Curves[InBulkElement.Index];
-				return Hierarchy->GetCurveValue(CurveElement);
-			}
-			return 0.f;
-		});
+			return Curves[InCurveIndex]->GetName();
+		};
+
+		auto GetValueFromIndex = [&Curves](int32 InCurveIndex)
+		{
+			 return Curves[InCurveIndex]->Value;
+		};
+		
+		auto GetIsValidCurveFromIndex = [&Curves](int32 InCurveIndex)
+        {
+            return Curves[InCurveIndex]->bIsValueSet;
+        };
+
+		FBlendedCurve ControlRigCurves;
+		UE::Anim::FCurveUtils::BuildUnsorted(ControlRigCurves, Curves.Num(), GetNameFromIndex, GetValueFromIndex, GetIsValidCurveFromIndex);
+		InOutput.Curve.Combine(ControlRigCurves);
 	}
 
 #if WITH_EDITOR
diff --git a/Engine/Plugins/Animation/ControlRig/Source/ControlRig/Private/Rigs/RigHierarchy.cpp b/Engine/Plugins/Animation/ControlRig/Source/ControlRig/Private/Rigs/RigHierarchy.cpp
--- a/Engine/Plugins/Animation/ControlRig/Source/ControlRig/Private/Rigs/RigHierarchy.cpp
+++ b/Engine/Plugins/Animation/ControlRig/Source/ControlRig/Private/Rigs/RigHierarchy.cpp
@@ -814,6 +814,18 @@ void URigHierarchy::ResetCurveValues()
 	}
 }
 
+void URigHierarchy::UnsetCurveValues(bool bSetupUndo)
+{
+	LLM_SCOPE_BYNAME(TEXT("Animation/ControlRig"));
+	for(int32 ElementIndex=0; ElementIndex<Elements.Num(); ElementIndex++)
+	{
+		if(FRigCurveElement* CurveElement = Cast<FRigCurveElement>(Elements[ElementIndex]))
+		{
+			UnsetCurveValue(CurveElement, bSetupUndo);
+		}
+	}
+}
+
 int32 URigHierarchy::Num(ERigElementType InElementType) const
 {
 	return ElementsPerType[RigElementTypeToFlatIndex(InElementType)].Num();
diff --git a/Engine/Plugins/Animation/ControlRig/Source/ControlRig/Public/AnimNode_ControlRigBase.h b/Engine/Plugins/Animation/ControlRig/Source/ControlRig/Public/AnimNode_ControlRigBase.h
--- a/Engine/Plugins/Animation/ControlRig/Source/ControlRig/Public/AnimNode_ControlRigBase.h
+++ b/Engine/Plugins/Animation/ControlRig/Source/ControlRig/Public/AnimNode_ControlRigBase.h
@@ -142,9 +142,6 @@ protected:
 	TMap<FName, FName> ControlRigCurveMappingByName;
 	
 	TMap<FName, int32> InputToControlIndex;
-
-	// Cached bulk curves for I/O
-	UE::Anim::FBulkCurves CachedBulkCurves;
 	
 	/** Node Mapping Container */
 	UPROPERTY(transient)
diff --git a/Engine/Plugins/Animation/ControlRig/Source/ControlRig/Public/Rigs/RigHierarchy.h b/Engine/Plugins/Animation/ControlRig/Source/ControlRig/Public/Rigs/RigHierarchy.h
--- a/Engine/Plugins/Animation/ControlRig/Source/ControlRig/Public/Rigs/RigHierarchy.h
+++ b/Engine/Plugins/Animation/ControlRig/Source/ControlRig/Public/Rigs/RigHierarchy.h
@@ -199,6 +199,11 @@ public:
 	UFUNCTION(BlueprintCallable, Category = URigHierarchy)
     void ResetCurveValues();
 
+	/**
+	 * Resets all curves to be unset (defaulting to 0.0)
+	 */
+	void UnsetCurveValues(bool bSetupUndo = false);
+
 	/**
 	 * Returns the number of elements in the Hierarchy.
 	 * @return The number of elements in the Hierarchy
diff --git a/Engine/Source/Runtime/Engine/Public/Animation/AnimCurveUtils.h b/Engine/Source/Runtime/Engine/Public/Animation/AnimCurveUtils.h
--- a/Engine/Source/Runtime/Engine/Public/Animation/AnimCurveUtils.h
+++ b/Engine/Source/Runtime/Engine/Public/Animation/AnimCurveUtils.h
@@ -143,6 +143,24 @@ private:
 		OutCurve.CheckDuplicates();
 	}
 
+	// Helper function for building curves
+	template<typename NamePredicateType, typename ValuePredicateType, typename ValidityPredicateType, typename CurveAllocatorType, typename CurveElementType>
+	FORCEINLINE_DEBUGGABLE static void BuildLinearUnfiltered(TBaseBlendedCurve<CurveAllocatorType, CurveElementType>& OutCurve, int32 InNumElements, const NamePredicateType& InNamePredicate, const ValuePredicateType& InValuePredicate, const ValidityPredicateType& InValidityPredicate)
+	{
+		OutCurve.Empty();
+		OutCurve.Reserve(InNumElements);
+
+		for(int32 ElementIndex = 0; ElementIndex < InNumElements; ++ElementIndex)
+		{
+			if (InValidityPredicate(ElementIndex))
+			{
+				OutCurve.Elements.Emplace(InNamePredicate(ElementIndex), InValuePredicate(ElementIndex));				
+			}
+		}
+		
+		OutCurve.CheckDuplicates();
+	}
+
 public:
 	// Helper function for building curves, applying any filtering
 	// Assumes (and enforces in debug builds) elements being built are in FName sorted order
@@ -301,6 +319,14 @@ public:
 		}
 	}
 
+	// Helper function for building curves, applying filtering through InValidityPredicate
+	template<typename NamePredicateType, typename ValuePredicateType, typename ValidityPredicateType, typename CurveAllocatorType, typename CurveElementType>
+	static void BuildUnsorted(TBaseBlendedCurve<CurveAllocatorType, CurveElementType>& OutCurve, int32 InNumElements, const NamePredicateType& InNamePredicate, const ValuePredicateType& InValuePredicate, const ValidityPredicateType& InValidityPredicate)
+	{
+		CURVE_PROFILE_CYCLE_COUNTER(FCurveUtils_BuildUnsorted);
+		BuildLinearUnfiltered(OutCurve, InNumElements, InNamePredicate, InValuePredicate, InValidityPredicate);
+	}
+
 	/**
 	 * Remove any curves in InOutCurve that are filtered by InFilter
 	 */
