RigVM: Move compiler context to RigVMDeveloper

diff --git a/Engine/Plugins/Animation/ControlRig/Source/ControlRigDeveloper/Private/ControlRigBlueprintCompiler.cpp b/Engine/Plugins/Animation/ControlRig/Source/ControlRigDeveloper/Private/ControlRigBlueprintCompiler.cpp
--- a/Engine/Plugins/Animation/ControlRig/Source/ControlRigDeveloper/Private/ControlRigBlueprintCompiler.cpp
+++ b/Engine/Plugins/Animation/ControlRig/Source/ControlRigDeveloper/Private/ControlRigBlueprintCompiler.cpp
@@ -2,20 +2,7 @@
 
 #include "ControlRigBlueprintCompiler.h"
 #include "ControlRig.h"
-#include "KismetCompiler.h"
 #include "ControlRigBlueprint.h"
-#include "Units/RigUnit.h"
-#include "Units/Execution/RigUnit_BeginExecution.h"
-#include "Graph/ControlRigGraph.h"
-#include "Graph/ControlRigGraphNode.h"
-#include "Graph/ControlRigGraphSchema.h"
-#include "RigVMBlueprintGeneratedClass.h"
-#include "Kismet2/KismetReinstanceUtilities.h"
-#include "Widgets/Notifications/SNotificationList.h"
-#include "Framework/Notifications/NotificationManager.h"
-
-DEFINE_LOG_CATEGORY_STATIC(LogControlRigCompiler, Log, All);
-#define LOCTEXT_NAMESPACE "ControlRigBlueprintCompiler"
 
 bool FControlRigBlueprintCompiler::CanCompile(const UBlueprint* Blueprint)
 {
@@ -37,55 +24,11 @@ void FControlRigBlueprintCompiler::Compile(UBlueprint* Blueprint, const FKismetC
 	Compiler.Compile();
 }
 
-void FControlRigBlueprintCompilerContext::MarkCompilationFailed(const FString& Message)
-{
-	DECLARE_SCOPE_HIERARCHICAL_COUNTER_FUNC()
-
-	UControlRigBlueprint* ControlRigBlueprint = Cast<UControlRigBlueprint>(Blueprint);
-	if (ControlRigBlueprint)
-	{
-		Blueprint->Status = BS_Error;
-		Blueprint->MarkPackageDirty();
-		UE_LOG(LogControlRigCompiler, Error, TEXT("%s"), *Message);
-		MessageLog.Error(*Message);
-
-#if WITH_EDITOR
-		FNotificationInfo Info(FText::FromString(Message));
-		Info.Image = FCoreStyle::Get().GetBrush(TEXT("MessageLog.Error"));
-		Info.bFireAndForget = true;
-		Info.FadeOutDuration = 5.0f;
-		Info.ExpireDuration = 5.0f;
-		TSharedPtr<SNotificationItem> NotificationPtr = FSlateNotificationManager::Get().AddNotification(Info);
-		NotificationPtr->SetCompletionState(SNotificationItem::CS_Success);
-#endif
-	}
-}
-
-void FControlRigBlueprintCompilerContext::PostCompile()
-{
-	DECLARE_SCOPE_HIERARCHICAL_COUNTER_FUNC()
-
-	UControlRigBlueprint* ControlRigBlueprint = Cast<UControlRigBlueprint>(Blueprint);
-	if (ControlRigBlueprint)
-	{
-		if (!ControlRigBlueprint->bIsRegeneratingOnLoad)
-		{
-			// todo: can we find out somehow if we are cooking?
-			ControlRigBlueprint->RecompileVM();
-		}
-	}
-
-	{
-		DECLARE_SCOPE_HIERARCHICAL_COUNTER(FKismetCompilerContext::PostCompile)
-		FKismetCompilerContext::PostCompile();
-	}
-}
-
 void FControlRigBlueprintCompilerContext::CopyTermDefaultsToDefaultObject(UObject* DefaultObject)
 {
 	DECLARE_SCOPE_HIERARCHICAL_COUNTER_FUNC()
 
-	FKismetCompilerContext::CopyTermDefaultsToDefaultObject(DefaultObject);
+	FRigVMBlueprintCompilerContext::CopyTermDefaultsToDefaultObject(DefaultObject);
 
 	UControlRigBlueprint* ControlRigBlueprint = Cast<UControlRigBlueprint>(Blueprint);
 	if (ControlRigBlueprint)
@@ -94,7 +37,6 @@ void FControlRigBlueprintCompilerContext::CopyTermDefaultsToDefaultObject(UObjec
 		// and in UControlRig::InitializeFromCDO,
 		// other Control Rig Instances are then initialized From the CDO 
 		UControlRig* ControlRig = CastChecked<UControlRig>(DefaultObject);
-		ControlRig->PostInitInstanceIfRequired();
 		
 		// copy hierarchy
 		{
@@ -118,60 +60,3 @@ void FControlRigBlueprintCompilerContext::CopyTermDefaultsToDefaultObject(UObjec
 		ControlRig->Influences = ControlRigBlueprint->Influences;
 	}
 }
-
-void FControlRigBlueprintCompilerContext::EnsureProperGeneratedClass(UClass*& TargetUClass)
-{
-	DECLARE_SCOPE_HIERARCHICAL_COUNTER_FUNC()
-
-	if( TargetUClass && !((UObject*)TargetUClass)->IsA(URigVMBlueprintGeneratedClass::StaticClass()) )
-	{
-		FKismetCompilerUtilities::ConsignToOblivion(TargetUClass, Blueprint->bIsRegeneratingOnLoad);
-		TargetUClass = nullptr;
-	}
-}
-
-void FControlRigBlueprintCompilerContext::SpawnNewClass(const FString& NewClassName)
-{
-	DECLARE_SCOPE_HIERARCHICAL_COUNTER_FUNC()
-
-	NewControlRigBlueprintGeneratedClass = FindObject<URigVMBlueprintGeneratedClass>(Blueprint->GetOutermost(), *NewClassName);
-
-	if (NewControlRigBlueprintGeneratedClass == nullptr)
-	{
-		NewControlRigBlueprintGeneratedClass = NewObject<URigVMBlueprintGeneratedClass>(Blueprint->GetOutermost(), FName(*NewClassName), RF_Public | RF_Transactional);
-	}
-	else
-	{
-		// Already existed, but wasn't linked in the Blueprint yet due to load ordering issues
-		FBlueprintCompileReinstancer::Create(NewControlRigBlueprintGeneratedClass);
-	}
-	NewClass = NewControlRigBlueprintGeneratedClass;
-}
-
-void FControlRigBlueprintCompilerContext::OnNewClassSet(UBlueprintGeneratedClass* ClassToUse)
-{
-	DECLARE_SCOPE_HIERARCHICAL_COUNTER_FUNC()
-
-	NewControlRigBlueprintGeneratedClass = CastChecked<URigVMBlueprintGeneratedClass>(ClassToUse);
-}
-
-void FControlRigBlueprintCompilerContext::CleanAndSanitizeClass(UBlueprintGeneratedClass* ClassToClean, UObject*& InOldCDO)
-{
-	DECLARE_SCOPE_HIERARCHICAL_COUNTER_FUNC()
-
-	FKismetCompilerContext::CleanAndSanitizeClass(ClassToClean, InOldCDO);
-
-	// Make sure our typed pointer is set
-	check(ClassToClean == NewClass && NewControlRigBlueprintGeneratedClass == NewClass);
-}
-
-void FControlRigBlueprintCompilerContext::PreCompileUpdateBlueprintOnLoad(UBlueprint* BP)
-{
-	if (UControlRigBlueprint* RigBlueprint = Cast<UControlRigBlueprint>(BP))
-	{
-		RigBlueprint->CreateMemberVariablesOnLoad();
-	}
-}
-
-
-#undef LOCTEXT_NAMESPACE
diff --git a/Engine/Plugins/Animation/ControlRig/Source/ControlRigDeveloper/Public/ControlRigBlueprintCompiler.h b/Engine/Plugins/Animation/ControlRig/Source/ControlRigDeveloper/Public/ControlRigBlueprintCompiler.h
--- a/Engine/Plugins/Animation/ControlRig/Source/ControlRigDeveloper/Public/ControlRigBlueprintCompiler.h
+++ b/Engine/Plugins/Animation/ControlRig/Source/ControlRigDeveloper/Public/ControlRigBlueprintCompiler.h
@@ -2,15 +2,9 @@
 
 #pragma once
 
-#include "KismetCompilerModule.h"
-#include "KismetCompiler.h"
+#include "RigVMBlueprintCompiler.h"
 
-class UControlRigGraph;
-class UControlRigGraphNode;
-class UControlRigBlueprint;
-class URigVMBlueprintGeneratedClass;
-
-class CONTROLRIGDEVELOPER_API FControlRigBlueprintCompiler : public IBlueprintCompiler
+class CONTROLRIGDEVELOPER_API FControlRigBlueprintCompiler : public FRigVMBlueprintCompiler
 {
 public:
 	/** IBlueprintCompiler interface */
@@ -18,40 +12,14 @@ public:
 	virtual void Compile(UBlueprint* Blueprint, const FKismetCompilerOptions& CompileOptions, FCompilerResultsLog& Results) override;
 };
 
-class CONTROLRIGDEVELOPER_API FControlRigBlueprintCompilerContext : public FKismetCompilerContext
+class CONTROLRIGDEVELOPER_API FControlRigBlueprintCompilerContext : public FRigVMBlueprintCompilerContext
 {
 public:
 	FControlRigBlueprintCompilerContext(UBlueprint* SourceSketch, FCompilerResultsLog& InMessageLog, const FKismetCompilerOptions& InCompilerOptions)
-		: FKismetCompilerContext(SourceSketch, InMessageLog, InCompilerOptions)
-		, NewControlRigBlueprintGeneratedClass(nullptr)
+		: FRigVMBlueprintCompilerContext(SourceSketch, InMessageLog, InCompilerOptions)
 	{
 	}
 
 	// FKismetCompilerContext interface
-	virtual void ValidateLink(const UEdGraphPin* PinA, const UEdGraphPin* PinB) const override {}
-	virtual void ValidatePin(const UEdGraphPin* Pin) const override {}
-	virtual void ValidateNode(const UEdGraphNode* Node) const override {}
-	virtual bool CanIgnoreNode(const UEdGraphNode* Node) const override { return true; }
-	virtual bool ShouldForceKeepNode(const UEdGraphNode* Node) const override { return false; }
-	virtual void PrecompileFunction(FKismetFunctionContext& Context, EInternalCompilerFlags InternalFlags) override {}
-	virtual void PostCompile() override;
 	virtual void CopyTermDefaultsToDefaultObject(UObject* DefaultObject) override;
-	virtual void EnsureProperGeneratedClass(UClass*& TargetUClass) override;
-	virtual void SpawnNewClass(const FString& NewClassName) override;
-	virtual void OnNewClassSet(UBlueprintGeneratedClass* ClassToUse) override;
-	virtual void PruneIsolatedNodes(const TArray<UEdGraphNode*>& RootSet, TArray<UEdGraphNode*>& GraphNodes) override {}
-	virtual void CleanAndSanitizeClass(UBlueprintGeneratedClass* ClassToClean, UObject*& InOldCDO) override;
-	virtual void MergeUbergraphPagesIn(UEdGraph* Ubergraph) override {};
-	virtual void PreCompileUpdateBlueprintOnLoad(UBlueprint* BP) override;
-	virtual void CreateFunctionList() override {}
-	virtual void ProcessOneFunctionGraph(UEdGraph* SourceGraph, bool bInternalFunction = false) override {}
-
-private:
-
-	// used to fail a compilation and mark the blueprint in error
-	void MarkCompilationFailed(const FString& Message);
-
-private:
-	/** the new class we are generating */
-	URigVMBlueprintGeneratedClass* NewControlRigBlueprintGeneratedClass;
 };
\ No newline at end of file
diff --git a/Engine/Plugins/Runtime/RigVM/Source/RigVMDeveloper/Private/RigVMBlueprintCompiler.cpp b/Engine/Plugins/Runtime/RigVM/Source/RigVMDeveloper/Private/RigVMBlueprintCompiler.cpp
new file mode 100644
--- /dev/null
+++ b/Engine/Plugins/Runtime/RigVM/Source/RigVMDeveloper/Private/RigVMBlueprintCompiler.cpp
@@ -0,0 +1,143 @@
+// Copyright Epic Games, Inc. All Rights Reserved.
+
+#include "RigVMBlueprintCompiler.h"
+#include "RigVMHost.h"
+#include "KismetCompiler.h"
+#include "RigVMBlueprint.h"
+#include "RigVMCore/RigVMStruct.h"
+#include "EdGraph/RigVMEdGraph.h"
+#include "EdGraph/RigVMEdGraphNode.h"
+#include "EdGraph/RigVMEdGraphSchema.h"
+#include "RigVMBlueprintGeneratedClass.h"
+#include "Kismet2/KismetReinstanceUtilities.h"
+#include "Widgets/Notifications/SNotificationList.h"
+#include "Framework/Notifications/NotificationManager.h"
+
+DEFINE_LOG_CATEGORY_STATIC(LogRigVMCompiler, Log, All);
+#define LOCTEXT_NAMESPACE "RigVMBlueprintCompiler"
+
+void FRigVMBlueprintCompiler::Compile(UBlueprint* Blueprint, const FKismetCompilerOptions& CompileOptions, FCompilerResultsLog& Results)
+{
+	DECLARE_SCOPE_HIERARCHICAL_COUNTER_FUNC()
+
+	FRigVMBlueprintCompilerContext Compiler(Blueprint, Results, CompileOptions);
+	Compiler.Compile();
+}
+
+void FRigVMBlueprintCompilerContext::MarkCompilationFailed(const FString& Message)
+{
+	DECLARE_SCOPE_HIERARCHICAL_COUNTER_FUNC()
+
+	URigVMBlueprint* RigVMBlueprint = Cast<URigVMBlueprint>(Blueprint);
+	if (RigVMBlueprint)
+	{
+		Blueprint->Status = BS_Error;
+		Blueprint->MarkPackageDirty();
+		UE_LOG(LogRigVMCompiler, Error, TEXT("%s"), *Message);
+		MessageLog.Error(*Message);
+
+#if WITH_EDITOR
+		FNotificationInfo Info(FText::FromString(Message));
+		Info.Image = FCoreStyle::Get().GetBrush(TEXT("MessageLog.Error"));
+		Info.bFireAndForget = true;
+		Info.FadeOutDuration = 5.0f;
+		Info.ExpireDuration = 5.0f;
+		TSharedPtr<SNotificationItem> NotificationPtr = FSlateNotificationManager::Get().AddNotification(Info);
+		NotificationPtr->SetCompletionState(SNotificationItem::CS_Success);
+#endif
+	}
+}
+
+void FRigVMBlueprintCompilerContext::PostCompile()
+{
+	DECLARE_SCOPE_HIERARCHICAL_COUNTER_FUNC()
+
+	URigVMBlueprint* RigVMBlueprint = Cast<URigVMBlueprint>(Blueprint);
+	if (RigVMBlueprint)
+	{
+		if (!RigVMBlueprint->bIsRegeneratingOnLoad)
+		{
+			// todo: can we find out somehow if we are cooking?
+			RigVMBlueprint->RecompileVM();
+		}
+	}
+
+	{
+		DECLARE_SCOPE_HIERARCHICAL_COUNTER(FKismetCompilerContext::PostCompile)
+		FKismetCompilerContext::PostCompile();
+	}
+}
+
+void FRigVMBlueprintCompilerContext::CopyTermDefaultsToDefaultObject(UObject* DefaultObject)
+{
+	DECLARE_SCOPE_HIERARCHICAL_COUNTER_FUNC()
+
+	FKismetCompilerContext::CopyTermDefaultsToDefaultObject(DefaultObject);
+
+	URigVMBlueprint* RigVMBlueprint = Cast<URigVMBlueprint>(Blueprint);
+	if (RigVMBlueprint)
+	{
+		// here, CDO is initialized from BP,
+		// and in URigVMHost::InitializeFromCDO,
+		// other instances are then initialized From the CDO 
+		URigVMHost* Host = CastChecked<URigVMHost>(DefaultObject);
+		Host->PostInitInstanceIfRequired();
+	}
+}
+
+void FRigVMBlueprintCompilerContext::EnsureProperGeneratedClass(UClass*& TargetUClass)
+{
+	DECLARE_SCOPE_HIERARCHICAL_COUNTER_FUNC()
+
+	if( TargetUClass && !((UObject*)TargetUClass)->IsA(URigVMBlueprintGeneratedClass::StaticClass()) )
+	{
+		FKismetCompilerUtilities::ConsignToOblivion(TargetUClass, Blueprint->bIsRegeneratingOnLoad);
+		TargetUClass = nullptr;
+	}
+}
+
+void FRigVMBlueprintCompilerContext::SpawnNewClass(const FString& NewClassName)
+{
+	DECLARE_SCOPE_HIERARCHICAL_COUNTER_FUNC()
+
+	NewRigVMBlueprintGeneratedClass = FindObject<URigVMBlueprintGeneratedClass>(Blueprint->GetOutermost(), *NewClassName);
+
+	if (NewRigVMBlueprintGeneratedClass == nullptr)
+	{
+		NewRigVMBlueprintGeneratedClass = NewObject<URigVMBlueprintGeneratedClass>(Blueprint->GetOutermost(), FName(*NewClassName), RF_Public | RF_Transactional);
+	}
+	else
+	{
+		// Already existed, but wasn't linked in the Blueprint yet due to load ordering issues
+		FBlueprintCompileReinstancer::Create(NewRigVMBlueprintGeneratedClass);
+	}
+	NewClass = NewRigVMBlueprintGeneratedClass;
+}
+
+void FRigVMBlueprintCompilerContext::OnNewClassSet(UBlueprintGeneratedClass* ClassToUse)
+{
+	DECLARE_SCOPE_HIERARCHICAL_COUNTER_FUNC()
+
+	NewRigVMBlueprintGeneratedClass = CastChecked<URigVMBlueprintGeneratedClass>(ClassToUse);
+}
+
+void FRigVMBlueprintCompilerContext::CleanAndSanitizeClass(UBlueprintGeneratedClass* ClassToClean, UObject*& InOldCDO)
+{
+	DECLARE_SCOPE_HIERARCHICAL_COUNTER_FUNC()
+
+	FKismetCompilerContext::CleanAndSanitizeClass(ClassToClean, InOldCDO);
+
+	// Make sure our typed pointer is set
+	check(ClassToClean == NewClass && NewRigVMBlueprintGeneratedClass == NewClass);
+}
+
+void FRigVMBlueprintCompilerContext::PreCompileUpdateBlueprintOnLoad(UBlueprint* BP)
+{
+	if (URigVMBlueprint* RigBlueprint = Cast<URigVMBlueprint>(BP))
+	{
+		RigBlueprint->CreateMemberVariablesOnLoad();
+	}
+}
+
+
+#undef LOCTEXT_NAMESPACE
diff --git a/Engine/Plugins/Runtime/RigVM/Source/RigVMDeveloper/Private/RigVMDeveloperModule.cpp b/Engine/Plugins/Runtime/RigVM/Source/RigVMDeveloper/Private/RigVMDeveloperModule.cpp
--- a/Engine/Plugins/Runtime/RigVM/Source/RigVMDeveloper/Private/RigVMDeveloperModule.cpp
+++ b/Engine/Plugins/Runtime/RigVM/Source/RigVMDeveloper/Private/RigVMDeveloperModule.cpp
@@ -5,8 +5,8 @@
 =============================================================================*/
 
 #include "RigVMDeveloperModule.h"
-#include "Modules/ModuleManager.h"
 #include "RigVMBlueprint.h"
+#include "RigVMBlueprintCompiler.h"
 
 DEFINE_LOG_CATEGORY(LogRigVMDeveloper);
 
@@ -14,15 +14,33 @@ class FRigVMDeveloperModule : public IModuleInterface
 {
 public:
 	/** IModuleInterface implementation */
-	virtual void StartupModule() override
-	{
-	}
+	virtual void StartupModule() override;
+	virtual void ShutdownModule() override;
 
-	virtual void ShutdownModule() override
-	{
-	}
+	FRigVMBlueprintCompiler RigVMBlueprintCompiler;
+	static TSharedPtr<FKismetCompilerContext> GetRigVMCompiler(UBlueprint* BP, FCompilerResultsLog& InMessageLog, const FKismetCompilerOptions& InCompileOptions);
 };
 
 IMPLEMENT_MODULE(FRigVMDeveloperModule, RigVMDeveloper);
 
+void FRigVMDeveloperModule::StartupModule()
+{
+	// Register blueprint compiler
+	FKismetCompilerContext::RegisterCompilerForBP(URigVMBlueprint::StaticClass(), &FRigVMDeveloperModule::GetRigVMCompiler);
+	IKismetCompilerInterface& KismetCompilerModule = FModuleManager::LoadModuleChecked<IKismetCompilerInterface>("KismetCompiler");
+	KismetCompilerModule.GetCompilers().Add(&RigVMBlueprintCompiler);
+}
+
+void FRigVMDeveloperModule::ShutdownModule()
+{
+	IKismetCompilerInterface* KismetCompilerModule = FModuleManager::GetModulePtr<IKismetCompilerInterface>("KismetCompiler");
+	if (KismetCompilerModule)
+	{
+		KismetCompilerModule->GetCompilers().Remove(&RigVMBlueprintCompiler);
+	}
+}
 
+TSharedPtr<FKismetCompilerContext> FRigVMDeveloperModule::GetRigVMCompiler(UBlueprint* BP, FCompilerResultsLog& InMessageLog, const FKismetCompilerOptions& InCompileOptions)
+{
+	return TSharedPtr<FKismetCompilerContext>(new FRigVMBlueprintCompilerContext(BP, InMessageLog, InCompileOptions));
+}
diff --git a/Engine/Plugins/Runtime/RigVM/Source/RigVMDeveloper/Public/RigVMBlueprintCompiler.h b/Engine/Plugins/Runtime/RigVM/Source/RigVMDeveloper/Public/RigVMBlueprintCompiler.h
new file mode 100644
--- /dev/null
+++ b/Engine/Plugins/Runtime/RigVM/Source/RigVMDeveloper/Public/RigVMBlueprintCompiler.h
@@ -0,0 +1,55 @@
+// Copyright Epic Games, Inc. All Rights Reserved.
+
+#pragma once
+
+#include "KismetCompilerModule.h"
+#include "KismetCompiler.h"
+
+class FRigVMBlueprintCompilerContext;
+class URigVMBlueprintGeneratedClass;
+
+class RIGVMDEVELOPER_API FRigVMBlueprintCompiler : public IBlueprintCompiler
+{
+public:
+	/** IBlueprintCompiler interface */
+	virtual bool CanCompile(const UBlueprint* Blueprint) override { return false; } // to be implemented by the specializations
+	virtual void Compile(UBlueprint* Blueprint, const FKismetCompilerOptions& CompileOptions, FCompilerResultsLog& Results) override;
+};
+
+class RIGVMDEVELOPER_API FRigVMBlueprintCompilerContext : public FKismetCompilerContext
+{
+public:
+	FRigVMBlueprintCompilerContext(UBlueprint* SourceSketch, FCompilerResultsLog& InMessageLog, const FKismetCompilerOptions& InCompilerOptions)
+		: FKismetCompilerContext(SourceSketch, InMessageLog, InCompilerOptions)
+		, NewRigVMBlueprintGeneratedClass(nullptr)
+	{
+	}
+
+	// FKismetCompilerContext interface
+	virtual void ValidateLink(const UEdGraphPin* PinA, const UEdGraphPin* PinB) const override {}
+	virtual void ValidatePin(const UEdGraphPin* Pin) const override {}
+	virtual void ValidateNode(const UEdGraphNode* Node) const override {}
+	virtual bool CanIgnoreNode(const UEdGraphNode* Node) const override { return true; }
+	virtual bool ShouldForceKeepNode(const UEdGraphNode* Node) const override { return false; }
+	virtual void PrecompileFunction(FKismetFunctionContext& Context, EInternalCompilerFlags InternalFlags) override {}
+	virtual void PostCompile() override;
+	virtual void CopyTermDefaultsToDefaultObject(UObject* DefaultObject) override;
+	virtual void EnsureProperGeneratedClass(UClass*& TargetUClass) override;
+	virtual void SpawnNewClass(const FString& NewClassName) override;
+	virtual void OnNewClassSet(UBlueprintGeneratedClass* ClassToUse) override;
+	virtual void PruneIsolatedNodes(const TArray<UEdGraphNode*>& RootSet, TArray<UEdGraphNode*>& GraphNodes) override {}
+	virtual void CleanAndSanitizeClass(UBlueprintGeneratedClass* ClassToClean, UObject*& InOldCDO) override;
+	virtual void MergeUbergraphPagesIn(UEdGraph* Ubergraph) override {};
+	virtual void PreCompileUpdateBlueprintOnLoad(UBlueprint* BP) override;
+	virtual void CreateFunctionList() override {}
+	virtual void ProcessOneFunctionGraph(UEdGraph* SourceGraph, bool bInternalFunction = false) override {}
+
+private:
+
+	// used to fail a compilation and mark the blueprint in error
+	void MarkCompilationFailed(const FString& Message);
+
+private:
+	/** the new class we are generating */
+	URigVMBlueprintGeneratedClass* NewRigVMBlueprintGeneratedClass;
+};
\ No newline at end of file
