Sequencer:  Control Rig: Fix keying Rotator controls from Rotation Order fallout

diff --git a/Engine/Plugins/Animation/ControlRig/Source/ControlRig/Private/Sequencer/MovieSceneControlRigParameterSection.cpp b/Engine/Plugins/Animation/ControlRig/Source/ControlRig/Private/Sequencer/MovieSceneControlRigParameterSection.cpp
--- a/Engine/Plugins/Animation/ControlRig/Source/ControlRig/Private/Sequencer/MovieSceneControlRigParameterSection.cpp
+++ b/Engine/Plugins/Animation/ControlRig/Source/ControlRig/Private/Sequencer/MovieSceneControlRigParameterSection.cpp
@@ -601,6 +601,11 @@ struct FParameterTransformChannelEditorData
 					
 					return FVector(Euler.GetTranslation());
 				}
+				else if(ControlElement->Settings.ControlType == ERigControlType::Position)
+				{
+					FVector3f Vector = ControlRig->GetHierarchy()->GetControlValue(ControlElement, ERigControlValueType::Current).Get<FVector3f>();
+					return FVector(Vector.X, Vector.Y, Vector.Z);
+				}
 			}
 		}
 		return TOptional<FVector>();
@@ -663,6 +668,11 @@ struct FParameterTransformChannelEditorData
 					
 					return FVector(Transform.GetScale3D());
 				}
+				else if (ControlElement->Settings.ControlType == ERigControlType::Scale)
+				{
+					FVector3f Vector = ControlRig->GetHierarchy()->GetControlValue(ControlElement, ERigControlValueType::Current).Get<FVector3f>();
+					return FVector(Vector.X, Vector.Y, Vector.Z);
+				}
 			}
 		}
 		return TOptional<FVector>();
@@ -1026,15 +1036,6 @@ void UMovieSceneControlRigParameterSection::PostLoad()
 					Handle->ControlRig = ControlRig;
 				}
 			}
-			/*
-			if (ConstraintChannel.Constraint.IsValid() == false && ConstraintChannel.ConstraintCopyToSpawn)
-			{
-				const FConstraintsManagerController& Controller = FConstraintsManagerController::Get(ControlRig->GetWorld());
-				UTickableConstraint* NewOne = Controller.AddConstraintFromCopy(ConstraintChannel.ConstraintCopyToSpawn);
-				ConstraintChannel.Constraint = NewOne;
-				ReconstructChannelProxy();
-			}
-			*/
 		}
 	}
 }
diff --git a/Engine/Plugins/Animation/ControlRig/Source/ControlRig/Private/Sequencer/MovieSceneControlRigParameterTemplate.cpp b/Engine/Plugins/Animation/ControlRig/Source/ControlRig/Private/Sequencer/MovieSceneControlRigParameterTemplate.cpp
--- a/Engine/Plugins/Animation/ControlRig/Source/ControlRig/Private/Sequencer/MovieSceneControlRigParameterTemplate.cpp
+++ b/Engine/Plugins/Animation/ControlRig/Source/ControlRig/Private/Sequencer/MovieSceneControlRigParameterTemplate.cpp
@@ -791,8 +791,12 @@ struct FControlRigParameterPreAnimatedTokenProducer : IMovieScenePreAnimatedToke
 
 								for (TNameAndValue<FVector>& Value : VectorValues)
 								{
-									if (ControlRig->FindControl(Value.Name))
+									if (FRigControlElement* ControlElement = ControlRig->FindControl(Value.Name))
 									{
+										if (ControlElement->Settings.ControlType == ERigControlType::Rotator)
+										{
+											RigHierarchy->SetControlSpecifiedEulerAngle(ControlElement, Value.Value);
+										}
 										ControlRig->SetControlValue<FVector3f>(Value.Name, (FVector3f)Value.Value, true, FRigControlModifiedContext(EControlRigSetKey::Never), bSetupUndo);
 									}
 								}
@@ -805,25 +809,25 @@ struct FControlRigParameterPreAnimatedTokenProducer : IMovieScenePreAnimatedToke
 										{
 										case ERigControlType::Transform:
 										{
-											ControlRig->SetControlValue<FRigControlValue::FTransform_Float>(Value.Name, Value.Value.ToFTransform(), true, FRigControlModifiedContext(EControlRigSetKey::Never), bSetupUndo);
 											FVector EulerAngle(Value.Value.Rotation.Roll, Value.Value.Rotation.Pitch, Value.Value.Rotation.Yaw);
 											RigHierarchy->SetControlSpecifiedEulerAngle(ControlElement, EulerAngle);
+											ControlRig->SetControlValue<FRigControlValue::FTransform_Float>(Value.Name, Value.Value.ToFTransform(), true, FRigControlModifiedContext(EControlRigSetKey::Never), bSetupUndo);
 											break;
 										}
 										case ERigControlType::TransformNoScale:
 										{
 											FTransformNoScale NoScale = Value.Value.ToFTransform();
-											ControlRig->SetControlValue<FRigControlValue::FTransformNoScale_Float>(Value.Name, NoScale, true, FRigControlModifiedContext(EControlRigSetKey::Never), bSetupUndo);
 											FVector EulerAngle(Value.Value.Rotation.Roll, Value.Value.Rotation.Pitch, Value.Value.Rotation.Yaw);
 											RigHierarchy->SetControlSpecifiedEulerAngle(ControlElement, EulerAngle);
+											ControlRig->SetControlValue<FRigControlValue::FTransformNoScale_Float>(Value.Name, NoScale, true, FRigControlModifiedContext(EControlRigSetKey::Never), bSetupUndo);
 											break;
 										}
 										case ERigControlType::EulerTransform:
 										{
 											FEulerTransform EulerTransform = Value.Value;
-											ControlRig->SetControlValue<FRigControlValue::FEulerTransform_Float>(Value.Name, EulerTransform, true, FRigControlModifiedContext(EControlRigSetKey::Never), bSetupUndo);
 											FVector EulerAngle(Value.Value.Rotation.Roll, Value.Value.Rotation.Pitch, Value.Value.Rotation.Yaw);
 											RigHierarchy->SetControlSpecifiedEulerAngle(ControlElement, EulerAngle);
+											ControlRig->SetControlValue<FRigControlValue::FEulerTransform_Float>(Value.Name, EulerTransform, true, FRigControlModifiedContext(EControlRigSetKey::Never), bSetupUndo);
 											break;
 										}
 										default:
@@ -927,7 +931,12 @@ struct FControlRigParameterPreAnimatedTokenProducer : IMovieScenePreAnimatedToke
 						SpaceValue.ControlRigElement = ControlElement->GetKey();
 						SpaceValue.SpaceType = EMovieSceneControlRigSpaceType::Parent;
 						Token.SpaceValues.Add(FControlSpaceAndValue(ControlElement->GetName(), SpaceValue));
-						const FVector3f Val = ControlRig->GetHierarchy()->GetControlValue(ControlElement, ERigControlValueType::Current).Get<FVector3f>();
+						FVector3f Val = ControlRig->GetHierarchy()->GetControlValue(ControlElement, ERigControlValueType::Current).Get<FVector3f>();
+						if (ControlElement->Settings.ControlType == ERigControlType::Rotator)
+						{
+							FVector Vector = ControlRig->GetHierarchy()->GetControlSpecifiedEulerAngle(ControlElement);
+							Val = FVector3f(Vector.X, Vector.Y, Vector.Z);
+						}
 						Token.VectorValues.Add(TNameAndValue<FVector>{ ControlElement->GetName(), (FVector)Val });
 						//mz todo specify rotator special so we can do quat interps
 						break;
@@ -1517,6 +1526,11 @@ struct TControlRigParameterActuatorVector : TMovieSceneBlendingActuator<FControl
 				&& (ControlElement->Settings.ControlType == ERigControlType::Position || ControlElement->Settings.ControlType == ERigControlType::Scale || ControlElement->Settings.ControlType == ERigControlType::Rotator))
 			{
 				FVector3f Val = ControlRig->GetHierarchy()->GetControlValue(ControlElement, ERigControlValueType::Current).Get<FVector3f>();
+				if (ControlElement->Settings.ControlType == ERigControlType::Rotator)
+				{
+					FVector Vector = ControlRig->GetHierarchy()->GetControlSpecifiedEulerAngle(ControlElement);
+					Val = FVector3f(Vector.X, Vector.Y, Vector.Z);
+				}
 				return FControlRigTrackTokenVector((FVector)Val);
 			}
 		}
@@ -1542,7 +1556,11 @@ struct TControlRigParameterActuatorVector : TMovieSceneBlendingActuator<FControl
 					ControlElement->Settings.AnimationType != ERigControlAnimationType::VisualCue
 					&& (ControlElement->Settings.ControlType == ERigControlType::Position || ControlElement->Settings.ControlType == ERigControlType::Scale || ControlElement->Settings.ControlType == ERigControlType::Rotator))
 				{
-						ControlRig->SetControlValue<FVector3f>(ParameterName, (FVector3f)InFinalValue.Value, true, EControlRigSetKey::Never,bSetupUndo);
+					if (ControlElement && ControlElement->Settings.ControlType == ERigControlType::Rotator)
+					{
+						ControlRig->GetHierarchy()->SetControlSpecifiedEulerAngle(ControlElement, InFinalValue.Value);
+					}
+					ControlRig->SetControlValue<FVector3f>(ParameterName, (FVector3f)InFinalValue.Value, true, EControlRigSetKey::Never, bSetupUndo);
 				}
 			}
 
@@ -1642,16 +1660,16 @@ struct TControlRigParameterActuatorTransform : TMovieSceneBlendingActuator<FCont
 				{
 					if (ControlElement && ControlElement->Settings.ControlType == ERigControlType::Transform)
 					{
-						ControlRig->SetControlValue<FRigControlValue::FTransform_Float>(ParameterName, InFinalValue.Value.ToFTransform(), true, EControlRigSetKey::Never, bSetupUndo);
 						FVector EulerAngle(InFinalValue.Value.Rotation.Roll, InFinalValue.Value.Rotation.Pitch, InFinalValue.Value.Rotation.Yaw);
 						ControlRig->GetHierarchy()->SetControlSpecifiedEulerAngle(ControlElement, EulerAngle);
+						ControlRig->SetControlValue<FRigControlValue::FTransform_Float>(ParameterName, InFinalValue.Value.ToFTransform(), true, EControlRigSetKey::Never, bSetupUndo);
 					}
 					else if (ControlElement && ControlElement->Settings.ControlType == ERigControlType::TransformNoScale)
 					{
 						const FTransformNoScale NoScale = InFinalValue.Value.ToFTransform();
-						ControlRig->SetControlValue<FRigControlValue::FTransformNoScale_Float>(ParameterName, NoScale, true, EControlRigSetKey::Never, bSetupUndo);
 						FVector EulerAngle(InFinalValue.Value.Rotation.Roll, InFinalValue.Value.Rotation.Pitch, InFinalValue.Value.Rotation.Yaw);
 						ControlRig->GetHierarchy()->SetControlSpecifiedEulerAngle(ControlElement, EulerAngle);
+						ControlRig->SetControlValue<FRigControlValue::FTransformNoScale_Float>(ParameterName, NoScale, true, EControlRigSetKey::Never, bSetupUndo);
 					}
 					else if (ControlElement && ControlElement->Settings.ControlType == ERigControlType::EulerTransform)
 					{
diff --git a/Engine/Plugins/Animation/ControlRig/Source/ControlRigEditor/Private/EditMode/ControlRigControlsProxy.cpp b/Engine/Plugins/Animation/ControlRig/Source/ControlRigEditor/Private/EditMode/ControlRigControlsProxy.cpp
--- a/Engine/Plugins/Animation/ControlRig/Source/ControlRigEditor/Private/EditMode/ControlRigControlsProxy.cpp
+++ b/Engine/Plugins/Animation/ControlRig/Source/ControlRigEditor/Private/EditMode/ControlRigControlsProxy.cpp
@@ -597,6 +597,11 @@ void UControlRigVectorControlProxy::PostEditChangeChainProperty(struct FProperty
 		if (ControlElement && ControlRig.IsValid())
 		{
 			//MUST set through ControlRig
+			if (ControlElement->Settings.ControlType == ERigControlType::Rotator)
+			{
+				FVector DVector(Vector.X, Vector.Y, Vector.Z);
+				ControlRig->GetHierarchy()->SetControlSpecifiedEulerAngle(ControlElement, DVector);
+			}
 			ControlRig->SetControlValue<FVector3f>(ControlName, Vector, true, EControlRigSetKey::DoNotCare,false);
 			ControlRig->Evaluate_AnyThread();
 
@@ -612,7 +617,12 @@ void UControlRigVectorControlProxy::ValueChanged()
 		Modify();
 		const FName PropertyName("Vector");
 		FTrackInstancePropertyBindings Binding(PropertyName, PropertyName.ToString());
-		const FVector3f Val = ControlRig.Get()->GetHierarchy()->GetControlValue(ControlElement, ERigControlValueType::Current).Get<FVector3f>();
+		FVector3f Val = ControlRig.Get()->GetHierarchy()->GetControlValue(ControlElement, ERigControlValueType::Current).Get<FVector3f>();
+		if (ControlElement->Settings.ControlType == ERigControlType::Rotator)
+		{ 
+			FVector DVector = ControlRig->GetHierarchy()->GetControlSpecifiedEulerAngle(ControlElement);
+			Val = FVector3f(DVector.X, DVector.Y, DVector.Z);
+		}
 		Binding.CallFunction<FVector3f>(*this, Val);
 	}
 }
@@ -624,6 +634,11 @@ void UControlRigVectorControlProxy::PostEditUndo()
 	if (ControlElement && ControlRig.IsValid() && ControlRig->GetHierarchy()->Contains(FRigElementKey(ControlName, ERigElementType::Control)))
 	{
 		ControlRig->SelectControl(ControlName, bSelected);
+		if (ControlElement && ControlElement->Settings.ControlType == ERigControlType::Rotator)
+		{
+			FVector DVector(Vector.X, Vector.Y, Vector.Z);
+			ControlRig->GetHierarchy()->SetControlSpecifiedEulerAngle(ControlElement, DVector);
+		}
 		ControlRig->SetControlValue<FVector3f>(ControlName, (FVector3f)Vector, true, EControlRigSetKey::Never,false);
 	}
 }
@@ -682,6 +697,11 @@ void UControlRigVectorControlProxy::SetKey(const IPropertyHandle& KeyedPropertyH
 				break;
 			}
 		}
+		if (ControlElement->Settings.ControlType == ERigControlType::Rotator)
+		{
+			FVector DVector(Vector.X, Vector.Y, Vector.Z);
+			ControlRig->GetHierarchy()->SetControlSpecifiedEulerAngle(ControlElement, DVector);
+		}
 		ControlRig->SetControlValue<FVector3f>(ControlName, (FVector3f)Vector, true, Context,false);
 	}
 }
