"Edit in Sequencer" button present in the composite, blend space and montage editors
#fix dynamic visibility for button according to currently active animation asset

diff --git a/Engine/Plugins/Animation/ControlRig/Source/ControlRigEditor/Private/ControlRigEditorModule.cpp b/Engine/Plugins/Animation/ControlRig/Source/ControlRigEditor/Private/ControlRigEditorModule.cpp
--- a/Engine/Plugins/Animation/ControlRig/Source/ControlRigEditor/Private/ControlRigEditorModule.cpp
+++ b/Engine/Plugins/Animation/ControlRig/Source/ControlRigEditor/Private/ControlRigEditorModule.cpp
@@ -891,10 +891,24 @@ void FControlRigEditorModule::AddControlRigExtenderToToolMenu(FName InToolMenuNa
 	FToolMenuOwnerScoped OwnerScoped(this);
 
 	UToolMenu* ToolMenu = UToolMenus::Get()->ExtendMenu(InToolMenuName);
+
+	FToolUIAction UIAction;
+	UIAction.IsActionVisibleDelegate.BindLambda([](const FToolMenuContext& Context) -> bool
+	{
+		if (const UAnimationToolMenuContext* MenuContext = Context.FindContext<UAnimationToolMenuContext>())
+		{					
+			if (const UAnimationAsset* AnimationAsset = MenuContext->AnimationEditor.Pin()->GetPersonaToolkit()->GetAnimationAsset())
+			{
+				return AnimationAsset->GetClass() == UAnimSequence::StaticClass();
+			}
+		}
+		return false;
+	});
+	
 	ToolMenu->AddMenuEntry("Sequencer",
 		FToolMenuEntry::InitComboButton(
 			"EditInSequencer",
-			FToolUIActionChoice(),
+			FToolUIActionChoice(UIAction),
 			FNewToolMenuChoice(
 				FNewToolMenuDelegate::CreateLambda([this](UToolMenu* InNewToolMenu)
 				{
