RigVMEditor: Refactoring editor module, blueprint actions etc

diff --git a/Engine/Plugins/Animation/ControlRig/Source/ControlRigDeveloper/Private/Graph/ControlRigGraph.cpp b/Engine/Plugins/Animation/ControlRig/Source/ControlRigDeveloper/Private/Graph/ControlRigGraph.cpp
--- a/Engine/Plugins/Animation/ControlRig/Source/ControlRigDeveloper/Private/Graph/ControlRigGraph.cpp
+++ b/Engine/Plugins/Animation/ControlRig/Source/ControlRigDeveloper/Private/Graph/ControlRigGraph.cpp
@@ -42,6 +42,27 @@ void UControlRigGraph::InitializeFromBlueprint(URigVMBlueprint* InBlueprint)
 
 TArray<TSharedPtr<FString>> UControlRigGraph::EmptyElementNameList;
 
+const TArray<TSharedPtr<FString>>* UControlRigGraph::GetNameListForWidget(const FString& InWidgetName) const
+{
+	if (InWidgetName == TEXT("BoneName"))
+	{
+		return GetBoneNameList();
+	}
+	if (InWidgetName == TEXT("ControlName"))
+	{
+		return GetControlNameListWithoutAnimationChannels();
+	}
+	if (InWidgetName == TEXT("SpaceName"))
+	{
+		return GetNullNameList();
+	}
+	if (InWidgetName == TEXT("CurveName"))
+	{
+		return GetCurveNameList();
+	}
+	return nullptr;
+}
+
 void UControlRigGraph::CacheNameLists(URigHierarchy* InHierarchy, const FRigVMDrawContainer* DrawContainer, TArray<TSoftObjectPtr<UControlRigShapeLibrary>> ShapeLibraries)
 {
 	if (UControlRigGraph* OuterGraph = Cast<UControlRigGraph>(GetOuter()))
diff --git a/Engine/Plugins/Animation/ControlRig/Source/ControlRigDeveloper/Public/Graph/ControlRigGraph.h b/Engine/Plugins/Animation/ControlRig/Source/ControlRigDeveloper/Public/Graph/ControlRigGraph.h
--- a/Engine/Plugins/Animation/ControlRig/Source/ControlRigDeveloper/Public/Graph/ControlRigGraph.h
+++ b/Engine/Plugins/Animation/ControlRig/Source/ControlRigDeveloper/Public/Graph/ControlRigGraph.h
@@ -54,6 +54,8 @@ public:
 		return GetElementNameList(ERigElementType::Curve);
 	}
 
+	virtual const TArray<TSharedPtr<FString>>* GetNameListForWidget(const FString& InWidgetName) const override;
+
 	void CacheNameLists(URigHierarchy* InHierarchy, const FRigVMDrawContainer* DrawContainer, TArray<TSoftObjectPtr<UControlRigShapeLibrary>> ShapeLibraries);
 	const TArray<TSharedPtr<FString>>* GetElementNameList(ERigElementType InElementType = ERigElementType::Bone) const;
 	const TArray<TSharedPtr<FString>>* GetElementNameList(URigVMPin* InPin = nullptr) const;
diff --git a/Engine/Plugins/Animation/ControlRig/Source/ControlRigDeveloper/Public/Graph/ControlRigGraphSchema.h b/Engine/Plugins/Animation/ControlRig/Source/ControlRigDeveloper/Public/Graph/ControlRigGraphSchema.h
--- a/Engine/Plugins/Animation/ControlRig/Source/ControlRigDeveloper/Public/Graph/ControlRigGraphSchema.h
+++ b/Engine/Plugins/Animation/ControlRig/Source/ControlRigDeveloper/Public/Graph/ControlRigGraphSchema.h
@@ -32,6 +32,6 @@ public:
 	virtual FLinearColor GetPinTypeColor(const FEdGraphPinType& PinType) const override;
 
 	// URigVMEdGraphSchema interface
-	virtual const FName& GetGraphNameForUI() const override { return GraphName_ControlRig; }
+	virtual const FName& GetRootGraphName() const override { return GraphName_ControlRig; }
 	virtual bool IsRigVMDefaultEvent(const FName& InEventName) const override;
 };
diff --git a/Engine/Plugins/Animation/ControlRig/Source/ControlRigEditor/Private/ControlRigBlueprintFactory.cpp b/Engine/Plugins/Animation/ControlRig/Source/ControlRigEditor/Private/ControlRigBlueprintFactory.cpp
--- a/Engine/Plugins/Animation/ControlRig/Source/ControlRigEditor/Private/ControlRigBlueprintFactory.cpp
+++ b/Engine/Plugins/Animation/ControlRig/Source/ControlRigEditor/Private/ControlRigBlueprintFactory.cpp
@@ -20,6 +20,7 @@
 #include "ClassViewerFilter.h"
 #include "ControlRigBlueprint.h"
 #include "ControlRigBlueprintActions.h"
+#include "ControlRigEditorModule.h"
 #include "RigVMBlueprintGeneratedClass.h"
 #include "Graph/ControlRigGraphSchema.h"
 #include "Graph/ControlRigGraph.h"
@@ -290,7 +291,7 @@ UObject* UControlRigBlueprintFactory::FactoryCreateNew(UClass* Class, UObject* I
 	else
 	{
 		UControlRigBlueprint* ControlRigBlueprint = CastChecked<UControlRigBlueprint>(FKismetEditorUtilities::CreateBlueprint(ParentClass, InParent, Name, BPTYPE_Normal, UControlRigBlueprint::StaticClass(), URigVMBlueprintGeneratedClass::StaticClass(), CallingContext));
-		CreateRigGraphIfRequired(ControlRigBlueprint);
+		FControlRigEditorModule::Get().CreateRootGraphIfRequired(ControlRigBlueprint);
 		return ControlRigBlueprint;
 	}
 }
@@ -310,29 +311,5 @@ UControlRigBlueprint* UControlRigBlueprintFactory::CreateControlRigFromSkeletalM
 	return FControlRigBlueprintActions::CreateControlRigFromSkeletalMeshOrSkeleton(InSelectedObject);
 }
 
-void UControlRigBlueprintFactory::CreateRigGraphIfRequired(UControlRigBlueprint* InBlueprint)
-{
-	if(InBlueprint == nullptr)
-	{
-		return;
-	}
-	
-	for(UEdGraph* EdGraph : InBlueprint->UbergraphPages)
-	{
-		if(EdGraph->IsA<UControlRigGraph>())
-		{
-			return;
-		}
-	}
-	
-	// add an initial graph for us to work in
-	const UControlRigGraphSchema* ControlRigGraphSchema = GetDefault<UControlRigGraphSchema>();
-	UEdGraph* ControlRigGraph = FBlueprintEditorUtils::CreateNewGraph(InBlueprint, ControlRigGraphSchema->GraphName_ControlRig, UControlRigGraph::StaticClass(), UControlRigGraphSchema::StaticClass());
-	ControlRigGraph->bAllowDeletion = false;
-	FBlueprintEditorUtils::AddUbergraphPage(InBlueprint, ControlRigGraph);
-	InBlueprint->LastEditedDocuments.AddUnique(ControlRigGraph);
-	InBlueprint->PostLoad();
-}
-
 #undef LOCTEXT_NAMESPACE
 
diff --git a/Engine/Plugins/Animation/ControlRig/Source/ControlRigEditor/Private/ControlRigBlueprintFactory.h b/Engine/Plugins/Animation/ControlRig/Source/ControlRigEditor/Private/ControlRigBlueprintFactory.h
--- a/Engine/Plugins/Animation/ControlRig/Source/ControlRigEditor/Private/ControlRigBlueprintFactory.h
+++ b/Engine/Plugins/Animation/ControlRig/Source/ControlRigEditor/Private/ControlRigBlueprintFactory.h
@@ -42,11 +42,5 @@ public:
 	 */
 	UFUNCTION(BlueprintCallable, Category = "Control Rig")
 	static UControlRigBlueprint* CreateControlRigFromSkeletalMeshOrSkeleton(UObject* InSelectedObject);
-
-	/**
-	 * Creates the riggraph if necessary
-	 * @param InBlueprint The blueprint to setup up
-	 */
-	static void CreateRigGraphIfRequired(UControlRigBlueprint* InBlueprint);
 };
 
