RigVMEditor: Refactoring editor module, blueprint actions etc

diff --git a/Engine/Plugins/Animation/ControlRig/Source/ControlRigEditor/Private/ControlRigEditorModule.cpp b/Engine/Plugins/Animation/ControlRig/Source/ControlRigEditor/Private/ControlRigEditorModule.cpp
--- a/Engine/Plugins/Animation/ControlRig/Source/ControlRigEditor/Private/ControlRigEditorModule.cpp
+++ b/Engine/Plugins/Animation/ControlRig/Source/ControlRigEditor/Private/ControlRigEditorModule.cpp
@@ -2,13 +2,10 @@
 
 #include "ControlRigEditorModule.h"
 #include "BlueprintActionDatabaseRegistrar.h"
-#include "BlueprintEditorModule.h"
 #include "BlueprintNodeSpawner.h"
 #include "PropertyEditorModule.h"
 #include "ControlRig.h"
 #include "ControlRigComponent.h"
-#include "ControlRigVariableDetailsCustomization.h"
-#include "EdGraph/RigVMEdGraphConnectionDrawingPolicy.h"
 #include "GraphEditorActions.h"
 #include "Kismet2/KismetEditorUtilities.h"
 #include "ISequencerModule.h"
@@ -17,26 +14,17 @@
 #include "Editor/RigVMEditorStyle.h"
 #include "Framework/Docking/LayoutExtender.h"
 #include "Framework/Application/SlateApplication.h"
-#include "LevelEditor.h"
-#include "MovieSceneToolsProjectSettings.h"
 #include "Styling/SlateStyle.h"
-#include "WorkflowOrientedApp/WorkflowTabManager.h"
 #include "Modules/ModuleManager.h"
-#include "Widgets/Docking/SDockTab.h"
 #include "EditorModeRegistry.h"
 #include "EditMode/ControlRigEditMode.h"
-#include "ILevelSequenceModule.h"
 #include "EditorModeManager.h"
-#include "EditMode/ControlRigEditMode.h"
 #include "Sequencer/MovieSceneControlRigParameterSection.h"
 #include "Sequencer/MovieSceneControlRigSectionDetailsCustomization.h"
 #include "EditMode/ControlRigEditModeCommands.h"
-#include "Materials/Material.h"
-#include "Framework/MultiBox/MultiBoxExtender.h"
 #include "Framework/MultiBox/MultiBoxBuilder.h"
 #include "ToolMenus.h"
 #include "ContentBrowserMenuContexts.h"
-#include "EditMode/ControlRigEditModeSettings.h"
 #include "Components/SkeletalMeshComponent.h"
 #include "ContentBrowserModule.h"
 #include "AssetRegistry/AssetRegistryModule.h"
@@ -44,22 +32,10 @@
 #include "ControlRigBlueprintActions.h"
 #include "RigVMBlueprintGeneratedClass.h"
 #include "ControlRigGizmoLibraryActions.h"
-#include "Graph/ControlRigGraphSchema.h"
-#include "Graph/ControlRigGraph.h"
-#include "EdGraph/EdGraph.h"
-#include "EdGraph/NodeSpawners/RigVMEdGraphUnitNodeSpawner.h"
-#include "EdGraph/NodeSpawners/RigVMEdGraphVariableNodeSpawner.h"
-#include "EdGraph/NodeSpawners/RigVMEdGraphTemplateNodeSpawner.h"
-#include "EdGraph/NodeSpawners/RigVMEdGraphEnumNodeSpawner.h"
-#include "EdGraph/NodeSpawners/RigVMEdGraphFunctionRefNodeSpawner.h"
-#include "EdGraph/NodeSpawners/RigVMEdGraphInvokeEntryNodeSpawner.h"
 #include "Kismet2/BlueprintEditorUtils.h"
 #include "Kismet2/KismetDebugUtilities.h"
-#include "Graph/ControlRigGraphNode.h"
 #include "EdGraphUtilities.h"
-#include "EdGraph/RigVMEdGraphPanelNodeFactory.h"
 #include "Graph/ControlRigGraphPanelPinFactory.h"
-#include "RigVMBlueprintUtils.h"
 #include "ControlRigBlueprintCommands.h"
 #include "ControlRigHierarchyCommands.h"
 #include "Animation/AnimSequence.h"
@@ -67,21 +43,13 @@
 #include "ControlRigElementDetails.h"
 #include "ControlRigCompilerDetails.h"
 #include "ControlRigDrawingDetails.h"
-#include "ControlRigGraphDetails.h"
 #include "ControlRigAnimGraphDetails.h"
-#include "ControlRigLocalVariableDetails.h"
-#include "ControlRigInfluenceMapDetails.h"
-#include "Animation/AnimSequence.h"
 #include "Subsystems/AssetEditorSubsystem.h"
 #include "Sequencer/ControlRigParameterTrackEditor.h"
 #include "ActorFactories/ActorFactorySkeletalMesh.h"
 #include "ControlRigThumbnailRenderer.h"
 #include "RigVMModel/RigVMController.h"
-#include "RigVMModel/Nodes/RigVMLibraryNode.h"
-#include "RigVMModel/RigVMVariableDescription.h"
 #include "ControlRigBlueprint.h"
-#include "RigVMFunctions/Simulation/RigVMFunction_AlphaInterp.h"
-#include "RigVMFunctions/Debug/RigVMFunction_VisualDebug.h"
 #include "SKismetInspector.h"
 #include "Dialogs/Dialogs.h"
 #include "Settings/ControlRigSettings.h"
@@ -97,18 +65,13 @@
 #include "Editor.h"
 #include "LevelEditorViewport.h"
 #include "Animation/SkeletalMeshActor.h"
-#include "ControlRig.h"
 #include "Sequencer/MovieSceneControlRigParameterTrack.h"
 #include "ControlRigObjectBinding.h"
-#include "EditMode/ControlRigEditMode.h"
 #include "LevelSequenceAnimSequenceLink.h"
-#include "AnimSequenceLevelSequenceLink.h"
 #include "Rigs/FKControlRig.h"
 #include "SBakeToControlRigDialog.h"
 #include "Units/Execution/RigUnit_InverseExecution.h"
 #include "Widgets/SRigVMGraphPinVariableBinding.h"
-#include "Widgets/SRigVMGraphChangePinType.h"
-#include "ControlRigBlueprintFactory.h"
 #include "ControlRigPythonLogDetails.h"
 #include "Dialog/SCustomDialog.h"
 #include "Sequencer/MovieSceneControlRigSpaceChannel.h"
@@ -125,17 +88,10 @@
 #include "ControlConstraintChannelInterface.h"
 #include "RigVMModel/Nodes/RigVMAggregateNode.h"
 #include "RigVMUserWorkflowRegistry.h"
-#include "UserDefinedStructureCompilerUtils.h"
-#include "RigVMModel/Nodes/RigVMDispatchNode.h"
 #include "Units/ControlRigNodeWorkflow.h"
-#include "Units/RigDispatchFactory.h"
 #include "RigVMFunctions/Math/RigVMMathLibrary.h"
 #include "Constraints/ControlRigTransformableHandle.h"
 #include "Constraints/TransformConstraintChannelInterface.h"
-#include "RigVMFunctions/RigVMDispatch_Array.h"
-#include "RigVMFunctions/RigVMDispatch_MakeStruct.h"
-#include "RigVMFunctions/RigVMDispatch_Constant.h"
-#include "EdGraph/RigVMEdGraphPanelNodeFactory.h"
 
 #define LOCTEXT_NAMESPACE "ControlRigEditorModule"
 
@@ -148,16 +104,10 @@ void FControlRigEditorModule::StartupModule()
 	FControlRigHierarchyCommands::Register();
 	FControlRigEditorStyle::Get();
 
-	MenuExtensibilityManager = MakeShareable(new FExtensibilityManager);
-	ToolBarExtensibilityManager = MakeShareable(new FExtensibilityManager);
-
-	// Register Blueprint editor variable customization
-	FBlueprintEditorModule& BlueprintEditorModule = FModuleManager::LoadModuleChecked<FBlueprintEditorModule>("Kismet");
-	BlueprintVariableCustomizationHandle = BlueprintEditorModule.RegisterVariableCustomization(FProperty::StaticClass(), FOnGetVariableCustomizationInstance::CreateStatic(&FControlRigVariableDetailsCustomization::MakeInstance));
-	BlueprintEditorModule.RegisterGraphCustomization(GetDefault<UControlRigGraphSchema>(), FOnGetGraphCustomizationInstance::CreateStatic(&FControlRigGraphDetails::MakeInstance));
+	EdGraphPanelPinFactory = MakeShared<FControlRigGraphPanelPinFactory>();
+	FEdGraphUtilities::RegisterVisualPinFactory(EdGraphPanelPinFactory);
 
-	// Register to fixup newly created BPs
-	FKismetEditorUtilities::RegisterOnBlueprintCreatedCallback(this, UControlRig::StaticClass(), FKismetEditorUtilities::FOnBlueprintCreated::CreateRaw(this, &FControlRigEditorModule::HandleNewBlueprintCreated));
+	StartupModuleCommon();
 
 	// Register details customizations for animation controller nodes
 	FPropertyEditorModule& PropertyEditorModule = FModuleManager::LoadModuleChecked<FPropertyEditorModule>("PropertyEditor");
@@ -233,16 +183,6 @@ void FControlRigEditorModule::StartupModule()
 		false,
 		8500);
 
-
-	ControlRigGraphPanelNodeFactory = MakeShared<FRigVMEdGraphPanelNodeFactory>();
-	FEdGraphUtilities::RegisterVisualNodeFactory(ControlRigGraphPanelNodeFactory);
-
-	ControlRigGraphPanelPinFactory = MakeShared<FControlRigGraphPanelPinFactory>();
-	FEdGraphUtilities::RegisterVisualPinFactory(ControlRigGraphPanelPinFactory);
-
-	ReconstructAllNodesDelegateHandle = FBlueprintEditorUtils::OnReconstructAllNodesEvent.AddStatic(&FRigVMBlueprintUtils::HandleReconstructAllNodes);
-	RefreshAllNodesDelegateHandle = FBlueprintEditorUtils::OnRefreshAllNodesEvent.AddStatic(&FRigVMBlueprintUtils::HandleRefreshAllNodes);
-
 	ICurveEditorModule& CurveEditorModule = FModuleManager::LoadModuleChecked<ICurveEditorModule>("CurveEditor");
 	FControlRigSpaceChannelCurveModel::ViewID = CurveEditorModule.RegisterView(FOnCreateCurveEditorView::CreateStatic(
 		[](TWeakPtr<FCurveEditor> WeakCurveEditor) -> TSharedRef<SCurveEditorView>
@@ -280,18 +220,15 @@ void FControlRigEditorModule::ShutdownModule()
 		CurveEditorModule->UnregisterView(FControlRigSpaceChannelCurveModel::ViewID);
 	}
 
+	ShutdownModuleCommon();
+
 	//UThumbnailManager::Get().UnregisterCustomRenderer(UControlRigBlueprint::StaticClass());
 	//UActorFactorySkeletalMesh::UnregisterDelegatesForAssetClass(UControlRigBlueprint::StaticClass());
 
-	FBlueprintEditorUtils::OnRefreshAllNodesEvent.Remove(RefreshAllNodesDelegateHandle);
-	FBlueprintEditorUtils::OnReconstructAllNodesEvent.Remove(ReconstructAllNodesDelegateHandle);
-
-	FEdGraphUtilities::UnregisterVisualPinFactory(ControlRigGraphPanelPinFactory);
-	FEdGraphUtilities::UnregisterVisualNodeFactory(ControlRigGraphPanelNodeFactory);
-
 	FEditorModeRegistry::Get().UnregisterMode(FControlRigEditorEditMode::ModeName);
 	FEditorModeRegistry::Get().UnregisterMode(FControlRigEditMode::ModeName);
 
+	FEdGraphUtilities::UnregisterVisualPinFactory(EdGraphPanelPinFactory);
 
 	ISequencerModule* SequencerModule = FModuleManager::GetModulePtr<ISequencerModule>("Sequencer");
 	if (SequencerModule)
@@ -308,18 +245,6 @@ void FControlRigEditorModule::ShutdownModule()
 		}
 	}
 
-	FKismetEditorUtilities::UnregisterAutoBlueprintNodeCreation(this);
-
-	if (!IsEngineExitRequested())
-	{
-		FBlueprintEditorModule* BlueprintEditorModule = FModuleManager::GetModulePtr<FBlueprintEditorModule>("Kismet");
-		if (BlueprintEditorModule)
-		{
-			BlueprintEditorModule->UnregisterVariableCustomization(FProperty::StaticClass(), BlueprintVariableCustomizationHandle);
-			BlueprintEditorModule->UnregisterGraphCustomization(GetDefault<UControlRigGraphSchema>());
-		}
-	}
-
 	FPropertyEditorModule* PropertyEditorModule = FModuleManager::GetModulePtr<FPropertyEditorModule>("PropertyEditor");
 	if (PropertyEditorModule)
 	{
@@ -347,6 +272,373 @@ void FControlRigEditorModule::ShutdownModule()
 	WorkflowHandles.Reset();
 }
 
+UClass* FControlRigEditorModule::GetRigVMBlueprintClass() const
+{
+	return UControlRigBlueprint::StaticClass();
+}
+
+void FControlRigEditorModule::GetNodeContextMenuActions(URigVMBlueprint* RigVMBlueprint,
+	const URigVMEdGraphNode* EdGraphNode, URigVMNode* ModelNode, UToolMenu* Menu) const
+{
+	FRigVMEditorModule::GetNodeContextMenuActions(RigVMBlueprint, EdGraphNode, ModelNode, Menu);
+
+	UControlRigBlueprint* ControlRigBlueprint = Cast<UControlRigBlueprint>(RigVMBlueprint);
+	if(ControlRigBlueprint == nullptr)
+	{
+		return;
+	}
+
+	URigVMGraph* Model = RigVMBlueprint->GetModel(EdGraphNode->GetGraph());
+	URigVMController* Controller = RigVMBlueprint->GetController(Model);
+
+	TArray<FName> SelectedNodeNames = Model->GetSelectNodes();
+	SelectedNodeNames.AddUnique(ModelNode->GetFName());
+
+	URigHierarchy* TemporaryHierarchy = NewObject<URigHierarchy>();
+	TemporaryHierarchy->CopyHierarchy(ControlRigBlueprint->Hierarchy);
+
+	TArray<FRigElementKey> RigElementsToSelect;
+	TMap<const URigVMPin*, FRigElementKey> PinToKey;
+
+	for(const FName& SelectedNodeName : SelectedNodeNames)
+	{
+		if (URigVMNode* FoundNode = Model->FindNodeByName(SelectedNodeName))
+		{
+			TSharedPtr<FStructOnScope> StructOnScope;
+			FRigVMStruct* StructMemory = nullptr;
+			UScriptStruct* ScriptStruct = nullptr;
+			if (URigVMUnitNode* UnitNode = Cast<URigVMUnitNode>(FoundNode))
+			{
+				ScriptStruct = UnitNode->GetScriptStruct();
+				if(ScriptStruct)
+				{
+					StructOnScope = UnitNode->ConstructStructInstance(false);
+					if(StructOnScope->GetStruct()->IsChildOf(FRigVMStruct::StaticStruct()))
+					{
+						StructMemory = (FRigVMStruct*)StructOnScope->GetStructMemory();
+						StructMemory->Execute();
+					}
+				}
+			}
+
+			const TArray<URigVMPin*> AllPins = FoundNode->GetAllPinsRecursively();
+			for (const URigVMPin* Pin : AllPins)
+			{
+				if (Pin->GetCPPType() == TEXT("FName"))
+				{
+					FRigElementKey Key;
+					if (Pin->GetCustomWidgetName() == TEXT("BoneName"))
+					{
+						Key = FRigElementKey(*Pin->GetDefaultValue(), ERigElementType::Bone);
+					}
+					else if (Pin->GetCustomWidgetName() == TEXT("ControlName"))
+					{
+						Key = FRigElementKey(*Pin->GetDefaultValue(), ERigElementType::Control);
+					}
+					else if (Pin->GetCustomWidgetName() == TEXT("SpaceName"))
+					{
+						Key = FRigElementKey(*Pin->GetDefaultValue(), ERigElementType::Null);
+					}
+					else if (Pin->GetCustomWidgetName() == TEXT("CurveName"))
+					{
+						Key = FRigElementKey(*Pin->GetDefaultValue(), ERigElementType::Curve);
+					}
+					else
+					{
+						continue;
+					}
+
+					RigElementsToSelect.AddUnique(Key);
+					PinToKey.Add(Pin, Key);
+				}
+				else if (Pin->GetCPPTypeObject() == FRigElementKey::StaticStruct() && !Pin->IsArray())
+				{
+					if (StructMemory == nullptr)
+					{
+						FString DefaultValue = Pin->GetDefaultValue();
+						if (!DefaultValue.IsEmpty())
+						{
+							FRigElementKey Key;
+							FRigElementKey::StaticStruct()->ImportText(*DefaultValue, &Key, nullptr, EPropertyPortFlags::PPF_None, nullptr, FRigElementKey::StaticStruct()->GetName(), true);
+							if (Key.IsValid())
+							{
+								RigElementsToSelect.AddUnique(Key);
+								if (URigVMPin* NamePin = Pin->FindSubPin(TEXT("Name")))
+								{
+									PinToKey.Add(NamePin, Key);
+								}
+							}
+						}
+					}
+					else
+					{
+						check(ScriptStruct);
+
+						TArray<FString> PropertyNames; 
+						if(!URigVMPin::SplitPinPath(Pin->GetSegmentPath(true), PropertyNames))
+						{
+							PropertyNames.Add(Pin->GetName());
+						}
+
+						UScriptStruct* Struct = ScriptStruct;
+						uint8* Memory = (uint8*)StructMemory; 
+
+						while(!PropertyNames.IsEmpty())
+						{
+							FString PropertyName;
+							PropertyNames.HeapPop(PropertyName);
+							
+							const FProperty* Property = ScriptStruct->FindPropertyByName(*PropertyName);
+							if(Property == nullptr)
+							{
+								Memory = nullptr;
+								break;
+							}
+
+							Memory = Property->ContainerPtrToValuePtr<uint8>(Memory);
+							
+							if(PropertyNames.IsEmpty())
+							{
+								continue;
+							}
+							
+							if(const FArrayProperty* ArrayProperty = CastField<FArrayProperty>(Property))
+							{
+								PropertyNames.HeapPop(PropertyName);
+
+								int32 ArrayIndex = FCString::Atoi(*PropertyName);
+								FScriptArrayHelper Helper(ArrayProperty, Memory);
+								if(!Helper.IsValidIndex(ArrayIndex))
+								{
+									Memory = nullptr;
+									break;
+								}
+
+								Memory = Helper.GetRawPtr(ArrayIndex);
+								Property = ArrayProperty->Inner;
+							}
+
+							if(const FStructProperty* StructProperty = CastField<FStructProperty>(Property))
+							{
+								Struct = StructProperty->Struct;
+							}
+						}
+
+						if(Memory)
+						{
+							const FRigElementKey& Key = *(const FRigElementKey*)Memory;
+							if (Key.IsValid())
+							{
+								RigElementsToSelect.AddUnique(Key);
+
+								if (URigVMPin* NamePin = Pin->FindSubPin(TEXT("Name")))
+								{
+									PinToKey.Add(NamePin, Key);
+								}
+							}
+						}
+					}
+				}
+				else if (Pin->GetCPPTypeObject() == FRigElementKeyCollection::StaticStruct() && Pin->GetDirection() == ERigVMPinDirection::Output)
+				{
+					if (StructMemory == nullptr)
+					{
+						// not supported for now
+					}
+					else
+					{
+						check(ScriptStruct);
+						if (const FProperty* Property = ScriptStruct->FindPropertyByName(Pin->GetFName()))
+						{
+							const FRigElementKeyCollection& Collection = *Property->ContainerPtrToValuePtr<FRigElementKeyCollection>(StructMemory);
+
+							if (Collection.Num() > 0)
+							{
+								RigElementsToSelect.Reset();
+								for (const FRigElementKey& Item : Collection)
+								{
+									RigElementsToSelect.AddUnique(Item);
+								}
+								break;
+							}
+						}
+					}
+				}
+			}
+		}
+	}
+
+	if (RigElementsToSelect.Num() > 0)
+	{
+		FToolMenuSection& Section = Menu->AddSection("RigVMEditorContextMenuHierarchy", LOCTEXT("HierarchyHeader", "Hierarchy"));
+		Section.AddMenuEntry(
+			"SelectRigElements",
+			LOCTEXT("SelectRigElements", "Select Rig Elements"),
+			LOCTEXT("SelectRigElements_Tooltip", "Selects the bone, controls or nulls associated with this node."),
+			FSlateIcon(),
+			FUIAction(FExecuteAction::CreateLambda([ControlRigBlueprint, RigElementsToSelect]() {
+
+				ControlRigBlueprint->GetHierarchyController()->SetSelection(RigElementsToSelect);
+				
+			})
+		));
+	}
+
+	if (RigElementsToSelect.Num() > 0)
+	{
+		FToolMenuSection& Section = Menu->AddSection("RigVMEditorContextMenuHierarchy", LOCTEXT("ToolsHeader", "Tools"));
+		Section.AddMenuEntry(
+			"SearchAndReplaceNames",
+			LOCTEXT("SearchAndReplaceNames", "Search & Replace / Mirror"),
+			LOCTEXT("SearchAndReplaceNames_Tooltip", "Searches within all names and replaces with a different text."),
+			FSlateIcon(),
+			FUIAction(FExecuteAction::CreateLambda([ControlRigBlueprint, Controller, PinToKey]() {
+
+				FRigVMMirrorSettings Settings;
+				TSharedPtr<FStructOnScope> StructToDisplay = MakeShareable(new FStructOnScope(FRigVMMirrorSettings::StaticStruct(), (uint8*)&Settings));
+
+				TSharedRef<SKismetInspector> KismetInspector = SNew(SKismetInspector);
+				KismetInspector->ShowSingleStruct(StructToDisplay);
+
+				TSharedRef<SCustomDialog> MirrorDialog = SNew(SCustomDialog)
+					.Title(FText(LOCTEXT("ControlRigHierarchyMirror", "Mirror Graph")))
+					.Content()
+					[
+						KismetInspector
+					]
+					.Buttons({
+						SCustomDialog::FButton(LOCTEXT("OK", "OK")),
+						SCustomDialog::FButton(LOCTEXT("Cancel", "Cancel"))
+				});
+				if (MirrorDialog->ShowModal() == 0)
+				{
+					Controller->OpenUndoBracket(TEXT("Mirroring Graph"));
+					int32 ReplacedNames = 0;
+					TArray<FString> UnchangedItems;
+
+					for (const TPair<const URigVMPin*, FRigElementKey>& Pair : PinToKey)
+					{
+						const URigVMPin* Pin = Pair.Key;
+						FRigElementKey Key = Pair.Value;
+
+						if (Key.Name.IsNone())
+						{
+							continue;
+						}
+
+						FString OldNameStr = Key.Name.ToString();
+						FString NewNameStr = OldNameStr.Replace(*Settings.SearchString, *Settings.ReplaceString, ESearchCase::CaseSensitive);
+						if(NewNameStr != OldNameStr)
+						{
+							Key.Name = *NewNameStr;
+							if(ControlRigBlueprint->Hierarchy->GetIndex(Key) != INDEX_NONE)
+							{
+								Controller->SetPinDefaultValue(Pin->GetPinPath(), NewNameStr, false, true, false, true);
+								ReplacedNames++;
+							}
+							else
+							{
+								// save the names of the items that we skipped during this search & replace
+								UnchangedItems.AddUnique(OldNameStr);
+							} 
+						}
+					}
+
+					if (UnchangedItems.Num() > 0)
+					{
+						FString ListOfUnchangedItems;
+						for (int Index = 0; Index < UnchangedItems.Num(); Index++)
+						{
+							// construct the string "item1, item2, item3"
+							ListOfUnchangedItems += UnchangedItems[Index];
+							if (Index != UnchangedItems.Num() - 1)
+							{
+								ListOfUnchangedItems += TEXT(", ");
+							}
+						}
+						
+						// inform the user that some items were skipped due to invalid new names
+						Controller->ReportAndNotifyError(FString::Printf(TEXT("Invalid Names after Search & Replace, action skipped for %s"), *ListOfUnchangedItems));
+					}
+
+					if (ReplacedNames > 0)
+					{ 
+						Controller->CloseUndoBracket();
+					}
+					else
+					{
+						Controller->CancelUndoBracket();
+					}
+				}
+			})
+		));
+	}
+
+	if (URigVMUnitNode* UnitNode = Cast<URigVMUnitNode>(EdGraphNode->GetModelNode()))
+	{
+		FToolMenuSection& SettingsSection = Menu->AddSection("RigVMEditorContextMenuSettings", LOCTEXT("SettingsHeader", "Settings"));
+		SettingsSection.AddMenuEntry(
+			"Save Default Expansion State",
+			LOCTEXT("SaveDefaultExpansionState", "Save Default Expansion State"),
+			LOCTEXT("SaveDefaultExpansionState_Tooltip", "Saves the expansion state of all pins of the node as the default."),
+			FSlateIcon(),
+			FUIAction(FExecuteAction::CreateLambda([UnitNode]() {
+
+#if WITH_EDITORONLY_DATA
+
+				FScopedTransaction Transaction(LOCTEXT("RigUnitDefaultExpansionStateChanged", "Changed Rig Unit Default Expansion State"));
+				UControlRigEditorSettings::Get()->Modify();
+
+				FControlRigSettingsPerPinBool& ExpansionMap = UControlRigEditorSettings::Get()->RigUnitPinExpansion.FindOrAdd(UnitNode->GetScriptStruct()->GetName());
+				ExpansionMap.Values.Empty();
+
+				TArray<URigVMPin*> Pins = UnitNode->GetAllPinsRecursively();
+				for (URigVMPin* Pin : Pins)
+				{
+					if (Pin->GetSubPins().Num() == 0)
+					{
+						continue;
+					}
+
+					FString PinPath = Pin->GetPinPath();
+					FString NodeName, RemainingPath;
+					URigVMPin::SplitPinPathAtStart(PinPath, NodeName, RemainingPath);
+					ExpansionMap.Values.FindOrAdd(RemainingPath) = Pin->IsExpanded();
+				}
+#endif
+			})
+		));
+	}
+}
+
+void FControlRigEditorModule::GetPinContextMenuActions(URigVMBlueprint* RigVMBlueprint, const UEdGraphPin* EdGraphPin, URigVMPin* ModelPin, UToolMenu* Menu) const
+{
+	FRigVMEditorModule::GetPinContextMenuActions(RigVMBlueprint, EdGraphPin, ModelPin, Menu);
+
+    // Add direct manipulation context menu entries
+	if(UControlRigBlueprint* ControlRigBlueprint = Cast<UControlRigBlueprint>(RigVMBlueprint))
+	{
+		if ((ModelPin->GetCPPType() == TEXT("FVector") ||
+			 ModelPin->GetCPPType() == TEXT("FQuat") ||
+			 ModelPin->GetCPPType() == TEXT("FTransform")) &&
+			(ModelPin->GetDirection() == ERigVMPinDirection::Input ||
+			 ModelPin->GetDirection() == ERigVMPinDirection::IO) &&
+			 ModelPin->GetPinForLink()->GetRootPin()->GetSourceLinks(true).Num() == 0)
+		{
+			FToolMenuSection& Section = Menu->AddSection("RigVMEditorContextMenuControlPin", LOCTEXT("ControlPin", "Direct Manipulation"));
+			Section.AddMenuEntry(
+				"DirectManipControlPin",
+				LOCTEXT("DirectManipControlPin", "Control Pin Value"),
+				LOCTEXT("DirectManipControlPin_Tooltip", "Configures the pin for direct interaction in the viewport"),
+				FSlateIcon(),
+				FUIAction(FExecuteAction::CreateLambda([ControlRigBlueprint, ModelPin]() {
+					ControlRigBlueprint->AddTransientControl(ModelPin);
+				})
+			));
+		}
+	}
+}
+
 TSharedRef< SWidget > FControlRigEditorModule::GenerateAnimationMenu(TWeakPtr<IAnimationEditor> InAnimationEditor)
 {
 	const bool bShouldCloseWindowAfterMenuSelection = true;
@@ -988,11 +1280,6 @@ void FControlRigEditorModule::ExtendAnimSequenceMenu()
 		}));
 }
 
-void FControlRigEditorModule::HandleNewBlueprintCreated(UBlueprint* InBlueprint)
-{
-	UControlRigBlueprintFactory::CreateRigGraphIfRequired(Cast<UControlRigBlueprint>(InBlueprint));
-}
-
 TSharedRef<IControlRigEditor> FControlRigEditorModule::CreateControlRigEditor(const EToolkitMode::Type Mode, const TSharedPtr< IToolkitHost >& InitToolkitHost, class UControlRigBlueprint* InBlueprint)
 {
 	TSharedRef< FControlRigEditor > NewControlRigEditor(new FControlRigEditor());
@@ -1000,1835 +1287,6 @@ TSharedRef<IControlRigEditor> FControlRigEditorModule::CreateControlRigEditor(co
 	return NewControlRigEditor;
 }
 
-void FControlRigEditorModule::GetTypeActions(UControlRigBlueprint* CRB, FBlueprintActionDatabaseRegistrar& ActionRegistrar)
-{
-	FAssetRegistryModule& AssetRegistryModule = FModuleManager::LoadModuleChecked<FAssetRegistryModule>("AssetRegistry");
-
-	// actions get registered under specific object-keys; the idea is that 
-	// actions might have to be updated (or deleted) if their object-key is  
-	// mutated (or removed)... here we use the class (so if the class 
-	// type disappears, then the action should go with it)
-	UClass* ActionKey = CRB->GetClass();
-	// to keep from needlessly instantiating a UBlueprintNodeSpawner, first   
-	// check to make sure that the registrar is looking for actions of this type
-	// (could be regenerating actions for a specific asset, and therefore the 
-	// registrar would only accept actions corresponding to that asset)
-	if (!ActionRegistrar.IsOpenForRegistration(ActionKey))
-	{
-		return;
-	}
-
-	FRigVMRegistry& Registry = FRigVMRegistry::Get();
-	Registry.RefreshEngineTypes();
-
-	for (const FRigVMTemplate& Template : Registry.GetTemplates())
-	{
-		// factories are registered below
-		if(Template.UsesDispatch())
-		{
-			continue;
-		}
-
-		// ignore templates that have only one permutation
-		if (Template.NumPermutations() <= 1)
-		{
-			continue;
-		}
-
-		// ignore templates which don't have a function backing it up
-		if(Template.GetPermutation(0) == nullptr)
-		{
-			continue;
-		}
-
-		if(!Template.SupportsExecuteContextStruct(FControlRigExecuteContext::StaticStruct()))
-		{
-			continue;
-		}
-
-		FText NodeCategory = FText::FromString(Template.GetCategory());
-		FText MenuDesc = FText::FromName(Template.GetName());
-		FText ToolTip = Template.GetTooltipText();
-
-		UBlueprintNodeSpawner* NodeSpawner = URigVMEdGraphTemplateNodeSpawner::CreateFromNotation(Template.GetNotation(), MenuDesc, NodeCategory, ToolTip);
-		check(NodeSpawner != nullptr);
-		ActionRegistrar.AddBlueprintAction(ActionKey, NodeSpawner);
-	};
-
-	for (const FRigVMDispatchFactory* Factory : Registry.GetFactories())
-	{
-		if(!Factory->SupportsExecuteContextStruct(FControlRigExecuteContext::StaticStruct()))
-		{
-			continue;
-		}
-
-		const FRigVMTemplate* Template = Factory->GetTemplate();
-		if(Template == nullptr)
-		{
-			continue;
-		}
-
-		FText NodeCategory = FText::FromString(Factory->GetCategory());
-		FText MenuDesc = FText::FromString(Factory->GetNodeTitle(FRigVMTemplateTypeMap()));
-		FText ToolTip = Factory->GetNodeTooltip(FRigVMTemplateTypeMap());
-
-		UBlueprintNodeSpawner* NodeSpawner = URigVMEdGraphTemplateNodeSpawner::CreateFromNotation(Template->GetNotation(), MenuDesc, NodeCategory, ToolTip);
-		check(NodeSpawner != nullptr);
-		ActionRegistrar.AddBlueprintAction(ActionKey, NodeSpawner);
-	};
-
-	// Add all rig units
-	for(const FRigVMFunction& Function : Registry.GetFunctions())
-	{
-		UScriptStruct* Struct = Function.Struct;
-		if (Struct == nullptr || !Struct->IsChildOf(FRigVMStruct::StaticStruct()))
-		{
-			continue;
-		}
-
-		if(!Function.SupportsExecuteContextStruct(FControlRigExecuteContext::StaticStruct()))
-		{
-			continue;
-		}
-
-		// skip rig units which have a template
-		if (Function.GetTemplate())
-		{
-			continue;
-		}
-
-		// skip deprecated units
-		if(Function.Struct->HasMetaData(FRigVMStruct::DeprecatedMetaName))
-		{
-			continue;
-		}
-
-		FString CategoryMetadata, DisplayNameMetadata, MenuDescSuffixMetadata;
-		Struct->GetStringMetaDataHierarchical(FRigVMStruct::CategoryMetaName, &CategoryMetadata);
-		Struct->GetStringMetaDataHierarchical(FRigVMStruct::DisplayNameMetaName, &DisplayNameMetadata);
-		Struct->GetStringMetaDataHierarchical(FRigVMStruct::MenuDescSuffixMetaName, &MenuDescSuffixMetadata);
-
-		if(DisplayNameMetadata.IsEmpty())
-		{
-			DisplayNameMetadata = Struct->GetDisplayNameText().ToString();
-		}
-		if (!MenuDescSuffixMetadata.IsEmpty())
-		{
-			MenuDescSuffixMetadata = TEXT(" ") + MenuDescSuffixMetadata;
-		}
-		FText NodeCategory = FText::FromString(CategoryMetadata);
-		FText MenuDesc = FText::FromString(DisplayNameMetadata + MenuDescSuffixMetadata);
-		FText ToolTip = Struct->GetToolTipText();
-
-		UBlueprintNodeSpawner* NodeSpawner = URigVMEdGraphUnitNodeSpawner::CreateFromStruct(Struct, Function.GetMethodName(), MenuDesc, NodeCategory, ToolTip);
-		check(NodeSpawner != nullptr);
-		ActionRegistrar.AddBlueprintAction(ActionKey, NodeSpawner);
-	};
-
-	for (TObjectIterator<UEnum> EnumIt; EnumIt; ++EnumIt)
-	{
-		UEnum* EnumToConsider = (*EnumIt);
-
-		if (EnumToConsider->HasMetaData(TEXT("Hidden")))
-		{
-			continue;
-		}
-
-		if (EnumToConsider->IsEditorOnly())
-		{
-			continue;
-		}
-
-		if(EnumToConsider->IsNative())
-		{
-			continue;
-		}
-
-		FText NodeCategory = FText::FromString(TEXT("Enum"));
-		FText MenuDesc = FText::FromString(FString::Printf(TEXT("Enum %s"), *EnumToConsider->GetName()));
-		FText ToolTip = MenuDesc;
-
-		UBlueprintNodeSpawner* NodeSpawner = URigVMEdGraphEnumNodeSpawner::CreateForEnum(EnumToConsider, MenuDesc, NodeCategory, ToolTip);
-		check(NodeSpawner != nullptr);
-		ActionRegistrar.AddBlueprintAction(ActionKey, NodeSpawner);
-	}
-
-	FArrayProperty* PublicGraphFunctionsProperty = CastField<FArrayProperty>(UControlRigBlueprint::StaticClass()->FindPropertyByName(TEXT("PublicGraphFunctions")));
-	FArrayProperty* PublicFunctionsProperty = CastField<FArrayProperty>(UControlRigBlueprint::StaticClass()->FindPropertyByName(TEXT("PublicFunctions")));
-	if(PublicGraphFunctionsProperty || PublicFunctionsProperty)
-	{
-		// find all control rigs in the project
-		TArray<FAssetData> ControlRigAssetDatas;
-		FARFilter ControlRigAssetFilter;
-		ControlRigAssetFilter.ClassPaths.Add(UControlRigBlueprint::StaticClass()->GetClassPathName());
-		ControlRigAssetFilter.ClassPaths.Add(URigVMBlueprintGeneratedClass::StaticClass()->GetClassPathName());
-		AssetRegistryModule.Get().GetAssets(ControlRigAssetFilter, ControlRigAssetDatas);
-
-		// loop over all control rigs in the project
-		TSet<FName> PackagesProcessed;
-		for(const FAssetData& ControlRigAssetData : ControlRigAssetDatas)
-		{
-			// Avoid duplication of spawners
-			if (PackagesProcessed.Contains(ControlRigAssetData.PackageName))
-			{
-				continue;
-			}
-			PackagesProcessed.Add(ControlRigAssetData.PackageName);
-
-			FString PublicGraphFunctionsString;
-			FString PublicFunctionsString;
-			if (PublicGraphFunctionsProperty)
-			{
-				PublicGraphFunctionsString = ControlRigAssetData.GetTagValueRef<FString>(PublicGraphFunctionsProperty->GetFName());
-			}
-			// Only look at the deprecated public functions if the PublicGraphFunctionsString is empty
-			if (PublicGraphFunctionsString.IsEmpty() && PublicFunctionsProperty)
-			{
-				PublicFunctionsString = ControlRigAssetData.GetTagValueRef<FString>(PublicFunctionsProperty->GetFName());
-			}
-			
-			if(PublicFunctionsString.IsEmpty() && PublicGraphFunctionsString.IsEmpty())
-			{
-				continue;
-			}
-
-			if (PublicFunctionsProperty && !PublicFunctionsString.IsEmpty())
-			{
-				TArray<FRigVMOldPublicFunctionData> PublicFunctions;
-				PublicFunctionsProperty->ImportText_Direct(*PublicFunctionsString, &PublicFunctions, nullptr, EPropertyPortFlags::PPF_None);
-				for(const FRigVMOldPublicFunctionData& PublicFunction : PublicFunctions)
-				{
-					UBlueprintNodeSpawner* NodeSpawner = URigVMEdGraphFunctionRefNodeSpawner::CreateFromAssetData(ControlRigAssetData, PublicFunction);
-					check(NodeSpawner != nullptr);
-					ActionRegistrar.AddBlueprintAction(ActionKey, NodeSpawner);
-				}
-			}
-
-			if (PublicGraphFunctionsProperty && !PublicGraphFunctionsString.IsEmpty())
-			{
-				TArray<FRigVMGraphFunctionHeader> PublicFunctions;
-				PublicGraphFunctionsProperty->ImportText_Direct(*PublicGraphFunctionsString, &PublicFunctions, nullptr, EPropertyPortFlags::PPF_None);
-				for(const FRigVMGraphFunctionHeader& PublicFunction : PublicFunctions)
-				{
-					UBlueprintNodeSpawner* NodeSpawner = URigVMEdGraphFunctionRefNodeSpawner::CreateFromAssetData(ControlRigAssetData, PublicFunction);
-					check(NodeSpawner != nullptr);
-					ActionRegistrar.AddBlueprintAction(ActionKey, NodeSpawner);
-				}
-			}
-		}
-	}
-}
-
-void FControlRigEditorModule::GetInstanceActions(UControlRigBlueprint* CRB, FBlueprintActionDatabaseRegistrar& ActionRegistrar)
-{
-	if (URigVMBlueprintGeneratedClass* GeneratedClass = CRB->GetRigVMBlueprintGeneratedClass())
-	{
-		if (UControlRig* CDO = Cast<UControlRig>(GeneratedClass->GetDefaultObject()))
-		{
-			static const FText NodeCategory = LOCTEXT("Variables", "Variables");
-
-			TArray<FRigVMExternalVariable> ExternalVariables = CDO->GetExternalVariables();
-			for (const FRigVMExternalVariable& ExternalVariable : ExternalVariables)
-			{
-				FText MenuDesc = FText::FromName(ExternalVariable.Name);
-				FText ToolTip = FText::FromString(FString::Printf(TEXT("Get the value of variable %s"), *ExternalVariable.Name.ToString()));
-				ActionRegistrar.AddBlueprintAction(GeneratedClass, URigVMEdGraphVariableNodeSpawner::CreateFromExternalVariable(CRB, ExternalVariable, true, MenuDesc, NodeCategory, ToolTip));
-
-				ToolTip = FText::FromString(FString::Printf(TEXT("Set the value of variable %s"), *ExternalVariable.Name.ToString()));
-				ActionRegistrar.AddBlueprintAction(GeneratedClass, URigVMEdGraphVariableNodeSpawner::CreateFromExternalVariable(CRB, ExternalVariable, false, MenuDesc, NodeCategory, ToolTip));
-			}
-		}
-
-		if (URigVMFunctionLibrary* LocalFunctionLibrary = CRB->GetLocalFunctionLibrary())
-		{
-			TArray<URigVMLibraryNode*> Functions = LocalFunctionLibrary->GetFunctions();
-			for (URigVMLibraryNode* Function : Functions)
-			{
-				UBlueprintNodeSpawner* NodeSpawner = URigVMEdGraphFunctionRefNodeSpawner::CreateFromFunction(Function);
-				check(NodeSpawner != nullptr);
-				ActionRegistrar.AddBlueprintAction(GeneratedClass, NodeSpawner);
-			}
-
-			static const FText NodeCategory = LOCTEXT("LocalVariables", "Local Variables");
-			for (URigVMLibraryNode* Function : Functions)
-			{
-				for (const FRigVMGraphVariableDescription& LocalVariable : Function->GetContainedGraph()->GetLocalVariables())
-				{
-					FText MenuDesc = FText::FromName(LocalVariable.Name);
-					FText ToolTip = FText::FromString(FString::Printf(TEXT("Get the value of variable %s"), *LocalVariable.Name.ToString()));
-					ActionRegistrar.AddBlueprintAction(GeneratedClass, URigVMEdGraphVariableNodeSpawner::CreateFromLocalVariable(CRB, Function->GetContainedGraph(), LocalVariable, true, MenuDesc, NodeCategory, ToolTip));
-
-					ToolTip = FText::FromString(FString::Printf(TEXT("Set the value of variable %s"), *LocalVariable.Name.ToString()));
-					ActionRegistrar.AddBlueprintAction(GeneratedClass, URigVMEdGraphVariableNodeSpawner::CreateFromLocalVariable(CRB, Function->GetContainedGraph(), LocalVariable, false, MenuDesc, NodeCategory, ToolTip));
-				}
-			}
-		}
-
-		for (URigVMGraph* Graph : CRB->GetAllModels())
-		{
-			if (Graph->GetEntryNode())
-			{
-				static const FText NodeCategory = LOCTEXT("InputArguments", "Input Arguments");
-				for (const FRigVMGraphVariableDescription& InputArgument : Graph->GetInputArguments())
-				{
-					FText MenuDesc = FText::FromName(InputArgument.Name);
-					FText ToolTip = FText::FromString(FString::Printf(TEXT("Get the value of input %s"), *InputArgument.Name.ToString()));
-					ActionRegistrar.AddBlueprintAction(GeneratedClass, URigVMEdGraphVariableNodeSpawner::CreateFromLocalVariable(CRB, Graph, InputArgument, true, MenuDesc, NodeCategory, ToolTip));
-				}			
-			}
-		}
-
-		const TArray<FName> EntryNames = CRB->GetRigVMClient()->GetEntryNames();
-		if(!EntryNames.IsEmpty())
-		{
-			static const FText NodeCategory = LOCTEXT("Events", "Events");
-			for (const FName& EntryName : EntryNames)
-			{
-				static constexpr TCHAR EventStr[] = TEXT("Event");
-				static const FString EventSuffix = FString::Printf(TEXT(" %s"), EventStr);
-				FString Suffix = EntryName.ToString().EndsWith(EventStr) ? FString() : EventSuffix;
-				FText MenuDesc = FText::FromString(FString::Printf(TEXT("Run %s%s"), *EntryName.ToString(), *Suffix));
-				FText ToolTip = FText::FromString(FString::Printf(TEXT("Runs the %s%s"), *EntryName.ToString(), *Suffix));
-				ActionRegistrar.AddBlueprintAction(GeneratedClass, URigVMEdGraphInvokeEntryNodeSpawner::CreateForEntry(CRB, EntryName, MenuDesc, NodeCategory, ToolTip));
-			}
-		}
-	}
-}
-
-FConnectionDrawingPolicy* FControlRigEditorModule::CreateConnectionDrawingPolicy(int32 InBackLayerID, int32 InFrontLayerID, float InZoomFactor, const FSlateRect& InClippingRect, class FSlateWindowElementList& InDrawElements, class UEdGraph* InGraphObj)
-{
-	return new FRigVMEdGraphConnectionDrawingPolicy(InBackLayerID, InFrontLayerID, InZoomFactor, InClippingRect, InDrawElements, InGraphObj);
-}
-
-void FControlRigEditorModule::GetContextMenuActions(const UControlRigGraphSchema* Schema, UToolMenu* Menu, UGraphNodeContextMenuContext* Context) const
-{
-	if (Menu && Context)
-	{
-		Schema->UEdGraphSchema::GetContextMenuActions(Menu, Context);
-
-		auto ShowWorkflowOptionsDialog = [](URigVMUserWorkflowOptions* InOptions)
-		{
-			FPropertyEditorModule& EditModule = FModuleManager::Get().GetModuleChecked<FPropertyEditorModule>("PropertyEditor");
-
-			FDetailsViewArgs DetailsViewArgs;
-			DetailsViewArgs.NameAreaSettings = FDetailsViewArgs::HideNameArea;
-			DetailsViewArgs.bHideSelectionTip = true;
-			DetailsViewArgs.DefaultsOnlyVisibility = EEditDefaultsOnlyNodeVisibility::Hide;
-			DetailsViewArgs.bAllowSearch = false;
-
-			const TSharedRef<class IDetailsView> PropertyView = EditModule.CreateDetailView( DetailsViewArgs );
-			PropertyView->SetObject(InOptions);
-
-			TSharedRef<SCustomDialog> OptionsDialog = SNew(SCustomDialog)
-				.Title(FText(LOCTEXT("ControlRigWorkflowOptions", "Options")))
-				.Content()
-				[
-					PropertyView
-				]
-				.Buttons({
-					SCustomDialog::FButton(LOCTEXT("OK", "OK")),
-					SCustomDialog::FButton(LOCTEXT("Cancel", "Cancel"))
-			});
-			return OptionsDialog->ShowModal() == 0;
-		};
-
-		if (UEdGraphPin* InGraphPin = (UEdGraphPin* )Context->Pin)
-		{
-			UEdGraph* Graph = InGraphPin->GetOwningNode()->GetGraph();
-
-			// per pin workflows
-			if(UControlRigBlueprint* RigBlueprint = Cast<UControlRigBlueprint>((UBlueprint*)Context->Blueprint))
-			{
-				if (URigVMPin* ModelPin = RigBlueprint->GetModel(Graph)->FindPin(InGraphPin->GetName()))
-				{
-					const TArray<FRigVMUserWorkflow> Workflows = ModelPin->GetNode()->GetSupportedWorkflows(ERigVMUserWorkflowType::PinContext, ModelPin);
-					if(!Workflows.IsEmpty())
-					{
-						URigVMController* Controller = RigBlueprint->GetController(ModelPin->GetGraph());
-
-						FToolMenuSection& SettingsSection = Menu->AddSection("EdGraphSchemaWorkflow", LOCTEXT("WorkflowHeader", "Workflow"));
-
-						for(const FRigVMUserWorkflow& Workflow : Workflows)
-						{
-							SettingsSection.AddMenuEntry(
-								*Workflow.GetTitle(),
-								FText::FromString(Workflow.GetTitle()),
-								FText::FromString(Workflow.GetTooltip()),
-								FSlateIcon(),
-								FUIAction(FExecuteAction::CreateLambda([Controller, Workflow, ModelPin, ShowWorkflowOptionsDialog]()
-								{
-									URigVMUserWorkflowOptions* Options = Controller->MakeOptionsForWorkflow(ModelPin, Workflow);
-
-									bool bPerform = true;
-									if(Options->RequiresDialog())
-									{
-										bPerform = ShowWorkflowOptionsDialog(Options);
-									}
-									if(bPerform)
-									{
-										Controller->PerformUserWorkflow(Workflow, Options);
-									}
-								}))
-							);
-						}
-					}
-				}
-			}
-
-			bool bIsExecutePin = false;
-			bool bIsUnknownPin = false;
-			if(UScriptStruct* ScriptStruct = Cast<UScriptStruct>(InGraphPin->PinType.PinSubCategoryObject))
-			{
-				bIsExecutePin = ScriptStruct->IsChildOf(FRigVMExecuteContext::StaticStruct());
-				bIsUnknownPin = ScriptStruct->IsChildOf(RigVMTypeUtils::GetWildCardCPPTypeObject());
-			}
-			const bool bIsEditablePin = !bIsExecutePin && !bIsUnknownPin;
-
-			if(UControlRigGraphNode* RigNode = Cast<UControlRigGraphNode>(InGraphPin->GetOwningNode()))
-			{
-				if(bIsEditablePin)
-				{
-					// Add the watch pin / unwatch pin menu items
-					FToolMenuSection& Section = Menu->AddSection("EdGraphSchemaWatches", LOCTEXT("WatchesHeader", "Watches"));
-					UBlueprint* OwnerBlueprint = FBlueprintEditorUtils::FindBlueprintForGraphChecked(Context->Graph);
-					{
-						if (FKismetDebugUtilities::IsPinBeingWatched(OwnerBlueprint, InGraphPin))
-						{
-							Section.AddMenuEntry(FGraphEditorCommands::Get().StopWatchingPin);
-						}
-						else
-						{
-							Section.AddMenuEntry(FGraphEditorCommands::Get().StartWatchingPin);
-						}
-					}
-				}
-			}
-
-			// Add alphainterp menu entries
-			if(UControlRigBlueprint* RigBlueprint = Cast<UControlRigBlueprint>((UBlueprint*)Context->Blueprint))
-			{
-				if (URigVMPin* ModelPin = RigBlueprint->GetModel(Graph)->FindPin(InGraphPin->GetName()))
-				{
-					URigVMController* Controller = RigBlueprint->GetController(ModelPin->GetGraph());
-
-					if (ModelPin->IsArray() && !bIsExecutePin)
-					{
-						FToolMenuSection& Section = Menu->AddSection("EdGraphSchemaPinArrays", LOCTEXT("PinArrays", "Arrays"));
-						Section.AddMenuEntry(
-							"ClearPinArray",
-							LOCTEXT("ClearPinArray", "Clear Array"),
-							LOCTEXT("ClearPinArray_Tooltip", "Removes all elements of the array."),
-							FSlateIcon(),
-							FUIAction(FExecuteAction::CreateLambda([Controller, ModelPin]() {
-								Controller->ClearArrayPin(ModelPin->GetPinPath());
-							})
-						));
-		}
-					if(ModelPin->IsArrayElement())
-					{
-						FToolMenuSection& Section = Menu->AddSection("EdGraphSchemaPinArrays", LOCTEXT("PinArrays", "Arrays"));
-						Section.AddMenuEntry(
-							"RemoveArrayPin",
-							LOCTEXT("RemoveArrayPin", "Remove Array Element"),
-							LOCTEXT("RemoveArrayPin_Tooltip", "Removes the selected element from the array"),
-							FSlateIcon(),
-							FUIAction(FExecuteAction::CreateLambda([Controller, ModelPin]() {
-								Controller->RemoveArrayPin(ModelPin->GetPinPath(), true, true);
-							})
-						));
-						Section.AddMenuEntry(
-							"DuplicateArrayPin",
-							LOCTEXT("DuplicateArrayPin", "Duplicate Array Element"),
-							LOCTEXT("DuplicateArrayPin_Tooltip", "Duplicates the selected element"),
-							FSlateIcon(),
-							FUIAction(FExecuteAction::CreateLambda([Controller, ModelPin]() {
-								Controller->DuplicateArrayPin(ModelPin->GetPinPath(), true, true);
-							})
-						));
-					}
-
-					if (URigVMAggregateNode* AggregateNode = Cast<URigVMAggregateNode>(ModelPin->GetNode()))
-					{
-						FToolMenuSection& Section = Menu->AddSection("EdGraphSchemaAggregatePin", LOCTEXT("AggregatePin", "Aggregates"));
-						Section.AddMenuEntry(
-							"RemoveAggregatePin",
-							LOCTEXT("RemoveAggregatePin", "Remove Aggregate Element"),
-							LOCTEXT("RemoveAggregatePin_Tooltip", "Removes the selected element from the aggregate"),
-							FSlateIcon(),
-							FUIAction(FExecuteAction::CreateLambda([Controller, ModelPin]() {
-								Controller->RemoveAggregatePin(ModelPin->GetPinPath(), true, true);
-							})
-						));
-					}
-
-					if (URigVMTemplateNode* TemplateNode = Cast<URigVMTemplateNode>(ModelPin->GetNode()))
-					{
-						if (!TemplateNode->IsSingleton())
-						{
-							FToolMenuSection& TemplatesSection = Menu->AddSection("EdGraphSchemaTemplates", LOCTEXT("TemplatesHeader", "Templates"));
-
-							if(const FRigVMTemplate* Template = TemplateNode->GetTemplate())
-							{
-								if(!ModelPin->IsExecuteContext())
-								{
-									if(const FRigVMTemplateArgument* Argument = Template->FindArgument(ModelPin->GetRootPin()->GetFName()))
-									{
-										if(!Argument->IsSingleton())
-										{
-											TArray<TRigVMTypeIndex> ResolvedTypeIndices = Argument->GetSupportedTypeIndices(TemplateNode->GetResolvedPermutationIndices(true));
-											TSharedRef<SRigVMGraphChangePinType> ChangePinTypeWidget =
-											SNew(SRigVMGraphChangePinType)
-											.Blueprint(RigBlueprint)
-											.Types(ResolvedTypeIndices)
-											.OnTypeSelected_Lambda([RigBlueprint, ModelPin](const TRigVMTypeIndex& TypeSelected)
-											{
-												if (URigVMController* Controller = RigBlueprint->GetController(ModelPin->GetGraph()))
-												{
-													Controller->ResolveWildCardPin(ModelPin, TypeSelected, true, true);
-												}
-											});
-
-											TemplatesSection.AddEntry(FToolMenuEntry::InitWidget("ChangePinTypeWidget", ChangePinTypeWidget, FText(), true));
-										}
-									}
-								}
-							
-								TemplatesSection.AddMenuEntry(
-									"Unresolve Template Node",
-									LOCTEXT("UnresolveTemplateNode", "Unresolve Template Node"),
-									LOCTEXT("UnresolveTemplateNode_Tooltip", "Removes any type information from the template node"),
-									FSlateIcon(),
-									FUIAction(FExecuteAction::CreateLambda([Controller, ModelPin]() {
-										const TArray<FName> Nodes = ModelPin->GetGraph()->GetSelectNodes();
-										Controller->UnresolveTemplateNodes(Nodes, true, true);
-									})
-								));
-							}
-						}
-					}
-
-					// menu entries to convert the node
-					if (URigVMRerouteNode* RerouteNode = Cast<URigVMRerouteNode>(ModelPin->GetNode()))
-					{
-						if(RerouteNode->GetLinkedSourceNodes().Num() == 0)
-						{
-							if(const URigVMPin* ValuePin = RerouteNode->FindPin(TEXT("Value")))
-							{
-								if(!ValuePin->IsExecuteContext())
-								{
-									FToolMenuSection& ConversionSection = Menu->AddSection("EdGraphSchemaConversion", LOCTEXT("ConversionHeader", "Conversion"));
-
-									if(ValuePin->IsArray())
-									{
-										ConversionSection.AddMenuEntry(
-											"Convert Reroute to Make Array",
-											LOCTEXT("ConvertReroutetoMakeArray", "Convert Reroute to Make Array"),
-											LOCTEXT("ConvertReroutetoMakeArray_Tooltip", "Converts the Reroute node to a to Make Array node and sustains the value"),
-											FSlateIcon(),
-											FUIAction(FExecuteAction::CreateLambda([Controller, RerouteNode]()
-											{
-												Controller->ConvertRerouteNodeToDispatch(RerouteNode, FRigVMDispatch_ArrayMake().GetTemplateNotation(), true, true);
-											}
-										)));
-									}
-									else if(ValuePin->IsStruct())
-									{
-										ConversionSection.AddMenuEntry(
-											"Convert Reroute to Make Struct",
-											LOCTEXT("ConvertReroutetoMakeStruct", "Convert Reroute to Make Struct"),
-											LOCTEXT("ConvertReroutetoMakeStruct_Tooltip", "Converts the Reroute node to a to Make Struct node and sustains the value"),
-											FSlateIcon(),
-											FUIAction(FExecuteAction::CreateLambda([Controller, RerouteNode]()
-											{
-												Controller->ConvertRerouteNodeToDispatch(RerouteNode, FRigVMDispatch_MakeStruct().GetTemplateNotation(), true, true);
-											}
-										)));
-									}
-									else
-									{
-										ConversionSection.AddMenuEntry(
-											"Convert Reroute to Constant",
-											LOCTEXT("ConvertReroutetoConstant", "Convert Reroute to Constant"),
-											LOCTEXT("ConvertReroutetoConstant_Tooltip", "Converts the Reroute node to a to Constant node and sustains the value"),
-											FSlateIcon(),
-											FUIAction(FExecuteAction::CreateLambda([Controller, RerouteNode]()
-											{
-												Controller->ConvertRerouteNodeToDispatch(RerouteNode, FRigVMDispatch_Constant().GetTemplateNotation(), true, true);
-											}
-										)));
-									}
-								}
-							}
-						}
-					}
-
-					if (ModelPin->GetDirection() == ERigVMPinDirection::Input &&
-							bIsEditablePin)
-					{
-						if (ModelPin->IsBoundToVariable())
-						{
-							FVector2D NodePosition = FVector2D(Context->Node->NodePosX - 200.f, Context->Node->NodePosY);
-
-							FToolMenuSection& VariablesSection = Menu->AddSection("EdGraphSchemaVariables", LOCTEXT("Variables", "Variables"));
-							VariablesSection.AddMenuEntry(
-								"MakeVariableNodeFromBinding",
-								LOCTEXT("MakeVariableNodeFromBinding", "Make Variable Node"),
-								LOCTEXT("MakeVariableNodeFromBinding_Tooltip", "Turns the variable binding on the pin to a variable node"),
-								FSlateIcon(),
-								FUIAction(FExecuteAction::CreateLambda([Controller, ModelPin, NodePosition]() {
-									Controller->MakeVariableNodeFromBinding(ModelPin->GetPinPath(), NodePosition, true, true);
-								})
-							));
-						}
-						else
-						{
-							FVector2D NodePosition = FVector2D(Context->Node->NodePosX - 200.f, Context->Node->NodePosY);
-
-							FToolMenuSection& VariablesSection = Menu->AddSection("EdGraphSchemaVariables", LOCTEXT("Variables", "Variables"));
-							VariablesSection.AddMenuEntry(
-								"PromotePinToVariable",
-								LOCTEXT("PromotePinToVariable", "Promote Pin To Variable"),
-								LOCTEXT("PromotePinToVariable_Tooltip", "Turns the pin into a variable"),
-								FSlateIcon(),
-								FUIAction(FExecuteAction::CreateLambda([Controller, ModelPin, NodePosition]() {
-
-									FModifierKeysState KeyState = FSlateApplication::Get().GetModifierKeys();
-									bool bCreateVariableNode = !KeyState.IsAltDown();
-
-									Controller->PromotePinToVariable(ModelPin->GetPinPath(), bCreateVariableNode, NodePosition, true, true);
-								})
-							));
-						}
-					}
-
-					if (Cast<URigVMUnitNode>(ModelPin->GetNode()) != nullptr || 
-						Cast<URigVMDispatchNode>(ModelPin->GetNode()) != nullptr || 
-						Cast<URigVMLibraryNode>(ModelPin->GetNode()) != nullptr)
-					{
-						if (ModelPin->GetDirection() == ERigVMPinDirection::Input &&
-							ModelPin->IsRootPin() &&
-							bIsEditablePin)
-						{
-							if (!ModelPin->IsBoundToVariable())
-							{
-								FToolMenuSection& VariablesSection = Menu->FindOrAddSection(TEXT("Variables"));
-
-								TSharedRef<SRigVMGraphVariableBinding> VariableBindingWidget =
-									SNew(SRigVMGraphVariableBinding)
-									.Blueprint(RigBlueprint)
-									.ModelPins({ModelPin})
-									.CanRemoveBinding(false);
-
-								VariablesSection.AddEntry(FToolMenuEntry::InitWidget("BindPinToVariableWidget", VariableBindingWidget, FText(), true));
-							}
-
-							FToolMenuSection& Section = Menu->AddSection("EdGraphSchemaPinDefaults", LOCTEXT("PinDefaults", "Pin Defaults"));
-							Section.AddMenuEntry(
-								"ResetPinDefaultValue",
-								LOCTEXT("ResetPinDefaultValue", "Reset Pin Value"),
-								LOCTEXT("ResetPinDefaultValue_Tooltip", "Resets the pin's value to its default."),
-								FSlateIcon(),
-								FUIAction(FExecuteAction::CreateLambda([Controller, ModelPin]() {
-									Controller->ResetPinDefaultValue(ModelPin->GetPinPath());
-								})
-							));
-						}
-					}
-
-					if ((ModelPin->GetCPPType() == TEXT("FVector") ||
-						 ModelPin->GetCPPType() == TEXT("FQuat") ||
-						 ModelPin->GetCPPType() == TEXT("FTransform")) &&
-						(ModelPin->GetDirection() == ERigVMPinDirection::Input ||
-						 ModelPin->GetDirection() == ERigVMPinDirection::IO) &&
-						 ModelPin->GetPinForLink()->GetRootPin()->GetSourceLinks(true).Num() == 0)
-					{
-						FToolMenuSection& Section = Menu->AddSection("EdGraphSchemaControlPin", LOCTEXT("ControlPin", "Direct Manipulation"));
-						Section.AddMenuEntry(
-							"DirectManipControlPin",
-							LOCTEXT("DirectManipControlPin", "Control Pin Value"),
-							LOCTEXT("DirectManipControlPin_Tooltip", "Configures the pin for direct interaction in the viewport"),
-							FSlateIcon(),
-							FUIAction(FExecuteAction::CreateLambda([RigBlueprint, ModelPin]() {
-								RigBlueprint->AddTransientControl(ModelPin);
-							})
-						));
-					}
-
-					if (ModelPin->GetRootPin() == ModelPin && (
-						Cast<URigVMUnitNode>(ModelPin->GetNode()) != nullptr ||
-						Cast<URigVMLibraryNode>(ModelPin->GetNode()) != nullptr))
-					{
-						if (ModelPin->HasInjectedNodes())
-						{
-							FToolMenuSection& Section = Menu->AddSection("EdGraphSchemaNodeEjectionInterp", LOCTEXT("NodeEjectionInterp", "Eject"));
-
-							Section.AddMenuEntry(
-								"EjectLastNode",
-								LOCTEXT("EjectLastNode", "Eject Last Node"),
-								LOCTEXT("EjectLastNode_Tooltip", "Eject the last injected node"),
-								FSlateIcon(),
-								FUIAction(FExecuteAction::CreateLambda([Controller, ModelPin]() {
-									Controller->OpenUndoBracket(TEXT("Eject node from pin"));
-									URigVMNode* Node = Controller->EjectNodeFromPin(ModelPin->GetPinPath(), true, true);
-									Controller->SelectNode(Node, true, true, true);
-									Controller->CloseUndoBracket();
-								})
-							));
-						}
-
-						if (ModelPin->GetCPPType() == TEXT("float") ||
-							ModelPin->GetCPPType() == TEXT("double") ||
-							ModelPin->GetCPPType() == TEXT("FVector"))
-						{
-							FToolMenuSection& Section = Menu->AddSection("EdGraphSchemaNodeInjectionInterp", LOCTEXT("NodeInjectionInterp", "Interpolate"));
-							URigVMNode* InterpNode = nullptr;
-							bool bBoundToVariable = false;
-							for (URigVMInjectionInfo* Injection : ModelPin->GetInjectedNodes())
-							{
-								FString TemplateName;
-								if (URigVMUnitNode* UnitNode = Cast<URigVMUnitNode>(Injection->Node))
-								{
-									if (UnitNode->GetScriptStruct()->GetStringMetaDataHierarchical(FRigVMStruct::TemplateNameMetaName, &TemplateName))
-									{
-										if (TemplateName == TEXT("AlphaInterp"))
-										{
-											InterpNode = Injection->Node;
-											break;
-										}
-									}
-								}
-								else if (URigVMVariableNode* VariableNode = Cast<URigVMVariableNode>(Injection->Node))
-								{
-									bBoundToVariable = true;
-									break;
-								}
-							}
-
-							if(!bBoundToVariable)
-							{
-								if (InterpNode == nullptr)
-								{
-									UScriptStruct* ScriptStruct = nullptr;
-
-									if ((ModelPin->GetCPPType() == TEXT("float")) || (ModelPin->GetCPPType() == TEXT("double")))
-									{
-										ScriptStruct = FRigVMFunction_AlphaInterp::StaticStruct();
-									}
-									else if (ModelPin->GetCPPType() == TEXT("FVector"))
-									{
-										ScriptStruct = FRigVMFunction_AlphaInterpVector::StaticStruct();
-									}
-									else
-									{
-										checkNoEntry();
-									}
-
-									Section.AddMenuEntry(
-										"AddAlphaInterp",
-										LOCTEXT("AddAlphaInterp", "Add Interpolate"),
-										LOCTEXT("AddAlphaInterp_Tooltip", "Injects an interpolate node"),
-										FSlateIcon(),
-										FUIAction(FExecuteAction::CreateLambda([Controller, InGraphPin, ModelPin, ScriptStruct]() {
-											Controller->OpenUndoBracket(TEXT("Add injected node"));
-											URigVMInjectionInfo* Injection = Controller->AddInjectedNode(ModelPin->GetPinPath(), ModelPin->GetDirection() != ERigVMPinDirection::Output, ScriptStruct, FRigUnit::GetMethodName(), TEXT("Value"), TEXT("Result"), FString(), true, true);
-											if (Injection)
-											{
-												TArray<FName> NodeNames;
-												NodeNames.Add(Injection->Node->GetFName());
-												Controller->SetNodeSelection(NodeNames);
-											}
-											Controller->CloseUndoBracket();
-										})
-									));
-								}
-								else
-								{
-									Section.AddMenuEntry(
-										"EditAlphaInterp",
-										LOCTEXT("EditAlphaInterp", "Edit Interpolate"),
-										LOCTEXT("EditAlphaInterp_Tooltip", "Edit the interpolate node"),
-										FSlateIcon(),
-										FUIAction(FExecuteAction::CreateLambda([RigBlueprint, InterpNode]() {
-										TArray<FName> NodeNames;
-										NodeNames.Add(InterpNode->GetFName());
-										RigBlueprint->GetController(InterpNode->GetGraph())->SetNodeSelection(NodeNames);
-									})
-										));
-									Section.AddMenuEntry(
-										"RemoveAlphaInterp",
-										LOCTEXT("RemoveAlphaInterp", "Remove Interpolate"),
-										LOCTEXT("RemoveAlphaInterp_Tooltip", "Removes the interpolate node"),
-										FSlateIcon(),
-										FUIAction(FExecuteAction::CreateLambda([Controller, InGraphPin, ModelPin, InterpNode]() {
-											Controller->RemoveInjectedNode(ModelPin->GetPinPath(), ModelPin->GetDirection() != ERigVMPinDirection::Output, true);
-										})
-									));
-								}
-							}
-						}
-
-						if (ModelPin->GetCPPType() == TEXT("FVector") ||
-							ModelPin->GetCPPType() == TEXT("FQuat") ||
-							ModelPin->GetCPPType() == TEXT("FTransform"))
-						{
-							FToolMenuSection& Section = Menu->AddSection("EdGraphSchemaNodeInjectionVisualDebug", LOCTEXT("NodeInjectionVisualDebug", "Visual Debug"));
-
-							URigVMNode* VisualDebugNode = nullptr;
-							bool bBoundToVariable = false;
-							for (URigVMInjectionInfo* Injection : ModelPin->GetInjectedNodes())
-							{
-								FString TemplateName;
-								if (URigVMUnitNode* UnitNode = Cast<URigVMUnitNode>(Injection->Node))
-								{
-									if (UnitNode->GetScriptStruct()->GetStringMetaDataHierarchical(FRigVMStruct::TemplateNameMetaName, &TemplateName))
-									{
-										if (TemplateName == TEXT("VisualDebug"))
-										{
-											VisualDebugNode = Injection->Node;
-											break;
-										}
-									}
-								}
-								else if (URigVMVariableNode* VariableNode = Cast<URigVMVariableNode>(Injection->Node))
-								{
-									bBoundToVariable = true;
-									break;
-								}
-							}
-
-							if (!bBoundToVariable)
-							{
-								if (VisualDebugNode == nullptr)
-								{
-									UScriptStruct* ScriptStruct = nullptr;
-
-									if (ModelPin->GetCPPType() == TEXT("FVector"))
-									{
-										ScriptStruct = FRigVMFunction_VisualDebugVectorNoSpace::StaticStruct();
-									}
-									else if (ModelPin->GetCPPType() == TEXT("FQuat"))
-									{
-										ScriptStruct = FRigVMFunction_VisualDebugQuatNoSpace::StaticStruct();
-									}
-									else if (ModelPin->GetCPPType() == TEXT("FTransform"))
-									{
-										ScriptStruct = FRigVMFunction_VisualDebugTransformNoSpace::StaticStruct();
-									}
-									else
-									{
-										checkNoEntry();
-									}
-
-									Section.AddMenuEntry(
-										"AddVisualDebug",
-										LOCTEXT("AddVisualDebug", "Add Visual Debug"),
-										LOCTEXT("AddVisualDebug_Tooltip", "Injects a visual debugging node"),
-										FSlateIcon(),
-										FUIAction(FExecuteAction::CreateLambda([RigBlueprint, Controller, InGraphPin, ModelPin, ScriptStruct]() {
-											URigVMInjectionInfo* Injection = Controller->AddInjectedNode(ModelPin->GetPinPath(), ModelPin->GetDirection() != ERigVMPinDirection::Output, ScriptStruct, FRigUnit::GetMethodName(), TEXT("Value"), TEXT("Value"), FString(), true, true);
-											if (Injection)
-											{
-												TArray<FName> NodeNames;
-												NodeNames.Add(Injection->Node->GetFName());
-												Controller->SetNodeSelection(NodeNames);
-
-												if (URigVMUnitNode* UnitNode = Cast<URigVMUnitNode>(ModelPin->GetNode()))
-												{
-													if (TSharedPtr<FStructOnScope> DefaultStructScope = UnitNode->ConstructStructInstance())
-													{
-														if(DefaultStructScope->GetStruct()->IsChildOf(FRigUnit::StaticStruct()))
-														{
-															FRigUnit* DefaultStruct = (FRigUnit*)DefaultStructScope->GetStructMemory();
-
-															FString PinPath = ModelPin->GetPinPath();
-															FString Left, Right;
-
-															FRigElementKey SpaceKey;
-															if (URigVMPin::SplitPinPathAtStart(PinPath, Left, Right))
-															{
-																SpaceKey = DefaultStruct->DetermineSpaceForPin(Right, RigBlueprint->Hierarchy);
-															}
-
-															if (SpaceKey.IsValid())
-															{
-																if (URigVMPin* SpacePin = Injection->Node->FindPin(TEXT("Space")))
-																{
-																	if(URigVMPin* SpaceTypePin = SpacePin->FindSubPin(TEXT("Type")))
-																	{
-																		FString SpaceTypeStr = StaticEnum<ERigElementType>()->GetDisplayNameTextByValue((int64)SpaceKey.Type).ToString();
-																		Controller->SetPinDefaultValue(SpaceTypePin->GetPinPath(), SpaceTypeStr, true, true, false, true);
-																	}
-																	if(URigVMPin* SpaceNamePin = SpacePin->FindSubPin(TEXT("Name")))
-																	{
-																		Controller->SetPinDefaultValue(SpaceNamePin->GetPinPath(), SpaceKey.Name.ToString(), true, true, false, true);
-																	}
-																}
-															}
-														}
-													}
-												}
-											}
-										})
-									));
-								}
-								else
-								{
-									Section.AddMenuEntry(
-										"EditVisualDebug",
-										LOCTEXT("EditVisualDebug", "Edit Visual Debug"),
-										LOCTEXT("EditVisualDebug_Tooltip", "Edit the visual debugging node"),
-										FSlateIcon(),
-										FUIAction(FExecuteAction::CreateLambda([Controller, VisualDebugNode]() {
-											TArray<FName> NodeNames;
-											NodeNames.Add(VisualDebugNode->GetFName());
-											Controller->SetNodeSelection(NodeNames);
-										})
-									));
-									Section.AddMenuEntry(
-										"ToggleVisualDebug",
-										LOCTEXT("ToggleVisualDebug", "Toggle Visual Debug"),
-										LOCTEXT("ToggleVisualDebug_Tooltip", "Toggle the visibility the visual debugging"),
-										FSlateIcon(),
-										FUIAction(FExecuteAction::CreateLambda([Controller, VisualDebugNode]() {
-											URigVMPin* EnabledPin = VisualDebugNode->FindPin(TEXT("bEnabled"));
-											check(EnabledPin);
-											Controller->SetPinDefaultValue(EnabledPin->GetPinPath(), EnabledPin->GetDefaultValue() == TEXT("True") ? TEXT("False") : TEXT("True"), false, true, false, true);
-										})
-									));
-									Section.AddMenuEntry(
-										"RemoveVisualDebug",
-										LOCTEXT("RemoveVisualDebug", "Remove Visual Debug"),
-										LOCTEXT("RemoveVisualDebug_Tooltip", "Removes the visual debugging node"),
-										FSlateIcon(),
-											FUIAction(FExecuteAction::CreateLambda([Controller, InGraphPin, ModelPin, VisualDebugNode]() {
-											Controller->RemoveNodeByName(VisualDebugNode->GetFName(), true, false);
-										})
-									));
-								}
-							}
-						}
-					}
-				}
-			}
-		}
-		else if(Context->Node) // right clicked on the node
-		{
-			if (UControlRigBlueprint* RigBlueprint = Cast<UControlRigBlueprint>((UBlueprint*)Context->Blueprint))
-			{
-				URigVMGraph* Model = RigBlueprint->GetModel(Context->Node->GetGraph());
-				URigVMController* Controller = RigBlueprint->GetController(Model);
-
-				URigHierarchy* TemporaryHierarchy = NewObject<URigHierarchy>();
-				TemporaryHierarchy->CopyHierarchy(RigBlueprint->Hierarchy);
-
-				TArray<FRigElementKey> RigElementsToSelect;
-				TMap<const URigVMPin*, FRigElementKey> PinToKey;
-				TArray<FName> SelectedNodeNames = Model->GetSelectNodes();
-				SelectedNodeNames.AddUnique(Context->Node->GetFName());
-
-				bool bCanNodeBeUpgraded = false; 
-
-				for(const FName& SelectedNodeName : SelectedNodeNames)
-				{
-					if (URigVMNode* ModelNode = Model->FindNodeByName(SelectedNodeName))
-					{
-						TSharedPtr<FStructOnScope> StructOnScope;
-						FRigUnit* StructMemory = nullptr;
-						UScriptStruct* ScriptStruct = nullptr;
-						if (URigVMUnitNode* UnitNode = Cast<URigVMUnitNode>(ModelNode))
-						{
-							ScriptStruct = UnitNode->GetScriptStruct();
-							if(ScriptStruct)
-							{
-								StructOnScope = UnitNode->ConstructStructInstance(false /* default */);
-								if(StructOnScope->GetStruct()->IsChildOf(FRigUnit::StaticStruct()))
-								{
-									StructMemory = (FRigUnit*)StructOnScope->GetStructMemory();
-
-									//FRigUnitContext RigUnitContext;
-									//RigUnitContext.Hierarchy = TemporaryHierarchy;
-									//RigUnitContext.State = EControlRigState::Update;
-									
-									StructMemory->Execute();
-								}
-							}
-						}
-
-						const TArray<URigVMPin*> AllPins = ModelNode->GetAllPinsRecursively();
-						for (const URigVMPin* Pin : AllPins)
-						{
-							if (Pin->GetCPPType() == TEXT("FName"))
-							{
-								FRigElementKey Key;
-								if (Pin->GetCustomWidgetName() == TEXT("BoneName"))
-								{
-									Key = FRigElementKey(*Pin->GetDefaultValue(), ERigElementType::Bone);
-								}
-								else if (Pin->GetCustomWidgetName() == TEXT("ControlName"))
-								{
-									Key = FRigElementKey(*Pin->GetDefaultValue(), ERigElementType::Control);
-								}
-								else if (Pin->GetCustomWidgetName() == TEXT("SpaceName"))
-								{
-									Key = FRigElementKey(*Pin->GetDefaultValue(), ERigElementType::Null);
-								}
-								else if (Pin->GetCustomWidgetName() == TEXT("CurveName"))
-								{
-									Key = FRigElementKey(*Pin->GetDefaultValue(), ERigElementType::Curve);
-								}
-								else
-								{
-									continue;
-								}
-
-								RigElementsToSelect.AddUnique(Key);
-								PinToKey.Add(Pin, Key);
-							}
-							else if (Pin->GetCPPTypeObject() == FRigElementKey::StaticStruct() && !Pin->IsArray())
-							{
-								if (StructMemory == nullptr)
-								{
-									FString DefaultValue = Pin->GetDefaultValue();
-									if (!DefaultValue.IsEmpty())
-									{
-										FRigElementKey Key;
-										FRigElementKey::StaticStruct()->ImportText(*DefaultValue, &Key, nullptr, EPropertyPortFlags::PPF_None, nullptr, FRigElementKey::StaticStruct()->GetName(), true);
-										if (Key.IsValid())
-										{
-											RigElementsToSelect.AddUnique(Key);
-											if (URigVMPin* NamePin = Pin->FindSubPin(TEXT("Name")))
-											{
-												PinToKey.Add(NamePin, Key);
-											}
-										}
-									}
-								}
-								else
-								{
-									check(ScriptStruct);
-
-									TArray<FString> PropertyNames; 
-									if(!URigVMPin::SplitPinPath(Pin->GetSegmentPath(true), PropertyNames))
-									{
-										PropertyNames.Add(Pin->GetName());
-									}
-
-									UScriptStruct* Struct = ScriptStruct;
-									uint8* Memory = (uint8*)StructMemory; 
-
-									while(!PropertyNames.IsEmpty())
-									{
-										FString PropertyName;
-										PropertyNames.HeapPop(PropertyName);
-										
-										const FProperty* Property = ScriptStruct->FindPropertyByName(*PropertyName);
-										if(Property == nullptr)
-										{
-											Memory = nullptr;
-											break;
-										}
-
-										Memory = Property->ContainerPtrToValuePtr<uint8>(Memory);
-										
-										if(PropertyNames.IsEmpty())
-										{
-											continue;
-										}
-										
-										if(const FArrayProperty* ArrayProperty = CastField<FArrayProperty>(Property))
-										{
-											PropertyNames.HeapPop(PropertyName);
-
-											int32 ArrayIndex = FCString::Atoi(*PropertyName);
-											FScriptArrayHelper Helper(ArrayProperty, Memory);
-											if(!Helper.IsValidIndex(ArrayIndex))
-											{
-												Memory = nullptr;
-												break;
-											}
-
-											Memory = Helper.GetRawPtr(ArrayIndex);
-											Property = ArrayProperty->Inner;
-										}
-
-										if(const FStructProperty* StructProperty = CastField<FStructProperty>(Property))
-										{
-											Struct = StructProperty->Struct;
-										}
-									}
-
-									if(Memory)
-									{
-										const FRigElementKey& Key = *(const FRigElementKey*)Memory;
-										if (Key.IsValid())
-										{
-											RigElementsToSelect.AddUnique(Key);
-
-											if (URigVMPin* NamePin = Pin->FindSubPin(TEXT("Name")))
-											{
-												PinToKey.Add(NamePin, Key);
-											}
-										}
-									}
-								}
-							}
-							else if (Pin->GetCPPTypeObject() == FRigElementKeyCollection::StaticStruct() && Pin->GetDirection() == ERigVMPinDirection::Output)
-							{
-								if (StructMemory == nullptr)
-								{
-									// not supported for now
-								}
-								else
-								{
-									check(ScriptStruct);
-									if (const FProperty* Property = ScriptStruct->FindPropertyByName(Pin->GetFName()))
-									{
-										const FRigElementKeyCollection& Collection = *Property->ContainerPtrToValuePtr<FRigElementKeyCollection>(StructMemory);
-
-										if (Collection.Num() > 0)
-										{
-											RigElementsToSelect.Reset();
-											for (const FRigElementKey& Item : Collection)
-											{
-												RigElementsToSelect.AddUnique(Item);
-											}
-											break;
-										}
-									}
-								}
-							}
-						}
-						
-						bCanNodeBeUpgraded = bCanNodeBeUpgraded || ModelNode->CanBeUpgraded();
-					}
-				}
-
-				if (RigElementsToSelect.Num() > 0)
-				{
-					FToolMenuSection& Section = Menu->AddSection("EdGraphSchemaHierarchy", LOCTEXT("HierarchyHeader", "Hierarchy"));
-					Section.AddMenuEntry(
-						"SelectRigElements",
-						LOCTEXT("SelectRigElements", "Select Rig Elements"),
-						LOCTEXT("SelectRigElements_Tooltip", "Selects the bone, controls or nulls associated with this node."),
-						FSlateIcon(),
-						FUIAction(FExecuteAction::CreateLambda([RigBlueprint, RigElementsToSelect]() {
-
-							RigBlueprint->GetHierarchyController()->SetSelection(RigElementsToSelect);
-							
-						})
-					));
-				}
-
-				if (RigElementsToSelect.Num() > 0)
-				{
-					FToolMenuSection& Section = Menu->AddSection("EdGraphSchemaHierarchy", LOCTEXT("ToolsHeader", "Tools"));
-					Section.AddMenuEntry(
-						"SearchAndReplaceNames",
-						LOCTEXT("SearchAndReplaceNames", "Search & Replace / Mirror"),
-						LOCTEXT("SearchAndReplaceNames_Tooltip", "Searches within all names and replaces with a different text."),
-						FSlateIcon(),
-						FUIAction(FExecuteAction::CreateLambda([RigBlueprint, Controller, PinToKey]() {
-
-							FRigVMMirrorSettings Settings;
-							TSharedPtr<FStructOnScope> StructToDisplay = MakeShareable(new FStructOnScope(FRigVMMirrorSettings::StaticStruct(), (uint8*)&Settings));
-
-							TSharedRef<SKismetInspector> KismetInspector = SNew(SKismetInspector);
-							KismetInspector->ShowSingleStruct(StructToDisplay);
-
-							TSharedRef<SCustomDialog> MirrorDialog = SNew(SCustomDialog)
-								.Title(FText(LOCTEXT("ControlRigHierarchyMirror", "Mirror Graph")))
-								.Content()
-								[
-									KismetInspector
-								]
-								.Buttons({
-									SCustomDialog::FButton(LOCTEXT("OK", "OK")),
-									SCustomDialog::FButton(LOCTEXT("Cancel", "Cancel"))
-							});
-							if (MirrorDialog->ShowModal() == 0)
-							{
-								Controller->OpenUndoBracket(TEXT("Mirroring Graph"));
-								int32 ReplacedNames = 0;
-								TArray<FString> UnchangedItems;
-
-								for (const TPair<const URigVMPin*, FRigElementKey>& Pair : PinToKey)
-								{
-									const URigVMPin* Pin = Pair.Key;
-									FRigElementKey Key = Pair.Value;
-
-									if (Key.Name.IsNone())
-									{
-										continue;
-									}
-
-									FString OldNameStr = Key.Name.ToString();
-									FString NewNameStr = OldNameStr.Replace(*Settings.SearchString, *Settings.ReplaceString, ESearchCase::CaseSensitive);
-									if(NewNameStr != OldNameStr)
-									{
-										Key.Name = *NewNameStr;
-										if(RigBlueprint->Hierarchy->GetIndex(Key) != INDEX_NONE)
-										{
-											Controller->SetPinDefaultValue(Pin->GetPinPath(), NewNameStr, false, true, false, true);
-											ReplacedNames++;
-										}
-										else
-										{
-											// save the names of the items that we skipped during this search & replace
-											UnchangedItems.AddUnique(OldNameStr);
-										} 
-									}
-								}
-
-								if (UnchangedItems.Num() > 0)
-								{
-									FString ListOfUnchangedItems;
-									for (int Index = 0; Index < UnchangedItems.Num(); Index++)
-									{
-										// construct the string "item1, item2, item3"
-										ListOfUnchangedItems += UnchangedItems[Index];
-										if (Index != UnchangedItems.Num() - 1)
-										{
-											ListOfUnchangedItems += TEXT(", ");
-										}
-									}
-									
-									// inform the user that some items were skipped due to invalid new names
-									Controller->ReportAndNotifyError(FString::Printf(TEXT("Invalid Names after Search & Replace, action skipped for %s"), *ListOfUnchangedItems));
-								}
-
-								if (ReplacedNames > 0)
-								{ 
-									Controller->CloseUndoBracket();
-								}
-								else
-								{
-									Controller->CancelUndoBracket();
-								}
-							}
-						})
-					));
-				}
-
-				// per node workflows
-				if (const UControlRigGraphNode* RigNode = Cast<const UControlRigGraphNode>(Context->Node))
-				{
-					if (URigVMNode* ModelNode = RigNode->GetModelNode())
-					{
-						if(ModelNode->IsEvent())
-						{
-							const FName& EventName = ModelNode->GetEventName();
-							const bool bCanRunOnce = !CastChecked<URigVMEdGraphSchema>(RigNode->GetSchema())->IsRigVMDefaultEvent(EventName);
-
-							FToolMenuSection& EventsSection = Menu->AddSection("EdGraphSchemaEvents", LOCTEXT("EventsHeader", "Events"));
-
-							EventsSection.AddMenuEntry(
-								"Switch to Event",
-								LOCTEXT("SwitchToEvent", "Switch to Event"),
-								LOCTEXT("SwitchToEvent_Tooltip", "Switches the Control Rig Editor to run this event permanently."),
-								FSlateIcon(),
-								FUIAction(FExecuteAction::CreateLambda([RigBlueprint, EventName]()
-								{
-									if(UControlRig* ControlRig = Cast<UControlRig>(RigBlueprint->GetObjectBeingDebugged()))
-									{
-										ControlRig->SetEventQueue({EventName});
-									}
-								}))
-							);
-
-							EventsSection.AddMenuEntry(
-								"Run Event Once",
-								LOCTEXT("RuntEventOnce", "Run Event Once"),
-								LOCTEXT("RuntEventOnce_Tooltip", "Runs the event once (for testing)"),
-								FSlateIcon(),
-								FUIAction(FExecuteAction::CreateLambda([RigBlueprint, EventName]()
-								{
-									if(UControlRig* ControlRig = Cast<UControlRig>(RigBlueprint->GetObjectBeingDebugged()))
-									{
-										const TArray<FName> PreviousEventQueue = ControlRig->GetEventQueue();
-										TArray<FName> NewEventQueue = PreviousEventQueue;
-										NewEventQueue.Add(EventName);
-										ControlRig->SetEventQueue(NewEventQueue);
-
-										TSharedPtr<FDelegateHandle> RunOnceHandle = MakeShareable(new FDelegateHandle);
-										*(RunOnceHandle.Get()) = ControlRig->OnExecuted_AnyThread().AddLambda(
-											[RunOnceHandle, EventName, PreviousEventQueue](URigVMHost* InRig, const FName& InEventName)
-											{
-												if(InEventName == EventName)
-												{
-													InRig->SetEventQueue(PreviousEventQueue);
-													InRig->OnExecuted_AnyThread().Remove(*RunOnceHandle.Get());
-												}
-											}
-										);
-									}
-								}),
-								FCanExecuteAction::CreateLambda([bCanRunOnce](){ return bCanRunOnce; }))
-							);
-						}
-
-						if(URigVMDispatchNode* DispatchNode = Cast<URigVMDispatchNode>(ModelNode))
-						{
-							if(DispatchNode->GetFactory()->GetFactoryName() == FRigVMDispatch_Constant().GetFactoryName())
-							{
-								if(const URigVMPin* ValuePin = DispatchNode->FindPin(TEXT("Value")))
-								{
-									// if the value pin has only links on the root pin
-									if((ValuePin->GetSourceLinks(false).Num() > 0) && (ValuePin->GetTargetLinks(false).Num() > 0))
-									{
-										FToolMenuSection& ConversionSection = Menu->AddSection("EdGraphSchemaConversion", LOCTEXT("ConversionHeader", "Conversion"));
-										ConversionSection.AddMenuEntry(
-											"Convert Constant to Reroute",
-											LOCTEXT("ConvertConstantToReroute", "Convert Constant to Reroute"),
-											LOCTEXT("ConvertConstantToReroute_Tooltip", "Converts the Constant node to a to Reroute node and sustains the value"),
-											FSlateIcon(),
-											FUIAction(FExecuteAction::CreateLambda([Controller, ValuePin, DispatchNode]()
-											{
-												Controller->OpenUndoBracket(TEXT("Replace Constant with Reroute"));
-												TArray<URigVMController::FLinkedPath> LinkedPaths = Controller->GetLinkedPaths(DispatchNode);
-												Controller->FastBreakLinkedPaths(LinkedPaths, true);
-												const FVector2D Position = DispatchNode->GetPosition();
-												const FString NodeName = DispatchNode->GetName();
-												const FString CPPType = ValuePin->GetCPPType();
-												FName CPPTypeObjectPath = NAME_None;
-												if(const UObject* CPPTypeObject = ValuePin->GetCPPTypeObject())
-												{
-													CPPTypeObjectPath = *CPPTypeObject->GetPathName();
-												}
-												
-												Controller->RemoveNode(DispatchNode, true, true);
-												Controller->AddFreeRerouteNode(CPPType, CPPTypeObjectPath, false, NAME_None, FString(), Position, NodeName, true);
-												Controller->RestoreLinkedPaths(LinkedPaths, URigVMController::FRestoreLinkedPathSettings(), true);
-												Controller->CloseUndoBracket();
-											}
-										)));
-									}
-								}
-							}
-						}
-						
-						const TArray<FRigVMUserWorkflow> Workflows = ModelNode->GetSupportedWorkflows(ERigVMUserWorkflowType::NodeContext, ModelNode);
-						if(!Workflows.IsEmpty())
-						{
-							FToolMenuSection& SettingsSection = Menu->AddSection("EdGraphSchemaWorkflow", LOCTEXT("WorkflowHeader", "Workflow"));
-
-							for(const FRigVMUserWorkflow& Workflow : Workflows)
-							{
-								SettingsSection.AddMenuEntry(
-									*Workflow.GetTitle(),
-									FText::FromString(Workflow.GetTitle()),
-									FText::FromString(Workflow.GetTooltip()),
-									FSlateIcon(),
-									FUIAction(FExecuteAction::CreateLambda([Controller, Workflow, ModelNode, ShowWorkflowOptionsDialog]()
-									{
-										URigVMUserWorkflowOptions* Options = Controller->MakeOptionsForWorkflow(ModelNode, Workflow);
-
-										bool bPerform = true;
-										if(Options->RequiresDialog())
-										{
-											bPerform = ShowWorkflowOptionsDialog(Options);
-										}
-										if(bPerform)
-										{
-											Controller->PerformUserWorkflow(Workflow, Options);
-										}
-									}))
-								);
-							}
-						}
-					}
-				}
-
-				if (const UControlRigGraphNode* RigNode = Cast<const UControlRigGraphNode>(Context->Node))
-				{
-					if (URigVMUnitNode* UnitNode = Cast<URigVMUnitNode>(RigNode->GetModelNode()))
-					{
-						FToolMenuSection& SettingsSection = Menu->AddSection("EdGraphSchemaSettings", LOCTEXT("SettingsHeader", "Settings"));
-						SettingsSection.AddMenuEntry(
-							"Save Default Expansion State",
-							LOCTEXT("SaveDefaultExpansionState", "Save Default Expansion State"),
-							LOCTEXT("SaveDefaultExpansionState_Tooltip", "Saves the expansion state of all pins of the node as the default."),
-							FSlateIcon(),
-							FUIAction(FExecuteAction::CreateLambda([UnitNode]() {
-
-#if WITH_EDITORONLY_DATA
-
-								FScopedTransaction Transaction(LOCTEXT("RigUnitDefaultExpansionStateChanged", "Changed Rig Unit Default Expansion State"));
-								UControlRigEditorSettings::Get()->Modify();
-
-								FControlRigSettingsPerPinBool& ExpansionMap = UControlRigEditorSettings::Get()->RigUnitPinExpansion.FindOrAdd(UnitNode->GetScriptStruct()->GetName());
-								ExpansionMap.Values.Empty();
-
-								TArray<URigVMPin*> Pins = UnitNode->GetAllPinsRecursively();
-								for (URigVMPin* Pin : Pins)
-								{
-									if (Pin->GetSubPins().Num() == 0)
-									{
-										continue;
-									}
-
-									FString PinPath = Pin->GetPinPath();
-									FString NodeName, RemainingPath;
-									URigVMPin::SplitPinPathAtStart(PinPath, NodeName, RemainingPath);
-									ExpansionMap.Values.FindOrAdd(RemainingPath) = Pin->IsExpanded();
-								}
-#endif
-							})
-						));
-					}
-
-					if (URigVMVariableNode* VariableNode = Cast<URigVMVariableNode>(RigNode->GetModelNode()))
-					{
-						FToolMenuSection& VariablesSection = Menu->AddSection("EdGraphSchemaVariables", LOCTEXT("VariablesSettingsHeader", "Variables"));
-						VariablesSection.AddMenuEntry(
-							"MakePindingsFromVariableNode",
-							LOCTEXT("MakeBindingsFromVariableNode", "Make Bindings From Node"),
-							LOCTEXT("MakeBindingsFromVariableNode_Tooltip", "Turns the variable node into one ore more variable bindings on the pin(s)"),
-							FSlateIcon(),
-							FUIAction(FExecuteAction::CreateLambda([RigBlueprint, Controller, VariableNode]() {
-								Controller->MakeBindingsFromVariableNode(VariableNode->GetFName());
-							})
-						));
-					}
-
-					
-					FToolMenuSection& DebugSection = Menu->AddSection("EdGraphSchemaDebug", LOCTEXT("DebugHeader", "Debug"));
-					bool bNoneHasBreakpoint = true;
-					TArray<URigVMNode*> SelectedNodes;
-					for (FName SelectedNodeName : SelectedNodeNames)
-					{
-						if (URigVMNode* ModelNode = Model->FindNodeByName(SelectedNodeName))
-						{
-							SelectedNodes.Add(ModelNode);
-							if (ModelNode->HasBreakpoint())
-							{
-								bNoneHasBreakpoint = false;
-							}
-						}
-					}
-					
-					if (bNoneHasBreakpoint)
-					{
-						DebugSection.AddMenuEntry(
-						"Add Breakpoint",
-						LOCTEXT("AddBreakpoint", "Add Breakpoint"),
-						LOCTEXT("AddBreakpoint_Tooltip", "Adds a breakpoint to the graph at this node"),
-						FSlateIcon(),
-						FUIAction(FExecuteAction::CreateLambda([Controller, SelectedNodes, RigBlueprint]()
-						{
-							for (URigVMNode* SelectedNode : SelectedNodes)
-							{
-								if (RigBlueprint->AddBreakpoint(SelectedNode))
-								{
-									SelectedNode->SetHasBreakpoint(true);
-								}
-							}
-						})));
-					}
-					else
-					{						
-						DebugSection.AddMenuEntry(
-                        "Remove Breakpoint",
-                        LOCTEXT("RemoveBreakpoint", "Remove Breakpoint"),
-                        LOCTEXT("RemoveBreakpoint_Tooltip", "Removes a breakpoint to the graph at this node"),
-                        FSlateIcon(),
-                        FUIAction(FExecuteAction::CreateLambda([Controller, SelectedNodes, RigBlueprint]()
-                        {
-                        	for (URigVMNode* SelectedNode : SelectedNodes)
-                        	{
-                        		if (SelectedNode->HasBreakpoint())
-                        		{
-                        			if (RigBlueprint->RemoveBreakpoint(SelectedNode))
-                        			{
-										SelectedNode->SetHasBreakpoint(false);
-									}
-                        		}
-							}                            
-                        })));
-					}
-					
-					
-
-					FToolMenuSection& OrganizationSection = Menu->AddSection("EdGraphSchemaOrganization", LOCTEXT("OrganizationHeader", "Organization"));
-					OrganizationSection.AddMenuEntry(
-						"Collapse Nodes",
-						LOCTEXT("CollapseNodes", "Collapse Nodes"),
-						LOCTEXT("CollapseNodes_Tooltip", "Turns the selected nodes into a single Collapse node"),
-						FSlateIcon(),
-						FUIAction(FExecuteAction::CreateLambda([Model, Controller]() {
-							TArray<FName> Nodes = Model->GetSelectNodes();
-							Controller->CollapseNodes(Nodes, FString(), true, true);
-						})
-					));
-					OrganizationSection.AddMenuEntry(
-						"Collapse to Function",
-						LOCTEXT("CollapseNodesToFunction", "Collapse to Function"),
-						LOCTEXT("CollapseNodesToFunction_Tooltip", "Turns the selected nodes into a new Function"),
-						FSlateIcon(),
-						FUIAction(FExecuteAction::CreateLambda([Model, Controller]() {
-							TArray<FName> Nodes = Model->GetSelectNodes();
-							Controller->OpenUndoBracket(TEXT("Collapse to Function"));
-							URigVMCollapseNode* CollapseNode = Controller->CollapseNodes(Nodes, TEXT("New Function"), true, true);
-							if(CollapseNode)
-							{
-								if (Controller->PromoteCollapseNodeToFunctionReferenceNode(CollapseNode->GetFName(), true, true).IsNone())
-								{
-									Controller->CancelUndoBracket();
-								}
-								else
-								{
-									Controller->CloseUndoBracket();
-								}
-							}
-							else
-							{
-								Controller->CancelUndoBracket();
-							}
-						})
-					));
-
-					if (URigVMCollapseNode* CollapseNode = Cast<URigVMCollapseNode>(RigNode->GetModelNode()))
-					{
-						OrganizationSection.AddMenuEntry(
-							"Promote To Function",
-							LOCTEXT("PromoteToFunction", "Promote To Function"),
-							LOCTEXT("PromoteToFunction_Tooltip", "Turns the Collapse Node into a Function"),
-							FSlateIcon(),
-							FUIAction(FExecuteAction::CreateLambda([Controller, CollapseNode]() {
-								Controller->PromoteCollapseNodeToFunctionReferenceNode(CollapseNode->GetFName(), true, true);
-							})
-						));
-					}
-
-					if (URigVMFunctionReferenceNode* FunctionReferenceNode = Cast<URigVMFunctionReferenceNode>(RigNode->GetModelNode()))
-					{
-						TSoftObjectPtr<URigVMFunctionReferenceNode> RefPtr(FunctionReferenceNode);
-						if(RefPtr.GetLongPackageName() != FunctionReferenceNode->GetReferencedFunctionHeader().LibraryPointer.LibraryNode.GetLongPackageName())
-						{
-							OrganizationSection.AddMenuEntry(
-							   "Localize Function",
-							   LOCTEXT("LocalizeFunction", "Localize Function"),
-							   LOCTEXT("LocalizeFunction_Tooltip", "Creates a local copy of the function backing the node."),
-							   FSlateIcon(),
-							   FUIAction(FExecuteAction::CreateLambda([RigBlueprint, FunctionReferenceNode]() {
-								   RigBlueprint->BroadcastRequestLocalizeFunctionDialog(FunctionReferenceNode->GetFunctionIdentifier(), true);
-							   })
-						    ));
-							
-
-							if(!FunctionReferenceNode->IsFullyRemapped() && FunctionReferenceNode->GetReferencedFunctionHeader().LibraryPointer.LibraryNode.ResolveObject())
-							{
-								FToolMenuSection& VariablesSection = Menu->AddSection("EdGraphSchemaVariables", LOCTEXT("Variables", "Variables"));
-								VariablesSection.AddMenuEntry(
-                                    "MakeVariablesFromFunctionReferenceNode",
-                                    LOCTEXT("MakeVariablesFromFunctionReferenceNode", "Create required variables"),
-                                    LOCTEXT("MakeVariablesFromFunctionReferenceNode_Tooltip", "Creates all required variables for this function and binds them"),
-                                    FSlateIcon(),
-                                    FUIAction(FExecuteAction::CreateLambda([Controller, FunctionReferenceNode, RigBlueprint]() {
-
-                                    	const TArray<FRigVMExternalVariable> ExternalVariables = FunctionReferenceNode->GetExternalVariables(false);
-                                    	if(!ExternalVariables.IsEmpty())
-                                    	{
-											FScopedTransaction Transaction(LOCTEXT("MakeVariablesFromFunctionReferenceNode", "Create required variables"));
-                                    		RigBlueprint->Modify();
-
-                                    		if (URigVMLibraryNode* LibraryNode = FunctionReferenceNode->LoadReferencedNode())
-                                    		{
-                                    			UControlRigBlueprint* ReferencedBlueprint = LibraryNode->GetTypedOuter<UControlRigBlueprint>();
-											   // ReferencedBlueprint != RigBlueprint - since only FunctionReferenceNodes from other assets have the potential to be unmapped
-                                    		
-											   for(const FRigVMExternalVariable& ExternalVariable : ExternalVariables)
-											   {
-												   FString DefaultValue;
-												   if(ReferencedBlueprint)
-												   {
-													   for(const FBPVariableDescription& NewVariable : ReferencedBlueprint->NewVariables)
-													   {
-														   if(NewVariable.VarName == ExternalVariable.Name)
-														   {
-															   DefaultValue = NewVariable.DefaultValue;
-															   break;
-														   }
-													   }
-												   }
-                                    			
-												   FName NewVariableName = RigBlueprint->AddHostMemberVariableFromExternal(ExternalVariable, DefaultValue);
-												   if(!NewVariableName.IsNone())
-												   {
-													   Controller->SetRemappedVariable(FunctionReferenceNode, ExternalVariable.Name, NewVariableName);
-												   }
-											   }
-                                    		}
-
-                                    		FBlueprintEditorUtils::MarkBlueprintAsModified(RigBlueprint);
-                                    	}
-                                        
-                                    })
-                                ));
-							}
-						}
-					}
-
-					if (URigVMFunctionReferenceNode* FunctionRefNode = Cast<URigVMFunctionReferenceNode>(RigNode->GetModelNode()))
-					{
-						OrganizationSection.AddMenuEntry(
-							"Promote To Collapse Node",
-							LOCTEXT("PromoteToCollapseNode", "Promote To Collapse Node"),
-							LOCTEXT("PromoteToCollapseNode_Tooltip", "Turns the Function Ref Node into a Collapse Node"),
-							FSlateIcon(),
-							FUIAction(FExecuteAction::CreateLambda([Controller, FunctionRefNode]() {
-								Controller->PromoteFunctionReferenceNodeToCollapseNode(FunctionRefNode->GetFName());
-								})
-							));
-					}
-
-					if (URigVMTemplateNode* TemplateNode = Cast<URigVMTemplateNode>(RigNode->GetModelNode()))
-					{
-						if (!TemplateNode->IsSingleton())
-						{
-							FToolMenuSection& TemplatesSection = Menu->AddSection("EdGraphSchemaTemplates", LOCTEXT("TemplatesHeader", "Templates"));
-							TemplatesSection.AddMenuEntry(
-								"Unresolve Template Node",
-								LOCTEXT("UnresolveTemplateNode", "Unresolve Template Node"),
-								LOCTEXT("UnresolveTemplateNode_Tooltip", "Removes any type information from the template node"),
-								FSlateIcon(),
-								FUIAction(FExecuteAction::CreateLambda([Controller, Model]() {
-									const TArray<FName> Nodes = Model->GetSelectNodes();
-									Controller->UnresolveTemplateNodes(Nodes, true, true);
-								})
-							));
-						}
-					}
-
-					if (URigVMLibraryNode* LibraryNode = Cast<URigVMLibraryNode>(RigNode->GetModelNode()))
-					
-					{
-						OrganizationSection.AddMenuEntry(
-							"Expand Node",
-							LOCTEXT("ExpandNode", "Expand Node"),
-							LOCTEXT("ExpandNode_Tooltip", "Expands the contents of the node into this graph"),
-							FSlateIcon(),
-							FUIAction(FExecuteAction::CreateLambda([Controller, LibraryNode]() {
-								Controller->OpenUndoBracket(TEXT("Expand node"));
-								TArray<URigVMNode*> ExpandedNodes = Controller->ExpandLibraryNode(LibraryNode->GetFName(), true, true);
-								if (ExpandedNodes.Num() > 0)
-								{
-									TArray<FName> ExpandedNodeNames;
-									for (URigVMNode* ExpandedNode : ExpandedNodes)
-									{
-										ExpandedNodeNames.Add(ExpandedNode->GetFName());
-									}
-									Controller->SetNodeSelection(ExpandedNodeNames);
-								}
-								Controller->CloseUndoBracket();
-							})
-						));
-					}
-
-					if(bCanNodeBeUpgraded)
-					{
-						FToolMenuSection& VersioningSection = Menu->AddSection("EdGraphSchemaVersioning", LOCTEXT("VersioningHeader", "Versioning"));
-						VersioningSection.AddMenuEntry(
-							"Upgrade Nodes",
-							LOCTEXT("UpgradeNodes", "Upgrade Nodes"),
-							LOCTEXT("UpgradeNodes_Tooltip", "Upgrades deprecated nodes to their current implementation"),
-							FSlateIcon(),
-							FUIAction(FExecuteAction::CreateLambda([Model, Controller]() {
-								TArray<FName> Nodes = Model->GetSelectNodes();
-								Controller->UpgradeNodes(Nodes, true, true);
-							})
-						));
-					}
-					
-					OrganizationSection.AddSubMenu("Alignment", LOCTEXT("AlignmentHeader", "Alignment"), FText(), FNewToolMenuDelegate::CreateLambda([](UToolMenu* AlignmentMenu)
-					{
-						{
-							FToolMenuSection& InSection = AlignmentMenu->AddSection("EdGraphSchemaAlignment", LOCTEXT("AlignHeader", "Align"));
-							InSection.AddMenuEntry(FGraphEditorCommands::Get().AlignNodesTop);
-							InSection.AddMenuEntry(FGraphEditorCommands::Get().AlignNodesMiddle);
-							InSection.AddMenuEntry(FGraphEditorCommands::Get().AlignNodesBottom);
-							InSection.AddMenuEntry(FGraphEditorCommands::Get().AlignNodesLeft);
-							InSection.AddMenuEntry(FGraphEditorCommands::Get().AlignNodesCenter);
-							InSection.AddMenuEntry(FGraphEditorCommands::Get().AlignNodesRight);
-							InSection.AddMenuEntry(FGraphEditorCommands::Get().StraightenConnections);
-						}
-
-						{
-							FToolMenuSection& InSection = AlignmentMenu->AddSection("EdGraphSchemaDistribution", LOCTEXT("DistributionHeader", "Distribution"));
-							InSection.AddMenuEntry(FGraphEditorCommands::Get().DistributeNodesHorizontally);
-							InSection.AddMenuEntry(FGraphEditorCommands::Get().DistributeNodesVertically);
-						}
-					}));
-
-					// this struct is only available in EngineTest for now
-					static const FString DecoratorObjectPath = TEXT("/Script/EngineTestEditor.EngineTestRigVM_SimpleDecorator");
-					if(const UScriptStruct* SimpleDecoratorStruct = Cast<UScriptStruct>(RigVMTypeUtils::FindObjectFromCPPTypeObjectPath(DecoratorObjectPath)))
-					{
-						if(URigVMNode* ModelNode = RigNode->GetModelNode())
-						{
-							FToolMenuSection& EngineTestSection = Menu->AddSection("EdGraphSchemaEngineTest", LOCTEXT("EngineTestHeader", "EngineTest"));
-							EngineTestSection.AddMenuEntry(
-								"Add simple decorator",
-								LOCTEXT("AddSimpleDecorator", "Add simple decorator"),
-								LOCTEXT("AddSimpleDecorator_Tooltip", "Adds a simple test decorator to the node"),
-								FSlateIcon(),
-								FUIAction(FExecuteAction::CreateLambda([Controller, ModelNode, SimpleDecoratorStruct]()
-								{
-									(void)Controller->AddDecorator(
-										ModelNode->GetFName(),
-										*SimpleDecoratorStruct->GetPathName(),
-										TEXT("Decorator"),
-										FString(), INDEX_NONE, true, true);
-								}))
-							);
-						}
-					}
-				}
-			}
-		}
-	}
-}
-
-void FControlRigEditorModule::PreChange(const UUserDefinedStruct* Changed,
-	FStructureEditorUtils::EStructureEditorChangeInfo ChangedType)
-{
-	// the following is similar to
-	// FUserDefinedStructureCompilerInner::ReplaceStructWithTempDuplicate()
-	// it is necessary since existing rigs need to be kept valid until after PreBPCompile
-	// there are other systems, such as sequencer, that might need to evaluate the rig
-	// for one last time during PreBPCompile
-	// Overrall sequence of events
-	// PreStructChange --1--> PostStructChange
-	//                              --2--> PreBPCompile --3--> PostBPCompile
-	
-	UUserDefinedStruct* StructureToReinstance = (UUserDefinedStruct*)Changed;
-
-	FUserDefinedStructureCompilerUtils::ReplaceStructWithTempDuplicateByPredicate(
-		StructureToReinstance,
-		[](FStructProperty* InStructProperty)
-		{
-			// make sure variable properties on the BP is patched
-			// since active rig instance still references it
-			if (URigVMBlueprintGeneratedClass* BPClass = Cast<URigVMBlueprintGeneratedClass>(InStructProperty->GetOwnerClass()))
-			{
-				if (BPClass->ClassGeneratedBy->IsA<UControlRigBlueprint>())
-				{
-					return true;
-				}
-			}
-			// similar story, VM instructions reference properties on the GeneratorClass
-			else if ((InStructProperty->GetOwnerStruct())->IsA<URigVMMemoryStorageGeneratorClass>())
-			{
-				return true;
-			}
-			else if (URigVMDetailsViewWrapperObject::IsValidClass(InStructProperty->GetOwnerClass()))
-			{
-				return true;
-			}
-			
-			return false;
-		},
-		[](UStruct* InStruct)
-		{
-			// refresh these since VM caching references them
-			if (URigVMMemoryStorageGeneratorClass* GeneratorClass = Cast<URigVMMemoryStorageGeneratorClass>(InStruct))
-			{
-				GeneratorClass->RefreshLinkedProperties();
-				GeneratorClass->RefreshPropertyPaths();	
-			}
-			else if (InStruct->IsChildOf(URigVMDetailsViewWrapperObject::StaticClass()))
-			{
-				URigVMDetailsViewWrapperObject::MarkOutdatedClass(Cast<UClass>(InStruct));
-			}
-		});
-	
-	// in the future we could only invalidate caches on affected rig instances, it shouldn't make too much of a difference though
-	for (TObjectIterator<UControlRig> It(RF_Transient | RF_ClassDefaultObject, /** bIncludeDerivedClasses */ true, /** InternalExcludeFlags */ EInternalObjectFlags::Garbage); It; ++It)
-	{
-		UControlRig* Rig = *It;
-		// rebuild property list and property path list
-		Rig->InvalidateCachedMemory();
-	}
-}
-
-void FControlRigEditorModule::PostChange(const UUserDefinedStruct* Changed,
-                                   FStructureEditorUtils::EStructureEditorChangeInfo ChangedType)
-{
-	TArray<UControlRigBlueprint*> BlueprintsToRefresh;
-	for (TObjectIterator<URigVMPin> It(RF_Transient | RF_ClassDefaultObject, /** bIncludeDerivedClasses */ true, /** InternalExcludeFlags */ EInternalObjectFlags::Garbage); It; ++It)
-	{
-		URigVMPin* Pin = *It;
-		// GetCPPTypeObject also makes sure the pin's type information is update to date
-		if (Pin && Pin->GetCPPTypeObject() == Changed)
-		{
-			if (UControlRigBlueprint* RigBlueprint = Pin->GetTypedOuter<UControlRigBlueprint>())
-			{
-				BlueprintsToRefresh.AddUnique(RigBlueprint);
-				
-				// this pin is part of a function definition
-				// update all BP that uses this function
-				if (Pin->GetGraph() == RigBlueprint->GetLocalFunctionLibrary())
-				{
-					TArray< TSoftObjectPtr<URigVMFunctionReferenceNode> > References;
-					References = RigBlueprint->RigVMClient.GetOrCreateFunctionLibrary(false)->GetReferencesForFunction(Pin->GetNode()->GetFName());
-					
-					for (const TSoftObjectPtr<URigVMFunctionReferenceNode>& Reference : References)
-					{
-						URigVMFunctionReferenceNode* RefNode = Reference.LoadSynchronous();
-						if (!RefNode)
-						{
-							continue;
-						}
-						
-						if (UControlRigBlueprint* FunctionUserBlueprint = RefNode->GetTypedOuter<UControlRigBlueprint>())
-						{
-							BlueprintsToRefresh.AddUnique(FunctionUserBlueprint);
-						}
-					}	
-				}
-			}
-		}
-	}
-
-	for (UControlRigBlueprint* RigBlueprint : BlueprintsToRefresh)
-	{
-		RigBlueprint->OnRigVMRegistryChanged();
-		RigBlueprint->MarkPackageDirty();
-	}
-	
-	for (UControlRigBlueprint* RigBlueprint : BlueprintsToRefresh)
-	{
-		// this should make sure variables in BP are updated with the latest struct object
-		// otherwise RigVMCompiler validation would complain about variable type - pin type mismatch
-		FCompilerResultsLog	ResultsLog;
-		FKismetEditorUtilities::CompileBlueprint(RigBlueprint, EBlueprintCompileOptions::None, &ResultsLog);
-		
-		// BP compiler always initialize the new CDO by copying from the old CDO,
-		// however, in case that a BP variable type has changed, the data old CDO would be invalid because
-		// while the old memory container still references the temp duplicated struct we created during PreChange()
-		// registers that reference the BP variable would be referencing the new struct as a result of
-		// FKismetCompilerContext::CompileClassLayout, so type mismatch would invalidate relevant copy operations
-		// so to simplify things, here we just reset all rigs upon error
-		if (ResultsLog.NumErrors > 0)
-		{
-			URigVMBlueprintGeneratedClass* RigClass = RigBlueprint->GetRigVMBlueprintGeneratedClass();
-			UControlRig* CDO = Cast<UControlRig>(RigClass->GetDefaultObject(true /* create if needed */));
-			if (CDO->GetVM() != nullptr)
-			{
-				CDO->GetVM()->Reset();
-			}
-			TArray<UObject*> ArchetypeInstances;
-			CDO->GetArchetypeInstances(ArchetypeInstances);
-			for (UObject* Instance : ArchetypeInstances)
-			{
-				if (UControlRig* InstanceRig = Cast<UControlRig>(Instance))
-				{
-					InstanceRig->GetVM()->Reset();
-				}
-			}
-		}
-	}
-}
-
 IMPLEMENT_MODULE(FControlRigEditorModule, ControlRigEditor)
 
 #undef LOCTEXT_NAMESPACE
