Control Rig: Additive control rig with control rig component

diff --git a/Engine/Plugins/Animation/ControlRig/Source/ControlRig/Private/ControlRig.cpp b/Engine/Plugins/Animation/ControlRig/Source/ControlRig/Private/ControlRig.cpp
--- a/Engine/Plugins/Animation/ControlRig/Source/ControlRig/Private/ControlRig.cpp
+++ b/Engine/Plugins/Animation/ControlRig/Source/ControlRig/Private/ControlRig.cpp
@@ -319,6 +319,96 @@ bool UControlRig::InitializeVM(const FName& InEventName)
 	return true;
 }
 
+void UControlRig::Evaluate_AnyThread()
+{
+	if (bIsAdditive && !bRequiresInitExecution)
+	{
+		// we can have other systems trying to poke into running instances of Control Rigs
+		// on the anim thread and query data, such as
+		// URigVMHostSkeletalMeshComponent::RebuildDebugDrawSkeleton,
+		// using a lock here to prevent them from having an inconsistent view of the rig at some
+		// intermediate stage of evaluation, for example, during evaluate, we can have a call
+		// to copy hierarchy, which empties the hierarchy for a short period of time
+		// and we don't want other systems to see that.
+		FScopeLock EvaluateLock(&GetEvaluateMutex());
+
+		if (PoseBeforeBackwardsSolve.Num() == 0)
+		{
+			// If the pose is empty, this is an indication that a new pose is coming in
+			PoseBeforeBackwardsSolve = GetHierarchy()->GetPose(false, ERigElementType::Bone, TArrayView<const FRigElementKey>());
+		}
+		else
+		{
+			// Restore the pose from the anim sequence
+			DynamicHierarchy->SetPose(PoseBeforeBackwardsSolve);
+		}
+		
+		// Backwards solve
+		{
+			TGuardValue<bool> UpdatePreferredAngles(DynamicHierarchy->bUpdatePreferedEulerAngleWhenSettingTransform, false);
+			Execute(FRigUnit_InverseExecution::EventName);
+		}
+
+		// Store control pose after backwards solve to figure out additive local transforms based on animation
+		ControlsAfterBackwardsSolve = GetHierarchy()->GetPose(false, ERigElementType::Control, TArrayView<const FRigElementKey>());
+
+		// Apply additive controls
+		for (TPair<FRigElementKey, FRigSetControlValueInfo>& Value : ControlValues)
+		{
+			if (FRigBaseElement* Element = GetHierarchy()->Find(Value.Key))
+			{
+				if (FRigControlElement* Control = Cast<FRigControlElement>(Element))
+				{
+					FRigSetControlValueInfo& Info = Value.Value;
+
+					// Transform from animation
+					const FRigControlValue PreviousValue = GetHierarchy()->GetControlValue(Control, ERigControlValueType::Current);
+					const FTransform PreviousTransform = PreviousValue.GetAsTransform(Control->Settings.ControlType, Control->Settings.PrimaryAxis);
+
+					// Additive transform from controls
+					const FRigControlValue& AdditiveValue = Info.Value;
+					const FTransform AdditiveTransform = AdditiveValue.GetAsTransform(Control->Settings.ControlType, Control->Settings.PrimaryAxis);
+
+					// Add them to find the final value
+					FTransform FinalTransform = AdditiveTransform * PreviousTransform;
+
+					FRigControlValue FinalValue;
+					FinalValue.SetFromTransform(FinalTransform, Control->Settings.ControlType, Control->Settings.PrimaryAxis);
+					
+					GetHierarchy()->SetControlValue(Control, FinalValue, ERigControlValueType::Current, Info.bSetupUndo, false, Info.bPrintPythonCommnds, false);
+					GetHierarchy()->SetPreferredEulerAnglesFromValue(Control, AdditiveValue, ERigControlValueType::Current, Info.bFixEulerFlips);
+					
+					if (Info.bNotify && OnControlModified.IsBound())
+					{
+						OnControlModified.Broadcast(this, Control, Info.Context);
+						Info.bNotify = false;
+					}
+				}
+			}
+		}
+		
+		// Forward solve
+		{
+			TGuardValue<bool> UpdatePreferredAngles(DynamicHierarchy->bUpdatePreferedEulerAngleWhenSettingTransform, false);
+			Execute(FRigUnit_BeginExecution::EventName);
+		}
+	}
+	else
+	{
+		Super::Evaluate_AnyThread();
+	}
+}
+
+void UControlRig::ResetControlValues()
+{
+	ControlValues.Reset();
+}
+
+void UControlRig::ClearPoseBeforeBackwardsSolve()
+{
+	PoseBeforeBackwardsSolve.Reset();
+}
+
 void UControlRig::InitializeFromCDO()
 {
 	Super::InitializeFromCDO();
@@ -856,8 +946,11 @@ bool UControlRig::Execute(const FName& InEventName)
 
 	if (ExecutedEvent.IsBound())
 	{
-		FControlRigBracketScope BracketScope(ExecuteBracket);
-		ExecutedEvent.Broadcast(this, InEventName);
+		if (!IsAdditive() || InEventName != FRigUnit_InverseExecution::EventName)
+		{
+			FControlRigBracketScope BracketScope(ExecuteBracket);
+			ExecutedEvent.Broadcast(this, InEventName);
+		}
 	}
 
 	// close remaining undo brackets from hierarchy
@@ -1518,6 +1611,53 @@ FTransform UControlRig::GetControlGlobalTransform(const FName& InControlName) co
 	return DynamicHierarchy->GetGlobalTransform(FRigElementKey(InControlName, ERigElementType::Control), false);
 }
 
+FRigControlValue UControlRig::GetControlValue(FRigControlElement* InControl, const ERigControlValueType& InValueType)
+{
+	if (bIsAdditive && InValueType == ERigControlValueType::Current)
+	{
+		const int32 ControlIndex = ControlsAfterBackwardsSolve.GetIndex(InControl->GetKey());
+		if (ControlIndex != INDEX_NONE)
+		{
+			// return local space control value (the one to be added after backwards solve)
+			const FRigPoseElement& AnimPose = ControlsAfterBackwardsSolve[ControlIndex];
+			const FRigControlValue& CurrentValue = GetHierarchy()->GetControlValue(InControl, InValueType);
+			const FTransform FinalTransform = CurrentValue.GetAsTransform(InControl->Settings.ControlType, InControl->Settings.PrimaryAxis);
+			const FTransform AdditiveTransform = FinalTransform * AnimPose.LocalTransform.Inverse();
+			FRigControlValue AdditiveValue;
+			AdditiveValue.SetFromTransform(AdditiveTransform, InControl->Settings.ControlType, InControl->Settings.PrimaryAxis);
+			return AdditiveValue;
+		}
+	}
+	return GetHierarchy()->GetControlValue(InControl, InValueType);
+}
+
+void UControlRig::SetControlValueImpl(const FName& InControlName, const FRigControlValue& InValue, bool bNotify,
+	const FRigControlModifiedContext& Context, bool bSetupUndo, bool bPrintPythonCommnds, bool bFixEulerFlips)
+{
+	const FRigElementKey Key(InControlName, ERigElementType::Control);
+
+	FRigControlElement* ControlElement = DynamicHierarchy->Find<FRigControlElement>(Key);
+	if(ControlElement == nullptr)
+	{
+		return;
+	}
+	if (bIsAdditive)
+	{
+		// Store the value to apply it after the backwards solve
+		FRigSetControlValueInfo Info = {InValue, bNotify, Context, bSetupUndo, bPrintPythonCommnds, bFixEulerFlips};
+		ControlValues.Add(ControlElement->GetKey(), Info);
+	}
+	else
+	{
+		DynamicHierarchy->SetControlValue(ControlElement, InValue, ERigControlValueType::Current, bSetupUndo, false, bPrintPythonCommnds, bFixEulerFlips);
+
+		if (bNotify && OnControlModified.IsBound())
+		{
+			OnControlModified.Broadcast(this, ControlElement, Context);
+		}
+	}
+}
+
 bool UControlRig::SetControlGlobalTransform(const FName& InControlName, const FTransform& InGlobalTransform, bool bNotify, const FRigControlModifiedContext& Context, bool bSetupUndo, bool bPrintPythonCommands, bool bFixEulerFlips)
 {
 	FTransform GlobalTransform = InGlobalTransform;
@@ -1556,8 +1696,21 @@ FRigControlValue UControlRig::GetControlValueFromGlobalTransform(const FName& In
 	{
 		if(DynamicHierarchy)
 		{
-			FTransform Transform = DynamicHierarchy->ComputeLocalControlValue(ControlElement, InGlobalTransform, InTransformType);
-			Value.SetFromTransform(Transform, ControlElement->Settings.ControlType, ControlElement->Settings.PrimaryAxis);
+			if (bIsAdditive)
+			{
+				const int32 ControlIndex = ControlsAfterBackwardsSolve.GetIndex(ControlElement->GetKey());
+				if (ControlIndex != INDEX_NONE)
+				{
+					const FTransform& AnimGlobalTransform = ControlsAfterBackwardsSolve[ControlIndex].GlobalTransform;
+					const FTransform AdditiveTransform = InGlobalTransform.GetRelativeTransform(AnimGlobalTransform);					
+					Value.SetFromTransform(AdditiveTransform, ControlElement->Settings.ControlType, ControlElement->Settings.PrimaryAxis);
+				}
+			}
+			else
+			{
+				FTransform Transform = DynamicHierarchy->ComputeLocalControlValue(ControlElement, InGlobalTransform, InTransformType);
+				Value.SetFromTransform(Transform, ControlElement->Settings.ControlType, ControlElement->Settings.PrimaryAxis);
+			}
 
 			if (ShouldApplyLimits())
 			{
@@ -1591,6 +1744,17 @@ FTransform UControlRig::GetControlLocalTransform(const FName& InControlName)
 	{
 		return FTransform::Identity;
 	}
+	if (bIsAdditive)
+	{
+		FRigElementKey ControlKey (InControlName, ERigElementType::Control);
+		if (FRigBaseElement* Element = DynamicHierarchy->Find(ControlKey))
+		{
+			if (FRigControlElement* ControlElement = Cast<FRigControlElement>(Element))
+			{
+				return GetControlValue(ControlElement, ERigControlValueType::Current).GetAsTransform(ControlElement->Settings.ControlType, ControlElement->Settings.PrimaryAxis);
+			}
+		}
+	}
 	return DynamicHierarchy->GetLocalTransform(FRigElementKey(InControlName, ERigElementType::Control));
 }
 
diff --git a/Engine/Plugins/Animation/ControlRig/Source/ControlRig/Private/ControlRigComponent.cpp b/Engine/Plugins/Animation/ControlRig/Source/ControlRig/Private/ControlRigComponent.cpp
--- a/Engine/Plugins/Animation/ControlRig/Source/ControlRig/Private/ControlRigComponent.cpp
+++ b/Engine/Plugins/Animation/ControlRig/Source/ControlRig/Private/ControlRigComponent.cpp
@@ -3,6 +3,7 @@
 #include "ControlRigComponent.h"
 #include "Engine/SkeletalMesh.h"
 #include "Units/Execution/RigUnit_BeginExecution.h"
+#include "Units/Execution/RigUnit_InverseExecution.h"
 #include "Units/Execution/RigUnit_Hierarchy.h"
 
 #include "SkeletalDebugRendering.h"
@@ -351,7 +352,11 @@ void UControlRigComponent::Initialize()
 
 	TGuardValue<bool> InitializeBracket(bIsInsideInitializeBracket, true);
 	
-	ClearMappedElements();
+	if (!UserDefinedElements.IsEmpty())
+	{
+		MappedElements = UserDefinedElements;
+	}
+	ValidateMappingData();
 
 #if WITH_EDITOR
 	if (bUpdateInEditor)
@@ -482,6 +487,11 @@ void UControlRigComponent::Update(float DeltaTime)
 				}
 #endif
 
+				// Animation will only evaluate if DeltaTime > 0
+				if (CR->IsAdditive() && DeltaTime > 0)
+				{
+					CR->ClearPoseBeforeBackwardsSolve();
+				}
 				CR->Evaluate_AnyThread();
 
 #if WITH_EDITOR
@@ -614,7 +624,8 @@ void UControlRigComponent::AddMappedComponents(TArray<FControlRigComponentMapped
 	AddMappedElements(ElementsToMap);
 }
 
-void UControlRigComponent::AddMappedSkeletalMesh(USkeletalMeshComponent* SkeletalMeshComponent, TArray<FControlRigComponentMappedBone> Bones, TArray<FControlRigComponentMappedCurve> Curves)
+void UControlRigComponent::AddMappedSkeletalMesh(USkeletalMeshComponent* SkeletalMeshComponent, TArray<FControlRigComponentMappedBone> Bones, TArray<
+                                                 FControlRigComponentMappedCurve> Curves, const EControlRigComponentMapDirection InDirection)
 {
 	if (SkeletalMeshComponent == nullptr)
 	{
@@ -706,6 +717,7 @@ void UControlRigComponent::AddMappedSkeletalMesh(USkeletalMeshComponent* Skeleta
 		ElementToMap.ElementName = BoneToMap.Source;
 		ElementToMap.ElementType = ERigElementType::Bone;
 		ElementToMap.TransformName = BoneToMap.Target;
+		ElementToMap.Direction = InDirection;
 
 		ElementsToMap.Add(ElementToMap);
 	}
@@ -732,9 +744,9 @@ void UControlRigComponent::AddMappedSkeletalMesh(USkeletalMeshComponent* Skeleta
 	AddMappedElements(ElementsToMap);
 }
 
-void UControlRigComponent::AddMappedCompleteSkeletalMesh(USkeletalMeshComponent* SkeletalMeshComponent)
+void UControlRigComponent::AddMappedCompleteSkeletalMesh(USkeletalMeshComponent* SkeletalMeshComponent, const EControlRigComponentMapDirection InDirection)
 {
-	AddMappedSkeletalMesh(SkeletalMeshComponent, TArray<FControlRigComponentMappedBone>(), TArray<FControlRigComponentMappedCurve>());
+	AddMappedSkeletalMesh(SkeletalMeshComponent, TArray<FControlRigComponentMappedBone>(), TArray<FControlRigComponentMappedCurve>(), InDirection);
 }
 
 void UControlRigComponent::SetBoneInitialTransformsFromSkeletalMesh(USkeletalMesh* InSkeletalMesh)
@@ -1509,98 +1521,140 @@ void UControlRigComponent::TransferOutputs()
 {
 	if (ControlRig)
 	{
-		USceneComponent* LastComponent = nullptr;
-		FControlRigAnimInstanceProxy* Proxy = nullptr;
-
-		for (const FControlRigComponentMappedElement& MappedElement : MappedElements)
+		if (ControlRig->IsAdditive())
 		{
-			if (LastComponent != MappedElement.SceneComponent || Proxy == nullptr)
+			if (MappedElements.Num() > 0)
 			{
-				Proxy = MappedElement.GetAnimProxyOnGameThread();
-				if (Proxy)
+				if (USkeletalMeshComponent* SkeletalMeshComponent = Cast<USkeletalMeshComponent>(MappedElements[0].SceneComponent))
 				{
-					Proxy->StoredTransforms.Reset();
-					Proxy->StoredCurves.Reset();
-					LastComponent = MappedElement.SceneComponent;
+					if (USkeletalMesh* SkeletalMesh = SkeletalMeshComponent->GetSkeletalMeshAsset())
+					{
+						TArray<FTransform> BoneSpaceTransforms;
+						BoneSpaceTransforms.SetNumUninitialized(SkeletalMesh->GetRefSkeleton().GetNum());
+						for (const FControlRigComponentMappedElement& MappedElement : MappedElements)
+						{
+							if (MappedElement.ElementIndex == INDEX_NONE || MappedElement.Direction == EControlRigComponentMapDirection::Input)
+							{
+								continue;
+							}
+
+							if (MappedElement.ElementType == ERigElementType::Bone ||
+								MappedElement.ElementType == ERigElementType::Control ||
+								MappedElement.ElementType == ERigElementType::Null)
+							{
+								if (MappedElement.SubIndex >= 0 && BoneSpaceTransforms.IsValidIndex(MappedElement.SubIndex))
+								{
+									FTransform Transform = ControlRig->GetHierarchy()->GetLocalTransform(MappedElement.ElementIndex);
+									BoneSpaceTransforms[MappedElement.SubIndex] = Transform;
+								}
+							}
+						}
+
+						TArray<FTransform> OutSpaceBases;
+						OutSpaceBases.SetNumUninitialized(BoneSpaceTransforms.Num());
+						SkeletalMesh->FillComponentSpaceTransforms(BoneSpaceTransforms, SkeletalMeshComponent->FillComponentSpaceTransformsRequiredBones, SkeletalMeshComponent->GetEditableComponentSpaceTransforms());
+#if WITH_EDITOR
+						SkeletalMeshComponent->ApplyEditedComponentSpaceTransforms();
+#endif
+					}
 				}
 			}
 		}
-
-		for (const FControlRigComponentMappedElement& MappedElement : MappedElements)
+		else
 		{
-			if (MappedElement.ElementIndex == INDEX_NONE || MappedElement.Direction == EControlRigComponentMapDirection::Input)
+			USceneComponent* LastComponent = nullptr;
+			FControlRigAnimInstanceProxy* Proxy = nullptr;
+
+			for (const FControlRigComponentMappedElement& MappedElement : MappedElements)
 			{
-				continue;
+				if (LastComponent != MappedElement.SceneComponent || Proxy == nullptr)
+				{
+					Proxy = MappedElement.GetAnimProxyOnGameThread();
+					if (Proxy)
+					{
+						Proxy->StoredTransforms.Reset();
+						Proxy->StoredCurves.Reset();
+						LastComponent = MappedElement.SceneComponent;
+					}
+				}
 			}
 
-			if (MappedElement.ElementType == ERigElementType::Bone ||
-				MappedElement.ElementType == ERigElementType::Control ||
-				MappedElement.ElementType == ERigElementType::Null)
+			for (const FControlRigComponentMappedElement& MappedElement : MappedElements)
 			{
-				FTransform Transform = ControlRig->GetHierarchy()->GetGlobalTransform(MappedElement.ElementIndex);
-				ConvertTransformFromRigSpace(Transform, MappedElement.Space);
-
-				Transform = Transform * MappedElement.Offset;
+				if (MappedElement.ElementIndex == INDEX_NONE || MappedElement.Direction == EControlRigComponentMapDirection::Input)
+				{
+					continue;
+				}
 
-				if (MappedElement.SubIndex >= 0)
+				if (MappedElement.ElementType == ERigElementType::Bone ||
+					MappedElement.ElementType == ERigElementType::Control ||
+					MappedElement.ElementType == ERigElementType::Null)
 				{
-					if (LastComponent != MappedElement.SceneComponent || Proxy == nullptr)
+					FTransform Transform = ControlRig->GetHierarchy()->GetGlobalTransform(MappedElement.ElementIndex);
+					ConvertTransformFromRigSpace(Transform, MappedElement.Space);
+
+					Transform = Transform * MappedElement.Offset;
+
+					if (MappedElement.SubIndex >= 0)
 					{
-						Proxy = MappedElement.GetAnimProxyOnGameThread();
-						if (Proxy)
+						if (LastComponent != MappedElement.SceneComponent || Proxy == nullptr)
 						{
-							LastComponent = MappedElement.SceneComponent;
+							Proxy = MappedElement.GetAnimProxyOnGameThread();
+							if (Proxy)
+							{
+								LastComponent = MappedElement.SceneComponent;
+							}
 						}
-					}
 
-					if (Proxy && (MappedElement.SceneComponent != nullptr))
-					{
-						if(MappedElement.Space == EControlRigComponentSpace::WorldSpace)
+						if (Proxy && (MappedElement.SceneComponent != nullptr))
 						{
-							Transform = Transform.GetRelativeTransform(MappedElement.SceneComponent->GetComponentToWorld());
+							if(MappedElement.Space == EControlRigComponentSpace::WorldSpace)
+							{
+								Transform = Transform.GetRelativeTransform(MappedElement.SceneComponent->GetComponentToWorld());
+							}
+							Proxy->StoredTransforms.FindOrAdd(FMeshPoseBoneIndex(MappedElement.SubIndex)) = Transform;
 						}
-						Proxy->StoredTransforms.FindOrAdd(FMeshPoseBoneIndex(MappedElement.SubIndex)) = Transform;
-					}
-					else if (UInstancedStaticMeshComponent* InstancingComponent = Cast<UInstancedStaticMeshComponent>(MappedElement.SceneComponent))
-					{
-						if (MappedElement.SubIndex < InstancingComponent->GetNumRenderInstances())
+						else if (UInstancedStaticMeshComponent* InstancingComponent = Cast<UInstancedStaticMeshComponent>(MappedElement.SceneComponent))
 						{
-							if (MappedElement.Weight < 1.f - SMALL_NUMBER)
+							if (MappedElement.SubIndex < InstancingComponent->GetNumRenderInstances())
 							{
-								FTransform Previous = FTransform::Identity;
-								InstancingComponent->GetInstanceTransform(MappedElement.SubIndex, Previous, true);
-								Transform = FControlRigMathLibrary::LerpTransform(Previous, Transform, FMath::Clamp<float>(MappedElement.Weight, 0.f, 1.f));
+								if (MappedElement.Weight < 1.f - SMALL_NUMBER)
+								{
+									FTransform Previous = FTransform::Identity;
+									InstancingComponent->GetInstanceTransform(MappedElement.SubIndex, Previous, true);
+									Transform = FControlRigMathLibrary::LerpTransform(Previous, Transform, FMath::Clamp<float>(MappedElement.Weight, 0.f, 1.f));
+								}
+								InstancingComponent->UpdateInstanceTransform(MappedElement.SubIndex, Transform, true, true, true);
 							}
-							InstancingComponent->UpdateInstanceTransform(MappedElement.SubIndex, Transform, true, true, true);
 						}
 					}
-				}
-				else
-				{
-					if (MappedElement.Weight < 1.f - SMALL_NUMBER)
+					else
 					{
-						FTransform Previous = MappedElement.SceneComponent->GetComponentToWorld();
-						Transform = FControlRigMathLibrary::LerpTransform(Previous, Transform, FMath::Clamp<float>(MappedElement.Weight, 0.f, 1.f));
+						if (MappedElement.Weight < 1.f - SMALL_NUMBER)
+						{
+							FTransform Previous = MappedElement.SceneComponent->GetComponentToWorld();
+							Transform = FControlRigMathLibrary::LerpTransform(Previous, Transform, FMath::Clamp<float>(MappedElement.Weight, 0.f, 1.f));
+						}
+						MappedElement.SceneComponent->SetWorldTransform(Transform);
 					}
-					MappedElement.SceneComponent->SetWorldTransform(Transform);
 				}
-			}
-			else if (MappedElement.ElementType == ERigElementType::Curve)
-			{
-				if (MappedElement.SubIndex >= 0)
+				else if (MappedElement.ElementType == ERigElementType::Curve)
 				{
-					if (LastComponent != MappedElement.SceneComponent || Proxy == nullptr)
+					if (MappedElement.SubIndex >= 0)
 					{
-						Proxy = MappedElement.GetAnimProxyOnGameThread();
-						if (Proxy)
+						if (LastComponent != MappedElement.SceneComponent || Proxy == nullptr)
 						{
-							LastComponent = MappedElement.SceneComponent;
+							Proxy = MappedElement.GetAnimProxyOnGameThread();
+							if (Proxy)
+							{
+								LastComponent = MappedElement.SceneComponent;
+							}
 						}
-					}
 
-					if (Proxy)
-					{
-						Proxy->StoredCurves.FindOrAdd(MappedElement.TransformName) = ControlRig->GetHierarchy()->GetCurveValue(MappedElement.ElementIndex);
+						if (Proxy)
+						{
+							Proxy->StoredCurves.FindOrAdd(MappedElement.TransformName) = ControlRig->GetHierarchy()->GetCurveValue(MappedElement.ElementIndex);
+						}
 					}
 				}
 			}
diff --git a/Engine/Plugins/Animation/ControlRig/Source/ControlRig/Private/Rigs/RigHierarchy.cpp b/Engine/Plugins/Animation/ControlRig/Source/ControlRig/Private/Rigs/RigHierarchy.cpp
--- a/Engine/Plugins/Animation/ControlRig/Source/ControlRig/Private/Rigs/RigHierarchy.cpp
+++ b/Engine/Plugins/Animation/ControlRig/Source/ControlRig/Private/Rigs/RigHierarchy.cpp
@@ -2402,6 +2402,12 @@ FRigPose URigHierarchy::GetPose(bool bInitial, ERigElementType InElementType,
 		{
 			PoseElement.LocalTransform = GetTransform(TransformElement, bInitial ? ERigTransformType::InitialLocal : ERigTransformType::CurrentLocal);
 			PoseElement.GlobalTransform = GetTransform(TransformElement, bInitial ? ERigTransformType::InitialGlobal : ERigTransformType::CurrentGlobal);
+
+			if(FRigControlElement* ControlElement = Cast<FRigControlElement>(Element))
+			{
+				PoseElement.PreferredEulerAngle = GetControlPreferredEulerAngles(ControlElement,
+					GetControlPreferredEulerRotationOrder(ControlElement), bInitial);
+			}
 		}
 		else if(FRigCurveElement* CurveElement = Cast<FRigCurveElement>(Element))
 		{
@@ -3540,6 +3546,47 @@ FRigControlValue URigHierarchy::GetControlValue(FRigControlElement* InControlEle
 	return Value;
 }
 
+void URigHierarchy::SetPreferredEulerAnglesFromValue(FRigControlElement* InControlElement, const FRigControlValue& InValue, const ERigControlValueType& InValueType, const bool bFixEulerFlips)
+{
+	const bool bInitial = InValueType == ERigControlValueType::Initial;
+	switch(InControlElement->Settings.ControlType)
+	{
+	case ERigControlType::Rotator:
+		{
+			InControlElement->PreferredEulerAngles.SetRotator(GetRotatorFromControlValue(InValue), bInitial, bFixEulerFlips);
+			break;
+		}
+	case ERigControlType::EulerTransform:
+		{
+			FEulerTransform EulerTransform = GetEulerTransformFromControlValue(InValue);
+			FQuat Quat = EulerTransform.GetRotation();
+			const FVector Angle = GetControlAnglesFromQuat(InControlElement, Quat);
+			InControlElement->PreferredEulerAngles.SetAngles(Angle, bInitial, InControlElement->PreferredEulerAngles.RotationOrder, bFixEulerFlips);
+			break;
+		}
+	case ERigControlType::Transform:
+		{
+			FTransform Transform = GetTransformFromControlValue(InValue);
+			FQuat Quat = Transform.GetRotation();
+			const FVector Angle = GetControlAnglesFromQuat(InControlElement, Quat);
+			InControlElement->PreferredEulerAngles.SetAngles(Angle, bInitial, InControlElement->PreferredEulerAngles.RotationOrder, bFixEulerFlips);
+			break;
+		}
+	case ERigControlType::TransformNoScale:
+		{
+			FTransform Transform = GetTransformNoScaleFromControlValue(InValue);
+			FQuat Quat = Transform.GetRotation();
+			const FVector Angle = GetControlAnglesFromQuat(InControlElement, Quat);
+			InControlElement->PreferredEulerAngles.SetAngles(Angle, bInitial, InControlElement->PreferredEulerAngles.RotationOrder, bFixEulerFlips);
+			break;
+		}
+	default:
+		{
+			break;
+		}
+	}
+};
+
 void URigHierarchy::SetControlValue(FRigControlElement* InControlElement, const FRigControlValue& InValue, ERigControlValueType InValueType, bool bSetupUndo, bool bForce, bool bPrintPythonCommands, bool bFixEulerFlips)
 {
 	LLM_SCOPE_BYNAME(TEXT("Animation/ControlRig"));
@@ -3548,47 +3595,6 @@ void URigHierarchy::SetControlValue(FRigControlElement* InControlElement, const
 
 	if(InControlElement != nullptr)
 	{
-		auto SetPreferredEulerAnglesFromValue = [this, InControlElement, &InValue, InValueType, bFixEulerFlips]()
-		{
-			const bool bInitial = InValueType == ERigControlValueType::Initial;
-			switch(InControlElement->Settings.ControlType)
-			{
-				case ERigControlType::Rotator:
-				{
-					InControlElement->PreferredEulerAngles.SetRotator(GetRotatorFromControlValue(InValue), bInitial, bFixEulerFlips);
-					break;
-				}
-				case ERigControlType::EulerTransform:
-				{
-					FEulerTransform EulerTransform = GetEulerTransformFromControlValue(InValue);
-					FQuat Quat = EulerTransform.GetRotation();
-					const FVector Angle = GetControlAnglesFromQuat(InControlElement, Quat);
-					InControlElement->PreferredEulerAngles.SetAngles(Angle, bInitial, InControlElement->PreferredEulerAngles.RotationOrder, bFixEulerFlips);
-					break;
-				}
-				case ERigControlType::Transform:
-				{
-					FTransform Transform = GetTransformFromControlValue(InValue);
-					FQuat Quat = Transform.GetRotation();
-					const FVector Angle = GetControlAnglesFromQuat(InControlElement, Quat);
-					InControlElement->PreferredEulerAngles.SetAngles(Angle, bInitial, InControlElement->PreferredEulerAngles.RotationOrder, bFixEulerFlips);
-					break;
-				}
-				case ERigControlType::TransformNoScale:
-				{
-					FTransform Transform = GetTransformNoScaleFromControlValue(InValue);
-					FQuat Quat = Transform.GetRotation();
-					const FVector Angle = GetControlAnglesFromQuat(InControlElement, Quat);
-					InControlElement->PreferredEulerAngles.SetAngles(Angle, bInitial, InControlElement->PreferredEulerAngles.RotationOrder, bFixEulerFlips);
-					break;
-				}
-				default:
-				{
-					break;
-				}
-			}
-		};
-
 		switch(InValueType)
 		{
 			case ERigControlValueType::Current:
@@ -3611,7 +3617,7 @@ void URigHierarchy::SetControlValue(FRigControlElement* InControlElement, const
 				);
 				if (bFixEulerFlips)
 				{
-					SetPreferredEulerAnglesFromValue();
+					SetPreferredEulerAnglesFromValue(InControlElement, Value, InValueType, bFixEulerFlips);
 				}
 				break;
 			}
@@ -3636,7 +3642,7 @@ void URigHierarchy::SetControlValue(FRigControlElement* InControlElement, const
 
 				if (bFixEulerFlips)
 				{
-					SetPreferredEulerAnglesFromValue();
+					SetPreferredEulerAnglesFromValue(InControlElement, Value, InValueType, bFixEulerFlips);
 				}
 				break;
 			}
diff --git a/Engine/Plugins/Animation/ControlRig/Source/ControlRig/Private/Sequencer/MovieSceneControlRigParameterSection.cpp b/Engine/Plugins/Animation/ControlRig/Source/ControlRig/Private/Sequencer/MovieSceneControlRigParameterSection.cpp
--- a/Engine/Plugins/Animation/ControlRig/Source/ControlRig/Private/Sequencer/MovieSceneControlRigParameterSection.cpp
+++ b/Engine/Plugins/Animation/ControlRig/Source/ControlRig/Private/Sequencer/MovieSceneControlRigParameterSection.cpp
@@ -63,7 +63,7 @@ struct FParameterFloatChannelEditorData
 			FRigControlElement* ControlElement = ControlRig->FindControl(ParameterName);
 			if (ControlElement)
 			{
-				return ControlRig->GetHierarchy()->GetControlValue(ControlElement, ERigControlValueType::Current).Get<float>();
+				return ControlRig->GetControlValue(ControlElement, ERigControlValueType::Current).Get<float>();
 			}
 		}
 		return TOptional<float>();
@@ -184,18 +184,17 @@ struct FParameterVectorChannelEditorData
 		
 				if (NumChannels == 2)
 				{
-					const FVector3f Vector = ControlRig->GetHierarchy()->GetControlValue(ControlElement, ERigControlValueType::Current).Get<FVector3f>();
+					const FVector3f Vector = ControlRig->GetControlValue(ControlElement, ERigControlValueType::Current).Get<FVector3f>();
 					return FVector4(Vector.X, Vector.Y, 0.f, 0.f);
 				}
 				else if (NumChannels == 3)
 				{
-					const FVector3f Vector = ControlRig->GetHierarchy()->GetControlValue(ControlElement, ERigControlValueType::Current).Get<FVector3f>();
+					const FVector3f Vector = ControlRig->GetControlValue(ControlElement, ERigControlValueType::Current).Get<FVector3f>();
 					return FVector4(Vector.X, Vector.Y, Vector.Z, 0.f);
 				}
 				else
 				{
-					const FRigControlValue::FTransform_Float Storage = ControlRig->GetHierarchy()
-						->GetControlValue(ControlElement, ERigControlValueType::Current).Get<FRigControlValue::FTransform_Float>();
+					const FRigControlValue::FTransform_Float Storage = ControlRig->GetControlValue(ControlElement, ERigControlValueType::Current).Get<FRigControlValue::FTransform_Float>();
 #if ENABLE_VECTORIZED_TRANSFORM
 					return FVector4(Storage.TranslationX, Storage.TranslationY, Storage.TranslationZ, Storage.TranslationW);
 #else
@@ -574,21 +573,21 @@ struct FParameterTransformChannelEditorData
 				if (ControlElement->Settings.ControlType == ERigControlType::Transform)
 				{
 					const FRigControlValue::FTransform_Float Transform = 
-						ControlRig->GetHierarchy()
+						ControlRig
 						->GetControlValue(ControlElement, ERigControlValueType::Current).Get<FRigControlValue::FTransform_Float>();
 					return FVector(Transform.GetTranslation());
 				}
 				else if  (ControlElement->Settings.ControlType == ERigControlType::TransformNoScale)
 				{
 					const FRigControlValue::FTransformNoScale_Float Transform = 
-						ControlRig->GetHierarchy()
+						ControlRig
 						->GetControlValue(ControlElement, ERigControlValueType::Current).Get<FRigControlValue::FTransformNoScale_Float>();
 					return FVector(Transform.GetTranslation());
 				}
 				else if (ControlElement->Settings.ControlType == ERigControlType::EulerTransform)
 				{
 					const FRigControlValue::FEulerTransform_Float Euler = 
-						ControlRig->GetHierarchy()
+						ControlRig
 						->GetControlValue(ControlElement, ERigControlValueType::Current).Get<FRigControlValue::FEulerTransform_Float>();
 
 					// switch translation to constraint space if needed
@@ -603,7 +602,7 @@ struct FParameterTransformChannelEditorData
 				}
 				else if(ControlElement->Settings.ControlType == ERigControlType::Position)
 				{
-					FVector3f Vector = ControlRig->GetHierarchy()->GetControlValue(ControlElement, ERigControlValueType::Current).Get<FVector3f>();
+					FVector3f Vector = ControlRig->GetControlValue(ControlElement, ERigControlValueType::Current).Get<FVector3f>();
 					return FVector(Vector.X, Vector.Y, Vector.Z);
 				}
 			}
@@ -648,14 +647,14 @@ struct FParameterTransformChannelEditorData
 				if (ControlElement->Settings.ControlType == ERigControlType::Transform)
 				{
 					const FRigControlValue::FTransform_Float Transform = 
-						ControlRig->GetHierarchy()
+						ControlRig
 						->GetControlValue(ControlElement, ERigControlValueType::Current).Get<FRigControlValue::FTransform_Float>();
 					return FVector(Transform.GetScale3D());
 				}
 				else if (ControlElement->Settings.ControlType == ERigControlType::EulerTransform)
 				{
 					const FRigControlValue::FEulerTransform_Float Transform = 
-						ControlRig->GetHierarchy()
+						ControlRig
 						->GetControlValue(ControlElement, ERigControlValueType::Current).Get<FRigControlValue::FEulerTransform_Float>();
 
 					// switch scale to constraint space if needed
@@ -670,7 +669,7 @@ struct FParameterTransformChannelEditorData
 				}
 				else if (ControlElement->Settings.ControlType == ERigControlType::Scale)
 				{
-					FVector3f Vector = ControlRig->GetHierarchy()->GetControlValue(ControlElement, ERigControlValueType::Current).Get<FVector3f>();
+					FVector3f Vector = ControlRig->GetControlValue(ControlElement, ERigControlValueType::Current).Get<FVector3f>();
 					return FVector(Vector.X, Vector.Y, Vector.Z);
 				}
 			}
@@ -1691,10 +1690,11 @@ void UMovieSceneControlRigParameterSection::ReconstructChannelProxy()
 			}
 		};
 #endif
-		
+
+		URigHierarchy* Hierarchy = ControlRig->GetHierarchy();
 		for (FRigControlElement* ControlElement : SortedControls)
 		{
-			if (!ControlRig->GetHierarchy()->IsAnimatable(ControlElement))
+			if (!Hierarchy->IsAnimatable(ControlElement))
 			{
 				continue;
 			}
@@ -1702,9 +1702,9 @@ void UMovieSceneControlRigParameterSection::ReconstructChannelProxy()
 			FName ParentControlName = NAME_None;
 			FText Group;
 
-			if(ControlRig->GetHierarchy()->ShouldBeGrouped(ControlElement))
+			if(Hierarchy->ShouldBeGrouped(ControlElement))
 			{
-				if(const FRigControlElement* ParentControlElement = Cast<FRigControlElement>(ControlRig->GetHierarchy()->GetFirstParent(ControlElement)))
+				if(const FRigControlElement* ParentControlElement = Cast<FRigControlElement>(Hierarchy->GetFirstParent(ControlElement)))
 				{
 					ParentControlName = ParentControlElement->GetName();
 					Group = FText::FromName(ParentControlElement->GetDisplayName());
@@ -2428,6 +2428,7 @@ void UMovieSceneControlRigParameterSection::RecreateWithThisControlRig(UControlR
 	ControlRig->GetControlsInOrder(SortedControls);
 
 	TMap<FName, FName> CurveControlNameRemapping;
+	URigHierarchy* Hierarchy = ControlRig->GetHierarchy();
 	if (GetLinkerCustomVersion(FUE5MainStreamObjectVersion::GUID) < FUE5MainStreamObjectVersion::FKControlNamingScheme)
 	{
 		for (FRigControlElement* ControlElement : SortedControls)
@@ -2439,7 +2440,7 @@ void UMovieSceneControlRigParameterSection::RecreateWithThisControlRig(UControlR
 				// Ensure name is valid, and curve actually exists in the hierarchy,
 				// this means we could not be renaming some controls for which the curves do not exist anymore, which ties back to comment at the top op the function
 				// with regards to non-associated curves
-				if (TargetCurveName != NAME_None && ControlRig->GetHierarchy()->Find(CurveKey))
+				if (TargetCurveName != NAME_None && Hierarchy->Find(CurveKey))
 				{
 					// Add mapping from old to new control naming scheme (previous was using uniform naming for both bones and curves)
 					CurveControlNameRemapping.Add(ControlElement->GetName(), UFKControlRig::GetControlName(TargetCurveName, ERigElementType::Bone));
@@ -2450,12 +2451,12 @@ void UMovieSceneControlRigParameterSection::RecreateWithThisControlRig(UControlR
 
 	for (FRigControlElement* ControlElement : SortedControls)
 	{
-		if (!ControlRig->GetHierarchy()->IsAnimatable(ControlElement))
+		if (!Hierarchy->IsAnimatable(ControlElement))
 		{
 			continue;
 		}
 
-		FName PreviousName = ControlRig->GetHierarchy()->GetPreviousName(ControlElement->GetKey());
+		FName PreviousName = Hierarchy->GetPreviousName(ControlElement->GetKey());
 		if (PreviousName != NAME_None && PreviousName != ControlElement->GetKey().Name)
 		{
 			RenameParameterName(PreviousName, ControlElement->GetKey().Name);
@@ -2474,7 +2475,7 @@ void UMovieSceneControlRigParameterSection::RecreateWithThisControlRig(UControlR
 			if (bSetDefault)
 			{
 				//or use IntialValue?
-				DefaultValue = ControlRig->GetHierarchy()->GetControlValue(ControlElement, ERigControlValueType::Current).Get<float>();
+				DefaultValue = ControlRig->GetControlValue(ControlElement, ERigControlValueType::Current).Get<float>();
 			}
 			AddScalarParameter(ControlElement->GetName(), DefaultValue, false);
 			break;
@@ -2485,7 +2486,7 @@ void UMovieSceneControlRigParameterSection::RecreateWithThisControlRig(UControlR
 			//only add bools,int, enums and space onto first sections, which is the same as the default one
 			if (bSetDefault)
 			{
-				DefaultValue = ControlRig->GetHierarchy()->GetControlValue(ControlElement, ERigControlValueType::Current).Get<bool>();
+				DefaultValue = ControlRig->GetControlValue(ControlElement, ERigControlValueType::Current).Get<bool>();
 				AddBoolParameter(ControlElement->GetName(), DefaultValue, false);
 
 			}
@@ -2499,7 +2500,7 @@ void UMovieSceneControlRigParameterSection::RecreateWithThisControlRig(UControlR
 				//only add bools,int, enums and space onto first sections, which is the same as the default one
 				if (bSetDefault)
 				{
-					DefaultValue = (uint8)ControlRig->GetHierarchy()->GetControlValue(ControlElement, ERigControlValueType::Current).Get<int32>();
+					DefaultValue = (uint8)ControlRig->GetControlValue(ControlElement, ERigControlValueType::Current).Get<int32>();
 					AddEnumParameter(ControlElement->GetName(), ControlElement->Settings.ControlEnum, DefaultValue, false);
 				}
 			}
@@ -2509,7 +2510,7 @@ void UMovieSceneControlRigParameterSection::RecreateWithThisControlRig(UControlR
 				//only add bools,int, enums and space onto first sections, which is the same as the default one
 				if (bSetDefault)
 				{
-					DefaultValue = ControlRig->GetHierarchy()->GetControlValue(ControlElement, ERigControlValueType::Current).Get<int32>();
+					DefaultValue = ControlRig->GetControlValue(ControlElement, ERigControlValueType::Current).Get<int32>();
 					AddIntegerParameter(ControlElement->GetName(), DefaultValue, false);
 				}
 			}
@@ -2521,7 +2522,7 @@ void UMovieSceneControlRigParameterSection::RecreateWithThisControlRig(UControlR
 			if (bSetDefault)
 			{
 				//or use IntialValue?
-				const FVector3f TempValue = ControlRig->GetHierarchy()->GetControlValue(ControlElement, ERigControlValueType::Current).Get<FVector3f>();
+				const FVector3f TempValue = ControlRig->GetControlValue(ControlElement, ERigControlValueType::Current).Get<FVector3f>();
 				DefaultValue = FVector2D(TempValue.X, TempValue.Y);
 			}
 			AddVector2DParameter(ControlElement->GetName(), DefaultValue, false);
@@ -2536,7 +2537,7 @@ void UMovieSceneControlRigParameterSection::RecreateWithThisControlRig(UControlR
 			if (bSetDefault)
 			{
 				//or use IntialValue?
-				DefaultValue = (FVector)ControlRig->GetHierarchy()->GetControlValue(ControlElement, ERigControlValueType::Current).Get<FVector3f>();
+				DefaultValue = (FVector)ControlRig->GetControlValue(ControlElement, ERigControlValueType::Current).Get<FVector3f>();
 			}
 			AddVectorParameter(ControlElement->GetName(), DefaultValue, false);
 			//mz todo specify rotator special so we can do quat interps
@@ -2552,21 +2553,21 @@ void UMovieSceneControlRigParameterSection::RecreateWithThisControlRig(UControlR
 				if (ControlElement->Settings.ControlType == ERigControlType::Transform)
 				{
 					DefaultValue = FEulerTransform(
-						ControlRig->GetHierarchy()->
+						ControlRig->
 						GetControlValue(ControlElement, ERigControlValueType::Current).Get<FRigControlValue::FTransform_Float>().ToTransform());
 				}
 				else if (ControlElement->Settings.ControlType == ERigControlType::EulerTransform)
 
 				{
 					FEulerTransform Euler = 
-						ControlRig->GetHierarchy()->
+						ControlRig->
 						GetControlValue(ControlElement, ERigControlValueType::Current).Get<FRigControlValue::FEulerTransform_Float>().ToTransform();
 					DefaultValue = Euler;
 				}
 				else
 				{
 					FTransformNoScale NoScale = 
-						ControlRig->GetHierarchy()->
+						ControlRig->
 						GetControlValue(ControlElement, ERigControlValueType::Current).Get<FRigControlValue::FTransformNoScale_Float>().ToTransform();
 					DefaultValue = FEulerTransform(NoScale.Rotation.Rotator(), NoScale.Location, FVector::OneVector);
 				}
@@ -2882,7 +2883,7 @@ void UMovieSceneControlRigParameterSection::RecordControlRigKey(FFrameNumber Fra
 			{
 				case ERigControlType::Bool:
 				{
-					bool Val = ControlRig->GetHierarchy()->GetControlValue(ControlElement, ERigControlValueType::Current).Get<bool>();
+					bool Val = ControlRig->GetControlValue(ControlElement, ERigControlValueType::Current).Get<bool>();
 					if (bSetDefault)
 					{
 						BoolChannels[ChannelIndex]->SetDefault(Val);
@@ -2894,7 +2895,7 @@ void UMovieSceneControlRigParameterSection::RecordControlRigKey(FFrameNumber Fra
 				{
 					if (ControlElement->Settings.ControlEnum)
 					{
-						uint8 Val = (uint8)ControlRig->GetHierarchy()->GetControlValue(ControlElement, ERigControlValueType::Current).Get<uint8>();
+						uint8 Val = (uint8)ControlRig->GetControlValue(ControlElement, ERigControlValueType::Current).Get<uint8>();
 						if (bSetDefault)
 						{
 							EnumChannels[ChannelIndex]->SetDefault(Val);
@@ -2903,7 +2904,7 @@ void UMovieSceneControlRigParameterSection::RecordControlRigKey(FFrameNumber Fra
 					}
 					else
 					{
-						int32 Val = ControlRig->GetHierarchy()->GetControlValue(ControlElement, ERigControlValueType::Current).Get<int32>();
+						int32 Val = ControlRig->GetControlValue(ControlElement, ERigControlValueType::Current).Get<int32>();
 						if (bSetDefault)
 						{
 							IntChannels[ChannelIndex]->SetDefault(Val);
@@ -2914,7 +2915,7 @@ void UMovieSceneControlRigParameterSection::RecordControlRigKey(FFrameNumber Fra
 				}
 				case ERigControlType::Float:
 				{
-					float Val = ControlRig->GetHierarchy()->GetControlValue(ControlElement, ERigControlValueType::Current).Get<float>();
+					float Val = ControlRig->GetControlValue(ControlElement, ERigControlValueType::Current).Get<float>();
 					if (bSetDefault)
 					{
 						FloatChannels[ChannelIndex]->SetDefault(Val);
@@ -2940,7 +2941,7 @@ void UMovieSceneControlRigParameterSection::RecordControlRigKey(FFrameNumber Fra
 				}
 				case ERigControlType::Vector2D:
 				{
-					FVector3f Val = ControlRig->GetHierarchy()->GetControlValue(ControlElement, ERigControlValueType::Current).Get<FVector3f>();
+					FVector3f Val = ControlRig->GetControlValue(ControlElement, ERigControlValueType::Current).Get<FVector3f>();
 					if (bSetDefault)
 					{
 						FloatChannels[ChannelIndex]->SetDefault(Val.X);
@@ -2972,7 +2973,7 @@ void UMovieSceneControlRigParameterSection::RecordControlRigKey(FFrameNumber Fra
 				case ERigControlType::Scale:
 				case ERigControlType::Rotator:
 				{
-					FVector3f Val = ControlRig->GetHierarchy()->GetControlValue(ControlElement, ERigControlValueType::Current).Get<FVector3f>();
+					FVector3f Val = ControlRig->GetControlValue(ControlElement, ERigControlValueType::Current).Get<FVector3f>();
 					if (ControlElement->Settings.ControlType == ERigControlType::Rotator &&
 						FloatChannels[ChannelIndex]->GetNumKeys() > 0)
 					{
@@ -3003,21 +3004,19 @@ void UMovieSceneControlRigParameterSection::RecordControlRigKey(FFrameNumber Fra
 					if (ControlElement->Settings.ControlType == ERigControlType::TransformNoScale)
 					{
 						FTransformNoScale NoScale = 
-							ControlRig->GetHierarchy()
-							->GetControlValue(ControlElement, ERigControlValueType::Current).Get<FRigControlValue::FTransformNoScale_Float>().ToTransform();
+							ControlRig->GetControlValue(ControlElement, ERigControlValueType::Current).Get<FRigControlValue::FTransformNoScale_Float>().ToTransform();
 						Val = NoScale;
 					}
 					else if (ControlElement->Settings.ControlType == ERigControlType::EulerTransform)
 					{
 						FEulerTransform Euler = 
-							ControlRig->GetHierarchy()
-							->GetControlValue(ControlElement, ERigControlValueType::Current).Get<FRigControlValue::FEulerTransform_Float>().ToTransform();
+							ControlRig->GetControlValue(ControlElement, ERigControlValueType::Current).Get<FRigControlValue::FEulerTransform_Float>().ToTransform();
 						Val = Euler.ToFTransform();
 					}
 					else
 					{
 						Val = 
-							ControlRig->GetHierarchy()
+							ControlRig
 							->GetControlValue(ControlElement, ERigControlValueType::Current).Get<FRigControlValue::FTransform_Float>().ToTransform();
 					}
 					FVector CurrentVector = Val.GetTranslation();
diff --git a/Engine/Plugins/Animation/ControlRig/Source/ControlRig/Private/Sequencer/MovieSceneControlRigParameterTemplate.cpp b/Engine/Plugins/Animation/ControlRig/Source/ControlRig/Private/Sequencer/MovieSceneControlRigParameterTemplate.cpp
--- a/Engine/Plugins/Animation/ControlRig/Source/ControlRig/Private/Sequencer/MovieSceneControlRigParameterTemplate.cpp
+++ b/Engine/Plugins/Animation/ControlRig/Source/ControlRig/Private/Sequencer/MovieSceneControlRigParameterTemplate.cpp
@@ -888,6 +888,7 @@ struct FControlRigParameterPreAnimatedTokenProducer : IMovieScenePreAnimatedToke
 
 		if (UControlRig* ControlRig = Cast<UControlRig>(&Object))
 		{
+			URigHierarchy* RigHierarchy = ControlRig->GetHierarchy();
 			TArray<FRigControlElement*> Controls = ControlRig->AvailableControls();
 			FRigControlValue Value;
 			
@@ -897,25 +898,25 @@ struct FControlRigParameterPreAnimatedTokenProducer : IMovieScenePreAnimatedToke
 				{
 					case ERigControlType::Bool:
 					{
-						const bool Val = ControlRig->GetHierarchy()->GetControlValue(ControlElement, ERigControlValueType::Current).Get<bool>();
+						const bool Val = ControlRig->GetControlValue(ControlElement, ERigControlValueType::Current).Get<bool>();
 						Token.BoolValues.Add(TNameAndValue<bool>{ ControlElement->GetName(), Val });
 						break;
 					}
 					case ERigControlType::Integer:
 					{
-						const int32 Val = ControlRig->GetHierarchy()->GetControlValue(ControlElement, ERigControlValueType::Current).Get<int32>();
+						const int32 Val = ControlRig->GetControlValue(ControlElement, ERigControlValueType::Current).Get<int32>();
 						Token.IntegerValues.Add(TNameAndValue<int32>{ ControlElement->GetName(), Val });
 						break;
 					}
 					case ERigControlType::Float:
 					{
-						const float Val = ControlRig->GetHierarchy()->GetControlValue(ControlElement, ERigControlValueType::Current).Get<float>();
+						const float Val = ControlRig->GetControlValue(ControlElement, ERigControlValueType::Current).Get<float>();
 						Token.ScalarValues.Add(TNameAndValue<float>{ ControlElement->GetName(), Val });
 						break;
 					}
 					case ERigControlType::Vector2D:
 					{
-						const FVector3f Val = ControlRig->GetHierarchy()->GetControlValue(ControlElement, ERigControlValueType::Current).Get<FVector3f>();
+						const FVector3f Val = ControlRig->GetControlValue(ControlElement, ERigControlValueType::Current).Get<FVector3f>();
 						Token.Vector2DValues.Add(TNameAndValue<FVector2D>{ ControlElement->GetName(), FVector2D(Val.X, Val.Y) });
 						break;
 					}
@@ -923,7 +924,6 @@ struct FControlRigParameterPreAnimatedTokenProducer : IMovieScenePreAnimatedToke
 					case ERigControlType::Scale:
 					case ERigControlType::Rotator:
 					{
-						URigHierarchy* RigHierarchy = ControlRig->GetHierarchy();
 						FMovieSceneControlRigSpaceBaseKey SpaceValue;
 						//@helge how can we get the correct current space here? this is for restoring it.
 						//for now just using parent space
@@ -931,10 +931,10 @@ struct FControlRigParameterPreAnimatedTokenProducer : IMovieScenePreAnimatedToke
 						SpaceValue.ControlRigElement = ControlElement->GetKey();
 						SpaceValue.SpaceType = EMovieSceneControlRigSpaceType::Parent;
 						Token.SpaceValues.Add(FControlSpaceAndValue(ControlElement->GetName(), SpaceValue));
-						FVector3f Val = ControlRig->GetHierarchy()->GetControlValue(ControlElement, ERigControlValueType::Current).Get<FVector3f>();
+						FVector3f Val = ControlRig->GetControlValue(ControlElement, ERigControlValueType::Current).Get<FVector3f>();
 						if (ControlElement->Settings.ControlType == ERigControlType::Rotator)
 						{
-							FVector Vector = ControlRig->GetHierarchy()->GetControlSpecifiedEulerAngle(ControlElement);
+							FVector Vector = RigHierarchy->GetControlSpecifiedEulerAngle(ControlElement);
 							Val = FVector3f(Vector.X, Vector.Y, Vector.Z);
 						}
 						Token.VectorValues.Add(TNameAndValue<FVector>{ ControlElement->GetName(), (FVector)Val });
@@ -950,9 +950,9 @@ struct FControlRigParameterPreAnimatedTokenProducer : IMovieScenePreAnimatedToke
 						SpaceValue.ControlRigElement = ControlElement->GetKey();
 						SpaceValue.SpaceType = EMovieSceneControlRigSpaceType::Parent;
 						Token.SpaceValues.Add(FControlSpaceAndValue(ControlElement->GetName(), SpaceValue));
-						const FTransform Val = ControlRig->GetHierarchy()->GetControlValue(ControlElement, ERigControlValueType::Current).Get<FRigControlValue::FTransform_Float>().ToTransform();
+						const FTransform Val = ControlRig->GetControlValue(ControlElement, ERigControlValueType::Current).Get<FRigControlValue::FTransform_Float>().ToTransform();
 						FEulerTransform EulerTransform(Val);
-						FVector Vector = ControlRig->GetHierarchy()->GetControlSpecifiedEulerAngle(ControlElement);
+						FVector Vector = RigHierarchy->GetControlSpecifiedEulerAngle(ControlElement);
 						EulerTransform.Rotation = FRotator(Vector.Y, Vector.Z, Vector.X);
 						Token.TransformValues.Add(TNameAndValue<FEulerTransform>{ ControlElement->GetName(), EulerTransform });
 						break;
@@ -960,10 +960,10 @@ struct FControlRigParameterPreAnimatedTokenProducer : IMovieScenePreAnimatedToke
 					case ERigControlType::TransformNoScale:
 					{
 						const FTransformNoScale NoScale = 
-							ControlRig->GetHierarchy()
+							ControlRig
 							->GetControlValue(ControlElement, ERigControlValueType::Current).Get<FRigControlValue::FTransformNoScale_Float>().ToTransform();
 						FEulerTransform EulerTransform(NoScale.ToFTransform());
-						FVector Vector = ControlRig->GetHierarchy()->GetControlSpecifiedEulerAngle(ControlElement);
+						FVector Vector = RigHierarchy->GetControlSpecifiedEulerAngle(ControlElement);
 						EulerTransform.Rotation = FRotator(Vector.Y, Vector.Z, Vector.X);
 						Token.TransformValues.Add(TNameAndValue<FEulerTransform>{ ControlElement->GetName(), EulerTransform });
 						break;
@@ -971,9 +971,9 @@ struct FControlRigParameterPreAnimatedTokenProducer : IMovieScenePreAnimatedToke
 					case ERigControlType::EulerTransform:
 					{
 						FEulerTransform EulerTransform = 
-							ControlRig->GetHierarchy()
+							ControlRig
 							->GetControlValue(ControlElement, ERigControlValueType::Current).Get<FRigControlValue::FEulerTransform_Float>().ToTransform();
-						FVector Vector = ControlRig->GetHierarchy()->GetControlSpecifiedEulerAngle(ControlElement);
+						FVector Vector = RigHierarchy->GetControlSpecifiedEulerAngle(ControlElement);
 						EulerTransform.Rotation = FRotator(Vector.Y, Vector.Z, Vector.X);
 						Token.TransformValues.Add(TNameAndValue<FEulerTransform>{ ControlElement->GetName(), EulerTransform });
 						break;
@@ -1378,7 +1378,7 @@ struct TControlRigParameterActuatorFloat : TMovieSceneBlendingActuator<FControlR
 				ControlElement->Settings.AnimationType != ERigControlAnimationType::VisualCue
 				&& ControlElement->Settings.ControlType == ERigControlType::Float)
 			{
-				const float Val = ControlRig->GetHierarchy()->GetControlValue(ControlElement, ERigControlValueType::Current).Get<float>();
+				const float Val = ControlRig->GetControlValue(ControlElement, ERigControlValueType::Current).Get<float>();
 				return FControlRigTrackTokenFloat(Val);
 			}
 
@@ -1451,7 +1451,7 @@ struct TControlRigParameterActuatorVector2D : TMovieSceneBlendingActuator<FContr
 				ControlElement->Settings.AnimationType != ERigControlAnimationType::VisualCue
 				&& (ControlElement->Settings.ControlType == ERigControlType::Vector2D))
 			{
-				FVector3f Val = ControlRig->GetHierarchy()->GetControlValue(ControlElement, ERigControlValueType::Current).Get<FVector3f>();
+				FVector3f Val = ControlRig->GetControlValue(ControlElement, ERigControlValueType::Current).Get<FVector3f>();
 				return FControlRigTrackTokenVector2D(FVector2D(Val.X, Val.Y));
 			}
 		}
@@ -1525,7 +1525,7 @@ struct TControlRigParameterActuatorVector : TMovieSceneBlendingActuator<FControl
 				ControlElement->Settings.AnimationType != ERigControlAnimationType::VisualCue
 				&& (ControlElement->Settings.ControlType == ERigControlType::Position || ControlElement->Settings.ControlType == ERigControlType::Scale || ControlElement->Settings.ControlType == ERigControlType::Rotator))
 			{
-				FVector3f Val = ControlRig->GetHierarchy()->GetControlValue(ControlElement, ERigControlValueType::Current).Get<FVector3f>();
+				FVector3f Val = ControlRig->GetControlValue(ControlElement, ERigControlValueType::Current).Get<FVector3f>();
 				if (ControlElement->Settings.ControlType == ERigControlType::Rotator)
 				{
 					FVector Vector = ControlRig->GetHierarchy()->GetControlSpecifiedEulerAngle(ControlElement);
@@ -1596,6 +1596,7 @@ struct TControlRigParameterActuatorTransform : TMovieSceneBlendingActuator<FCont
 		const UMovieSceneControlRigParameterSection* Section = SectionData.Get();
 
 		UControlRig* ControlRig = Section ? GetControlRig(Section, InObject) : nullptr;
+		URigHierarchy* Hierarchy = ControlRig->GetHierarchy();
 
 		if (ControlRig && Section && (Section->ControlsToSet.Num() == 0 || Section->ControlsToSet.Contains(ParameterName)))
 		{
@@ -1605,25 +1606,25 @@ struct TControlRigParameterActuatorTransform : TMovieSceneBlendingActuator<FCont
 			{
 				if (ControlElement && ControlElement->Settings.ControlType == ERigControlType::Transform)
 				{
-					const FTransform Val = ControlRig->GetHierarchy()->GetControlValue(ControlElement, ERigControlValueType::Current).Get<FRigControlValue::FTransform_Float>().ToTransform();
+					const FTransform Val = ControlRig->GetControlValue(ControlElement, ERigControlValueType::Current).Get<FRigControlValue::FTransform_Float>().ToTransform();
 					FEulerTransform EulerTransform(Val);
-					FVector Vector = ControlRig->GetHierarchy()->GetControlSpecifiedEulerAngle(ControlElement);
+					FVector Vector = Hierarchy->GetControlSpecifiedEulerAngle(ControlElement);
 					EulerTransform.Rotation = FRotator(Vector.Y, Vector.Z, Vector.X);
 					return FControlRigTrackTokenTransform(EulerTransform);
 				}
 				else if (ControlElement && ControlElement->Settings.ControlType == ERigControlType::TransformNoScale)
 				{
-					FTransformNoScale ValNoScale = ControlRig->GetHierarchy()->GetControlValue(ControlElement, ERigControlValueType::Current).Get<FRigControlValue::FTransformNoScale_Float>().ToTransform();
+					FTransformNoScale ValNoScale = ControlRig->GetControlValue(ControlElement, ERigControlValueType::Current).Get<FRigControlValue::FTransformNoScale_Float>().ToTransform();
 					FTransform Val = ValNoScale;
 					FEulerTransform EulerTransform(Val);
-					FVector Vector = ControlRig->GetHierarchy()->GetControlSpecifiedEulerAngle(ControlElement);
+					FVector Vector = Hierarchy->GetControlSpecifiedEulerAngle(ControlElement);
 					EulerTransform.Rotation = FRotator(Vector.Y, Vector.Z, Vector.X);
 					return FControlRigTrackTokenTransform(EulerTransform);
 				}
 				else if (ControlElement && ControlElement->Settings.ControlType == ERigControlType::EulerTransform)
 				{
-					FEulerTransform EulerTransform = ControlRig->GetHierarchy()->GetControlValue(ControlElement, ERigControlValueType::Current).Get<FRigControlValue::FEulerTransform_Float>().ToTransform();
-					FVector Vector = ControlRig->GetHierarchy()->GetControlSpecifiedEulerAngle(ControlElement);
+					FEulerTransform EulerTransform = ControlRig->GetControlValue(ControlElement, ERigControlValueType::Current).Get<FRigControlValue::FEulerTransform_Float>().ToTransform();
+					FVector Vector = Hierarchy->GetControlSpecifiedEulerAngle(ControlElement);
 					EulerTransform.Rotation = FRotator(Vector.Y, Vector.Z, Vector.X);
 					return FControlRigTrackTokenTransform(EulerTransform);
 				}
@@ -1654,6 +1655,7 @@ struct TControlRigParameterActuatorTransform : TMovieSceneBlendingActuator<FCont
 
 			if (ControlRig && (Section->ControlsToSet.Num() == 0 || Section->ControlsToSet.Contains(ParameterName)))
 			{
+				URigHierarchy* Hierarchy = ControlRig->GetHierarchy();
 				FRigControlElement* ControlElement = ControlRig->FindControl(ParameterName);
 				if (ControlElement && ControlElement->Settings.AnimationType != ERigControlAnimationType::ProxyControl &&
 					ControlElement->Settings.AnimationType != ERigControlAnimationType::VisualCue)
@@ -1661,21 +1663,21 @@ struct TControlRigParameterActuatorTransform : TMovieSceneBlendingActuator<FCont
 					if (ControlElement && ControlElement->Settings.ControlType == ERigControlType::Transform)
 					{
 						FVector EulerAngle(InFinalValue.Value.Rotation.Roll, InFinalValue.Value.Rotation.Pitch, InFinalValue.Value.Rotation.Yaw);
-						ControlRig->GetHierarchy()->SetControlSpecifiedEulerAngle(ControlElement, EulerAngle);
+						Hierarchy->SetControlSpecifiedEulerAngle(ControlElement, EulerAngle);
 						ControlRig->SetControlValue<FRigControlValue::FTransform_Float>(ParameterName, InFinalValue.Value.ToFTransform(), true, EControlRigSetKey::Never, bSetupUndo);
 					}
 					else if (ControlElement && ControlElement->Settings.ControlType == ERigControlType::TransformNoScale)
 					{
 						const FTransformNoScale NoScale = InFinalValue.Value.ToFTransform();
 						FVector EulerAngle(InFinalValue.Value.Rotation.Roll, InFinalValue.Value.Rotation.Pitch, InFinalValue.Value.Rotation.Yaw);
-						ControlRig->GetHierarchy()->SetControlSpecifiedEulerAngle(ControlElement, EulerAngle);
+						Hierarchy->SetControlSpecifiedEulerAngle(ControlElement, EulerAngle);
 						ControlRig->SetControlValue<FRigControlValue::FTransformNoScale_Float>(ParameterName, NoScale, true, EControlRigSetKey::Never, bSetupUndo);
 					}
 					else if (ControlElement && ControlElement->Settings.ControlType == ERigControlType::EulerTransform)
 					{
 						FVector EulerAngle(InFinalValue.Value.Rotation.Roll, InFinalValue.Value.Rotation.Pitch, InFinalValue.Value.Rotation.Yaw);
-						FQuat Quat = ControlRig->GetHierarchy()->GetControlQuaternion(ControlElement, EulerAngle);
-						ControlRig->GetHierarchy()->SetControlSpecifiedEulerAngle(ControlElement, EulerAngle);
+						FQuat Quat = Hierarchy->GetControlQuaternion(ControlElement, EulerAngle);
+						Hierarchy->SetControlSpecifiedEulerAngle(ControlElement, EulerAngle);
 						FRotator UERotator(Quat);
 						FEulerTransform Transform = InFinalValue.Value;
 						Transform.Rotation = UERotator;
diff --git a/Engine/Plugins/Animation/ControlRig/Source/ControlRig/Public/ControlRig.h b/Engine/Plugins/Animation/ControlRig/Source/ControlRig/Public/ControlRig.h
--- a/Engine/Plugins/Animation/ControlRig/Source/ControlRig/Public/ControlRig.h
+++ b/Engine/Plugins/Animation/ControlRig/Source/ControlRig/Public/ControlRig.h
@@ -103,6 +103,16 @@ public:
 	/** Initialize the VM */
 	virtual bool InitializeVM(const FName& InEventName) override;
 
+	/** Evaluates the ControlRig */
+	virtual void Evaluate_AnyThread() override;
+
+	/** Removes any stored additive control values */
+	void ResetControlValues();
+
+	/** Resets the stored pose coming from the anim sequence.
+	 * This usually indicates a new pose should be stored. */
+	void ClearPoseBeforeBackwardsSolve();
+
 	/** Setup bindings to a runtime object (or clear by passing in nullptr). */
 	void SetObjectBinding(TSharedPtr<IControlRigObjectBinding> InObjectBinding)
 	{
@@ -184,27 +194,30 @@ public:
 	FRigControlValue GetControlValue(const FName& InControlName)
 	{
 		const FRigElementKey Key(InControlName, ERigElementType::Control);
+		if (FRigBaseElement* Element = DynamicHierarchy->Find(Key))
+		{
+			if (FRigControlElement* ControlElement = Cast<FRigControlElement>(Element))
+			{
+				return GetControlValue(ControlElement, ERigControlValueType::Current);
+			}
+		}
 		return DynamicHierarchy->GetControlValue(Key);
 	}
 
+	FRigControlValue GetControlValue(FRigControlElement* InControl, const ERigControlValueType& InValueType);
+
 	// Sets the relative value of a Control
 	virtual void SetControlValueImpl(const FName& InControlName, const FRigControlValue& InValue, bool bNotify = true,
-		const FRigControlModifiedContext& Context = FRigControlModifiedContext(), bool bSetupUndo = true, bool bPrintPythonCommnds = false, bool bFixEulerFlips = false)
-	{
-		const FRigElementKey Key(InControlName, ERigElementType::Control);
+		const FRigControlModifiedContext& Context = FRigControlModifiedContext(), bool bSetupUndo = true, bool bPrintPythonCommnds = false, bool bFixEulerFlips = false);
 
-		FRigControlElement* ControlElement = DynamicHierarchy->Find<FRigControlElement>(Key);
-		if(ControlElement == nullptr)
-		{
-			return;
-		}
-
-		DynamicHierarchy->SetControlValue(ControlElement, InValue, ERigControlValueType::Current, bSetupUndo, false, bPrintPythonCommnds, bFixEulerFlips);
-
-		if (bNotify && OnControlModified.IsBound())
+	FTransform GetInitialLocalTransform(const FRigElementKey &InKey)
+	{
+		if (bIsAdditive)
 		{
-			OnControlModified.Broadcast(this, ControlElement, Context);
+			// The initial value of all additive controls is always Identity
+			return FTransform::Identity;
 		}
+		return GetHierarchy()->GetInitialLocalTransform(InKey);
 	}
 
 	bool SetControlGlobalTransform(const FName& InControlName, const FTransform& InGlobalTransform, bool bNotify = true, const FRigControlModifiedContext& Context = FRigControlModifiedContext(), bool bSetupUndo = true, bool bPrintPythonCommands = false, bool bFixEulerFlips = false);
@@ -569,6 +582,25 @@ private:
 
 #endif
 
+	/** An additive contrrl rig runs a backwards solve before applying additive control values
+	 * and running the forward solve
+	 */
+	UPROPERTY(transient)
+	bool bIsAdditive = false;
+
+	struct FRigSetControlValueInfo
+	{
+		FRigControlValue Value;
+		bool bNotify;
+		FRigControlModifiedContext Context;
+		bool bSetupUndo;
+		bool bPrintPythonCommnds;
+		bool bFixEulerFlips;
+	};
+	FRigPose PoseBeforeBackwardsSolve;
+	FRigPose ControlsAfterBackwardsSolve;
+	TMap<FRigElementKey, FRigSetControlValueInfo> ControlValues; // Additive values in local space (to add after backwards solve)
+
 	float DebugBoneRadiusMultiplier;
 	
 #if WITH_EDITOR	
@@ -579,7 +611,13 @@ public:
 	void SetControlsVisible(const bool bIsVisible) { bControlsVisible = bIsVisible; }
 	bool GetControlsVisible()const { return bControlsVisible;}
 
-#endif	
+#endif
+	
+	bool IsAdditive() const { return bIsAdditive; }
+	void SetIsAdditive(const bool bInIsAdditive)
+	{
+		bIsAdditive = bInIsAdditive;
+	}
 
 private:
 
diff --git a/Engine/Plugins/Animation/ControlRig/Source/ControlRig/Public/ControlRigComponent.h b/Engine/Plugins/Animation/ControlRig/Source/ControlRig/Public/ControlRigComponent.h
--- a/Engine/Plugins/Animation/ControlRig/Source/ControlRig/Public/ControlRigComponent.h
+++ b/Engine/Plugins/Animation/ControlRig/Source/ControlRig/Public/ControlRigComponent.h
@@ -291,11 +291,11 @@ public:
 
 	/** Adds a series of mapped bones to the rig, should not be used before OnPreInitialize Event */
 	UFUNCTION(BlueprintCallable, Category = "ControlRig", meta = (DisplayName = "Add Mapped Skeletal Mesh Bone Array", UnsafeDuringActorConstruction = "true"))
-	void AddMappedSkeletalMesh(USkeletalMeshComponent* SkeletalMeshComponent, TArray<FControlRigComponentMappedBone> Bones, TArray<FControlRigComponentMappedCurve> Curves);
+	void AddMappedSkeletalMesh(USkeletalMeshComponent* SkeletalMeshComponent, TArray<FControlRigComponentMappedBone> Bones, TArray<FControlRigComponentMappedCurve> Curves, const EControlRigComponentMapDirection InDirection = EControlRigComponentMapDirection::Output);
 
 	/** Adds all matching bones to the rig, should not be used before OnPreInitialize Event */
 	UFUNCTION(BlueprintCallable, Category = "ControlRig", meta = (DisplayName = "Add Mapped Skeletal Mesh", UnsafeDuringActorConstruction = "true"))
-	void AddMappedCompleteSkeletalMesh(USkeletalMeshComponent* SkeletalMeshComponent);
+	void AddMappedCompleteSkeletalMesh(USkeletalMeshComponent* SkeletalMeshComponent, const EControlRigComponentMapDirection InDirection = EControlRigComponentMapDirection::Output);
 
 	/** Setup the initial transforms / ref pose of the bones based on a skeletal mesh */
 	UFUNCTION(BlueprintCallable, Category = "ControlRig", meta = (DisplayName = "Set Bone Initial Transforms From Skeletal Mesh"))
diff --git a/Engine/Plugins/Animation/ControlRig/Source/ControlRig/Public/Rigs/RigHierarchy.h b/Engine/Plugins/Animation/ControlRig/Source/ControlRig/Public/Rigs/RigHierarchy.h
--- a/Engine/Plugins/Animation/ControlRig/Source/ControlRig/Public/Rigs/RigHierarchy.h
+++ b/Engine/Plugins/Animation/ControlRig/Source/ControlRig/Public/Rigs/RigHierarchy.h
@@ -3626,6 +3626,9 @@ public:
 	 */
 	FRigControlValue GetControlValue(FRigControlElement* InControlElement, ERigControlValueType InValueType) const;
 
+	void SetPreferredEulerAnglesFromValue(FRigControlElement* InControlElement, const FRigControlValue& InValue,
+	                                      const ERigControlValueType& InValueType, bool bFixEulerFlips);
+
 	template<typename T>
 	T GetControlValue(FRigControlElement* InControlElement, ERigControlValueType InValueType) const
 	{
diff --git a/Engine/Plugins/Animation/ControlRig/Source/ControlRig/Public/Rigs/RigHierarchyPose.h b/Engine/Plugins/Animation/ControlRig/Source/ControlRig/Public/Rigs/RigHierarchyPose.h
--- a/Engine/Plugins/Animation/ControlRig/Source/ControlRig/Public/Rigs/RigHierarchyPose.h
+++ b/Engine/Plugins/Animation/ControlRig/Source/ControlRig/Public/Rigs/RigHierarchyPose.h
@@ -31,6 +31,9 @@ public:
 	UPROPERTY()
 	FTransform LocalTransform;
 
+	UPROPERTY()
+	FVector PreferredEulerAngle;
+
 	UPROPERTY()
 	float CurveValue;
 };
diff --git a/Engine/Plugins/Animation/ControlRig/Source/ControlRigEditor/Private/EditMode/ControlRigControlsProxy.cpp b/Engine/Plugins/Animation/ControlRig/Source/ControlRigEditor/Private/EditMode/ControlRigControlsProxy.cpp
--- a/Engine/Plugins/Animation/ControlRig/Source/ControlRigEditor/Private/EditMode/ControlRigControlsProxy.cpp
+++ b/Engine/Plugins/Animation/ControlRig/Source/ControlRigEditor/Private/EditMode/ControlRigControlsProxy.cpp
@@ -224,7 +224,7 @@ void UControlRigTransformControlProxy::ValueChanged()
 		Modify();
 		const FName PropertyName("Transform");
 		FTrackInstancePropertyBindings Binding(PropertyName, PropertyName.ToString());
-		const FTransform NewTransform = ControlRig.Get()->GetHierarchy()->GetControlValue(ControlElement, ERigControlValueType::Current).Get<FRigControlValue::FTransform_Float>().ToTransform();
+		const FTransform NewTransform = ControlRig.Get()->GetControlValue(ControlElement, ERigControlValueType::Current).Get<FRigControlValue::FTransform_Float>().ToTransform();
 		FEulerTransform EulerTransform(NewTransform);
 		EulerTransform.Rotation = ControlRig->GetHierarchy()->GetControlPreferredRotator(ControlElement);
 		Binding.CallFunction<FEulerTransform>(*this, EulerTransform);
@@ -304,7 +304,7 @@ void UControlRigTransformNoScaleControlProxy::ValueChanged()
 		Modify();
 		const FName PropertyName("Transform");
 		FTrackInstancePropertyBindings Binding(PropertyName, PropertyName.ToString());
-		const FTransformNoScale NewTransform = ControlRig.Get()->GetHierarchy()->GetControlValue(ControlElement, ERigControlValueType::Current).Get<FRigControlValue::FTransformNoScale_Float>().ToTransform();
+		const FTransformNoScale NewTransform = ControlRig.Get()->GetControlValue(ControlElement, ERigControlValueType::Current).Get<FRigControlValue::FTransformNoScale_Float>().ToTransform();
 		Binding.CallFunction<FTransformNoScale>(*this, NewTransform);
 	}
 }
@@ -358,9 +358,10 @@ void UControlRigEulerTransformControlProxy::PostEditChangeChainProperty(struct F
 		FRigControlElement* ControlElement = GetControlElement();
 		if (ControlElement && ControlRig.IsValid())
 		{
+			URigHierarchy* Hierarchy = ControlRig->GetHierarchy();
 			FVector EulerAngle(Transform.Rotation.Roll, Transform.Rotation.Pitch, Transform.Rotation.Yaw);
-			FQuat Quat = ControlRig->GetHierarchy()->GetControlQuaternion(ControlElement, EulerAngle);
-			ControlRig->GetHierarchy()->SetControlSpecifiedEulerAngle(ControlElement, EulerAngle);
+			FQuat Quat = Hierarchy->GetControlQuaternion(ControlElement, EulerAngle);
+			Hierarchy->SetControlSpecifiedEulerAngle(ControlElement, EulerAngle);
 			FRotator UERotator(Quat);
 			FEulerTransform UETransform = Transform;
 			UETransform.Rotation = UERotator;
@@ -380,7 +381,7 @@ void UControlRigEulerTransformControlProxy::ValueChanged()
 		Modify();
 		const FName PropertyName("Transform");
 		FTrackInstancePropertyBindings Binding(PropertyName, PropertyName.ToString());
-		FEulerTransform NewTransform = ControlRig.Get()->GetHierarchy()->GetControlValue(ControlElement, ERigControlValueType::Current).Get<FRigControlValue::FEulerTransform_Float>().ToTransform();
+		FEulerTransform NewTransform = ControlRig.Get()->GetControlValue(ControlElement, ERigControlValueType::Current).Get<FRigControlValue::FEulerTransform_Float>().ToTransform();
 		NewTransform.Rotation = ControlRig->GetHierarchy()->GetControlPreferredRotator(ControlElement);
 		Binding.CallFunction<FEulerTransform>(*this, NewTransform);
 
@@ -450,7 +451,7 @@ void UControlRigFloatControlProxy::ValueChanged()
 		Modify();
 		const FName PropertyName("Float");
 		FTrackInstancePropertyBindings Binding(PropertyName, PropertyName.ToString());
-		const float Val = ControlRig.Get()->GetHierarchy()->GetControlValue(ControlElement, ERigControlValueType::Current).Get<float>();
+		const float Val = ControlRig.Get()->GetControlValue(ControlElement, ERigControlValueType::Current).Get<float>();
 		Binding.CallFunction<float>(*this, Val);
 	}
 }
@@ -500,7 +501,7 @@ void UControlRigIntegerControlProxy::ValueChanged()
 		Modify();
 		const FName PropertyName("Integer");
 		FTrackInstancePropertyBindings Binding(PropertyName, PropertyName.ToString());
-		const int32 Val = ControlRig.Get()->GetHierarchy()->GetControlValue(ControlElement, ERigControlValueType::Current).Get<int32>();
+		const int32 Val = ControlRig.Get()->GetControlValue(ControlElement, ERigControlValueType::Current).Get<int32>();
 		Binding.CallFunction<int32>(*this, Val);
 	}
 }
@@ -553,7 +554,7 @@ void UControlRigEnumControlProxy::ValueChanged()
 
 		FControlRigEnumControlProxyValue Val;
 		Val.EnumType = ControlElement->Settings.ControlEnum;
-		Val.EnumIndex = ControlRig.Get()->GetHierarchy()->GetControlValue(ControlElement, ERigControlValueType::Current).Get<int32>();
+		Val.EnumIndex = ControlRig.Get()->GetControlValue(ControlElement, ERigControlValueType::Current).Get<int32>();
 
 		Binding.CallFunction<FControlRigEnumControlProxyValue>(*this, Val);
 	}
@@ -617,7 +618,7 @@ void UControlRigVectorControlProxy::ValueChanged()
 		Modify();
 		const FName PropertyName("Vector");
 		FTrackInstancePropertyBindings Binding(PropertyName, PropertyName.ToString());
-		FVector3f Val = ControlRig.Get()->GetHierarchy()->GetControlValue(ControlElement, ERigControlValueType::Current).Get<FVector3f>();
+		FVector3f Val = ControlRig.Get()->GetControlValue(ControlElement, ERigControlValueType::Current).Get<FVector3f>();
 		if (ControlElement->Settings.ControlType == ERigControlType::Rotator)
 		{ 
 			FVector DVector = ControlRig->GetHierarchy()->GetControlSpecifiedEulerAngle(ControlElement);
@@ -738,7 +739,7 @@ void UControlRigVector2DControlProxy::ValueChanged()
 		Modify();
 		const FName PropertyName("Vector2D");
 		FTrackInstancePropertyBindings Binding(PropertyName, PropertyName.ToString());
-		const FVector3f TempValue = ControlRig.Get()->GetHierarchy()->GetControlValue(ControlElement, ERigControlValueType::Current).Get<FVector3f>();
+		const FVector3f TempValue = ControlRig.Get()->GetControlValue(ControlElement, ERigControlValueType::Current).Get<FVector3f>();
 		const FVector2D Val(TempValue.X, TempValue.Y);
 		Binding.CallFunction<FVector2D>(*this, Val);
 	}
@@ -800,7 +801,7 @@ void UControlRigBoolControlProxy::ValueChanged()
 		Modify();
 		const FName PropertyName("Bool");
 		FTrackInstancePropertyBindings Binding(PropertyName, PropertyName.ToString());
-		const bool Val = ControlRig.Get()->GetHierarchy()->GetControlValue(ControlElement, ERigControlValueType::Current).Get<bool>();
+		const bool Val = ControlRig.Get()->GetControlValue(ControlElement, ERigControlValueType::Current).Get<bool>();
 		Binding.CallFunction<bool>(*this, Val);
 	}
 }
diff --git a/Engine/Plugins/Animation/ControlRig/Source/ControlRigEditor/Private/EditMode/ControlRigEditMode.cpp b/Engine/Plugins/Animation/ControlRig/Source/ControlRigEditor/Private/EditMode/ControlRigEditMode.cpp
--- a/Engine/Plugins/Animation/ControlRig/Source/ControlRigEditor/Private/EditMode/ControlRigEditMode.cpp
+++ b/Engine/Plugins/Animation/ControlRig/Source/ControlRigEditor/Private/EditMode/ControlRigEditMode.cpp
@@ -2914,7 +2914,7 @@ void FControlRigEditMode::ResetTransforms(bool bSelectionOnly)
 				}
 			}
 			
-			const FTransform InitialLocalTransform = ControlRig->GetHierarchy()->GetInitialLocalTransform(ElementToReset);
+			const FTransform InitialLocalTransform = ControlRig->GetInitialLocalTransform(ElementToReset);
 			ControlRig->Modify();
 			if (bHasNonDefaultParent == true) //possibly not at default parent so switch to it
 			{
@@ -4103,8 +4103,10 @@ void FControlRigEditMode::MoveControlShape(AControlRigShapeActor* ShapeActor, co
 					bPrintPythonCommands = World->IsPreviewWorld();
 				}
 
-				SetControlShapeTransform(ShapeActor, NewTransform, ToWorldTransform, Context, bPrintPythonCommands);			
-				
+				ControlRig->Evaluate_AnyThread();
+				SetControlShapeTransform(ShapeActor, NewTransform, ToWorldTransform, Context, bPrintPythonCommands);
+				NotifyDrivenControls(ControlRig, ShapeActor->GetElementKey());
+				ControlRig->Evaluate_AnyThread();
 				ShapeActor->SetGlobalTransform(CurrentTransform);
 				if (bCalcLocal)
 				{
@@ -4112,9 +4114,6 @@ void FControlRigEditMode::MoveControlShape(AControlRigShapeActor* ShapeActor, co
 					InOutLocal = NewLocal.GetRelativeTransform(InOutLocal);
 				}
 
-				ControlRig->Evaluate_AnyThread();
-
-				NotifyDrivenControls(ControlRig, ShapeActor->GetElementKey());
 			}
 		}
 	}
@@ -4630,7 +4629,7 @@ void FControlRigEditMode::NotifyDrivenControls(UControlRig* InControlRig, const
 			{
 				if(DrivenKey.Type == ERigElementType::Control)
 				{
-					const FTransform DrivenTransform = InControlRig->GetHierarchy()->GetLocalTransform(DrivenKey);
+					const FTransform DrivenTransform = InControlRig->GetControlLocalTransform(DrivenKey.Name);
 					InControlRig->SetControlLocalTransform(DrivenKey.Name, DrivenTransform, true, Context, false /*undo*/, true/* bFixEulerFlips*/);
 				}
 			}
diff --git a/Engine/Plugins/Animation/ControlRig/Source/ControlRigEditor/Private/Sequencer/ControlRigParameterTrackEditor.cpp b/Engine/Plugins/Animation/ControlRig/Source/ControlRigEditor/Private/Sequencer/ControlRigParameterTrackEditor.cpp
--- a/Engine/Plugins/Animation/ControlRig/Source/ControlRigEditor/Private/Sequencer/ControlRigParameterTrackEditor.cpp
+++ b/Engine/Plugins/Animation/ControlRig/Source/ControlRigEditor/Private/Sequencer/ControlRigParameterTrackEditor.cpp
@@ -65,6 +65,7 @@
 #include "Widgets/Layout/SScrollBox.h"
 #include "Widgets/Input/SCheckBox.h"
 #include "Units/Execution/RigUnit_InverseExecution.h"
+#include "Units/Execution/RigUnit_BeginExecution.h"
 #include "Tracks/MovieSceneSkeletalAnimationTrack.h"
 #include "Exporters/AnimSeqExportOption.h"
 #include "SBakeToControlRigDialog.h"
@@ -477,6 +478,8 @@ void FControlRigParameterTrackEditor::UnbindControlRig(UControlRig* ControlRig)
 			Binding->OnControlRigBind().RemoveAll(this);
 		}
 		ControlRig->ControlRigBound().RemoveAll(this);
+		ControlRig->Rename(nullptr, GetTransientPackage(), REN_ForceNoResetLoaders | REN_DoNotDirty | REN_DontCreateRedirectors | REN_NonTransactional);
+		ControlRig->MarkAsGarbage();
 		
 		BoundControlRigs.Remove(ControlRig);
 		ClearOutAllSpaceAndConstraintDelegates(ControlRig);
@@ -485,25 +488,13 @@ void FControlRigParameterTrackEditor::UnbindControlRig(UControlRig* ControlRig)
 void FControlRigParameterTrackEditor::UnbindAllControlRigs()
 {
 	ClearOutAllSpaceAndConstraintDelegates();
-	for (TWeakObjectPtr<UControlRig>& ObjectPtr : BoundControlRigs)
+	TArray<TWeakObjectPtr<UControlRig>> ControlRigs = BoundControlRigs;
+	for (TWeakObjectPtr<UControlRig>& ObjectPtr : ControlRigs)
 	{
 		if (ObjectPtr.IsValid())
 		{
 			UControlRig* ControlRig = ObjectPtr.Get();
-			UMovieSceneControlRigParameterTrack* Track = FindTrack(ControlRig);
-			if (Track)
-			{
-				Track->SpaceChannelAdded().RemoveAll(this);
-				Track->ConstraintChannelAdded().RemoveAll(this);
-			}
-			ControlRig->ControlModified().RemoveAll(this);
-			ControlRig->OnInitialized_AnyThread().RemoveAll(this);
-			ControlRig->ControlSelected().RemoveAll(this);
-			if (const TSharedPtr<IControlRigObjectBinding> Binding = ControlRig->GetObjectBinding())
-			{
-				Binding->OnControlRigBind().RemoveAll(this);
-			}
-			ControlRig->ControlRigBound().RemoveAll(this);
+			UnbindControlRig(ControlRig);
 		}
 	}
 	BoundControlRigs.SetNum(0);
@@ -641,6 +632,16 @@ void FControlRigParameterTrackEditor::BuildObjectBindingContextMenu(FMenuBuilder
 					NAME_None,
 					EUserInterfaceActionType::Button);
 
+				MenuBuilder.AddMenuEntry(
+					LOCTEXT("AddAdditiveControlRig", "Add Additive Control Rig"),
+					LOCTEXT("AddAdditiveControlRigTooltip", "Add additive Control Rig and add a track for it"),
+					FSlateIcon(),
+					FUIAction(
+						FExecuteAction::CreateRaw(this, &FControlRigParameterTrackEditor::AddAdditiveControlRig, ObjectBindings[0], BoundObject, SkelMeshComp, Skeleton),
+						FCanExecuteAction::CreateRaw(this, &FControlRigParameterTrackEditor::CanAddAdditiveConrolRig)),
+					NAME_None,
+					EUserInterfaceActionType::Button);
+
 				MenuBuilder.AddMenuEntry(
 					LOCTEXT("FilterAssetBySkeleton", "Filter Asset By Skeleton"),
 					LOCTEXT("FilterAssetBySkeletonTooltip", "Filters Control Rig assets to match current skeleton"),
@@ -819,6 +820,358 @@ void FControlRigParameterTrackEditor::ConvertToFKControlRig(FGuid ObjectBinding,
 	BakeToControlRig(UFKControlRig::StaticClass(), ObjectBinding, BoundObject, SkelMeshComp, Skeleton);
 }
 
+class SAddAdditiveControlRigOptionsWindow : public SCompoundWidget
+{
+public:
+	SLATE_BEGIN_ARGS(SAddAdditiveControlRigOptionsWindow)
+		: _ControlRigClass(nullptr)
+		, _Skeleton(nullptr)
+		, _WidgetWindow()
+	{}
+
+		SLATE_ARGUMENT(TSubclassOf<UControlRig>*, ControlRigClass)
+		SLATE_ARGUMENT(USkeleton*, Skeleton)
+		SLATE_ARGUMENT(TSharedPtr<SWindow>, WidgetWindow)
+		SLATE_END_ARGS()
+
+public:
+	void Construct(const FArguments& InArgs);
+
+	void OnAdd(UClass* InClass)
+	{
+		if (WidgetWindow.IsValid())
+		{
+			WidgetWindow.Pin()->RequestDestroyWindow();
+		}
+		*ControlRigClass = TSubclassOf<UControlRig>(InClass);
+		return;
+	}
+
+	FReply OnCancel()
+	{
+		if (WidgetWindow.IsValid())
+		{
+			WidgetWindow.Pin()->RequestDestroyWindow();
+		}
+		return FReply::Handled();
+	}
+
+	SAddAdditiveControlRigOptionsWindow()
+		: ControlRigClass(nullptr)
+		, Skeleton(nullptr)
+	{}
+
+private:
+	TSubclassOf<UControlRig>* ControlRigClass;
+	USkeleton* Skeleton;
+	TWeakPtr< SWindow > WidgetWindow;
+};
+
+
+void SAddAdditiveControlRigOptionsWindow::Construct(const FArguments& InArgs)
+{
+	ControlRigClass = InArgs._ControlRigClass;
+	Skeleton = InArgs._Skeleton;
+	WidgetWindow = InArgs._WidgetWindow;
+
+	FText CancelText = LOCTEXT("AnimSequenceOptions_Cancel", "Cancel");
+	FText CancelTooltipText = LOCTEXT("AnimSequenceOptions_Cancel_ToolTip", "Cancel control rig creation");
+
+	TSharedPtr<SBox> HeaderToolBox;
+	TSharedPtr<SHorizontalBox> AnimHeaderButtons;
+	TSharedPtr<SBox> InspectorBox;
+	this->ChildSlot
+	[
+		SNew(SBox)
+		[
+			SNew(SVerticalBox)
+			+ SVerticalBox::Slot()
+			.AutoHeight()
+			.Padding(2)
+			[
+				SAssignNew(HeaderToolBox, SBox)
+			]
+			+ SVerticalBox::Slot()
+			.AutoHeight()
+			.Padding(2)
+			[
+				SNew(SBorder)
+				.Padding(FMargin(3))
+				.BorderImage(FAppStyle::GetBrush("ToolPanel.GroupBorder"))
+				[
+					SNew(SHorizontalBox)
+					+ SHorizontalBox::Slot()
+					.AutoWidth()
+					[
+						SNew(STextBlock)
+						.Font(FAppStyle::GetFontStyle("CurveEd.LabelFont"))
+						.Text(LOCTEXT("Export_CurrentFileTitle", "Current File: "))
+					]
+				]
+			]
+			+ SVerticalBox::Slot()
+			.FillHeight(1.0f)
+			.Padding(2)
+			[
+				SAssignNew(InspectorBox, SBox)
+			]
+		]
+	];
+
+	//todo move to .h for ue5
+	class FControlRigClassFilter : public IClassViewerFilter
+	{
+	public:
+		FControlRigClassFilter(bool bInCheckSkeleton, bool bInCheckAnimatable, bool bInCheckInversion, USkeleton* InSkeleton) :
+			bFilterAssetBySkeleton(bInCheckSkeleton),
+			bFilterExposesAnimatableControls(bInCheckAnimatable),
+			bFilterInversion(bInCheckInversion),
+			AssetRegistry(FModuleManager::GetModuleChecked<FAssetRegistryModule>(TEXT("AssetRegistry")).Get())
+		{
+			if (InSkeleton)
+			{
+				SkeletonName = FAssetData(InSkeleton).GetExportTextName();
+			}
+		}
+		bool bFilterAssetBySkeleton;
+		bool bFilterExposesAnimatableControls;
+		bool bFilterInversion;
+
+		FString SkeletonName;
+		const IAssetRegistry& AssetRegistry;
+
+		bool MatchesFilter(const FAssetData& AssetData)
+		{
+			bool bExposesAnimatableControls = AssetData.GetTagValueRef<bool>(TEXT("bExposesAnimatableControls"));
+			if (bFilterExposesAnimatableControls == true && bExposesAnimatableControls == false)
+			{
+				return false;
+			}
+			if (bFilterInversion)
+			{
+				bool bHasInversion = false;
+				FAssetDataTagMapSharedView::FFindTagResult Tag = AssetData.TagsAndValues.FindTag(TEXT("SupportedEventNames"));
+				if (Tag.IsSet())
+				{
+					FString EventString = FRigUnit_InverseExecution::EventName.ToString();
+					FString OldEventString = FString(TEXT("Inverse"));
+					TArray<FString> SupportedEventNames;
+					Tag.GetValue().ParseIntoArray(SupportedEventNames, TEXT(","), true);
+
+					for (const FString& Name : SupportedEventNames)
+					{
+						if (Name.Contains(EventString) || Name.Contains(OldEventString))
+						{
+							bHasInversion = true;
+							break;
+						}
+					}
+					if (bHasInversion == false)
+					{
+						return false;
+					}
+				}
+			}
+			if (bFilterAssetBySkeleton)
+			{
+				FString PreviewSkeletalMesh = AssetData.GetTagValueRef<FString>(TEXT("PreviewSkeletalMesh"));
+				if (PreviewSkeletalMesh.Len() > 0)
+				{
+					FAssetData SkelMeshData = AssetRegistry.GetAssetByObjectPath(FSoftObjectPath(PreviewSkeletalMesh));
+					FString PreviewSkeleton = SkelMeshData.GetTagValueRef<FString>(TEXT("Skeleton"));
+					if (PreviewSkeleton == SkeletonName)
+					{
+						return true;
+					}
+				}
+				FString PreviewSkeleton = AssetData.GetTagValueRef<FString>(TEXT("PreviewSkeleton"));
+				if (PreviewSkeleton == SkeletonName)
+				{
+					return true;
+				}
+				FString SourceHierarchyImport = AssetData.GetTagValueRef<FString>(TEXT("SourceHierarchyImport"));
+				if (SourceHierarchyImport == SkeletonName)
+				{
+					return true;
+				}
+				FString SourceCurveImport = AssetData.GetTagValueRef<FString>(TEXT("SourceCurveImport"));
+				if (SourceCurveImport == SkeletonName)
+				{
+					return true;
+				}
+				return false;
+			}
+			return true;
+
+		}
+		bool IsClassAllowed(const FClassViewerInitializationOptions& InInitOptions, const UClass* InClass, TSharedRef< FClassViewerFilterFuncs > InFilterFuncs) override
+		{
+			if (InClass)
+			{
+				const bool bChildOfObjectClass = InClass->IsChildOf(UControlRig::StaticClass());
+				const bool bMatchesFlags = !InClass->HasAnyClassFlags(CLASS_Hidden | CLASS_HideDropDown | CLASS_Deprecated | CLASS_Abstract);
+				const bool bNotNative = !InClass->IsNative();
+
+				if (bChildOfObjectClass && bMatchesFlags && bNotNative)
+				{
+					FAssetData AssetData(InClass);
+					return MatchesFilter(AssetData);
+				}
+			}
+			return false;
+		}
+
+		virtual bool IsUnloadedClassAllowed(const FClassViewerInitializationOptions& InInitOptions, const TSharedRef< const IUnloadedBlueprintData > InUnloadedClassData, TSharedRef< FClassViewerFilterFuncs > InFilterFuncs) override
+		{
+			const bool bChildOfObjectClass = InUnloadedClassData->IsChildOf(UControlRig::StaticClass());
+			const bool bMatchesFlags = !InUnloadedClassData->HasAnyClassFlags(CLASS_Hidden | CLASS_HideDropDown | CLASS_Deprecated | CLASS_Abstract);
+			if (bChildOfObjectClass && bMatchesFlags)
+			{
+				FString GeneratedClassPathString = InUnloadedClassData->GetClassPathName().ToString();
+				FString BlueprintPath = GeneratedClassPathString.LeftChop(2); // Chop off _C
+				FAssetData AssetData = AssetRegistry.GetAssetByObjectPath(FSoftObjectPath(BlueprintPath));
+				return MatchesFilter(AssetData);
+
+			}
+			return false;
+		}
+
+	};
+
+	FClassViewerInitializationOptions Options;
+	Options.bShowUnloadedBlueprints = true;
+	Options.NameTypeToDisplay = EClassViewerNameTypeToDisplay::DisplayName;
+
+	TSharedPtr<FControlRigClassFilter> ClassFilter = MakeShareable(new FControlRigClassFilter(true, false, true, Skeleton));
+	Options.ClassFilters.Add(ClassFilter.ToSharedRef());
+	Options.bShowNoneOption = false;
+
+	FClassViewerModule& ClassViewerModule = FModuleManager::LoadModuleChecked<FClassViewerModule>("ClassViewer");
+	TSharedRef<SWidget> ClassViewer = ClassViewerModule.CreateClassViewer(Options, FOnClassPicked::CreateRaw(this, &SAddAdditiveControlRigOptionsWindow::OnAdd));
+	InspectorBox->SetContent(ClassViewer->AsShared());
+}
+
+void FControlRigParameterTrackEditor::AddAdditiveControlRig(FGuid ObjectBinding, UObject* BoundActor, USkeletalMeshComponent* SkelMeshComp, USkeleton* Skeleton)
+{
+	FSlateApplication::Get().DismissAllMenus();
+
+	TSharedPtr<SWindow> ParentWindow;
+	if (FModuleManager::Get().IsModuleLoaded("MainFrame"))
+	{
+		IMainFrameModule& MainFrame = FModuleManager::LoadModuleChecked<IMainFrameModule>("MainFrame");
+		ParentWindow = MainFrame.GetParentWindow();
+	}
+
+	TSharedRef<SWindow> Window = SNew(SWindow)
+				.Title(LOCTEXT("AnimSeqTitle", "Choose Additive Rig"))
+				.SizingRule(ESizingRule::UserSized)
+				.AutoCenter(EAutoCenter::PrimaryWorkArea)
+				.ClientSize(FVector2D(500, 445));
+
+	TSharedPtr<SAddAdditiveControlRigOptionsWindow> OptionWindow;
+	TSubclassOf<UControlRig> ControlRigClass;
+	Window->SetContent
+	(
+		SAssignNew(OptionWindow, SAddAdditiveControlRigOptionsWindow)
+		.ControlRigClass(&ControlRigClass)
+		.Skeleton(Skeleton)
+		.WidgetWindow(Window)
+	);
+
+	FSlateApplication::Get().AddModalWindow(Window, ParentWindow, false);
+	
+	
+	const TSharedPtr<ISequencer> SequencerParent = GetSequencer();
+
+	if (ControlRigClass && SequencerParent.IsValid())
+	{
+		UMovieSceneSequence* OwnerSequence = GetSequencer()->GetFocusedMovieSceneSequence();
+		SequencerParent->PreAnimatedState.DiscardAndRemoveEntityTokensForObject(*SkelMeshComp);
+		UMovieScene* OwnerMovieScene = OwnerSequence->GetMovieScene();
+		{
+			const TSharedPtr<ISequencer> ParentSequencer = GetSequencer();
+			FMovieSceneSequenceIDRef Template = ParentSequencer->GetFocusedTemplateID();
+			FMovieSceneSequenceTransform RootToLocalTransform = ParentSequencer->GetFocusedMovieSceneSequenceTransform();
+			
+			OwnerMovieScene->Modify();
+			UMovieSceneControlRigParameterTrack* Track = Cast<UMovieSceneControlRigParameterTrack>(AddTrack(OwnerMovieScene, ObjectBinding, UMovieSceneControlRigParameterTrack::StaticClass(), NAME_None));
+			if (Track)
+			{
+				Track->Modify();
+
+				FString ObjectName = ControlRigClass->GetName();
+				ObjectName.RemoveFromEnd(TEXT("_C"));
+				const FString AdditiveObjectName = ObjectName + "_Additive";
+
+				AActor* Actor = Cast<AActor>(BoundActor);
+				if (!Actor)
+				{
+					OwnerMovieScene->RemoveTrack(*Track);
+					return;
+				}
+				
+				UControlRigComponent* CRComponent = Cast<UControlRigComponent>(Actor->AddComponentByClass(UControlRigComponent::StaticClass(), false, FTransform::Identity, false));
+				check(CRComponent);
+				CRComponent->SetControlRigClass(ControlRigClass);
+				CRComponent->RegisterComponent();
+				CRComponent->AttachToComponent(SkelMeshComp, FAttachmentTransformRules::SnapToTargetIncludingScale);
+				
+				CRComponent->AddMappedCompleteSkeletalMesh(SkelMeshComp, EControlRigComponentMapDirection::Input);
+				CRComponent->AddMappedCompleteSkeletalMesh(SkelMeshComp, EControlRigComponentMapDirection::Output);
+				
+				UControlRig* ControlRig = CRComponent->GetControlRig();
+				ControlRig->SetIsAdditive(true);
+				if (!ControlRig->SupportsEvent(FRigUnit_InverseExecution::EventName))
+				{
+					OwnerMovieScene->RemoveTrack(*Track);
+					return;
+				}
+
+				FControlRigEditMode* ControlRigEditMode = GetEditMode();
+				if (!ControlRigEditMode)
+				{
+					ControlRigEditMode = GetEditMode(true);
+				}
+				
+				ControlRig->Modify();
+				ControlRig->SetBoneInitialTransformsFromSkeletalMeshComponent(SkelMeshComp, true);
+				ControlRig->Evaluate_AnyThread();
+
+				const bool bSequencerOwnsControlRig = true;
+				UMovieSceneSection* NewSection = Track->CreateControlRigSection(0, ControlRig, bSequencerOwnsControlRig);
+				UMovieSceneControlRigParameterSection* ParamSection = Cast<UMovieSceneControlRigParameterSection>(NewSection);
+
+				//mz todo need to have multiple rigs with same class
+				Track->SetTrackName(FName(*ObjectName));
+				Track->SetDisplayName(FText::FromString(ObjectName));
+
+				GetSequencer()->EmptySelection();
+				GetSequencer()->SelectSection(NewSection);
+				GetSequencer()->ThrobSectionSelection();
+				GetSequencer()->NotifyMovieSceneDataChanged(EMovieSceneDataChangeType::MovieSceneStructureItemAdded);
+				
+				//Finish Setup
+				if (ControlRigEditMode)
+				{
+					ControlRigEditMode->AddControlRigObject(ControlRig, GetSequencer());
+				}
+				BindControlRig(ControlRig);
+
+				GetSequencer()->NotifyMovieSceneDataChanged(EMovieSceneDataChangeType::MovieSceneStructureItemAdded);
+			}
+		}
+	}
+}
+
+bool FControlRigParameterTrackEditor::CanAddAdditiveConrolRig()
+{
+	if (FControlRigEditMode* ControlRigEditMode = GetEditMode())
+	{
+		return ControlRigEditMode->GetControlRigsArray(false).IsEmpty();
+	}
+	return true;
+}
+
 void FControlRigParameterTrackEditor::BakeToControlRigSubMenu(FMenuBuilder& MenuBuilder, FGuid ObjectBinding, UObject* BoundObject, USkeletalMeshComponent* SkelMeshComp, USkeleton* Skeleton)
 {
 	const TSharedPtr<ISequencer> ParentSequencer = GetSequencer();
@@ -1797,8 +2150,8 @@ static void EvaluateThisControl(UMovieSceneControlRigParameterSection* Section,
 					{
 						const FEulerTransform& Euler = Value.GetValue();
 						FVector EulerAngle(Euler.Rotation.Roll, Euler.Rotation.Pitch, Euler.Rotation.Yaw);
-						FQuat Quat = ControlRig->GetHierarchy()->GetControlQuaternion(ControlElement, EulerAngle);
-						ControlRig->GetHierarchy()->SetControlSpecifiedEulerAngle(ControlElement, EulerAngle);
+						FQuat Quat = RigHierarchy->GetControlQuaternion(ControlElement, EulerAngle);
+						RigHierarchy->SetControlSpecifiedEulerAngle(ControlElement, EulerAngle);
 						FRotator UERotator(Quat);
 						FEulerTransform Transform = Euler;
 						Transform.Rotation = UERotator;
@@ -2627,7 +2980,7 @@ void FControlRigParameterTrackEditor::ClearOutAllSpaceAndConstraintDelegates(con
 					for (FConstraintAndActiveChannel& Channel: ConstraintChannels)
 					{
 						Channel.ActiveChannel.OnKeyMovedEvent().Clear();
-						Channel.ActiveChannel.OnKeyDeletedEvent().Clear();							
+						Channel.ActiveChannel.OnKeyDeletedEvent().Clear();
 					}
 
 					if (CRSection->OnConstraintRemovedHandle.IsValid())
@@ -3258,6 +3611,7 @@ void FControlRigParameterTrackEditor::GetControlRigKeys(
 	{
 		ChannelsToKey = EControlRigContextChannelToKey::AllTransform;
 	}
+	URigHierarchy* Hierarchy = InControlRig->GetHierarchy();
 
 	//Need seperate index fo bools,ints and enums and floats since there are seperate entries for each later when they are accessed by the set key stuff.
 	int32 SpaceChannelIndex = 0;
@@ -3266,7 +3620,7 @@ void FControlRigParameterTrackEditor::GetControlRigKeys(
 		FRigControlElement* ControlElement = Controls[ControlIndex];
 		check(ControlElement);
 
-		if (!InControlRig->GetHierarchy()->IsAnimatable(ControlElement))
+		if (!Hierarchy->IsAnimatable(ControlElement))
 		{
 			continue;
 		}
@@ -3279,7 +3633,7 @@ void FControlRigParameterTrackEditor::GetControlRigKeys(
 			bool bMaskKeyOut = (ControlIndex >= ControlsMask.Num() || ControlsMask[ControlIndex] == false);
 			bool bSetKey = ControlElement->GetName() == ParameterName && !bMaskKeyOut;
 
-			FRigControlValue ControlValue = InControlRig->GetHierarchy()->GetControlValue(ControlElement, ERigControlValueType::Current);
+			FRigControlValue ControlValue = InControlRig->GetControlValue(ControlElement, ERigControlValueType::Current);
 
 			switch (ControlElement->Settings.ControlType)
 			{
@@ -3363,7 +3717,7 @@ void FControlRigParameterTrackEditor::GetControlRigKeys(
 			case ERigControlType::EulerTransform:
 			{
 				FVector Translation, Scale(1.0f, 1.0f, 1.0f);
-				FVector Vector = InControlRig->GetHierarchy()->GetControlSpecifiedEulerAngle(ControlElement);
+				FVector Vector = Hierarchy->GetControlSpecifiedEulerAngle(ControlElement);
 				FRotator Rotation = FRotator(Vector.Y, Vector.Z, Vector.X);
 				if (ControlElement->Settings.ControlType == ERigControlType::TransformNoScale)
 				{
@@ -3396,7 +3750,7 @@ void FControlRigParameterTrackEditor::GetControlRigKeys(
 							FVector Angle = Rotation.Euler();
 							//need to wind rotators still
 							ControlElement->PreferredEulerAngles.SetAngles(Angle, false, ControlElement->PreferredEulerAngles.RotationOrder, true);
-							Angle = InControlRig->GetHierarchy()->GetControlSpecifiedEulerAngle(ControlElement);
+							Angle = Hierarchy->GetControlSpecifiedEulerAngle(ControlElement);
 							Rotation = FRotator(Vector.Y, Vector.Z, Vector.X);
 						}
 						else
diff --git a/Engine/Plugins/Animation/ControlRig/Source/ControlRigEditor/Private/Sequencer/ControlRigParameterTrackEditor.h b/Engine/Plugins/Animation/ControlRig/Source/ControlRigEditor/Private/Sequencer/ControlRigParameterTrackEditor.h
--- a/Engine/Plugins/Animation/ControlRig/Source/ControlRigEditor/Private/Sequencer/ControlRigParameterTrackEditor.h
+++ b/Engine/Plugins/Animation/ControlRig/Source/ControlRigEditor/Private/Sequencer/ControlRigParameterTrackEditor.h
@@ -191,6 +191,12 @@ private:
 	/** Convert to FK Control Rig*/
 	void ConvertToFKControlRig(FGuid ObjectBinding, UObject* BoundObject, USkeletalMeshComponent* SkelMeshComp, USkeleton* Skeleton);
 
+	/** Add additive Control Rig*/
+	void AddAdditiveControlRig(FGuid ObjectBinding, UObject* BoundObject, USkeletalMeshComponent* SkelMeshComp, USkeleton* Skeleton);
+
+	/** Returns true if an additive control rig can be added to the sequence */
+	bool CanAddAdditiveConrolRig();
+
 	/** Bake To Control Rig Sub Menu*/
 	void BakeToControlRigSubMenu(FMenuBuilder& MenuBuilder, FGuid ObjectBinding, UObject* BoundObject, USkeletalMeshComponent* SkelMeshComp,USkeleton* Skeleton);
 	
