Control Rig: Fix test data single frame record not working
#rnx

diff --git a/Engine/Plugins/Animation/ControlRig/Source/ControlRig/Private/ControlRigTestData.cpp b/Engine/Plugins/Animation/ControlRig/Source/ControlRig/Private/ControlRigTestData.cpp
--- a/Engine/Plugins/Animation/ControlRig/Source/ControlRig/Private/ControlRigTestData.cpp
+++ b/Engine/Plugins/Animation/ControlRig/Source/ControlRig/Private/ControlRigTestData.cpp
@@ -266,8 +266,10 @@ bool UControlRigTestData::Record(UControlRig* InControlRig, double InRecordingDu
 			const double TimeDelta = TimeNow - TimeAtStartOfRecording;
 			if(DesiredRecordingDuration <= TimeDelta)
 			{
-				ClearDelegates(InControlRig);
 				DesiredRecordingDuration = 0.0;
+
+				// Once clear delegates is called, we no longer have access to this pointer
+				ClearDelegates(InControlRig);
 			}
 		}
 	);
