ControlRig: Fix crash during compilation in editor

diff --git a/Engine/Plugins/Animation/ControlRig/Source/ControlRigEditor/Private/Editor/ControlRigEditor.cpp b/Engine/Plugins/Animation/ControlRig/Source/ControlRigEditor/Private/Editor/ControlRigEditor.cpp
--- a/Engine/Plugins/Animation/ControlRig/Source/ControlRigEditor/Private/Editor/ControlRigEditor.cpp
+++ b/Engine/Plugins/Animation/ControlRig/Source/ControlRigEditor/Private/Editor/ControlRigEditor.cpp
@@ -2232,7 +2232,10 @@ void FControlRigEditor::Compile()
 
 		// default to always reset all bone modifications 
 		ResetAllBoneModification(); 
-		
+
+		// remove the current host from the editor
+		SetHost(nullptr);
+
 		{
 			TGuardValue<bool> GuardCompileReEntry(bIsCompilingThroughUI, true);
 			FBlueprintEditor::Compile();
@@ -4298,6 +4301,11 @@ void FControlRigEditor::UpdateControlRig()
 				{
 					TestDataStrongPtr->ReleaseReplay();
 				}
+
+				if(!IsValid(CurrentHost))
+				{
+					SetHost(nullptr);
+				}
 				
 				// if this control rig is from a temporary step,
 				// for example the reinstancing class, clear it 
diff --git a/Engine/Plugins/Runtime/RigVM/Source/RigVMEditor/Private/Editor/RigVMEditor.cpp b/Engine/Plugins/Runtime/RigVM/Source/RigVMEditor/Private/Editor/RigVMEditor.cpp
--- a/Engine/Plugins/Runtime/RigVM/Source/RigVMEditor/Private/Editor/RigVMEditor.cpp
+++ b/Engine/Plugins/Runtime/RigVM/Source/RigVMEditor/Private/Editor/RigVMEditor.cpp
@@ -32,7 +32,7 @@ void FRigVMEditor::HandleModifiedEvent(ERigVMGraphNotifType InNotifType, URigVMG
 void FRigVMEditor::SetHost(URigVMHost* InHost)
 {
 	Host = InHost;
-	if(Host)
+	if(Host && IsValid(Host))
 	{
 		OnPreviewHostUpdated().Broadcast(this);
 	}
