Control Rig: Fix warning due to conflict with window titles of additive control rig widget
#rnx

diff --git a/Engine/Plugins/Animation/ControlRig/Source/ControlRigEditor/Private/Sequencer/ControlRigParameterTrackEditor.cpp b/Engine/Plugins/Animation/ControlRig/Source/ControlRigEditor/Private/Sequencer/ControlRigParameterTrackEditor.cpp
--- a/Engine/Plugins/Animation/ControlRig/Source/ControlRigEditor/Private/Sequencer/ControlRigParameterTrackEditor.cpp
+++ b/Engine/Plugins/Animation/ControlRig/Source/ControlRigEditor/Private/Sequencer/ControlRigParameterTrackEditor.cpp
@@ -1063,7 +1063,7 @@ void FControlRigParameterTrackEditor::AddAdditiveControlRig(FGuid ObjectBinding,
 	}
 
 	TSharedRef<SWindow> Window = SNew(SWindow)
-				.Title(LOCTEXT("AnimSeqTitle", "Choose Additive Rig"))
+				.Title(LOCTEXT("AdditiveRigTitle", "Choose Additive Rig"))
 				.SizingRule(ESizingRule::UserSized)
 				.AutoCenter(EAutoCenter::PrimaryWorkArea)
 				.ClientSize(FVector2D(500, 445));
