Control Rig: Fetch asset user data from the ControlRig CDO
#rnx

diff --git a/Engine/Plugins/Animation/ControlRig/Source/ControlRig/Private/ControlRig.cpp b/Engine/Plugins/Animation/ControlRig/Source/ControlRig/Private/ControlRig.cpp
--- a/Engine/Plugins/Animation/ControlRig/Source/ControlRig/Private/ControlRig.cpp
+++ b/Engine/Plugins/Animation/ControlRig/Source/ControlRig/Private/ControlRig.cpp
@@ -2531,7 +2531,15 @@ const TArray<UAssetUserData*>* UControlRig::GetAssetUserDataArray() const
 	}
 
 	CombinedAssetUserData.Reset();
-	CombinedAssetUserData.Append(AssetUserData);
+
+	if (UControlRig* CDO = Cast<UControlRig>(GetClass()->GetDefaultObject(false)))
+	{
+		CombinedAssetUserData.Append(CDO->AssetUserData);
+	}
+	else
+	{
+		CombinedAssetUserData.Append(AssetUserData);
+	}
 
 	if(OuterSceneComponent.IsValid())
 	{
