Sequencer: Control Rig: Fix setting zero value so it doesn't always add and uses default interoplation

diff --git a/Engine/Plugins/Animation/ControlRig/Source/ControlRig/Private/Sequencer/MovieSceneControlRigParameterSection.cpp b/Engine/Plugins/Animation/ControlRig/Source/ControlRig/Private/Sequencer/MovieSceneControlRigParameterSection.cpp
--- a/Engine/Plugins/Animation/ControlRig/Source/ControlRig/Private/Sequencer/MovieSceneControlRigParameterSection.cpp
+++ b/Engine/Plugins/Animation/ControlRig/Source/ControlRig/Private/Sequencer/MovieSceneControlRigParameterSection.cpp
@@ -2325,7 +2325,7 @@ float UMovieSceneControlRigParameterSection::GetTotalWeightValue(FFrameTime InTi
 	return WeightVal;
 }
 
-void UMovieSceneControlRigParameterSection::KeyZeroValue(FFrameNumber InFrame, bool bSelectedControls)
+void UMovieSceneControlRigParameterSection::KeyZeroValue(FFrameNumber InFrame, EMovieSceneKeyInterpolation DefaultInterpolation, bool bSelectedControls)
 {
 	TArray<FName> SelectedControls;
 	if (bSelectedControls && ControlRig)
@@ -2343,7 +2343,7 @@ void UMovieSceneControlRigParameterSection::KeyZeroValue(FFrameNumber InFrame, b
 	{
 		if (SelectedControls.Num() == 0 || SelectedControls.Contains(Scalar.ParameterName))
 		{
-			Scalar.ParameterCurve.AddCubicKey(InFrame, 0.0f, ERichCurveTangentMode::RCTM_Auto);
+			AddKeyToChannel(&Scalar.ParameterCurve, InFrame, 0.0f, DefaultInterpolation);
 			Scalar.ParameterCurve.AutoSetTangents();
 		}
 	}
@@ -2351,9 +2351,9 @@ void UMovieSceneControlRigParameterSection::KeyZeroValue(FFrameNumber InFrame, b
 	{
 		if (SelectedControls.Num() == 0 || SelectedControls.Contains(Vector2D.ParameterName))
 		{
-			Vector2D.XCurve.AddCubicKey(InFrame, 0.0f, ERichCurveTangentMode::RCTM_Auto);
+			AddKeyToChannel(&Vector2D.XCurve, InFrame, 0.0f, DefaultInterpolation);
 			Vector2D.XCurve.AutoSetTangents();
-			Vector2D.YCurve.AddCubicKey(InFrame, 0.0f, ERichCurveTangentMode::RCTM_Auto);
+			AddKeyToChannel(&Vector2D.YCurve, InFrame, 0.0f, DefaultInterpolation);
 			Vector2D.YCurve.AutoSetTangents();
 		}
 	}
@@ -2361,11 +2361,11 @@ void UMovieSceneControlRigParameterSection::KeyZeroValue(FFrameNumber InFrame, b
 	{
 		if (SelectedControls.Num() == 0 || SelectedControls.Contains(Vector.ParameterName))
 		{
-			Vector.XCurve.AddCubicKey(InFrame, 0.0f, ERichCurveTangentMode::RCTM_Auto);
+			AddKeyToChannel(&Vector.XCurve, InFrame, 0.0f, DefaultInterpolation);
 			Vector.XCurve.AutoSetTangents();
-			Vector.YCurve.AddCubicKey(InFrame, 0.0f, ERichCurveTangentMode::RCTM_Auto);
+			AddKeyToChannel(&Vector.YCurve, InFrame, 0.0f, DefaultInterpolation);
 			Vector.YCurve.AutoSetTangents();
-			Vector.ZCurve.AddCubicKey(InFrame, 0.0f, ERichCurveTangentMode::RCTM_Auto);
+			AddKeyToChannel(&Vector.ZCurve, InFrame, 0.0f, DefaultInterpolation);
 			Vector.ZCurve.AutoSetTangents();
 		}
 	}
@@ -2375,17 +2375,17 @@ void UMovieSceneControlRigParameterSection::KeyZeroValue(FFrameNumber InFrame, b
 		{
 			for (int32 Index = 0; Index < 3; ++Index)
 			{
-				Transform.Translation[Index].AddCubicKey(InFrame, 0.0f, ERichCurveTangentMode::RCTM_Auto);
+				AddKeyToChannel(&Transform.Translation[Index], InFrame, 0.0f, DefaultInterpolation);
 				Transform.Translation[Index].AutoSetTangents();
-				Transform.Rotation[Index].AddCubicKey(InFrame, 0.0f, ERichCurveTangentMode::RCTM_Auto);
+				AddKeyToChannel(&Transform.Rotation[Index], InFrame, 0.0f, DefaultInterpolation);
 				Transform.Rotation[Index].AutoSetTangents();
 				if (GetBlendType() == EMovieSceneBlendType::Additive)
 				{
-					Transform.Scale[Index].AddCubicKey(InFrame, 0.0f, ERichCurveTangentMode::RCTM_Auto);
+					AddKeyToChannel(&Transform.Scale[Index], InFrame, 0.0f, DefaultInterpolation);
 				}
 				else
 				{
-					Transform.Scale[Index].AddCubicKey(InFrame, 1.0f, ERichCurveTangentMode::RCTM_Auto);
+					AddKeyToChannel(&Transform.Scale[Index], InFrame, 1.0f, DefaultInterpolation);
 				}
 				Transform.Scale[Index].AutoSetTangents();
 
@@ -2394,9 +2394,9 @@ void UMovieSceneControlRigParameterSection::KeyZeroValue(FFrameNumber InFrame, b
 	}
 }
 
-void UMovieSceneControlRigParameterSection::KeyWeightValue(FFrameNumber InFrame, float InVal)
+void UMovieSceneControlRigParameterSection::KeyWeightValue(FFrameNumber InFrame, EMovieSceneKeyInterpolation DefaultInterpolation, float InVal)
 {
-	Weight.AddCubicKey(InFrame, InVal, ERichCurveTangentMode::RCTM_Auto);
+	AddKeyToChannel(&Weight, InFrame, InVal, DefaultInterpolation);
 	Weight.AutoSetTangents();
 }
 
diff --git a/Engine/Plugins/Animation/ControlRig/Source/ControlRig/Public/Sequencer/MovieSceneControlRigParameterSection.h b/Engine/Plugins/Animation/ControlRig/Source/ControlRig/Public/Sequencer/MovieSceneControlRigParameterSection.h
--- a/Engine/Plugins/Animation/ControlRig/Source/ControlRig/Public/Sequencer/MovieSceneControlRigParameterSection.h
+++ b/Engine/Plugins/Animation/ControlRig/Source/ControlRig/Public/Sequencer/MovieSceneControlRigParameterSection.h
@@ -476,10 +476,10 @@ public:
 	TOptional<FMovieSceneControlRigSpaceBaseKey> EvaluateSpaceChannel(const  FFrameTime& InTime, FName InParameterName);
 
 	/** Key Zero Values on all or just selected controls in these section at the specified time */
-	void KeyZeroValue(FFrameNumber InFrame, bool bSelected);
+	void KeyZeroValue(FFrameNumber InFrame, EMovieSceneKeyInterpolation DefaultInterpolation, bool bSelected);
 
 	/** Key the Weights to the specified value */
-	void KeyWeightValue(FFrameNumber InFrame, float InVal);
+	void KeyWeightValue(FFrameNumber InFrame, EMovieSceneKeyInterpolation DefaultInterpolation, float InVal);
 ;
 	/** Remove All Keys, but maybe not space keys if bIncludeSpaceKeys is false */
 	void RemoveAllKeys(bool bIncludeSpaceKeys);
diff --git a/Engine/Plugins/Animation/ControlRig/Source/ControlRigEditor/Private/Sequencer/ControlRigParameterTrackEditor.cpp b/Engine/Plugins/Animation/ControlRig/Source/ControlRigEditor/Private/Sequencer/ControlRigParameterTrackEditor.cpp
--- a/Engine/Plugins/Animation/ControlRig/Source/ControlRigEditor/Private/Sequencer/ControlRigParameterTrackEditor.cpp
+++ b/Engine/Plugins/Animation/ControlRig/Source/ControlRigEditor/Private/Sequencer/ControlRigParameterTrackEditor.cpp
@@ -5037,7 +5037,6 @@ void FControlRigParameterSection::CollapseAllLayers()
 	}
 }
 
-
 void FControlRigParameterSection::KeyZeroValue()
 {
 	UMovieSceneControlRigParameterSection* ParameterSection = CastChecked<UMovieSceneControlRigParameterSection>(WeakSection.Get());
@@ -5045,9 +5044,11 @@ void FControlRigParameterSection::KeyZeroValue()
 	FScopedTransaction Transaction(LOCTEXT("KeyZeroValue", "Key Zero Value"));
 	ParameterSection->Modify();
 	FFrameTime Time = SequencerPtr->GetLocalTime().Time;
-	ParameterSection->KeyZeroValue(Time.GetFrame(),true);
+	EMovieSceneKeyInterpolation DefaultInterpolation = SequencerPtr->GetKeyInterpolation();
+	ParameterSection->KeyZeroValue(Time.GetFrame(), DefaultInterpolation, true);
 	SequencerPtr->NotifyMovieSceneDataChanged(EMovieSceneDataChangeType::TrackValueChanged);
 }
+
 void FControlRigParameterSection::KeyWeightValue(float Val)
 {
 	UMovieSceneControlRigParameterSection* ParameterSection = CastChecked<UMovieSceneControlRigParameterSection>(WeakSection.Get());
@@ -5061,7 +5062,8 @@ void FControlRigParameterSection::KeyWeightValue(float Val)
 		SequencerPtr->NotifyMovieSceneDataChanged(EMovieSceneDataChangeType::MovieSceneStructureItemsChanged);
 	}
 	FFrameTime Time = SequencerPtr->GetLocalTime().Time;
-	ParameterSection->KeyWeightValue(Time.GetFrame(), Val);
+	EMovieSceneKeyInterpolation DefaultInterpolation = SequencerPtr->GetKeyInterpolation();
+	ParameterSection->KeyWeightValue(Time.GetFrame(), DefaultInterpolation, Val);
 	SequencerPtr->NotifyMovieSceneDataChanged(EMovieSceneDataChangeType::TrackValueChanged);
 }
 
