RigVM: Move most of the functionality from ControlRigEditor to RigVMEditor

diff --git a/Engine/Plugins/Animation/ControlRig/Source/ControlRigDeveloper/Private/ControlRigBlueprint.cpp b/Engine/Plugins/Animation/ControlRig/Source/ControlRigDeveloper/Private/ControlRigBlueprint.cpp
--- a/Engine/Plugins/Animation/ControlRig/Source/ControlRigDeveloper/Private/ControlRigBlueprint.cpp
+++ b/Engine/Plugins/Animation/ControlRig/Source/ControlRigDeveloper/Private/ControlRigBlueprint.cpp
@@ -386,6 +386,13 @@ UClass* UControlRigBlueprint::GetRigVMEditorSettingsClass() const
 	return UControlRigEditorSettings::StaticClass();
 }
 
+#if WITH_EDITOR
+IRigVMEditorModule* UControlRigBlueprint::GetEditorModule() const
+{
+	return &IControlRigEditorModule::Get();
+}
+#endif
+
 TArray<FString> UControlRigBlueprint::GeneratePythonCommands(const FString InNewBlueprintName)
 {
 	TArray<FString> InternalCommands;
diff --git a/Engine/Plugins/Animation/ControlRig/Source/ControlRigDeveloper/Private/Graph/ControlRigGraphNode.cpp b/Engine/Plugins/Animation/ControlRig/Source/ControlRigDeveloper/Private/Graph/ControlRigGraphNode.cpp
--- a/Engine/Plugins/Animation/ControlRig/Source/ControlRigDeveloper/Private/Graph/ControlRigGraphNode.cpp
+++ b/Engine/Plugins/Animation/ControlRig/Source/ControlRigDeveloper/Private/Graph/ControlRigGraphNode.cpp
@@ -14,14 +14,5 @@ UControlRigGraphNode::UControlRigGraphNode()
 {
 }
 
-void UControlRigGraphNode::GetNodeContextMenuActions(UToolMenu* Menu, UGraphNodeContextMenuContext* Context) const
-{
-#if WITH_EDITOR
-	// todo - part of this is likely generic 
-	const UControlRigGraphSchema* Schema = Cast<UControlRigGraphSchema>(GetSchema());
-	IControlRigEditorModule::Get().GetContextMenuActions(Schema, Menu, Context);
-#endif
-}
-
 #undef LOCTEXT_NAMESPACE
 
diff --git a/Engine/Plugins/Animation/ControlRig/Source/ControlRigDeveloper/Public/ControlRigBlueprint.h b/Engine/Plugins/Animation/ControlRig/Source/ControlRigDeveloper/Public/ControlRigBlueprint.h
--- a/Engine/Plugins/Animation/ControlRig/Source/ControlRigDeveloper/Public/ControlRigBlueprint.h
+++ b/Engine/Plugins/Animation/ControlRig/Source/ControlRigDeveloper/Public/ControlRigBlueprint.h
@@ -3,6 +3,7 @@
 #pragma once
 
 #include "CoreMinimal.h"
+#include "ControlRigBlueprintGeneratedClass.h"
 #include "UObject/ObjectMacros.h"
 #include "Engine/Blueprint.h"
 #include "ControlRigDefines.h"
@@ -37,6 +38,7 @@ public:
 	UControlRigBlueprint();
 
 	// URigVMBlueprint interface
+	virtual UClass* GetRigVMBlueprintGeneratedClassPrototype() const override { return UControlRigBlueprintGeneratedClass::StaticClass(); }
 	virtual UClass* GetRigVMSchemaClass() const override { return UControlRigSchema::StaticClass(); }
 	virtual UScriptStruct* GetRigVMExecuteContextStruct() const override { return FControlRigExecuteContext::StaticStruct(); }
 	virtual UClass* GetRigVMEdGraphClass() const override;
@@ -44,6 +46,10 @@ public:
 	virtual UClass* GetRigVMEdGraphSchemaClass() const override;
 	virtual TArray<FString> GeneratePythonCommands(const FString InNewBlueprintName) override;
 	virtual UClass* GetRigVMEditorSettingsClass() const override;
+#if WITH_EDITOR
+	virtual IRigVMEditorModule* GetEditorModule() const override;
+#endif
+
 	virtual void Serialize(FArchive& Ar) override;
 
 #if WITH_EDITOR
diff --git a/Engine/Plugins/Animation/ControlRig/Source/ControlRigDeveloper/Public/Graph/ControlRigGraphNode.h b/Engine/Plugins/Animation/ControlRig/Source/ControlRigDeveloper/Public/Graph/ControlRigGraphNode.h
--- a/Engine/Plugins/Animation/ControlRig/Source/ControlRigDeveloper/Public/Graph/ControlRigGraphNode.h
+++ b/Engine/Plugins/Animation/ControlRig/Source/ControlRigDeveloper/Public/Graph/ControlRigGraphNode.h
@@ -39,9 +39,6 @@ public:
 
 	UControlRigGraphNode();
 
-	// URigVMEdGraphNode
-	virtual void GetNodeContextMenuActions(UToolMenu* Menu, UGraphNodeContextMenuContext* Context) const override;
-
 private:
 
 	friend class SRigVMGraphNode;
diff --git a/Engine/Plugins/Animation/ControlRig/Source/ControlRigEditor/Private/ControlRigEditorModule.cpp b/Engine/Plugins/Animation/ControlRig/Source/ControlRigEditor/Private/ControlRigEditorModule.cpp
--- a/Engine/Plugins/Animation/ControlRig/Source/ControlRigEditor/Private/ControlRigEditorModule.cpp
+++ b/Engine/Plugins/Animation/ControlRig/Source/ControlRigEditor/Private/ControlRigEditorModule.cpp
@@ -36,7 +36,7 @@
 #include "Kismet2/KismetDebugUtilities.h"
 #include "EdGraphUtilities.h"
 #include "Graph/ControlRigGraphPanelPinFactory.h"
-#include "ControlRigBlueprintCommands.h"
+#include <Editor/ControlRigEditorCommands.h>
 #include "ControlRigHierarchyCommands.h"
 #include "Animation/AnimSequence.h"
 #include "Editor/ControlRigEditorEditMode.h"
@@ -100,7 +100,7 @@ DEFINE_LOG_CATEGORY(LogControlRigEditor);
 void FControlRigEditorModule::StartupModule()
 {
 	FControlRigEditModeCommands::Register();
-	FControlRigBlueprintCommands::Register();
+	FControlRigEditorCommands::Register();
 	FControlRigHierarchyCommands::Register();
 	FControlRigEditorStyle::Get();
 
@@ -1283,7 +1283,7 @@ void FControlRigEditorModule::ExtendAnimSequenceMenu()
 TSharedRef<IControlRigEditor> FControlRigEditorModule::CreateControlRigEditor(const EToolkitMode::Type Mode, const TSharedPtr< IToolkitHost >& InitToolkitHost, class UControlRigBlueprint* InBlueprint)
 {
 	TSharedRef< FControlRigEditor > NewControlRigEditor(new FControlRigEditor());
-	NewControlRigEditor->InitControlRigEditor(Mode, InitToolkitHost, InBlueprint);
+	NewControlRigEditor->InitRigVMEditor(Mode, InitToolkitHost, InBlueprint);
 	return NewControlRigEditor;
 }
 
