RigVM: Move most of the functionality from ControlRigEditor to RigVMEditor

diff --git a/Engine/Plugins/Animation/ControlRig/Source/ControlRigEditor/Private/Editor/ControlRigEditor.cpp b/Engine/Plugins/Animation/ControlRig/Source/ControlRigEditor/Private/Editor/ControlRigEditor.cpp
--- a/Engine/Plugins/Animation/ControlRig/Source/ControlRigEditor/Private/Editor/ControlRigEditor.cpp
+++ b/Engine/Plugins/Animation/ControlRig/Source/ControlRigEditor/Private/Editor/ControlRigEditor.cpp
@@ -14,7 +14,7 @@
 #include "Editor/Transactor.h"
 #include "Graph/ControlRigGraph.h"
 #include "BlueprintActionDatabase.h"
-#include "ControlRigBlueprintCommands.h"
+#include "ControlRigEditorCommands.h"
 #include "Kismet2/BlueprintEditorUtils.h"
 #include "IPersonaToolkit.h"
 #include "PersonaModule.h"
@@ -112,10 +112,6 @@
 
 TAutoConsoleVariable<bool> CVarControlRigShowTestingToolbar(TEXT("ControlRig.Test.EnableTestingToolbar"), false, TEXT("When true we'll show the testing toolbar in Control Rig Editor."));
 
-const FName ControlRigEditorAppName(TEXT("ControlRigEditorApp"));
-
-const FName FControlRigEditorModes::ControlRigEditorMode("Rigging");
-
 namespace ControlRigEditorTabs
 {
 	const FName DetailsTab(TEXT("DetailsTab"));
@@ -123,95 +119,7 @@ namespace ControlRigEditorTabs
 // 	const FName AdvancedPreviewTab(TEXT("AdvancedPreviewTab"));
 };
 
-struct FControlRigZoomLevelsContainer : public FZoomLevelsContainer
-{
-	struct FControlRigZoomLevelEntry
-	{
-	public:
-		FControlRigZoomLevelEntry(float InZoomAmount, const FText& InDisplayText, EGraphRenderingLOD::Type InLOD)
-			: DisplayText(FText::Format(NSLOCTEXT("GraphEditor", "Zoom", "Zoom {0}"), InDisplayText))
-		, ZoomAmount(InZoomAmount)
-		, LOD(InLOD)
-		{
-		}
-
-	public:
-		FText DisplayText;
-		float ZoomAmount;
-		EGraphRenderingLOD::Type LOD;
-	};
-	
-	FControlRigZoomLevelsContainer()
-	{
-		ZoomLevels.Reserve(22);
-		ZoomLevels.Add(FControlRigZoomLevelEntry(0.025f, FText::FromString(TEXT("-14")), EGraphRenderingLOD::LowestDetail));
-		ZoomLevels.Add(FControlRigZoomLevelEntry(0.070f, FText::FromString(TEXT("-13")), EGraphRenderingLOD::LowestDetail));
-		ZoomLevels.Add(FControlRigZoomLevelEntry(0.100f, FText::FromString(TEXT("-12")), EGraphRenderingLOD::LowestDetail));
-		ZoomLevels.Add(FControlRigZoomLevelEntry(0.125f, FText::FromString(TEXT("-11")), EGraphRenderingLOD::LowestDetail));
-		ZoomLevels.Add(FControlRigZoomLevelEntry(0.150f, FText::FromString(TEXT("-10")), EGraphRenderingLOD::LowestDetail));
-		ZoomLevels.Add(FControlRigZoomLevelEntry(0.175f, FText::FromString(TEXT("-9")), EGraphRenderingLOD::LowestDetail));
-		ZoomLevels.Add(FControlRigZoomLevelEntry(0.200f, FText::FromString(TEXT("-8")), EGraphRenderingLOD::LowestDetail));
-		ZoomLevels.Add(FControlRigZoomLevelEntry(0.225f, FText::FromString(TEXT("-7")), EGraphRenderingLOD::LowDetail));
-		ZoomLevels.Add(FControlRigZoomLevelEntry(0.250f, FText::FromString(TEXT("-6")), EGraphRenderingLOD::LowDetail));
-		ZoomLevels.Add(FControlRigZoomLevelEntry(0.375f, FText::FromString(TEXT("-5")), EGraphRenderingLOD::MediumDetail));
-		ZoomLevels.Add(FControlRigZoomLevelEntry(0.500f, FText::FromString(TEXT("-4")), EGraphRenderingLOD::MediumDetail));
-		ZoomLevels.Add(FControlRigZoomLevelEntry(0.675f, FText::FromString(TEXT("-3")), EGraphRenderingLOD::MediumDetail));
-		ZoomLevels.Add(FControlRigZoomLevelEntry(0.750f, FText::FromString(TEXT("-2")), EGraphRenderingLOD::DefaultDetail));
-		ZoomLevels.Add(FControlRigZoomLevelEntry(0.875f, FText::FromString(TEXT("-1")), EGraphRenderingLOD::DefaultDetail));
-		ZoomLevels.Add(FControlRigZoomLevelEntry(1.000f, FText::FromString(TEXT("1:1")), EGraphRenderingLOD::DefaultDetail));
-		ZoomLevels.Add(FControlRigZoomLevelEntry(1.250f, FText::FromString(TEXT("+1")), EGraphRenderingLOD::DefaultDetail));
-		ZoomLevels.Add(FControlRigZoomLevelEntry(1.375f, FText::FromString(TEXT("+2")), EGraphRenderingLOD::DefaultDetail));
-		ZoomLevels.Add(FControlRigZoomLevelEntry(1.500f, FText::FromString(TEXT("+3")), EGraphRenderingLOD::FullyZoomedIn));
-		ZoomLevels.Add(FControlRigZoomLevelEntry(1.675f, FText::FromString(TEXT("+4")), EGraphRenderingLOD::FullyZoomedIn));
-		ZoomLevels.Add(FControlRigZoomLevelEntry(1.750f, FText::FromString(TEXT("+5")), EGraphRenderingLOD::FullyZoomedIn));
-		ZoomLevels.Add(FControlRigZoomLevelEntry(1.875f, FText::FromString(TEXT("+6")), EGraphRenderingLOD::FullyZoomedIn));
-		ZoomLevels.Add(FControlRigZoomLevelEntry(2.000f, FText::FromString(TEXT("+7")), EGraphRenderingLOD::FullyZoomedIn));
-	}
-
-	float GetZoomAmount(int32 InZoomLevel) const override
-	{
-		checkSlow(ZoomLevels.IsValidIndex(InZoomLevel));
-		return ZoomLevels[InZoomLevel].ZoomAmount;
-	}
-
-	int32 GetNearestZoomLevel(float InZoomAmount) const override
-	{
-		for (int32 ZoomLevelIndex=0; ZoomLevelIndex < GetNumZoomLevels(); ++ZoomLevelIndex)
-		{
-			if (InZoomAmount <= GetZoomAmount(ZoomLevelIndex))
-			{
-				return ZoomLevelIndex;
-			}
-		}
-
-		return GetDefaultZoomLevel();
-	}
-	
-	FText GetZoomText(int32 InZoomLevel) const override
-	{
-		checkSlow(ZoomLevels.IsValidIndex(InZoomLevel));
-		return ZoomLevels[InZoomLevel].DisplayText;
-	}
-	
-	int32 GetNumZoomLevels() const override
-	{
-		return ZoomLevels.Num();
-	}
-	
-	int32 GetDefaultZoomLevel() const override
-	{
-		return 14;
-	}
-
-	EGraphRenderingLOD::Type GetLOD(int32 InZoomLevel) const override
-	{
-		checkSlow(ZoomLevels.IsValidIndex(InZoomLevel));
-		return ZoomLevels[InZoomLevel].LOD;
-	}
-
-	TArray<FControlRigZoomLevelEntry> ZoomLevels;
-};
-
+const FName FControlRigEditorModes::ControlRigEditorMode = TEXT("Rigging");
 const TArray<FName> FControlRigEditor::ForwardsSolveEventQueue = {FRigUnit_BeginExecution::EventName};
 const TArray<FName> FControlRigEditor::BackwardsSolveEventQueue = {FRigUnit_InverseExecution::EventName};
 const TArray<FName> FControlRigEditor::ConstructionEventQueue = {FRigUnit_PrepareForExecution::EventName};
@@ -221,41 +129,20 @@ FControlRigEditor::FControlRigEditor()
 	: IControlRigEditor()
 	, PreviewInstance(nullptr)
 	, ActiveController(nullptr)
-	, bControlRigEditorInitialized(false)
-	, bIsSettingObjectBeingDebugged(false)
 	, bExecutionControlRig(true)
-	, bIsCompilingThroughUI(false)
-	, bAnyErrorsLeft(false)
-	, LastEventQueue(ConstructionEventQueue)
-	, ExecutionMode(EControlRigExecutionModeType::EControlRigExecutionModeType_Release)
-	, LastDebuggedRig()
 	, RigHierarchyTabCount(0)
-	, HaltedAtNode(nullptr)
-	, bSuspendDetailsPanelRefresh(false)
-	, bAllowBulkEdits(false)
 	, bIsConstructionEventRunning(false)
 	, LastHierarchyHash(INDEX_NONE)
 {
+	LastEventQueue = ConstructionEventQueue;
 }
 
 FControlRigEditor::~FControlRigEditor()
 {
-	ClearDetailObject();
-
 	UControlRigBlueprint* RigBlueprint = GetControlRigBlueprint();
 
-	ControlRigEditorClosedDelegate.Broadcast(this, RigBlueprint);
-
-	if(PropertyChangedHandle.IsValid())
-	{
-		FCoreUObjectDelegates::OnObjectPropertyChanged.Remove(PropertyChangedHandle);
-	}
-
 	if (RigBlueprint)
 	{
-		// clear editor related data from the debugged control rig instance 
-		RigBlueprint->SetObjectBeingDebugged(nullptr);
-
 		UControlRigBlueprint::sCurrentlyOpenedRigBlueprints.Remove(RigBlueprint);
 
 		RigBlueprint->OnHierarchyModified().RemoveAll(this);
@@ -263,24 +150,6 @@ FControlRigEditor::~FControlRigEditor()
 		{
 			RigBlueprint->OnHierarchyModified().RemoveAll(EditMode);
 		}
-		RigBlueprint->OnRefreshEditor().RemoveAll(this);
-		RigBlueprint->OnVariableDropped().RemoveAll(this);
-		RigBlueprint->OnBreakpointAdded().RemoveAll(this);
-		RigBlueprint->OnNodeDoubleClicked().RemoveAll(this);
-		RigBlueprint->OnGraphImported().RemoveAll(this);
-		RigBlueprint->OnRequestLocalizeFunctionDialog().RemoveAll(this);
-		RigBlueprint->OnRequestBulkEditDialog().Unbind();
-		RigBlueprint->OnRequestBreakLinksDialog().Unbind();
-		RigBlueprint->OnRequestPinTypeSelectionDialog().Unbind();
-		RigBlueprint->OnRequestJumpToHyperlink().Unbind();
-		RigBlueprint->OnReportCompilerMessage().RemoveAll(this);
-
-#if WITH_EDITOR
-		RigBlueprint->SetDebugMode(false);
-		RigBlueprint->ClearBreakpoints();
-		SetHaltedNode(nullptr);
-		RigBlueprint->OnGetFocusedGraph().Unbind();
-#endif
 	}
 
 	if (UWorld* PreviewWorld = GetPersonaToolkit()->GetPreviewScene()->GetWorld())
@@ -296,11 +165,32 @@ FControlRigEditor::~FControlRigEditor()
 	}
 }
 
+UObject* FControlRigEditor::GetOuterForHost() const
+{
+	UControlRigSkeletalMeshComponent* EditorSkelComp = Cast<UControlRigSkeletalMeshComponent>(GetPersonaToolkit()->GetPreviewScene()->GetPreviewMeshComponent());
+	if(EditorSkelComp)
+	{
+		return EditorSkelComp;
+	}
+	return FRigVMEditor::GetOuterForHost();
+}
+
 UClass* FControlRigEditor::GetDetailWrapperClass() const
 {
 	return UControlRigWrapperObject::StaticClass();
 }
 
+void FControlRigEditor::CreateEmptyGraphContent(URigVMController* InController)
+{
+	URigVMNode* Node = InController->AddUnitNode(FRigUnit_BeginExecution::StaticStruct(), FRigUnit::GetMethodName(), FVector2D::ZeroVector, FString(), false);
+	if (Node)
+	{
+		TArray<FName> NodeNames;
+		NodeNames.Add(Node->GetFName());
+		InController->SetNodeSelection(NodeNames, false);
+	}
+}
+
 UControlRigBlueprint* FControlRigEditor::GetControlRigBlueprint() const
 {
 	return Cast<UControlRigBlueprint>(GetRigVMBlueprint());
@@ -324,116 +214,16 @@ URigHierarchy* FControlRigEditor::GetHierarchyBeingDebugged() const
 	return nullptr;
 }
 
-void FControlRigEditor::ExtendMenu()
-{
-	if(MenuExtender.IsValid())
-	{
-		RemoveMenuExtender(MenuExtender);
-		MenuExtender.Reset();
-	}
-
-	MenuExtender = MakeShareable(new FExtender);
-
-	AddMenuExtender(MenuExtender);
-
-	// add extensible menu if exists
-	FControlRigEditorModule& ControlRigEditorModule = FModuleManager::LoadModuleChecked<FControlRigEditorModule>("ControlRigEditor");
-	AddMenuExtender(ControlRigEditorModule.GetMenuExtensibilityManager()->GetAllExtenders(GetToolkitCommands(), GetEditingObjects()));
-}
-
-void FControlRigEditor::InitControlRigEditor(const EToolkitMode::Type Mode, const TSharedPtr< class IToolkitHost >& InitToolkitHost, UControlRigBlueprint* InControlRigBlueprint)
+void FControlRigEditor::InitRigVMEditor(const EToolkitMode::Type Mode, const TSharedPtr< class IToolkitHost >& InitToolkitHost, URigVMBlueprint* InRigVMBlueprint)
 {
 	DECLARE_SCOPE_HIERARCHICAL_COUNTER_FUNC()
 
-	check(InControlRigBlueprint);
-
-	FBlueprintCompilationManager::FlushCompilationQueue(nullptr);
-
-	FPersonaModule& PersonaModule = FModuleManager::GetModuleChecked<FPersonaModule>("Persona");
-
-	FPersonaToolkitArgs PersonaToolkitArgs;
-	PersonaToolkitArgs.OnPreviewSceneCreated = FOnPreviewSceneCreated::FDelegate::CreateSP(this, &FControlRigEditor::HandlePreviewSceneCreated);
-	PersonaToolkitArgs.bPreviewMeshCanUseDifferentSkeleton = true;
-	PersonaToolkit = PersonaModule.CreatePersonaToolkit(InControlRigBlueprint, PersonaToolkitArgs);
-
-	// set delegate prior to setting mesh
-	// otherwise, you don't get delegate
-	PersonaToolkit->GetPreviewScene()->RegisterOnPreviewMeshChanged(FOnPreviewMeshChanged::CreateSP(this, &FControlRigEditor::HandlePreviewMeshChanged));
-
-	// Set a default preview mesh, if any
-	PersonaToolkit->SetPreviewMesh(InControlRigBlueprint->GetPreviewMesh(), false);
-
-	Toolbox = SNew(SBorder)
-		.BorderImage(FAppStyle::GetBrush("ToolPanel.GroupBorder"))
-		.Padding(0.f);
-
-	if (!Toolbar.IsValid())
-	{
-		Toolbar = MakeShareable(new FBlueprintEditorToolbar(SharedThis(this)));
-	}
-
-	// Build up a list of objects being edited in this asset editor
-	TArray<UObject*> ObjectsBeingEdited;
-	ObjectsBeingEdited.Add(InControlRigBlueprint);
-
-	// Initialize the asset editor and spawn tabs
-	const bool bCreateDefaultStandaloneMenu = true;
-	const bool bCreateDefaultToolbar = true;
-	InitAssetEditor(Mode, InitToolkitHost, ControlRigEditorAppName, FTabManager::FLayout::NullLayout, bCreateDefaultStandaloneMenu, bCreateDefaultToolbar, ObjectsBeingEdited);
-
-	CreateDefaultCommands();
-
-	UControlRigBlueprint::sCurrentlyOpenedRigBlueprints.AddUnique(InControlRigBlueprint);
-
-	TArray<UBlueprint*> ControlRigBlueprints;
-	ControlRigBlueprints.Add(InControlRigBlueprint);
-
-	InControlRigBlueprint->InitializeModelIfRequired();
-
-	CommonInitialization(ControlRigBlueprints, false);
-	
-	// user-defined-struct can change even after load
-	// refresh the models such that pins are updated to match
-	// the latest struct member layout
-	InControlRigBlueprint->RefreshAllModels(ERigVMBlueprintLoadType::CheckUserDefinedStructs);
-
-	{
-		TArray<UEdGraph*> EdGraphs;
-		InControlRigBlueprint->GetAllGraphs(EdGraphs);
-
-		for (UEdGraph* Graph : EdGraphs)
-		{
-			UControlRigGraph* RigGraph = Cast<UControlRigGraph>(Graph);
-			if (RigGraph == nullptr)
-			{
-				continue;
-			}
-
-			RigGraph->InitializeFromBlueprint(InControlRigBlueprint);
-		}
-
-	}
-
-	InControlRigBlueprint->OnModified().AddSP(this, &FControlRigEditor::HandleModifiedEvent);
-	InControlRigBlueprint->OnVMCompiled().AddSP(this, &FControlRigEditor::HandleVMCompiledEvent);
-	InControlRigBlueprint->OnRequestInspectObject().AddSP(this, &FControlRigEditor::SetDetailObjects);
-
-	BindCommands();
+	FRigVMEditor::InitRigVMEditor(Mode, InitToolkitHost, InRigVMBlueprint);
 
-	AddApplicationMode(
-		FControlRigEditorModes::ControlRigEditorMode,
-		MakeShareable(new FControlRigEditorMode(SharedThis(this))));
+	UControlRigBlueprint* ControlRigBlueprint = CastChecked<UControlRigBlueprint>(InRigVMBlueprint);
 
-	ExtendMenu();
-	ExtendToolbar();
-	RegenerateMenusAndToolbars();
-
-	// Activate the initial mode (which will populate with a real layout)
-	SetCurrentMode(FControlRigEditorModes::ControlRigEditorMode);
-
-	// Activate our edit mode
-	GetEditorModeManager().SetDefaultMode(FControlRigEditorEditMode::ModeName);
-	GetEditorModeManager().ActivateMode(FControlRigEditorEditMode::ModeName);
+	CreatePersonaToolKitIfRequired();
+	UControlRigBlueprint::sCurrentlyOpenedRigBlueprints.AddUnique(ControlRigBlueprint);
 
 	if (FControlRigEditMode* EditMode = GetEditMode())
 	{
@@ -446,87 +236,10 @@ void FControlRigEditor::InitControlRigEditor(const EToolkitMode::Type Mode, cons
 		PersonaToolkit->GetPreviewScene()->SetRemoveAttachedComponentFilter(FOnRemoveAttachedComponentFilter::CreateSP(EditMode, &FControlRigEditMode::CanRemoveFromPreviewScene));
 	}
 
-	UpdateControlRig();
-
 	// Post-layout initialization
 	PostLayoutBlueprintEditorInitialization();
 
-	// tabs opened before reload
-	FString ActiveTabNodePath;
-	TArray<FString> OpenedTabNodePaths;
-
-	if (ControlRigBlueprints.Num() > 0)
-	{
-		bool bBroughtGraphToFront = false;
-		for(UEdGraph* Graph : ControlRigBlueprints[0]->UbergraphPages)
-		{
-			if (UControlRigGraph* RigGraph = Cast<UControlRigGraph>(Graph))
-			{
-				if (!bBroughtGraphToFront)
-				{
-					OpenGraphAndBringToFront(Graph, false);
-					bBroughtGraphToFront = true;
-				}
-
-				RigGraph->OnGraphNodeClicked.AddSP(this, &FControlRigEditor::OnGraphNodeClicked);
-				ActiveTabNodePath = RigGraph->ModelNodePath;
-			}
-		}
-	}
-
 	{
-		if (URigVMGraph* Model = InControlRigBlueprint->GetDefaultModel())
-		{
-			if (Model->GetNodes().Num() == 0)
-			{
-				URigVMNode* Node = InControlRigBlueprint->GetController()->AddUnitNode(FRigUnit_BeginExecution::StaticStruct(), FRigUnit::GetMethodName(), FVector2D::ZeroVector, FString(), false);
-				if (Node)
-				{
-					TArray<FName> NodeNames;
-					NodeNames.Add(Node->GetFName());
-					InControlRigBlueprint->GetController()->SetNodeSelection(NodeNames, false);
-				}
-			}
-			else
-			{
-				// remember all ed graphs which were visible as tabs
-				TArray<UEdGraph*> EdGraphs;
-				InControlRigBlueprint->GetAllGraphs(EdGraphs);
-
-				for (UEdGraph* EdGraph : EdGraphs)
-				{
-					if (UControlRigGraph* RigGraph = Cast<UControlRigGraph>(EdGraph))
-					{
-						TArray<TSharedPtr<SDockTab>> TabsForEdGraph;
-						FindOpenTabsContainingDocument(EdGraph, TabsForEdGraph);
-
-						if (TabsForEdGraph.Num() > 0)
-						{
-							OpenedTabNodePaths.Add(RigGraph->ModelNodePath);
-
-							if(RigGraph->bIsFunctionDefinition)
-							{
-								CloseDocumentTab(RigGraph);
-							}
-						}
-					}
-				}
-
-				InControlRigBlueprint->RebuildGraphFromModel();
-
-				// selection state does not need to be persistent, even though it is saved in the RigVM.
-				for (URigVMGraph* Graph : InControlRigBlueprint->GetAllModels())
-				{
-					InControlRigBlueprint->GetController(Graph)->ClearNodeSelection(false);
-				}
-
-				if (UPackage* Package = InControlRigBlueprint->GetOutermost())
-				{
-					Package->SetDirtyFlag(InControlRigBlueprint->IsMarkedDirtyDuringLoad());
-				}
-			}
-		}
-
 		// listening to the BP's event instead of BP's Hierarchy's Event ensure a propagation order of
 		// 1. Hierarchy change in BP
 		// 2. BP propagate to instances
@@ -539,337 +252,80 @@ void FControlRigEditor::InitControlRigEditor(const EToolkitMode::Type Mode, cons
 		// 3. BP performs propagation again in UControlRigBlueprint::HandleHierarchyModified, invalidates the rig element
 		//    that the UI is observing
 		// 4. Editor UI shows an invalid rig element
-		InControlRigBlueprint->OnHierarchyModified().AddSP(this, &FControlRigEditor::OnHierarchyModified);
-		
-		InControlRigBlueprint->OnRefreshEditor().AddSP(this, &FControlRigEditor::HandleRefreshEditorFromBlueprint);
-		InControlRigBlueprint->OnVariableDropped().AddSP(this, &FControlRigEditor::HandleVariableDroppedFromBlueprint);
-		InControlRigBlueprint->OnBreakpointAdded().AddSP(this, &FControlRigEditor::HandleBreakpointAdded);
+		ControlRigBlueprint->OnHierarchyModified().AddSP(this, &FControlRigEditor::OnHierarchyModified);
 
 		if (FControlRigEditMode* EditMode = GetEditMode())
 		{
-			InControlRigBlueprint->OnHierarchyModified().AddSP(EditMode, &FControlRigEditMode::OnHierarchyModified_AnyThread);
-		}
-
-		InControlRigBlueprint->OnNodeDoubleClicked().AddSP(this, &FControlRigEditor::OnNodeDoubleClicked);
-		InControlRigBlueprint->OnGraphImported().AddSP(this, &FControlRigEditor::OnGraphImported);
-		InControlRigBlueprint->OnRequestLocalizeFunctionDialog().AddSP(this, &FControlRigEditor::OnRequestLocalizeFunctionDialog);
-		InControlRigBlueprint->OnRequestBulkEditDialog().BindSP(this, &FControlRigEditor::OnRequestBulkEditDialog);
-		InControlRigBlueprint->OnRequestBreakLinksDialog().BindSP(this, &FControlRigEditor::OnRequestBreakLinksDialog);
-		InControlRigBlueprint->OnRequestPinTypeSelectionDialog().BindSP(this, &FControlRigEditor::OnRequestPinTypeSelectionDialog);
-		InControlRigBlueprint->OnRequestJumpToHyperlink().BindSP(this, &FControlRigEditor::HandleJumpToHyperlink);
-#if WITH_EDITOR
-		InControlRigBlueprint->OnGetFocusedGraph().BindSP(this, &FControlRigEditor::GetFocusedModel);
-#endif
-	}
-
-	UpdateStaleWatchedPins();
-
-	for (const FString& OpenedTabNodePath : OpenedTabNodePaths)
-	{
-		if (UEdGraph* EdGraph = InControlRigBlueprint->GetEdGraph(OpenedTabNodePath))
-		{
-			OpenDocument(EdGraph, FDocumentTracker::RestorePreviousDocument);
+			ControlRigBlueprint->OnHierarchyModified().AddSP(EditMode, &FControlRigEditMode::OnHierarchyModified_AnyThread);
 		}
 	}
 
-	if (UEdGraph* ActiveGraph = InControlRigBlueprint->GetEdGraph(ActiveTabNodePath))
-	{
-		OpenGraphAndBringToFront(ActiveGraph, true);
-	}
-
-	FRigVMBlueprintUtils::HandleRefreshAllNodes(InControlRigBlueprint);
-
-	bControlRigEditorInitialized = true;
+	CreateRigHierarchyToGraphDragAndDropMenu();
+}
 
-	if (ControlRigBlueprints.Num() > 0)
+void FControlRigEditor::CreatePersonaToolKitIfRequired()
+{
+	if(PersonaToolkit.IsValid())
 	{
-		if(ControlRigBlueprints[0]->Status == BS_Error)
-		{
-			Compile();
-		}
+		return;
 	}
 	
-	FFunctionGraphTask::CreateAndDispatchWhenReady([this, InControlRigBlueprint]()
-	{
-		// no need to do anything if the the CR is not opened anymore
-		// (i.e. destructor has been called before that task actually got a chance to start) 
-		if (!UControlRigBlueprint::sCurrentlyOpenedRigBlueprints.Contains(InControlRigBlueprint))
-		{
-			return;		
-		}
-		
-		TSharedPtr<FTabManager> TabManager = GetTabManager();
-		if (!TabManager.IsValid())
-		{
-			return;
-		}
-		
-		// Always show the myblueprint tab
-		static const FTabId MyBlueprintTabId(FBlueprintEditorTabs::MyBlueprintID);
-		if (!TabManager->FindExistingLiveTab(MyBlueprintTabId).IsValid())
-		{
-			TabManager->TryInvokeTab(MyBlueprintTabId);
-		}
-		
-	}, TStatId(), nullptr, ENamedThreads::GameThread);
-	
-	CreateRigHierarchyToGraphDragAndDropMenu();
+	UControlRigBlueprint* ControlRigBlueprint = GetControlRigBlueprint();
 
-#if WITH_EDITOR
-	FString BlueprintName = InControlRigBlueprint->GetPathName();
-	RigVMPythonUtils::PrintPythonContext(BlueprintName);
-#endif
-
-	TArray<UScriptStruct*> StructsToCustomize = {
-		TBaseStructure<FVector>::Get(),
-		TBaseStructure<FVector2D>::Get(),
-		TBaseStructure<FVector4>::Get(),
-		TBaseStructure<FRotator>::Get(),
-		TBaseStructure<FQuat>::Get(),
-		TBaseStructure<FTransform>::Get(),
-		TBaseStructure<FEulerTransform>::Get(),
-	};
+	FPersonaModule& PersonaModule = FModuleManager::GetModuleChecked<FPersonaModule>("Persona");
 
-	for(UScriptStruct* StructToCustomize : StructsToCustomize)
-	{
-		Inspector->GetPropertyView()->RegisterInstancedCustomPropertyTypeLayout(StructToCustomize->GetFName(),
-			FOnGetPropertyTypeCustomizationInstance::CreateLambda([=]()
-			{
-				return FRigVMGraphMathTypeDetailCustomization::MakeInstance();
-			}));
-	}
+	FPersonaToolkitArgs PersonaToolkitArgs;
+	PersonaToolkitArgs.OnPreviewSceneCreated = FOnPreviewSceneCreated::FDelegate::CreateSP(this, &FControlRigEditor::HandlePreviewSceneCreated);
+	PersonaToolkitArgs.bPreviewMeshCanUseDifferentSkeleton = true;
+	PersonaToolkit = PersonaModule.CreatePersonaToolkit(ControlRigBlueprint, PersonaToolkitArgs);
+
+	// set delegate prior to setting mesh
+	// otherwise, you don't get delegate
+	PersonaToolkit->GetPreviewScene()->RegisterOnPreviewMeshChanged(FOnPreviewMeshChanged::CreateSP(this, &FControlRigEditor::HandlePreviewMeshChanged));
 
-	PropertyChangedHandle = FCoreUObjectDelegates::OnObjectPropertyChanged.AddSP(this, &FControlRigEditor::OnPropertyChanged);
+	// Set a default preview mesh, if any
+	PersonaToolkit->SetPreviewMesh(ControlRigBlueprint->GetPreviewMesh(), false);
 }
 
-void FControlRigEditor::BindCommands()
+const FName FControlRigEditor::GetEditorAppName() const
 {
-	DECLARE_SCOPE_HIERARCHICAL_COUNTER_FUNC()
-
-	/*
-	GetToolkitCommands()->MapAction(
-		FControlRigBlueprintCommands::Get().ExecuteGraph,
-		FExecuteAction::CreateSP(this, &FControlRigEditor::ToggleExecuteGraph), 
-		FCanExecuteAction(), 
-		FIsActionChecked::CreateSP(this, &FControlRigEditor::IsExecuteGraphOn));
-	*/
-
-	GetToolkitCommands()->MapAction(
-		FControlRigBlueprintCommands::Get().AutoCompileGraph,
-		FExecuteAction::CreateSP(this, &FControlRigEditor::ToggleAutoCompileGraph), 
-		FIsActionChecked::CreateSP(this, &FControlRigEditor::CanAutoCompileGraph),
-		FIsActionChecked::CreateSP(this, &FControlRigEditor::IsAutoCompileGraphOn));
-
-	GetToolkitCommands()->MapAction(
-		FControlRigBlueprintCommands::Get().ToggleEventQueue,
-		FExecuteAction::CreateSP(this, &FControlRigEditor::ToggleEventQueue),
-		FCanExecuteAction());
-
-	GetToolkitCommands()->MapAction(
-		FControlRigBlueprintCommands::Get().ConstructionEvent,
-		FExecuteAction::CreateSP(this, &FControlRigEditor::SetEventQueue, TArray<FName>(ConstructionEventQueue)),
-		FCanExecuteAction());
-
-	GetToolkitCommands()->MapAction(
-		FControlRigBlueprintCommands::Get().ForwardsSolveEvent,
-		FExecuteAction::CreateSP(this, &FControlRigEditor::SetEventQueue, TArray<FName>(ForwardsSolveEventQueue)),
-		FCanExecuteAction());
-
-	GetToolkitCommands()->MapAction(
-		FControlRigBlueprintCommands::Get().BackwardsSolveEvent,
-		FExecuteAction::CreateSP(this, &FControlRigEditor::SetEventQueue, TArray<FName>(BackwardsSolveEventQueue)),
-		FCanExecuteAction());
-
-	GetToolkitCommands()->MapAction(
-		FControlRigBlueprintCommands::Get().BackwardsAndForwardsSolveEvent,
-		FExecuteAction::CreateSP(this, &FControlRigEditor::SetEventQueue, TArray<FName>(BackwardsAndForwardsSolveEventQueue)),
-		FCanExecuteAction());
-
-	GetToolkitCommands()->MapAction(
-		FControlRigBlueprintCommands::Get().ToggleExecutionMode,
-		FExecuteAction::CreateSP(this, &FControlRigEditor::ToggleExecutionMode),
-		FCanExecuteAction());
-
-	GetToolkitCommands()->MapAction(
-		FControlRigBlueprintCommands::Get().ReleaseMode,
-		FExecuteAction::CreateSP(this, &FControlRigEditor::SetExecutionMode, EControlRigExecutionModeType_Release),
-		FCanExecuteAction());
-
-	GetToolkitCommands()->MapAction(
-		FControlRigBlueprintCommands::Get().DebugMode,
-		FExecuteAction::CreateSP(this, &FControlRigEditor::SetExecutionMode, EControlRigExecutionModeType_Debug),
-		FCanExecuteAction());
-
-	GetToolkitCommands()->MapAction(
-		FControlRigBlueprintCommands::Get().ResumeExecution,
-		FExecuteAction::CreateSP(this, &FControlRigEditor::HandleBreakpointActionRequested, ERigVMBreakpointAction::Resume),			
-		FIsActionChecked::CreateSP(this, &FControlRigEditor::IsHaltedAtBreakpoint));
-
-	GetToolkitCommands()->MapAction(
-		FControlRigBlueprintCommands::Get().ShowCurrentStatement,
-		FExecuteAction::CreateSP(this, &FControlRigEditor::HandleShowCurrentStatement),
-		FIsActionChecked::CreateSP(this, &FControlRigEditor::IsHaltedAtBreakpoint));
-
-	GetToolkitCommands()->MapAction(
-		FControlRigBlueprintCommands::Get().StepOver,
-		FExecuteAction::CreateSP(this, &FControlRigEditor::HandleBreakpointActionRequested, ERigVMBreakpointAction::StepOver),
-		FIsActionChecked::CreateSP(this, &FControlRigEditor::IsHaltedAtBreakpoint));
-
-	GetToolkitCommands()->MapAction(
-		FControlRigBlueprintCommands::Get().StepInto,
-		FExecuteAction::CreateSP(this, &FControlRigEditor::HandleBreakpointActionRequested, ERigVMBreakpointAction::StepInto),
-		FIsActionChecked::CreateSP(this, &FControlRigEditor::IsHaltedAtBreakpoint));
-
-	GetToolkitCommands()->MapAction(
-		FControlRigBlueprintCommands::Get().StepOut,
-		FExecuteAction::CreateSP(this, &FControlRigEditor::HandleBreakpointActionRequested, ERigVMBreakpointAction::StepOut),
-		FIsActionChecked::CreateSP(this, &FControlRigEditor::IsHaltedAtBreakpoint));
-
-	GetToolkitCommands()->MapAction(
-		FControlRigBlueprintCommands::Get().FrameSelection,
-		FExecuteAction::CreateSP(this, &FControlRigEditor::FrameSelection),
-		FCanExecuteAction());
+	static const FName ControlRigEditorAppName(TEXT("ControlRigEditorApp"));
+	return ControlRigEditorAppName;
 }
 
-void FControlRigEditor::ToggleExecuteGraph()
+const FName FControlRigEditor::GetEditorModeName() const
 {
-	if (GetControlRig())
-	{
-		bExecutionControlRig = !bExecutionControlRig;
-
-		// This is required now since we update execution/input flag on update controlrig
-		// @fixme: change this to just change flag instead of updating whole control rig
-		// I'll do this before first check-in
-		UpdateControlRig();
-	}
+	return FControlRigEditorEditMode::ModeName;
 }
 
-bool FControlRigEditor::IsExecuteGraphOn() const
+TSharedPtr<FApplicationMode> FControlRigEditor::CreateEditorMode()
 {
-	return bExecutionControlRig;
+	CreatePersonaToolKitIfRequired();
+	return MakeShareable(new FControlRigEditorMode(SharedThis(this)));
 }
 
-void FControlRigEditor::ToggleAutoCompileGraph()
+FText FControlRigEditor::GetTestAssetName() const
 {
-	if (UControlRigBlueprint* RigBlueprint = Cast<UControlRigBlueprint>(GetBlueprintObj()))
+	if(TestDataStrongPtr.IsValid())
 	{
-		RigBlueprint->SetAutoVMRecompile(!RigBlueprint->GetAutoVMRecompile());
-		if (RigBlueprint->GetAutoVMRecompile())
-		{
-			RigBlueprint->RequestAutoVMRecompilation();
-		}
+		return FText::FromString(TestDataStrongPtr->GetName());
 	}
+
+	static const FText NoTestAsset = LOCTEXT("NoTestAsset", "No Test Asset");
+	return NoTestAsset;
 }
 
-bool FControlRigEditor::IsAutoCompileGraphOn() const
+FText FControlRigEditor::GetTestAssetTooltip() const
 {
-	if (UControlRigBlueprint* RigBlueprint = Cast<UControlRigBlueprint>(GetBlueprintObj()))
+	if(TestDataStrongPtr.IsValid())
 	{
-		return RigBlueprint->GetAutoVMRecompile();
+		return FText::FromString(TestDataStrongPtr->GetPathName());
 	}
-	return false;
-}
-
-void FControlRigEditor::ToggleEventQueue()
-{
-	SetEventQueue(LastEventQueue);
+	static const FText NoTestAssetTooltip = LOCTEXT("NoTestAssetTooltip", "Click the record button to the left to record a new frame");
+	return NoTestAssetTooltip;
 }
 
-void FControlRigEditor::ToggleExecutionMode()
-{
-	SetExecutionMode((ExecutionMode == EControlRigExecutionModeType_Debug) ?
-		EControlRigExecutionModeType_Release
-		: EControlRigExecutionModeType_Debug);
-}
-
-TSharedRef<SWidget> FControlRigEditor::GenerateEventQueueMenuContent()
-{
-	FMenuBuilder MenuBuilder(true, GetToolkitCommands());
-
-	MenuBuilder.BeginSection(TEXT("Events"));
-	MenuBuilder.AddMenuEntry(FControlRigBlueprintCommands::Get().ConstructionEvent, TEXT("Setup"), TAttribute<FText>(), TAttribute<FText>(), GetEventQueueIcon(ConstructionEventQueue));
-	MenuBuilder.AddMenuEntry(FControlRigBlueprintCommands::Get().ForwardsSolveEvent, TEXT("Forwards Solve"), TAttribute<FText>(), TAttribute<FText>(), GetEventQueueIcon(ForwardsSolveEventQueue));
-	MenuBuilder.AddMenuEntry(FControlRigBlueprintCommands::Get().BackwardsSolveEvent, TEXT("Backwards Solve"), TAttribute<FText>(), TAttribute<FText>(), GetEventQueueIcon(BackwardsSolveEventQueue));
-	MenuBuilder.EndSection();
-
-	MenuBuilder.BeginSection(TEXT("Validation"));
-	MenuBuilder.AddMenuEntry(FControlRigBlueprintCommands::Get().BackwardsAndForwardsSolveEvent, TEXT("BackwardsAndForwards"), TAttribute<FText>(), TAttribute<FText>(), GetEventQueueIcon(BackwardsAndForwardsSolveEventQueue));
-	MenuBuilder.EndSection();
-
-	if (const UControlRigBlueprint* RigBlueprint = Cast<UControlRigBlueprint>(GetBlueprintObj()))
-	{
-		URigVMEdGraphSchema* Schema = CastChecked<URigVMEdGraphSchema>(RigBlueprint->GetRigVMEdGraphSchemaClass()->GetDefaultObject());
-		
-		bool bFoundUserDefinedEvent = false;
-		const TArray<FName> EntryNames = RigBlueprint->GetRigVMClient()->GetEntryNames();
-		for(const FName& EntryName : EntryNames)
-		{
-			if(Schema->IsRigVMDefaultEvent(EntryName))
-			{
-				continue;
-			}
-
-			if(!bFoundUserDefinedEvent)
-			{
-				MenuBuilder.AddSeparator();
-				bFoundUserDefinedEvent = true;
-			}
-
-			FString EventNameStr = EntryName.ToString();
-			if(!EventNameStr.EndsWith(TEXT("Event")))
-			{
-				EventNameStr += TEXT(" Event");
-			}
-
-			MenuBuilder.AddMenuEntry(
-				FText::FromString(EventNameStr),
-				FText::FromString(FString::Printf(TEXT("Runs the user defined %s"), *EventNameStr)),
-				GetEventQueueIcon({EntryName}),
-				FUIAction(
-					FExecuteAction::CreateSP(this, &FControlRigEditor::SetEventQueue, TArray<FName>({EntryName})),
-					FCanExecuteAction()
-				)
-			);
-		}
-	}
-
-	return MenuBuilder.MakeWidget();
-}
-
-TSharedRef<SWidget> FControlRigEditor::GenerateExecutionModeMenuContent()
-{
-	FMenuBuilder MenuBuilder(true, GetToolkitCommands());
-
-	MenuBuilder.BeginSection(TEXT("Events"));
-	MenuBuilder.AddMenuEntry(FControlRigBlueprintCommands::Get().ReleaseMode, TEXT("Release"), TAttribute<FText>(), TAttribute<FText>(), GetExecutionModeIcon(EControlRigExecutionModeType_Release));
-	MenuBuilder.AddMenuEntry(FControlRigBlueprintCommands::Get().DebugMode, TEXT("Debug"), TAttribute<FText>(), TAttribute<FText>(), GetExecutionModeIcon(EControlRigExecutionModeType_Debug));
-	MenuBuilder.EndSection();
-
-	return MenuBuilder.MakeWidget();
-}
-
-FText FControlRigEditor::GetTestAssetName() const
-{
-	if(TestDataStrongPtr.IsValid())
-	{
-		return FText::FromString(TestDataStrongPtr->GetName());
-	}
-
-	static const FText NoTestAsset = LOCTEXT("NoTestAsset", "No Test Asset");
-	return NoTestAsset;
-}
-
-FText FControlRigEditor::GetTestAssetTooltip() const
-{
-	if(TestDataStrongPtr.IsValid())
-	{
-		return FText::FromString(TestDataStrongPtr->GetPathName());
-	}
-	static const FText NoTestAssetTooltip = LOCTEXT("NoTestAssetTooltip", "Click the record button to the left to record a new frame");
-	return NoTestAssetTooltip;
-}
-
-bool FControlRigEditor::SetTestAssetPath(const FString& InAssetPath)
+bool FControlRigEditor::SetTestAssetPath(const FString& InAssetPath)
 {
 	if(TestDataStrongPtr.IsValid())
 	{
@@ -1055,111 +511,11 @@ void FControlRigEditor::ToggleTestData()
 	}
 }
 
-void FControlRigEditor::ExtendToolbar()
+void FControlRigEditor::FillToolbar(FToolBarBuilder& ToolbarBuilder, bool bEndSection)
 {
-	DECLARE_SCOPE_HIERARCHICAL_COUNTER_FUNC()
-
-	// If the ToolbarExtender is valid, remove it before rebuilding it
-	if(ToolbarExtender.IsValid())
-	{
-		RemoveToolbarExtender(ToolbarExtender);
-		ToolbarExtender.Reset();
-	}
-
-	ToolbarExtender = MakeShareable(new FExtender);
-
-	AddToolbarExtender(ToolbarExtender);
-
-	FControlRigEditorModule& ControlRigEditorModule = FModuleManager::LoadModuleChecked<FControlRigEditorModule>("ControlRigEditor");
-	AddToolbarExtender(ControlRigEditorModule.GetToolBarExtensibilityManager()->GetAllExtenders(GetToolkitCommands(), GetEditingObjects()));
-
-	TArray<FRigVMEditorModule::FRigVMEditorToolbarExtender> ToolbarExtenderDelegates = ControlRigEditorModule.GetAllRigVMEditorToolbarExtenders();
-
-	for (auto& ToolbarExtenderDelegate : ToolbarExtenderDelegates)
+	FRigVMEditor::FillToolbar(ToolbarBuilder, false);
+	
 	{
-		if(ToolbarExtenderDelegate.IsBound())
-		{
-			AddToolbarExtender(ToolbarExtenderDelegate.Execute(GetToolkitCommands(), SharedThis(this)));
-		}
-	}
-
-	ToolbarExtender->AddToolBarExtension(
-		"Asset",
-		EExtensionHook::After,
-		GetToolkitCommands(),
-		FToolBarExtensionDelegate::CreateSP(this, &FControlRigEditor::FillToolbar)
-	);
-}
-
-void FControlRigEditor::FillToolbar(FToolBarBuilder& ToolbarBuilder)
-{
-	ToolbarBuilder.BeginSection("Toolbar");
-	{
-		/*
-		ToolbarBuilder.AddToolBarButton(FControlRigBlueprintCommands::Get().ExecuteGraph,
-			NAME_None, TAttribute<FText>(), TAttribute<FText>(), FSlateIcon(FRigVMEditorStyle::Get().GetStyleSetName(), "RigVM.ExecuteGraph"));
-		*/
-
-		ToolbarBuilder.AddToolBarButton(
-			FControlRigBlueprintCommands::Get().ToggleEventQueue,
-			NAME_None, 
-			TAttribute<FText>::Create(TAttribute<FText>::FGetter::CreateSP(this, &FControlRigEditor::GetEventQueueLabel)),
-			TAttribute<FText>(), 
-			TAttribute<FSlateIcon>::Create(TAttribute<FSlateIcon>::FGetter::CreateSP(this, &FControlRigEditor::GetEventQueueIcon))
-		);
-
-		FUIAction DefaultAction;
-		ToolbarBuilder.AddComboButton(
-			DefaultAction,
-			FOnGetContent::CreateSP(this, &FControlRigEditor::GenerateEventQueueMenuContent),
-			LOCTEXT("EventQueue_Label", "Available Events"),
-			LOCTEXT("EventQueue_ToolTip", "Pick between different events / modes for testing the Control Rig"),
-			FSlateIcon(FAppStyle::GetAppStyleSetName(), "LevelEditor.Recompile"),
-			true);
-
-		ToolbarBuilder.AddToolBarButton(FControlRigBlueprintCommands::Get().AutoCompileGraph,
-			NAME_None, TAttribute<FText>(), TAttribute<FText>(), FSlateIcon(FRigVMEditorStyle::Get().GetStyleSetName(), "RigVM.AutoCompileGraph"));
-
-		ToolbarBuilder.AddWidget(SNew(SBlueprintEditorSelectedDebugObjectWidget, SharedThis(this)));
-
-		ToolbarBuilder.AddSeparator();
-		ToolbarBuilder.AddToolBarButton(
-			FControlRigBlueprintCommands::Get().ToggleExecutionMode,
-			NAME_None, 
-			TAttribute<FText>::Create(TAttribute<FText>::FGetter::CreateSP(this, &FControlRigEditor::GetExecutionModeLabel)),
-			TAttribute<FText>(), 
-			TAttribute<FSlateIcon>::Create(TAttribute<FSlateIcon>::FGetter::CreateSP(this, &FControlRigEditor::GetExecutionModeIcon))
-		);
-		
-		FUIAction DefaultExecutionMode;
-		ToolbarBuilder.AddComboButton(
-			DefaultExecutionMode,
-			FOnGetContent::CreateSP(this, &FControlRigEditor::GenerateExecutionModeMenuContent),
-			LOCTEXT("ExecutionMode_Label", "Execution Modes"),
-			LOCTEXT("ExecutionMode_ToolTip", "Pick between different execution modes for testing the Control Rig"),
-			FSlateIcon(FAppStyle::GetAppStyleSetName(), "LevelEditor.Recompile"),
-			true);
-
-		ToolbarBuilder.BeginStyleOverride(FName("Toolbar.BackplateLeftPlay"));
-		ToolbarBuilder.AddToolBarButton(FControlRigBlueprintCommands::Get().ResumeExecution,
-			NAME_None, TAttribute<FText>(), TAttribute<FText>(), FSlateIcon(FAppStyle::GetAppStyleSetName(), "PlayWorld.ResumePlaySession"));
-
-		ToolbarBuilder.BeginStyleOverride(FName("Toolbar.BackplateLeft"));
-		ToolbarBuilder.AddToolBarButton(FControlRigBlueprintCommands::Get().ShowCurrentStatement,
-			NAME_None, TAttribute<FText>(), TAttribute<FText>(), FSlateIcon(FAppStyle::GetAppStyleSetName(), "PlayWorld.ShowCurrentStatement"));
-
-		ToolbarBuilder.BeginStyleOverride(FName("Toolbar.BackplateCenter"));
-		ToolbarBuilder.AddToolBarButton(FControlRigBlueprintCommands::Get().StepOver,
-			NAME_None, TAttribute<FText>(), TAttribute<FText>(), FSlateIcon(FAppStyle::GetAppStyleSetName(), "PlayWorld.StepOver"));
-		ToolbarBuilder.AddToolBarButton(FControlRigBlueprintCommands::Get().StepInto,
-			NAME_None, TAttribute<FText>(), TAttribute<FText>(), FSlateIcon(FAppStyle::GetAppStyleSetName(), "PlayWorld.StepInto"));
-
-		ToolbarBuilder.BeginStyleOverride(FName("Toolbar.BackplateRight"));
-		ToolbarBuilder.AddToolBarButton(FControlRigBlueprintCommands::Get().StepOut,
-			NAME_None, TAttribute<FText>(), TAttribute<FText>(), FSlateIcon(FAppStyle::GetAppStyleSetName(), "PlayWorld.StepOut"));
-
-		ToolbarBuilder.EndStyleOverride();
-
 		if(CVarControlRigShowTestingToolbar.GetValueOnAnyThread())
 		{
 			ToolbarBuilder.AddSeparator();
@@ -1286,22 +642,16 @@ void FControlRigEditor::FillToolbar(FToolBarBuilder& ToolbarBuilder)
 			);
 		}
 	}
-	ToolbarBuilder.EndSection();
-}
-
-TArray<FName> FControlRigEditor::GetEventQueue() const
-{
-	if (const UControlRig* ControlRig = GetControlRig())
+	
+	if(bEndSection)
 	{
-		return ControlRig->GetEventQueue();
+		ToolbarBuilder.EndSection();
 	}
-
-	return ForwardsSolveEventQueue;
 }
 
-void FControlRigEditor::SetEventQueue(TArray<FName> InEventQueue)
+TArray<FName> FControlRigEditor::GetDefaultEventQueue() const
 {
-	return SetEventQueue(InEventQueue, false);
+	return ForwardsSolveEventQueue;
 }
 
 void FControlRigEditor::SetEventQueue(TArray<FName> InEventQueue, bool bCompile)
@@ -1311,10 +661,6 @@ void FControlRigEditor::SetEventQueue(TArray<FName> InEventQueue, bool bCompile)
 		return;
 	}
 
-	LastEventQueue = GetEventQueue();
-
-	SetHaltedNode(nullptr);
-
 	TArray<FRigElementKey> PreviousSelection;
 	if (UControlRigBlueprint* RigBlueprint = Cast<UControlRigBlueprint>(GetBlueprintObj()))
 	{
@@ -1339,12 +685,12 @@ void FControlRigEditor::SetEventQueue(TArray<FName> InEventQueue, bool bCompile)
 		}
 	}
 
+	FRigVMEditor::SetEventQueue(InEventQueue, bCompile);
+
 	if (UControlRig* ControlRig = GetControlRig())
 	{
 		if (InEventQueue.Num() > 0)
 		{
-			ControlRig->SetEventQueue(InEventQueue);
-
 			if (UControlRigBlueprint* RigBlueprint = Cast<UControlRigBlueprint>(GetBlueprintObj()))
 			{
 				RigBlueprint->Validator->SetControlRig(ControlRig);
@@ -1399,7 +745,7 @@ int32 FControlRigEditor::GetEventQueueComboValue() const
 	{
 		return 3;
 	}
-	return INDEX_NONE;
+	return FRigVMEditor::GetEventQueueComboValue();
 }
 
 FText FControlRigEditor::GetEventQueueLabel() const
@@ -1434,201 +780,37 @@ FText FControlRigEditor::GetEventQueueLabel() const
 		}
 		return FText::FromString(EventName);
 	}
+	
 	return LOCTEXT("CustomEventQueue", "Custom Event Queue");
 }
 
-FSlateIcon FControlRigEditor::GetEventQueueIcon(const TArray<FName>& InEventQueue)
+FSlateIcon FControlRigEditor::GetEventQueueIcon(const TArray<FName>& InEventQueue) const
 {
 	if(InEventQueue == ConstructionEventQueue)
 	{
-		return FSlateIcon(FControlRigEditorStyle::Get().GetStyleSetName(), "ControlRig.ConstructionMode");
+		return FSlateIcon(FRigVMEditorStyle::Get().GetStyleSetName(), "ControlRig.ConstructionMode");
 	}
 	if(InEventQueue == ForwardsSolveEventQueue)
 	{
-		return FSlateIcon(FControlRigEditorStyle::Get().GetStyleSetName(), "ControlRig.ForwardsSolveEvent");
+		return FSlateIcon(FRigVMEditorStyle::Get().GetStyleSetName(), "ControlRig.ForwardsSolveEvent");
 	}
 	if(InEventQueue == BackwardsSolveEventQueue)
 	{
-		return FSlateIcon(FControlRigEditorStyle::Get().GetStyleSetName(), "ControlRig.BackwardsSolveEvent");
+		return FSlateIcon(FRigVMEditorStyle::Get().GetStyleSetName(), "ControlRig.BackwardsSolveEvent");
 	}
 	if(InEventQueue == BackwardsAndForwardsSolveEventQueue)
 	{
-		return FSlateIcon(FControlRigEditorStyle::Get().GetStyleSetName(), "ControlRig.BackwardsAndForwardsSolveEvent");
+		return FSlateIcon(FRigVMEditorStyle::Get().GetStyleSetName(), "ControlRig.BackwardsAndForwardsSolveEvent");
 	}
 
 	return FSlateIcon(FAppStyle::GetAppStyleSetName(), "GraphEditor.Event_16x");
 }
 
-FSlateIcon FControlRigEditor::GetEventQueueIcon() const
-{
-	return GetEventQueueIcon(GetEventQueue());
-}
-
-void FControlRigEditor::SetExecutionMode(const EControlRigExecutionModeType InExecutionMode)
-{
-	if (ExecutionMode == InExecutionMode)
-	{
-		return;
-	}
-
-	ExecutionMode = InExecutionMode;
-	GetControlRigBlueprint()->SetDebugMode(InExecutionMode == EControlRigExecutionModeType_Debug);
-	Compile();
-	
-	if (UControlRig* ControlRig = GetControlRig())
-	{
-		ControlRig->SetIsInDebugMode(InExecutionMode == EControlRigExecutionModeType_Debug);
-	}
-
-	SetHaltedNode(nullptr);
-	
-	RefreshDetailView();
-}
-
-int32 FControlRigEditor::GetExecutionModeComboValue() const
-{
-	return (int32) ExecutionMode;
-}
-
-FText FControlRigEditor::GetExecutionModeLabel() const
-{
-	if (ExecutionMode == EControlRigExecutionModeType_Debug)
-	{
-		return FText::FromString(TEXT("DebugMode"));
-	}
-	return FText::FromString(TEXT("ReleaseMode"));
-}
-
-FSlateIcon FControlRigEditor::GetExecutionModeIcon(EControlRigExecutionModeType InExecutionMode)
-{
-	if (InExecutionMode == EControlRigExecutionModeType_Debug)
-	{
-		return FSlateIcon(FAppStyle::GetAppStyleSetName(), "LevelEditor.Tabs.Debug");
-	}
-	return FSlateIcon(FRigVMEditorStyle::Get().GetStyleSetName(), "RigVM.ReleaseMode");
-}
-
-FSlateIcon FControlRigEditor::GetExecutionModeIcon() const
-{
-	return GetExecutionModeIcon(ExecutionMode);
-}
-
-void FControlRigEditor::GetCustomDebugObjects(TArray<FCustomDebugObject>& DebugList) const
-{
-	UControlRigBlueprint* RigBlueprint = Cast<UControlRigBlueprint>(GetBlueprintObj());
-	if (RigBlueprint == nullptr)
-	{
-		return;
-	}
-
-	if (UControlRig* ControlRig = GetControlRig())
-	{
-		FCustomDebugObject DebugObject;
-		DebugObject.Object = ControlRig;
-		DebugObject.NameOverride = GetCustomDebugObjectLabel(ControlRig);
-		DebugList.Add(DebugObject);
-	}
-
-	URigVMBlueprintGeneratedClass* GeneratedClass = RigBlueprint->GetRigVMBlueprintGeneratedClass();
-	if (GeneratedClass)
-	{
-		struct Local
-		{
-			static bool IsPendingKillOrUnreachableRecursive(UObject* InObject)
-			{
-				if (InObject != nullptr)
-				{
-					if (!IsValidChecked(InObject) || InObject->IsUnreachable())
-					{
-						return true;
-					}
-					return IsPendingKillOrUnreachableRecursive(InObject->GetOuter());
-				}
-				return false;
-			}
-
-			static bool OuterNameContainsRecursive(UObject* InObject, const FString& InStringToSearch)
-			{
-				if (InObject == nullptr)
-				{
-					return false;
-				}
-
-				UObject* InObjectOuter = InObject->GetOuter();
-				if (InObjectOuter == nullptr)
-				{
-					return false;
-				}
-
-				if (InObjectOuter->GetName().Contains(InStringToSearch))
-				{
-					return true;
-				}
-
-				return OuterNameContainsRecursive(InObjectOuter, InStringToSearch);
-			}
-		};
-
-		if (UObject* DefaultObject = GeneratedClass->GetDefaultObject(false))
-		{
-			TArray<UObject*> ArchetypeInstances;
-			DefaultObject->GetArchetypeInstances(ArchetypeInstances);
-
-			for (UObject* Instance : ArchetypeInstances)
-			{
-				UControlRig* InstanceControlRig = Cast<UControlRig>(Instance);
-				if (InstanceControlRig && InstanceControlRig != GetControlRig())
-				{
-					if (InstanceControlRig->GetOuter() == nullptr)
-					{
-						continue;
-					}
-
-					UWorld* World = InstanceControlRig->GetWorld();
-					if (World == nullptr)
-					{
-						continue;
-					}
-
-					// ensure to only allow preview actors in preview worlds
-					if (World->IsPreviewWorld())
-					{
-						if (!Local::OuterNameContainsRecursive(InstanceControlRig, TEXT("Preview")))
-						{
-							continue;
-						}
-					}
-
-					if (Local::IsPendingKillOrUnreachableRecursive(InstanceControlRig))
-					{
-						continue;
-					}
-
-					FCustomDebugObject DebugObject;
-					DebugObject.Object = InstanceControlRig;
-					DebugObject.NameOverride = GetCustomDebugObjectLabel(InstanceControlRig);
-					DebugList.Add(DebugObject);
-				}
-			}
-		}
-	}
-}
-
 void FControlRigEditor::HandleSetObjectBeingDebugged(UObject* InObject)
 {
+	FRigVMEditor::HandleSetObjectBeingDebugged(InObject);
+	
 	UControlRig* DebuggedControlRig = Cast<UControlRig>(InObject);
-
-	if (DebuggedControlRig == nullptr)
-	{
-		// fall back to our default control rig (which still can be nullptr)
-		if (GetControlRig() != nullptr && GetBlueprintObj() && !bIsSettingObjectBeingDebugged)
-		{
-			TGuardValue<bool> GuardSettingObjectBeingDebugged(bIsSettingObjectBeingDebugged, true);
-			GetBlueprintObj()->SetObjectBeingDebugged(GetControlRig());
-			return;
-		}
-	}
-
 	if(UControlRig* PreviouslyDebuggedControlRig = Cast<UControlRig>(GetBlueprintObj()->GetObjectBeingDebugged()))
 	{
 		if(!PreviouslyDebuggedControlRig->HasAnyFlags(RF_BeginDestroyed))
@@ -1641,16 +823,6 @@ void FControlRigEditor::HandleSetObjectBeingDebugged(UObject* InObject)
 
 	if (UControlRigBlueprint* RigBlueprint = Cast<UControlRigBlueprint>(GetBlueprintObj()))
 	{
-		if (URigVMBlueprintGeneratedClass* GeneratedClass = RigBlueprint->GetRigVMBlueprintGeneratedClass())
-		{
-			UControlRig* CDO = Cast<UControlRig>(GeneratedClass->GetDefaultObject(true /* create if needed */));
-			if (CDO->VM->GetInstructions().Num() <= 1 /* only exit */)
-			{
-				RigBlueprint->RecompileVM();
-				RigBlueprint->RequestRigVMInit();
-			}
-		}
-
 		RigBlueprint->Validator->SetControlRig(DebuggedControlRig);
 	}
 
@@ -1658,7 +830,6 @@ void FControlRigEditor::HandleSetObjectBeingDebugged(UObject* InObject)
 	{
 		bool bIsExternalControlRig = DebuggedControlRig != GetControlRig();
 		bool bShouldExecute = (!bIsExternalControlRig) && bExecutionControlRig;
-		DebuggedControlRig->RigVMLog = &ControlRigLog;
 		GetControlRigBlueprint()->Hierarchy->HierarchyForSelectionPtr = DebuggedControlRig->DynamicHierarchy;
 
 		UControlRigSkeletalMeshComponent* EditorSkelComp = Cast<UControlRigSkeletalMeshComponent>(GetPersonaToolkit()->GetPreviewScene()->GetPreviewMeshComponent());
@@ -1714,366 +885,291 @@ void FControlRigEditor::HandleSetObjectBeingDebugged(UObject* InObject)
 			EditMode->SetObjects(nullptr,  nullptr,nullptr);
 		}
 	}
-
-	RefreshDetailView();
-	LastDebuggedRig = GetCustomDebugObjectLabel(DebuggedControlRig);
 }
 
-FString FControlRigEditor::GetCustomDebugObjectLabel(UObject* ObjectBeingDebugged) const
+void FControlRigEditor::SetDetailViewForRigElements()
 {
-	if (ObjectBeingDebugged == nullptr)
-	{
-		return FString();
-	}
-
-	if (ObjectBeingDebugged == GetControlRig())
+	if(IsDetailsPanelRefreshSuspended())
 	{
-		return TEXT("Control Rig Editor Preview");
+		return;
 	}
 
-	if (AActor* ParentActor = ObjectBeingDebugged->GetTypedOuter<AActor>())
-	{
-		return FString::Printf(TEXT("%s in %s"), *GetBlueprintObj()->GetName(), *ParentActor->GetName());
-	}
+	ClearDetailObject();
 
-	return GetBlueprintObj()->GetName();
-}
+	UControlRigBlueprint* RigBlueprint = Cast<UControlRigBlueprint>(GetBlueprintObj());
+	URigHierarchy* HierarchyBeingDebugged = GetHierarchyBeingDebugged();
+	TArray<UObject*> Objects;
 
-UBlueprint* FControlRigEditor::GetBlueprintObj() const
-{
-	const TArray<UObject*>& EditingObjs = GetEditingObjects();
-	for (UObject* Obj : EditingObjs)
+	TArray<FRigElementKey> CurrentSelection = HierarchyBeingDebugged->GetSelectedKeys();
+	for(const FRigElementKey& SelectedKey : CurrentSelection)
 	{
-		if (Obj->IsA<UControlRigBlueprint>()) 
+		FRigBaseElement* Element = HierarchyBeingDebugged->Find(SelectedKey);
+		if (Element == nullptr)
 		{
-			return (UBlueprint*)Obj;
+			continue;
 		}
-	}
-	return nullptr;
-}
 
-TSubclassOf<UEdGraphSchema> FControlRigEditor::GetDefaultSchemaClass() const
-{
-	return UControlRigGraphSchema::StaticClass();
+		URigVMDetailsViewWrapperObject* WrapperObject = URigVMDetailsViewWrapperObject::MakeInstance(GetDetailWrapperClass(), GetBlueprintObj(), Element->GetElementStruct(), (uint8*)Element, HierarchyBeingDebugged);
+		WrapperObject->GetWrappedPropertyChangedChainEvent().AddSP(this, &FControlRigEditor::OnWrappedPropertyChangedChainEvent);
+		WrapperObject->AddToRoot();
+
+		Objects.Add(WrapperObject);
+	}
+	
+	SetDetailObjects(Objects);
 }
 
-class FMemoryTypeMetaData : public ISlateMetaData
+void FControlRigEditor::RefreshDetailView()
 {
-public:
-	SLATE_METADATA_TYPE(FMemoryTypeMetaData, ISlateMetaData)
-
-		FMemoryTypeMetaData(ERigVMMemoryType InMemoryType)
-		: MemoryType(InMemoryType)
+	if(DetailViewShowsAnyRigElement())
 	{
+		SetDetailViewForRigElements();
+		return;
 	}
-	ERigVMMemoryType MemoryType;
-};
 
-void FControlRigEditor::SetDetailObjects(const TArray<UObject*>& InObjects)
-{
-	SetDetailObjects(InObjects, true);
+	FRigVMEditor::RefreshDetailView();
 }
 
-void FControlRigEditor::SetDetailObjects(const TArray<UObject*>& InObjects, bool bChangeUISelectionState)
+bool FControlRigEditor::DetailViewShowsAnyRigElement() const
 {
-	if(bSuspendDetailsPanelRefresh)
-	{
-		return;
-	}
+	return DetailViewShowsStruct(FRigBaseElement::StaticStruct());
+}
 
-	if(InObjects.Num() == 1)
+bool FControlRigEditor::DetailViewShowsRigElement(FRigElementKey InKey) const
+{
+	TArray< TWeakObjectPtr<UObject> > SelectedObjects = Inspector->GetSelectedObjects();
+	for (TWeakObjectPtr<UObject> SelectedObject : SelectedObjects)
 	{
-		if(URigVMMemoryStorage* Memory = Cast<URigVMMemoryStorage>(InObjects[0]))
+		if (SelectedObject.IsValid())
 		{
-			FPropertyEditorModule& EditModule = FModuleManager::Get().GetModuleChecked<FPropertyEditorModule>("PropertyEditor");
-
-			FDetailsViewArgs DetailsViewArgs;
-			DetailsViewArgs.NameAreaSettings = FDetailsViewArgs::HideNameArea;
-			DetailsViewArgs.bHideSelectionTip = true;
-
-			TSharedRef<IDetailsView> DetailsView = EditModule.CreateDetailView( DetailsViewArgs );
-			TSharedRef<SDockTab> DockTab = SNew(SDockTab)
-			.Label( LOCTEXT("ControlRigMemoryDetails", "Control Rig Memory Details") )
-			.AddMetaData<FMemoryTypeMetaData>(FMemoryTypeMetaData(Memory->GetMemoryType()))
-			.TabRole(ETabRole::NomadTab)
-			[
-				DetailsView
-			];
-
-			FName TabId = *FString::Printf(TEXT("ControlRigMemoryDetails_%d"), (int32)Memory->GetMemoryType());
-			if(TSharedPtr<SDockTab> ActiveTab = GetTabManager()->FindExistingLiveTab(TabId))
+			if(URigVMDetailsViewWrapperObject* WrapperObject = Cast<URigVMDetailsViewWrapperObject>(SelectedObject.Get()))
 			{
-				ActiveTab->RequestCloseTab();
+				if (const UScriptStruct* WrappedStruct = WrapperObject->GetWrappedStruct())
+				{
+					if (WrappedStruct->IsChildOf(FRigBaseElement::StaticStruct()))
+					{
+						if(WrapperObject->GetContent<FRigBaseElement>().GetKey() == InKey)
+						{
+							return true;
+						}
+					}
+				}
 			}
-
-			GetTabManager()->InsertNewDocumentTab(
-				FBlueprintEditorTabs::DetailsID,
-				TabId,
-				FTabManager::FLastMajorOrNomadTab(TEXT("ControlRigMemoryDetails")),
-				DockTab
-			);
-
-			FFunctionGraphTask::CreateAndDispatchWhenReady([DetailsView, InObjects]()
-			{
-				
-				DetailsView->SetObject(InObjects[0]);
-				
-			}, TStatId(), NULL, ENamedThreads::GameThread);
-
-			return;
 		}
 	}
+	return false;
+}
 
-	ClearDetailObject(bChangeUISelectionState);
-
-	if (InObjects.Num() == 1)
+void FControlRigEditor::Compile()
+{
 	{
-		if (InObjects[0]->GetClass()->GetDefaultObject() == InObjects[0])
+		DECLARE_SCOPE_HIERARCHICAL_COUNTER_FUNC()
+
+		TUniquePtr<UControlRigBlueprint::FControlValueScope> ValueScope;
+		if (!UControlRigEditorSettings::Get()->bResetControlsOnCompile) // if we need to retain the controls
 		{
-			EditClassDefaults_Clicked();
-			return;
+			ValueScope = MakeUnique<UControlRigBlueprint::FControlValueScope>(GetControlRigBlueprint());
 		}
-		else if (InObjects[0] == GetBlueprintObj())
+
+		UControlRigBlueprint* ControlRigBlueprint = GetControlRigBlueprint();
+		if (ControlRigBlueprint == nullptr)
 		{
-			EditGlobalOptions_Clicked();
 			return;
 		}
-	}
 
-	TArray<UObject*> FilteredObjects;
+		TArray< TWeakObjectPtr<UObject> > SelectedObjects = Inspector->GetSelectedObjects();
 
-	TArray<URigVMNode*> ModelNodes;
-	for(UObject* InObject : InObjects)
-	{
-		if (URigVMNode* ModelNode = Cast<URigVMNode>(InObject))
+		if(IsConstructionModeEnabled())
 		{
-			ModelNodes.Add(ModelNode);
+			SetEventQueue(ForwardsSolveEventQueue, false);
 		}
-	}
-	
-	for(UObject* InObject : InObjects)
-	{
-		if (URigVMLibraryNode* LibraryNode = Cast<URigVMLibraryNode>(InObject))
+
+		// clear transient controls such that we don't leave
+		// a phantom shape in the viewport
+		// have to do this before compile() because during compile
+		// a new control rig instance is created without the transient controls
+		// so clear is never called for old transient controls
+		ControlRigBlueprint->ClearTransientControls();
+
+		// default to always reset all bone modifications 
+		ResetAllBoneModification(); 
+
 		{
-			if(!LibraryNode->IsA<URigVMFunctionReferenceNode>())
-			{
-				if (UEdGraph* EdGraph = GetControlRigBlueprint()->GetEdGraph(LibraryNode->GetContainedGraph()))
-				{
-					FilteredObjects.AddUnique(EdGraph);
-					continue;
-				}
-			}
+			FRigVMEditor::Compile();
 		}
-		else if (Cast<URigVMFunctionEntryNode>(InObject) || Cast<URigVMFunctionReturnNode>(InObject))
+
+		// ensure the skeletal mesh is still bound
+		UControlRigSkeletalMeshComponent* SkelMeshComponent = Cast<UControlRigSkeletalMeshComponent>(GetPersonaToolkit()->GetPreviewScene()->GetPreviewMeshComponent());
+		if (SkelMeshComponent)
 		{
-			if (UEdGraph* EdGraph = GetControlRigBlueprint()->GetEdGraph(CastChecked<URigVMNode>(InObject)->GetGraph()))
+			bool bWasCreated = false;
+			FAnimCustomInstanceHelper::BindToSkeletalMeshComponent<UControlRigLayerInstance>(SkelMeshComponent, bWasCreated);
+			if (bWasCreated)
 			{
-				FilteredObjects.AddUnique(EdGraph);
-				continue;
+				OnAnimInitialized();
 			}
 		}
-		else if (URigVMCommentNode* CommentNode = Cast<URigVMCommentNode>(InObject))
+		
+		if (SelectedObjects.Num() > 0)
 		{
-			if (UControlRigGraph* EdGraph = Cast<UControlRigGraph>(GetControlRigBlueprint()->GetEdGraph(CastChecked<URigVMNode>(InObject)->GetGraph())))
-			{
-				if(UEdGraphNode* EdGraphNode = EdGraph->FindNodeForModelNodeName(CommentNode->GetFName()))
-				{
-					FilteredObjects.AddUnique(EdGraphNode);
-					continue;
-				}
-			}
+			RefreshDetailView();
 		}
 
-		if (URigVMNode* ModelNode = Cast<URigVMNode>(InObject))
+		if (UControlRigEditorSettings::Get()->bResetControlTransformsOnCompile)
 		{
-			// check if we know the dynamic class already
-			const URigVMDetailsViewWrapperObject* CDOWrapper = CastChecked<URigVMDetailsViewWrapperObject>(GetDetailWrapperClass()->GetDefaultObject());
-			bool bClassCreated = CDOWrapper->GetClassForNodes(ModelNodes, false) == nullptr;
+			ControlRigBlueprint->Hierarchy->ForEach<FRigControlElement>([ControlRigBlueprint](FRigControlElement* ControlElement) -> bool
+            {
+				const FTransform Transform = ControlRigBlueprint->Hierarchy->GetInitialLocalTransform(ControlElement->GetIndex());
 
-			// create the wrapper object
-			if(URigVMDetailsViewWrapperObject* WrapperObject = URigVMDetailsViewWrapperObject::MakeInstance(GetDetailWrapperClass(), GetBlueprintObj(), ModelNodes, ModelNode))
-			{
-				WrapperObject->GetWrappedPropertyChangedChainEvent().AddSP(this, &FControlRigEditor::OnWrappedPropertyChangedChainEvent);
-				WrapperObject->AddToRoot();
+				/*/
+				if (UControlRig* ControlRig = GetControlRig())
+				{
+					ControlRig->Modify();
+					ControlRig->GetControlHierarchy().SetLocalTransform(Control.Index, Transform);
+					ControlRig->ControlModified().Broadcast(ControlRig, Control, EControlRigSetKey::DoNotCare);
+				}
+				*/
 
-				// todo: use transform widget for transforms
-				// todo: use rotation widget for rotations
-				
-				FilteredObjects.Add(WrapperObject);
-				continue;
-			}
+				ControlRigBlueprint->Hierarchy->SetLocalTransform(ControlElement->GetIndex(), Transform);
+				return true;
+			});
 		}
 
+		ControlRigBlueprint->PropagatePoseFromBPToInstances();
 
-		FilteredObjects.Add(InObject);
-	}
-
-	for(UObject* FilteredObject : FilteredObjects)
-	{
-		if(URigVMDetailsViewWrapperObject* WrapperObject = Cast<URigVMDetailsViewWrapperObject>(FilteredObject))
+		if (FControlRigEditMode* EditMode = GetEditMode())
 		{
-			WrapperObjects.Add(TStrongObjectPtr<URigVMDetailsViewWrapperObject>(WrapperObject));
+			EditMode->RecreateControlShapeActors(GetHierarchyBeingDebugged()->GetSelectedKeys());
 		}
 	}
-
-	SKismetInspector::FShowDetailsOptions Options;
-	Options.bForceRefresh = true;
-	Inspector->ShowDetailsForObjects(FilteredObjects, Options);
 }
 
-void FControlRigEditor::SetDetailViewForRigElements()
+void FControlRigEditor::SaveAsset_Execute()
 {
-	if(bSuspendDetailsPanelRefresh)
-	{
-		return;
-	}
-
-	ClearDetailObject();
+	FRigVMEditor::SaveAsset_Execute();
 
+	// Save the new state of the hierarchy in the default object, so that it has the correct values on load
 	UControlRigBlueprint* RigBlueprint = Cast<UControlRigBlueprint>(GetBlueprintObj());
-	URigHierarchy* HierarchyBeingDebugged = GetHierarchyBeingDebugged();
-	TArray<UObject*> Objects;
-
-	TArray<FRigElementKey> CurrentSelection = HierarchyBeingDebugged->GetSelectedKeys();
-	for(const FRigElementKey& SelectedKey : CurrentSelection)
+	if(const UControlRig* ControlRig = GetControlRig())
 	{
-		FRigBaseElement* Element = HierarchyBeingDebugged->Find(SelectedKey);
-		if (Element == nullptr)
-		{
-			continue;
-		}
-
-		URigVMDetailsViewWrapperObject* WrapperObject = URigVMDetailsViewWrapperObject::MakeInstance(GetDetailWrapperClass(), GetBlueprintObj(), Element->GetElementStruct(), (uint8*)Element, HierarchyBeingDebugged);
-		WrapperObject->GetWrappedPropertyChangedChainEvent().AddSP(this, &FControlRigEditor::OnWrappedPropertyChangedChainEvent);
-		WrapperObject->AddToRoot();
-
-		Objects.Add(WrapperObject);
+		const UControlRig* CDO = ControlRig->GetClass()->GetDefaultObject<UControlRig>();
+		CDO->DynamicHierarchy->CopyHierarchy(RigBlueprint->Hierarchy);
 	}
-	
-	SetDetailObjects(Objects);
 }
 
-void FControlRigEditor::SetDetailViewForGraph(URigVMGraph* InGraph)
+void FControlRigEditor::SaveAssetAs_Execute()
 {
-	check(InGraph);
-	
-	if(bSuspendDetailsPanelRefresh)
-	{
-		return;
-	}
+	FRigVMEditor::SaveAssetAs_Execute();
 
-	ClearDetailObject();
-
-	TArray<UObject*> SelectedNodes;
-	TArray<FName> SelectedNodeNames = InGraph->GetSelectNodes();
-	for(FName SelectedNodeName : SelectedNodeNames)
+	// Save the new state of the hierarchy in the default object, so that it has the correct values on load
+	UControlRigBlueprint* RigBlueprint = Cast<UControlRigBlueprint>(GetBlueprintObj());
+	if(const UControlRig* ControlRig = GetControlRig())
 	{
-		if(URigVMNode* Node = InGraph->FindNodeByName(SelectedNodeName))
-		{
-			SelectedNodes.Add(Node);
-		}
+		const UControlRig* CDO = ControlRig->GetClass()->GetDefaultObject<UControlRig>();
+		CDO->DynamicHierarchy->CopyHierarchy(RigBlueprint->Hierarchy);
 	}
-
-	SetDetailObjects(SelectedNodes);
 }
 
-void FControlRigEditor::SetDetailViewForFocusedGraph()
+FName FControlRigEditor::GetToolkitFName() const
 {
-	if(bSuspendDetailsPanelRefresh)
-	{
-		return;
-	}
+	return FName("ControlRigEditor");
+}
 
-	URigVMGraph* Model = GetFocusedModel();
-	if(Model == nullptr)
-	{
-		return;
-	}
+FText FControlRigEditor::GetBaseToolkitName() const
+{
+	return LOCTEXT("AppLabel", "Control Rig Editor");
+}
 
-	SetDetailViewForGraph(Model);
+FString FControlRigEditor::GetWorldCentricTabPrefix() const
+{
+	return LOCTEXT("WorldCentricTabPrefix", "Control Rig Editor ").ToString();
 }
 
-void FControlRigEditor::SetDetailViewForLocalVariable()
+FReply FControlRigEditor::OnSpawnGraphNodeByShortcut(FInputChord InChord, const FVector2D& InPosition, UEdGraph* InGraph)
 {
-	FName VariableName;
-	TArray< TWeakObjectPtr<UObject> > SelectedObjects = Inspector->GetSelectedObjects();
-	for (TWeakObjectPtr<UObject> SelectedObject : SelectedObjects)
+	const FReply SuperReply = FRigVMEditor::OnSpawnGraphNodeByShortcut(InChord, InPosition, InGraph);
+	if(SuperReply.IsEventHandled())
 	{
-		if (SelectedObject.IsValid())
+		return SuperReply;
+	}
+
+	if(!InChord.HasAnyModifierKeys())
+	{
+		if(UControlRigGraph* RigGraph = Cast<UControlRigGraph>(InGraph))
 		{
-			if(URigVMDetailsViewWrapperObject* WrapperObject = Cast<URigVMDetailsViewWrapperObject>(SelectedObject.Get()))
+			if(URigVMController* Controller = RigGraph->GetController())
 			{
-				VariableName = WrapperObject->GetContent<FRigVMGraphVariableDescription>().Name;
-				break;
+				if(InChord.Key == EKeys::S)
+				{
+					Controller->AddUnitNode(FRigUnit_SequenceAggregate::StaticStruct(), FRigUnit::GetMethodName(), InPosition, FString(), true, true);
+				}
+				else if(InChord.Key == EKeys::One)
+				{
+					Controller->AddUnitNode(FRigUnit_GetTransform::StaticStruct(), FRigUnit::GetMethodName(), InPosition, FString(), true, true);
+				}
+				else if(InChord.Key == EKeys::Two)
+				{
+					Controller->AddUnitNode(FRigUnit_SetTransform::StaticStruct(), FRigUnit::GetMethodName(), InPosition, FString(), true, true);
+				}
+				else if(InChord.Key == EKeys::Three)
+				{
+					Controller->AddUnitNode(FRigUnit_ParentConstraint::StaticStruct(), FRigUnit::GetMethodName(), InPosition, FString(), true, true);
+				}
+				else if(InChord.Key == EKeys::Four)
+				{
+					Controller->AddUnitNode(FRigUnit_GetControlFloat::StaticStruct(), FRigUnit::GetMethodName(), InPosition, FString(), true, true);
+				}
+				else if(InChord.Key == EKeys::Five)
+				{
+					Controller->AddUnitNode(FRigUnit_SetCurveValue::StaticStruct(), FRigUnit::GetMethodName(), InPosition, FString(), true, true);
+				}
 			}
 		}
 	}
-		
-	SelectLocalVariable(GetFocusedGraph(), VariableName);
-}
 
-void FControlRigEditor::RefreshDetailView()
-{
-	if(DetailViewShowsAnyRigElement())
-	{
-		SetDetailViewForRigElements();
-	}
-	else if(DetailViewShowsAnyRigUnit())
-	{
-		SetDetailViewForFocusedGraph();
-	}
-	else if(DetailViewShowsLocalVariable())
-	{
-		SetDetailViewForLocalVariable();	
-	}
-	else
-	{
-		// detail view is showing other stuff; could be a BP variable for example
-		// in this case wrapper objects are not in use, yet still rooted
-		// and preventing their outer objects from getting GCed after a Compile()
-		// so let's take the chance to manually clear them here.
-		ClearDetailsViewWrapperObjects();
-	}
+	return FReply::Unhandled();
 }
 
-bool FControlRigEditor::DetailViewShowsAnyRigElement() const
+void FControlRigEditor::PostTransaction(bool bSuccess, const FTransaction* Transaction, bool bIsRedo)
 {
-	return DetailViewShowsStruct(FRigBaseElement::StaticStruct());
-}
+	EnsureValidRigElementsInDetailPanel();
 
-bool FControlRigEditor::DetailViewShowsAnyRigUnit() const
-{
-	if (DetailViewShowsStruct(FRigUnit::StaticStruct()))
+	if (UControlRigBlueprint* RigBlueprint = Cast<UControlRigBlueprint>(GetBlueprintObj()))
 	{
-		return true;
-	}
+		// Do not compile here. ControlRigBlueprint::PostTransacted decides when it is necessary to compile depending
+		// on the properties that are affected.
+		//Compile();
+		
+		USkeletalMesh* PreviewMesh = GetPersonaToolkit()->GetPreviewScene()->GetPreviewMesh();
+		if (PreviewMesh != RigBlueprint->GetPreviewMesh())
+		{
+			RigBlueprint->SetPreviewMesh(PreviewMesh);
+			GetPersonaToolkit()->SetPreviewMesh(PreviewMesh, true);
+		}
 
-	const TArray< TWeakObjectPtr<UObject> >& SelectedObjects = Inspector->GetSelectedObjects();
-	for (TWeakObjectPtr<UObject> SelectedObject : SelectedObjects)
-	{
-		if (SelectedObject.IsValid())
+		if (UControlRig* DebuggedControlRig = Cast<UControlRig>(RigBlueprint->GetObjectBeingDebugged()))
 		{
-			if(URigVMDetailsViewWrapperObject* WrapperObject = Cast<URigVMDetailsViewWrapperObject>(SelectedObject.Get()))
+			if(URigHierarchy* Hierarchy = DebuggedControlRig->GetHierarchy())
 			{
-				const FString Notation = WrapperObject->GetWrappedNodeNotation();
-				if(!Notation.IsEmpty())
+				if(Hierarchy->Num() == 0)
 				{
-					return true;
+					OnHierarchyChanged();
 				}
 			}
 		}
+
+		if (FControlRigEditMode* EditMode = GetEditMode())
+		{
+			EditMode->RequestToRecreateControlShapeActors();
+		}
 	}
-	
-	return false;
 }
 
-bool FControlRigEditor::DetailViewShowsLocalVariable() const
+void FControlRigEditor::EnsureValidRigElementsInDetailPanel()
 {
-	return DetailViewShowsStruct(FRigVMGraphVariableDescription::StaticStruct());
-}
+	UControlRigBlueprint* ControlRigBP = GetControlRigBlueprint();
+	URigHierarchy* Hierarchy = ControlRigBP->Hierarchy; 
 
-bool FControlRigEditor::DetailViewShowsStruct(UScriptStruct* InStruct) const
-{
 	TArray< TWeakObjectPtr<UObject> > SelectedObjects = Inspector->GetSelectedObjects();
 	for (TWeakObjectPtr<UObject> SelectedObject : SelectedObjects)
 	{
@@ -2083,1338 +1179,105 @@ bool FControlRigEditor::DetailViewShowsStruct(UScriptStruct* InStruct) const
 			{
 				if(const UScriptStruct* WrappedStruct = WrapperObject->GetWrappedStruct())
 				{
-					if (WrappedStruct->IsChildOf(InStruct))
+					if (WrappedStruct->IsChildOf(FRigBaseElement::StaticStruct()))
 					{
-						return true;
+						FRigElementKey Key = WrapperObject->GetContent<FRigBaseElement>().GetKey();
+						if(!Hierarchy->Contains(Key))
+						{
+							ClearDetailObject();
+						}
 					}
 				}
 			}
 		}
 	}
-	return false;
 }
 
-bool FControlRigEditor::DetailViewShowsRigElement(FRigElementKey InKey) const
+void FControlRigEditor::OnAnimInitialized()
 {
-	TArray< TWeakObjectPtr<UObject> > SelectedObjects = Inspector->GetSelectedObjects();
-	for (TWeakObjectPtr<UObject> SelectedObject : SelectedObjects)
-	{
-		if (SelectedObject.IsValid())
-		{
-			if(URigVMDetailsViewWrapperObject* WrapperObject = Cast<URigVMDetailsViewWrapperObject>(SelectedObject.Get()))
-			{
-				if (const UScriptStruct* WrappedStruct = WrapperObject->GetWrappedStruct())
-				{
-					if (WrappedStruct->IsChildOf(FRigBaseElement::StaticStruct()))
-					{
-						if(WrapperObject->GetContent<FRigBaseElement>().GetKey() == InKey)
-						{
-							return true;
-						}
-					}
-				}
-			}
-		}
-	}
-	return false;
-}
-
-void FControlRigEditor::ClearDetailObject(bool bChangeUISelectionState)
-{
-	if(bSuspendDetailsPanelRefresh)
-	{
-		return;
-	}
-
-	ClearDetailsViewWrapperObjects();
-	
-	Inspector->GetPropertyView()->SetObjects(TArray<UObject*>(), true); // clear property view synchronously
-	Inspector->ShowDetailsForObjects(TArray<UObject*>());
-	Inspector->ShowSingleStruct(TSharedPtr<FStructOnScope>());
-
-	if (bChangeUISelectionState)
-	{
-		SetUISelectionState(FBlueprintEditor::SelectionState_Graph);
-	}
-}
-
-void FControlRigEditor::ClearDetailsViewWrapperObjects()
-{
-	for(const TStrongObjectPtr<URigVMDetailsViewWrapperObject>& WrapperObjectPtr : WrapperObjects)
-	{
-		if(WrapperObjectPtr.IsValid())
-		{
-			URigVMDetailsViewWrapperObject* WrapperObject = WrapperObjectPtr.Get();
-			WrapperObject->RemoveFromRoot();
-			WrapperObject->Rename(nullptr, GetTransientPackage(), REN_ForceNoResetLoaders | REN_DoNotDirty | REN_DontCreateRedirectors | REN_NonTransactional);
-			WrapperObject->MarkAsGarbage();
-		}
-	}
-	WrapperObjects.Reset();
-}
-
-
-void FControlRigEditor::CreateDefaultCommands() 
-{
-	if (GetBlueprintObj())
-	{
-		FBlueprintEditor::CreateDefaultCommands();
-	}
-	else
-	{
-		ToolkitCommands->MapAction( FGenericCommands::Get().Undo, 
-			FExecuteAction::CreateSP( this, &FControlRigEditor::UndoAction ));
-		ToolkitCommands->MapAction( FGenericCommands::Get().Redo, 
-			FExecuteAction::CreateSP( this, &FControlRigEditor::RedoAction ));
-	}
-}
-
-void FControlRigEditor::OnCreateGraphEditorCommands(TSharedPtr<FUICommandList> GraphEditorCommandsList)
-{
-}
-
-TSharedRef<SGraphEditor> FControlRigEditor::CreateGraphEditorWidget(TSharedRef<FTabInfo> InTabInfo, UEdGraph* InGraph)
-{
-	TSharedRef<SGraphEditor> GraphEditor = FBlueprintEditor::CreateGraphEditorWidget(InTabInfo, InGraph);
-	GraphEditor->GetGraphPanel()->SetZoomLevelsContainer<FControlRigZoomLevelsContainer>();
-	return GraphEditor;
-}
-
-void FControlRigEditor::Compile()
-{
-	{
-		DECLARE_SCOPE_HIERARCHICAL_COUNTER_FUNC()
-
-		TUniquePtr<UControlRigBlueprint::FControlValueScope> ValueScope;
-		if (!UControlRigEditorSettings::Get()->bResetControlsOnCompile) // if we need to retain the controls
-		{
-			ValueScope = MakeUnique<UControlRigBlueprint::FControlValueScope>(GetControlRigBlueprint());
-		}
-
-		// force to disable the supended notif brackets
-		if (UControlRigBlueprint* RigBlueprint = Cast<UControlRigBlueprint>(GetBlueprintObj()))
-		{
-			RigBlueprint->bSuspendModelNotificationsForOthers = false;
-			RigBlueprint->bSuspendModelNotificationsForSelf = false;
-		}
-
-		UControlRigBlueprint* RigBlueprint = Cast<UControlRigBlueprint>(GetBlueprintObj());
-		if (RigBlueprint == nullptr)
-		{
-			return;
-		}
-
-		RigBlueprint->GetCompileLog().Messages.Reset();
-
-		FString LastDebuggedObjectName = GetCustomDebugObjectLabel(RigBlueprint->GetObjectBeingDebugged());
-		RigBlueprint->SetObjectBeingDebugged(nullptr);
-
-		TArray< TWeakObjectPtr<UObject> > SelectedObjects = Inspector->GetSelectedObjects();
-
-		if (UControlRig* ControlRig = GetControlRig())
-		{
-			ControlRig->OnInitialized_AnyThread().Clear();
-			ControlRig->OnExecuted_AnyThread().Clear();
-			ControlRig->GetExtendedExecuteContext().ExecutionHalted().RemoveAll(this);
-		}
-
-		if(IsConstructionModeEnabled())
-		{
-			SetEventQueue(ForwardsSolveEventQueue, false);
-		}
-
-		// clear transient controls such that we don't leave
-		// a phantom shape in the viewport
-		// have to do this before compile() because during compile
-		// a new control rig instance is created without the transient controls
-		// so clear is never called for old transient controls
-		RigBlueprint->ClearTransientControls();
-
-		// default to always reset all bone modifications 
-		ResetAllBoneModification(); 
-
-		// remove the current host from the editor
-		SetHost(nullptr);
-
-		{
-			TGuardValue<bool> GuardCompileReEntry(bIsCompilingThroughUI, true);
-			FBlueprintEditor::Compile();
-		}
-
-		// ensure the skeletal mesh is still bound
-		UControlRigSkeletalMeshComponent* SkelMeshComponent = Cast<UControlRigSkeletalMeshComponent>(GetPersonaToolkit()->GetPreviewScene()->GetPreviewMeshComponent());
-		if (SkelMeshComponent)
-		{
-			bool bWasCreated = false;
-			FAnimCustomInstanceHelper::BindToSkeletalMeshComponent<UControlRigLayerInstance>(SkelMeshComponent, bWasCreated);
-			if (bWasCreated)
-			{
-				OnAnimInitialized();
-			}
-		}
-		
-		if (UControlRig* ControlRig = GetControlRig())
-		{
-			ControlRigLog.Reset();
-			ControlRig->RigVMLog = &ControlRigLog;
-
-			URigVMBlueprintGeneratedClass* GeneratedClass = Cast<URigVMBlueprintGeneratedClass>(ControlRig->GetClass());
-			if (GeneratedClass)
-			{
-				UControlRig* CDO = Cast<UControlRig>(GeneratedClass->GetDefaultObject(true /* create if needed */));
-				FRigVMInstructionArray Instructions = CDO->VM->GetInstructions();
-
-				if (Instructions.Num() <= 1) // just the "done" operator
-				{
-					FNotificationInfo Info(LOCTEXT("ControlRigBlueprintCompilerEmptyRigMessage", "The Control Rig you compiled doesn't do anything. Did you forget to add a Begin_Execution node?"));
-					Info.bFireAndForget = true;
-					Info.FadeOutDuration = 5.0f;
-					Info.ExpireDuration = 5.0f;
-					TSharedPtr<SNotificationItem> NotificationPtr = FSlateNotificationManager::Get().AddNotification(Info);
-					NotificationPtr->SetCompletionState(SNotificationItem::CS_Success);
-				}
-			}
-		}
-
-		TArray<FCustomDebugObject> DebugList;
-		GetCustomDebugObjects(DebugList);
-
-		for (const FCustomDebugObject& DebugObject : DebugList)
-		{
-			if (DebugObject.NameOverride == LastDebuggedObjectName)
-			{
-				RigBlueprint->SetObjectBeingDebugged(DebugObject.Object);
-			}
-		}
-
-		// invalidate all node titles
-		TArray<UEdGraph*> EdGraphs;
-		RigBlueprint->GetAllGraphs(EdGraphs);
-		for (UEdGraph* EdGraph : EdGraphs)
-		{
-			UControlRigGraph* RigGraph = Cast<UControlRigGraph>(EdGraph);
-			if (RigGraph == nullptr)
-			{
-				continue;
-			}
-
-			for (UEdGraphNode* EdNode : RigGraph->Nodes)
-			{
-				if (UControlRigGraphNode* RigNode = Cast<UControlRigGraphNode>(EdNode))
-				{
-					RigNode->InvalidateNodeTitle();
-				}
-			}
-		}
-
-		// store the defaults from the CDO back on the new variables list
-		bool bAnyVariableValueChanged = false;
-		for(FBPVariableDescription& NewVariable : RigBlueprint->NewVariables)
-		{
-			bAnyVariableValueChanged |= UpdateDefaultValueForVariable(NewVariable, true);
-		}
-		if (bAnyVariableValueChanged)
-		{
-			// Go over all the instances to update the default values from CDO
-			for (const FCustomDebugObject& DebugObject : DebugList)
-			{
-				if (UControlRig* DebuggedRig = Cast<UControlRig>(DebugObject.Object))
-				{
-					DebuggedRig->CopyExternalVariableDefaultValuesFromCDO();
-				}
-			}
-		}
-
-		if (SelectedObjects.Num() > 0)
-		{
-			RefreshDetailView();
-		}
-
-		if (UControlRigEditorSettings::Get()->bResetControlTransformsOnCompile)
-		{
-			RigBlueprint->Hierarchy->ForEach<FRigControlElement>([RigBlueprint](FRigControlElement* ControlElement) -> bool
-            {
-				const FTransform Transform = RigBlueprint->Hierarchy->GetInitialLocalTransform(ControlElement->GetIndex());
-
-				/*/
-				if (UControlRig* ControlRig = GetControlRig())
-				{
-					ControlRig->Modify();
-					ControlRig->GetControlHierarchy().SetLocalTransform(Control.Index, Transform);
-					ControlRig->ControlModified().Broadcast(ControlRig, Control, EControlRigSetKey::DoNotCare);
-				}
-				*/
-
-				RigBlueprint->Hierarchy->SetLocalTransform(ControlElement->GetIndex(), Transform);
-				return true;
-			});
-		}
-
-		RigBlueprint->PropagatePoseFromBPToInstances();
-
-		if (FControlRigEditMode* EditMode = GetEditMode())
-		{
-			EditMode->RecreateControlShapeActors(GetHierarchyBeingDebugged()->GetSelectedKeys());
-		}
-	}
-
-	// enable this for creating a new unit test
-	// DumpUnitTestCode();
-
-	// FStatsHierarchical::EndMeasurements();
-	// FMessageLog LogForMeasurements("ControlRigLog");
-	// FStatsHierarchical::DumpMeasurements(LogForMeasurements);
-}
-
-void FControlRigEditor::SaveAsset_Execute()
-{
-	LastDebuggedRig = GetCustomDebugObjectLabel(GetBlueprintObj()->GetObjectBeingDebugged());
-	FBlueprintEditor::SaveAsset_Execute();
-
-	// Save the new state of the hierarchy in the default object, so that it has the correct values on load
-	UControlRigBlueprint* RigBlueprint = Cast<UControlRigBlueprint>(GetBlueprintObj());
-	if(const UControlRig* ControlRig = GetControlRig())
-	{
-		const UControlRig* CDO = ControlRig->GetClass()->GetDefaultObject<UControlRig>();
-		CDO->DynamicHierarchy->CopyHierarchy(RigBlueprint->Hierarchy);
-	}
-
-	FBlueprintActionDatabase& ActionDatabase = FBlueprintActionDatabase::Get();
-	ActionDatabase.ClearAssetActions(UControlRigBlueprint::StaticClass());
-	ActionDatabase.RefreshClassActions(UControlRigBlueprint::StaticClass());
-}
-
-void FControlRigEditor::SaveAssetAs_Execute()
-{
-	LastDebuggedRig = GetCustomDebugObjectLabel(GetBlueprintObj()->GetObjectBeingDebugged());
-	FBlueprintEditor::SaveAssetAs_Execute();
-
-	// Save the new state of the hierarchy in the default object, so that it has the correct values on load
-	UControlRigBlueprint* RigBlueprint = Cast<UControlRigBlueprint>(GetBlueprintObj());
-	if(const UControlRig* ControlRig = GetControlRig())
-	{
-		const UControlRig* CDO = ControlRig->GetClass()->GetDefaultObject<UControlRig>();
-		CDO->DynamicHierarchy->CopyHierarchy(RigBlueprint->Hierarchy);
-	}
-
-	FBlueprintActionDatabase& ActionDatabase = FBlueprintActionDatabase::Get();
-	ActionDatabase.ClearAssetActions(UControlRigBlueprint::StaticClass());
-	ActionDatabase.RefreshClassActions(UControlRigBlueprint::StaticClass());
-}
-
-FName FControlRigEditor::GetToolkitFName() const
-{
-	return FName("ControlRigEditor");
-}
-
-FName FControlRigEditor::GetToolkitContextFName() const
-{
-	return GetToolkitFName();
-}
-
-FText FControlRigEditor::GetBaseToolkitName() const
-{
-	return LOCTEXT("AppLabel", "Control Rig Editor");
-}
-
-FText FControlRigEditor::GetToolkitToolTipText() const
-{
-	return FAssetEditorToolkit::GetToolTipTextForObject(GetBlueprintObj());
-}
-
-FString FControlRigEditor::GetWorldCentricTabPrefix() const
-{
-	return LOCTEXT("WorldCentricTabPrefix", "Control Rig Editor ").ToString();
-}
-
-FLinearColor FControlRigEditor::GetWorldCentricTabColorScale() const
-{
-	return FLinearColor( 0.5f, 0.25f, 0.35f, 0.5f );
-}
-
-void FControlRigEditor::InitToolMenuContext(FToolMenuContext& MenuContext)
-{
-	FBlueprintEditor::InitToolMenuContext(MenuContext);
-
-	if (UControlRigBlueprint* RigBlueprint = GetControlRigBlueprint())
-	{
-		URigVMGraph* Model = nullptr;
-		URigVMNode* Node = nullptr;
-		URigVMPin* Pin = nullptr;
-		
-		if (UGraphNodeContextMenuContext* GraphNodeContext = MenuContext.FindContext<UGraphNodeContextMenuContext>())
-		{
-		
-			if (GraphNodeContext->Node)
-			{
-				Model = RigBlueprint->GetModel(GraphNodeContext->Graph);
-				if(Model)
-				{
-					Node = Model->FindNodeByName(GraphNodeContext->Node->GetFName());
-				}
-			}
-		
-			if (GraphNodeContext->Pin && Node)
-			{
-				Pin = Model->FindPin(GraphNodeContext->Pin->GetName());
-			}
-		}
-		
-		UControlRigContextMenuContext* ControlRigMenuContext = NewObject<UControlRigContextMenuContext>();
-		FControlRigMenuSpecificContext MenuSpecificContext;	
-		MenuSpecificContext.GraphNodeContextMenuContext = FControlRigGraphNodeContextMenuContext(Model, Node, Pin);
-		ControlRigMenuContext->Init(SharedThis(this), MenuSpecificContext);
-
-		MenuContext.AddObject(ControlRigMenuContext);
-	}
-}
-
-bool FControlRigEditor::TransactionObjectAffectsBlueprint(UObject* InTransactedObject)
-{
-	UControlRigBlueprint* RigBlueprint = Cast<UControlRigBlueprint>(GetBlueprintObj());
-	if (RigBlueprint == nullptr)
-	{
-		return false;
-	}
-
-	if (InTransactedObject->GetOuter() == GetFocusedController())
-	{
-		return false;
-	}
-	return FBlueprintEditor::TransactionObjectAffectsBlueprint(InTransactedObject);
-}
-
-bool FControlRigEditor::CanAddNewLocalVariable() const
-{
-	const URigVMGraph* Graph = GetFocusedModel();
-	const URigVMGraph* ParentGraph = Graph->GetParentGraph();				
-	if (ParentGraph && ParentGraph->IsA<URigVMFunctionLibrary>())
-	{
-		return true;
-	}
-	return false;
-}
-
-void FControlRigEditor::OnAddNewLocalVariable()
-{
-	if (!CanAddNewLocalVariable())
-	{
-		return;
-	}
-
-	FRigVMGraphVariableDescription LastTypeVar;
-	LastTypeVar.ChangeType(MyBlueprintWidget->GetLastPinTypeUsed());
-	FRigVMGraphVariableDescription NewVar = GetFocusedController()->AddLocalVariable(TEXT("NewLocalVar"), LastTypeVar.CPPType, LastTypeVar.CPPTypeObject, LastTypeVar.DefaultValue, true, true);
-	if(NewVar.Name.IsNone())
-	{
-		LogSimpleMessage( LOCTEXT("AddLocalVariable_Error", "Adding new local variable failed.") );
-	}
-	else
-	{
-		RenameNewlyAddedAction(NewVar.Name);
-	}
-}
-
-void FControlRigEditor::OnPasteNewLocalVariable(const FBPVariableDescription& VariableDescription)
-{
-	if (!CanAddNewLocalVariable())
-	{
-		return;
-	}
-
-	FRigVMGraphVariableDescription TypeVar;
-	TypeVar.ChangeType(VariableDescription.VarType);
-	FRigVMGraphVariableDescription NewVar = GetFocusedController()->AddLocalVariable(VariableDescription.VarName, TypeVar.CPPType, TypeVar.CPPTypeObject, VariableDescription.DefaultValue, true, true);
-	if(NewVar.Name.IsNone())
-	{
-		LogSimpleMessage( LOCTEXT("PasteLocalVariable_Error", "Pasting new local variable failed.") );
-	}
-}
-
-void FControlRigEditor::DeleteSelectedNodes()
-{
-	DECLARE_SCOPE_HIERARCHICAL_COUNTER_FUNC()
-
-	UControlRigBlueprint* RigBlueprint = Cast<UControlRigBlueprint>(GetBlueprintObj());
-	if (RigBlueprint == nullptr)
-	{
-		return;
-	}
-
-	const FGraphPanelSelectionSet SelectedNodes = GetSelectedNodes();
-	SetUISelectionState(NAME_None);
-
-	bool bRelinkPins = false;
-	TArray<URigVMNode*> NodesToRemove;
-
-	for (FGraphPanelSelectionSet::TConstIterator NodeIt(SelectedNodes); NodeIt; ++NodeIt)
-	{
-		if (UEdGraphNode* Node = Cast<UEdGraphNode>(*NodeIt))
-		{
-			if (Node->CanUserDeleteNode())
-			{
-				AnalyticsTrackNodeEvent(GetBlueprintObj(), Node, true);
-				if (UControlRigGraphNode* RigNode = Cast<UControlRigGraphNode>(Node))
-				{
-					bRelinkPins = bRelinkPins || FSlateApplication::Get().GetModifierKeys().IsShiftDown();
-
-					if(URigVMNode* ModelNode = GetFocusedController()->GetGraph()->FindNodeByName(*RigNode->ModelNodePath))
-					{
-						NodesToRemove.Add(ModelNode);
-					}
-				}
-				else if (UEdGraphNode_Comment* CommentNode = Cast<UEdGraphNode_Comment>(Node))
-				{
-					if(URigVMNode* ModelNode = GetFocusedController()->GetGraph()->FindNodeByName(CommentNode->GetFName()))
-					{
-						NodesToRemove.Add(ModelNode);
-					}
-				}
-				else
-				{
-					Node->GetGraph()->RemoveNode(Node);
-				}
-			}
-		}
-	}
-
-	if(NodesToRemove.IsEmpty())
-	{
-		return;
-	}
-
-	GetFocusedController()->OpenUndoBracket(TEXT("Delete selected nodes"));
-	if(bRelinkPins && NodesToRemove.Num() == 1)
-	{
-		GetFocusedController()->RelinkSourceAndTargetPins(NodesToRemove[0], true);;
-	}
-	GetFocusedController()->RemoveNodes(NodesToRemove, true);
-	GetFocusedController()->CloseUndoBracket();
-}
-
-bool FControlRigEditor::CanDeleteNodes() const
-{
-	return true;
-}
-
-void FControlRigEditor::CopySelectedNodes()
-{
-	FString ExportedText = GetFocusedController()->ExportSelectedNodesToText();
-	FPlatformApplicationMisc::ClipboardCopy(*ExportedText);
-}
-
-bool FControlRigEditor::CanCopyNodes() const
-{
-	return GetFocusedModel()->GetSelectNodes().Num() > 0;
-}
-
-bool FControlRigEditor::CanPasteNodes() const
-{
-	FString TextToImport;
-	FPlatformApplicationMisc::ClipboardPaste(TextToImport);
-	return GetFocusedController()->CanImportNodesFromText(TextToImport);
-}
-
-FReply FControlRigEditor::OnSpawnGraphNodeByShortcut(FInputChord InChord, const FVector2D& InPosition,
-                                                     UEdGraph* InGraph)
-{
-	if(!InChord.HasAnyModifierKeys())
-	{
-		if(UControlRigGraph* RigGraph = Cast<UControlRigGraph>(InGraph))
-		{
-			if(URigVMController* Controller = RigGraph->GetController())
-			{
-				if(InChord.Key == EKeys::B)
-				{
-					Controller->AddUnitNode(FRigVMFunction_ControlFlowBranch::StaticStruct(), FRigVMStruct::ExecuteName, InPosition, FString(), true, true);
-				}
-				else if(InChord.Key == EKeys::S)
-				{
-					Controller->AddUnitNode(FRigUnit_SequenceAggregate::StaticStruct(), FRigUnit::GetMethodName(), InPosition, FString(), true, true);
-				}
-				else if(InChord.Key == EKeys::One)
-				{
-					Controller->AddUnitNode(FRigUnit_GetTransform::StaticStruct(), FRigUnit::GetMethodName(), InPosition, FString(), true, true);
-				}
-				else if(InChord.Key == EKeys::Two)
-				{
-					Controller->AddUnitNode(FRigUnit_SetTransform::StaticStruct(), FRigUnit::GetMethodName(), InPosition, FString(), true, true);
-				}
-				else if(InChord.Key == EKeys::Three)
-				{
-					Controller->AddUnitNode(FRigUnit_ParentConstraint::StaticStruct(), FRigUnit::GetMethodName(), InPosition, FString(), true, true);
-				}
-				else if(InChord.Key == EKeys::Four)
-				{
-					Controller->AddUnitNode(FRigUnit_GetControlFloat::StaticStruct(), FRigUnit::GetMethodName(), InPosition, FString(), true, true);
-				}
-				else if(InChord.Key == EKeys::Five)
-				{
-					Controller->AddUnitNode(FRigUnit_SetCurveValue::StaticStruct(), FRigUnit::GetMethodName(), InPosition, FString(), true, true);
-				}
-			}
-		}
-	}
-
-	return FReply::Unhandled();
-}
-
-void FControlRigEditor::PasteNodes()
-{
-	DECLARE_SCOPE_HIERARCHICAL_COUNTER_FUNC()
-
-	GetFocusedController()->OpenUndoBracket(TEXT("Pasted Nodes."));
-
-	FVector2D PasteLocation = FSlateApplication::Get().GetCursorPos();
-
-	TSharedPtr<SDockTab> ActiveTab = DocumentManager->GetActiveTab();
-	if (ActiveTab.IsValid())
-	{
-		TSharedPtr<SGraphEditor> GraphEditor = StaticCastSharedRef<SGraphEditor>(ActiveTab->GetContent());
-		if (GraphEditor.IsValid())
-		{
-			PasteLocation = GraphEditor->GetPasteLocation();
-
-		}
-	}
-
-	FString TextToImport;
-	FPlatformApplicationMisc::ClipboardPaste(TextToImport);
-
-	TGuardValue<FRigVMController_RequestLocalizeFunctionDelegate> RequestLocalizeDelegateGuard(
-		GetFocusedController()->RequestLocalizeFunctionDelegate,
-		FRigVMController_RequestLocalizeFunctionDelegate::CreateLambda([this](FRigVMGraphFunctionIdentifier& InFunctionToLocalize)
-		{
-			OnRequestLocalizeFunctionDialog(InFunctionToLocalize, GetControlRigBlueprint(), true);
-
-		   const URigVMLibraryNode* LocalizedFunctionNode = GetControlRigBlueprint()->GetLocalFunctionLibrary()->FindPreviouslyLocalizedFunction(InFunctionToLocalize);
-		   return LocalizedFunctionNode != nullptr;
-		
-		})
-	);
-	
-	TArray<FName> NodeNames = GetFocusedController()->ImportNodesFromText(TextToImport, true, true);
-
-	if (NodeNames.Num() > 0)
-	{
-		FBox2D Bounds;
-		Bounds.bIsValid = false;
-
-		TArray<FName> NodesToSelect;
-		for (const FName& NodeName : NodeNames)
-		{
-			const URigVMNode* Node = GetFocusedModel()->FindNodeByName(NodeName);
-			check(Node);
-
-			if (Node->IsInjected())
-			{
-				continue;
-			}
-			NodesToSelect.Add(NodeName);
-
-			FVector2D Position = Node->GetPosition();
-			FVector2D Size = Node->GetSize();
-
-			if (!Bounds.bIsValid)
-			{
-				Bounds.Min = Bounds.Max = Position;
-				Bounds.bIsValid = true;
-			}
-			Bounds += Position;
-			Bounds += Position + Size;
-		}
-
-		for (const FName& NodeName : NodesToSelect)
-		{
-			const URigVMNode* Node = GetFocusedModel()->FindNodeByName(NodeName);
-			check(Node);
-
-			FVector2D Position = Node->GetPosition();
-			GetFocusedController()->SetNodePositionByName(NodeName, PasteLocation + Position - Bounds.GetCenter(), true, false, true);
-		}
-
-		GetFocusedController()->SetNodeSelection(NodesToSelect);
-		GetFocusedController()->CloseUndoBracket();
-	}
-	else
-	{
-		GetFocusedController()->CancelUndoBracket();
-	}
-}
-
-void FControlRigEditor::PostUndo(bool bSuccess)
-{
-	const FTransaction* Transaction = GEditor->Trans->GetTransaction(GEditor->Trans->GetQueueLength() - GEditor->Trans->GetUndoCount());
-	IControlRigEditor::PostUndo(bSuccess);
-	PostTransaction(bSuccess, Transaction, false);
-}
-
-void FControlRigEditor::PostRedo(bool bSuccess)
-{
-	const FTransaction* Transaction = GEditor->Trans->GetTransaction(GEditor->Trans->GetQueueLength() - GEditor->Trans->GetUndoCount() - 1);
-	IControlRigEditor::PostRedo(bSuccess);
-	PostTransaction(bSuccess, Transaction, true);
-}
-
-void FControlRigEditor::PostTransaction(bool bSuccess, const FTransaction* Transaction, bool bIsRedo)
-{
-	EnsureValidRigElementsInDetailPanel();
-
-	if (UControlRigBlueprint* RigBlueprint = Cast<UControlRigBlueprint>(GetBlueprintObj()))
-	{
-		// Do not compile here. ControlRigBlueprint::PostTransacted decides when it is necessary to compile depending
-		// on the properties that are affected.
-		//Compile();
-		
-		USkeletalMesh* PreviewMesh = GetPersonaToolkit()->GetPreviewScene()->GetPreviewMesh();
-		if (PreviewMesh != RigBlueprint->GetPreviewMesh())
-		{
-			RigBlueprint->SetPreviewMesh(PreviewMesh);
-			GetPersonaToolkit()->SetPreviewMesh(PreviewMesh, true);
-		}
-
-		if (UControlRig* DebuggedControlRig = Cast<UControlRig>(RigBlueprint->GetObjectBeingDebugged()))
-		{
-			if(URigHierarchy* Hierarchy = DebuggedControlRig->GetHierarchy())
-			{
-				if(Hierarchy->Num() == 0)
-				{
-					OnHierarchyChanged();
-				}
-			}
-		}
-
-		if (FControlRigEditMode* EditMode = GetEditMode())
-		{
-			EditMode->RequestToRecreateControlShapeActors();
-		}
-	}
-}
-
-void FControlRigEditor::JumpToHyperlink(const UObject* ObjectReference, bool bRequestRename)
-{
-	if(const UControlRigGraph* Graph = Cast<UControlRigGraph>(ObjectReference))
-	{
-		OpenGraphAndBringToFront((UEdGraph*)Graph, true);
-		return;
-	}
-	
-	IControlRigEditor::JumpToHyperlink(ObjectReference, bRequestRename);
-}
-
-void FControlRigEditor::EnsureValidRigElementsInDetailPanel()
-{
-	UControlRigBlueprint* ControlRigBP = GetControlRigBlueprint();
-	URigHierarchy* Hierarchy = ControlRigBP->Hierarchy; 
-
-	TArray< TWeakObjectPtr<UObject> > SelectedObjects = Inspector->GetSelectedObjects();
-	for (TWeakObjectPtr<UObject> SelectedObject : SelectedObjects)
-	{
-		if (SelectedObject.IsValid())
-		{
-			if(URigVMDetailsViewWrapperObject* WrapperObject = Cast<URigVMDetailsViewWrapperObject>(SelectedObject.Get()))
-			{
-				if(const UScriptStruct* WrappedStruct = WrapperObject->GetWrappedStruct())
-				{
-					if (WrappedStruct->IsChildOf(FRigBaseElement::StaticStruct()))
-					{
-						FRigElementKey Key = WrapperObject->GetContent<FRigBaseElement>().GetKey();
-						if(!Hierarchy->Contains(Key))
-						{
-							ClearDetailObject();
-						}
-					}
-				}
-			}
-		}
-	}
-}
-
-void FControlRigEditor::OnStartWatchingPin()
-{
-	if (UEdGraphPin* Pin = GetCurrentlySelectedPin())
-	{
-		GetFocusedController()->SetPinIsWatched(Pin->GetName(), true);
-	}
-}
-
-bool FControlRigEditor::CanStartWatchingPin() const
-{
-	if (UEdGraphPin* Pin = GetCurrentlySelectedPin())
-	{
-		if (URigVMPin* ModelPin = GetFocusedModel()->FindPin(Pin->GetName()))
-		{
-			return ModelPin->GetParentPin() == nullptr &&
-					!ModelPin->RequiresWatch();
-		}
-	}
-	return false;
-}
-
-void FControlRigEditor::OnStopWatchingPin()
-{
-	if (UEdGraphPin* Pin = GetCurrentlySelectedPin())
-	{
-		GetFocusedController()->SetPinIsWatched(Pin->GetName(), false);
-	}
-}
-
-bool FControlRigEditor::CanStopWatchingPin() const
-{
-	if (UEdGraphPin* Pin = GetCurrentlySelectedPin())
-	{
-		if (URigVMPin* ModelPin = GetFocusedModel()->FindPin(Pin->GetName()))
-		{
-			return ModelPin->RequiresWatch();
-		}
-	}
-	return false;
-}
-
-void FControlRigEditor::OnToolkitHostingStarted(const TSharedRef<class IToolkit>& Toolkit)
-{
-	TSharedPtr<SWidget> InlineContent = Toolkit->GetInlineContent();
-	if (InlineContent.IsValid())
-	{
-		Toolbox->SetContent(InlineContent.ToSharedRef());
-	}
-}
-
-void FControlRigEditor::OnToolkitHostingFinished(const TSharedRef<class IToolkit>& Toolkit)
-{
-	Toolbox->SetContent(SNullWidget::NullWidget);
-}
-
-void FControlRigEditor::OnActiveTabChanged( TSharedPtr<SDockTab> PreviouslyActive, TSharedPtr<SDockTab> NewlyActivated )
-{
-	if (!NewlyActivated.IsValid())
-	{
-		TArray<UObject*> ObjArray;
-		Inspector->ShowDetailsForObjects(ObjArray);
-	}
-	else 
-	{
-		FBlueprintEditor::OnActiveTabChanged(PreviouslyActive, NewlyActivated);
-	}
-}
-
-void FControlRigEditor::OnAnimInitialized()
-{
-	UControlRigSkeletalMeshComponent* EditorSkelComp = Cast<UControlRigSkeletalMeshComponent>(GetPersonaToolkit()->GetPreviewScene()->GetPreviewMeshComponent());
-	if (EditorSkelComp)
-	{
-		EditorSkelComp->bRequiredBonesUpToDateDuringTick = 0;
-
-		UControlRigLayerInstance* AnimInstance = Cast<UControlRigLayerInstance>(EditorSkelComp->GetAnimInstance());
-		if (AnimInstance && GetControlRig())
-		{
-			// update control rig data to anim instance since animation system has been reinitialized
-			FInputBlendPose Filter;
-			AnimInstance->ResetControlRigTracks();
-			AnimInstance->AddControlRigTrack(0, GetControlRig());
-			AnimInstance->UpdateControlRigTrack(0, 1.0f, FControlRigIOSettings::MakeEnabled(), bExecutionControlRig);
-		}
-	}
-}
-
-void FControlRigEditor::UndoAction()
-{
-	GEditor->UndoTransaction();
-}
-
-void FControlRigEditor::RedoAction()
-{
-	GEditor->RedoTransaction();
-}
-
-void FControlRigEditor::CreateDefaultTabContents(const TArray<UBlueprint*>& InBlueprints)
-{
-	FBlueprintEditor::CreateDefaultTabContents(InBlueprints);
-}
-
-void FControlRigEditor::NewDocument_OnClicked(ECreatedDocumentType GraphType)
-{
-	if (GraphType == FBlueprintEditor::CGT_NewFunctionGraph)
-	{
-		if (UControlRigBlueprint* Blueprint = GetControlRigBlueprint())
-		{
-			if (URigVMController* Controller = Blueprint->GetOrCreateController(Blueprint->GetLocalFunctionLibrary()))
-			{
-				if (const URigVMLibraryNode* FunctionNode = Controller->AddFunctionToLibrary(TEXT("New Function"), true, FVector2D::ZeroVector, true, true))
-				{
-					if (const UEdGraph* NewGraph = Blueprint->GetEdGraph(FunctionNode->GetContainedGraph()))
-					{
-						OpenDocument(NewGraph, FDocumentTracker::OpenNewDocument);
-						RenameNewlyAddedAction(FunctionNode->GetFName());
-					}
-
-				}
-			}
-		}
-	}
-	else if(GraphType == FBlueprintEditor::CGT_NewEventGraph)
-	{
-		if (UControlRigBlueprint* Blueprint = GetControlRigBlueprint())
-		{
-			if(URigVMGraph* Model = Blueprint->AddModel(UControlRigGraphSchema::GraphName_ControlRig.ToString()))
-			{
-				if (const UEdGraph* NewGraph = Blueprint->GetEdGraph(Model))
-				{
-					OpenDocument(NewGraph, FDocumentTracker::OpenNewDocument);
-					RenameNewlyAddedAction(NewGraph->GetFName());
-				}
-			}
-		}
-	}
-}
-
-bool FControlRigEditor::IsSectionVisible(NodeSectionID::Type InSectionID) const
-{
-	switch (InSectionID)
-	{
-		case NodeSectionID::GRAPH:
-		case NodeSectionID::VARIABLE:
-		case NodeSectionID::FUNCTION:
-		{
-			return true;
-		}
-		case NodeSectionID::LOCAL_VARIABLE:
-		{
-			if(const URigVMGraph* Graph = GetFocusedModel())
-			{
-				const URigVMGraph* ParentGraph = Graph->GetParentGraph();				
-				if (ParentGraph && ParentGraph->IsA<URigVMFunctionLibrary>())
-				{
-					return true;
-				}
-			}
-		}
-		default:
-		{
-			break;
-		}
-	}
-	return false;
-}
-
-FGraphAppearanceInfo FControlRigEditor::GetGraphAppearance(UEdGraph* InGraph) const
-{
-	FGraphAppearanceInfo AppearanceInfo = FBlueprintEditor::GetGraphAppearance(InGraph);
-
-	if (GetBlueprintObj()->IsA(UControlRigBlueprint::StaticClass()))
-	{
-		AppearanceInfo.CornerText = LOCTEXT("AppearanceCornerText_ControlRig", "RIG");
-
-		if(const UControlRig* ControlRig = GetControlRig())
-		{
-			if(ControlRig->VM && ControlRig->VM->IsNativized())
-			{
-				AppearanceInfo.InstructionFade = 1;
-				AppearanceInfo.InstructionText = FText::FromString(
-					FString::Printf(TEXT("This graph runs a nativized VM (U%s)."), *ControlRig->VM->GetNativizedClass()->GetName())
-				);
-			}
-		}
-	}
-
-	return AppearanceInfo;
-}
-
-void FControlRigEditor::HandleModifiedEvent(ERigVMGraphNotifType InNotifType, URigVMGraph* InGraph, UObject* InSubject)
-{
-	DECLARE_SCOPE_HIERARCHICAL_COUNTER_FUNC()
-
-	UControlRigBlueprint* ControlRigBlueprint = CastChecked<UControlRigBlueprint>(GetBlueprintObj());
-	if (ControlRigBlueprint == nullptr)
-	{
-		return;
-	}
-
-	switch (InNotifType)
+	UControlRigSkeletalMeshComponent* EditorSkelComp = Cast<UControlRigSkeletalMeshComponent>(GetPersonaToolkit()->GetPreviewScene()->GetPreviewMeshComponent());
+	if (EditorSkelComp)
 	{
-		case ERigVMGraphNotifType::NodeSelectionChanged:
-		case ERigVMGraphNotifType::NodeSelected:
-		case ERigVMGraphNotifType::NodeDeselected:
-		{
-			if (UControlRigGraph* RigGraph = Cast<UControlRigGraph>(ControlRigBlueprint->GetEdGraph(InGraph)))
-			{
-				TSharedPtr<SGraphEditor> GraphEd = GetGraphEditor(RigGraph);
-				URigVMNode* Node = Cast<URigVMNode>(InSubject);
-				if (InNotifType == ERigVMGraphNotifType::NodeSelectionChanged)
-				{
-					const TArray<FName> SelectedNodes = InGraph->GetSelectNodes();
-					if (!SelectedNodes.IsEmpty())
-					{
-						Node = Cast<URigVMNode>(InGraph->FindNodeByName(SelectedNodes.Last()));	
-					}
-				}
-
-				if (GraphEd.IsValid() && Node != nullptr)
-				{
-					SetDetailViewForGraph(Node->GetGraph());
-
-					if (!RigGraph->bIsSelecting)
-					{
-						TGuardValue<bool> SelectingGuard(RigGraph->bIsSelecting, true);
-						if (UEdGraphNode* EdNode = RigGraph->FindNodeForModelNodeName(Node->GetFName()))
-						{
-							GraphEd->SetNodeSelection(EdNode, InNotifType == ERigVMGraphNotifType::NodeSelected);
-						}
-					}
-				}
-			}
-			break;
-		}
-		case ERigVMGraphNotifType::PinDefaultValueChanged:
-		{
-			URigVMPin* Pin = Cast<URigVMPin>(InSubject);
-
-			if(URigVMPin* RootPin = Pin->GetRootPin())
-			{
-				const FString DefaultValue = RootPin->GetDefaultValue();
-				if(!DefaultValue.IsEmpty())
-				{
-					// sync the value change with the unit(s) displayed 
-					TArray< TWeakObjectPtr<UObject> > SelectedObjects = Inspector->GetSelectedObjects();
-					for (TWeakObjectPtr<UObject> SelectedObject : SelectedObjects)
-					{
-						if (SelectedObject.IsValid())
-						{
-							if(URigVMDetailsViewWrapperObject* WrapperObject = Cast<URigVMDetailsViewWrapperObject>(SelectedObject.Get()))
-							{
-								if(WrapperObject->GetOuter() == Pin->GetNode())
-								{
-									const FProperty* TargetProperty = WrapperObject->GetClass()->FindPropertyByName(RootPin->GetFName());
-									if(TargetProperty)
-									{
-										uint8* PropertyStorage = TargetProperty->ContainerPtrToValuePtr<uint8>(WrapperObject);
-
-										// we are ok with not reacting to errors here
-										FRigVMPinDefaultValueImportErrorContext ErrorPipe;										
-										TargetProperty->ImportText_Direct(*DefaultValue, PropertyStorage, nullptr, PPF_None, &ErrorPipe);
-									}
-								}
-							}
-						}
-					}
-				}
-
-				if(const URigVMUnitNode* UnitNode = Cast<URigVMUnitNode>(RootPin->GetNode()))
-				{
-					if(UnitNode->IsEvent())
-					{
-						FBlueprintEditorUtils::MarkBlueprintAsStructurallyModified(GetBlueprintObj());
-						CacheNameLists();
-					}
-				}
-			}
-	
-			break;
-		}
-		case ERigVMGraphNotifType::PinArraySizeChanged:
-		case ERigVMGraphNotifType::PinBoundVariableChanged:
-		case ERigVMGraphNotifType::PinTypeChanged:
-		{
-			URigVMPin* Pin = Cast<URigVMPin>(InSubject);
+		EditorSkelComp->bRequiredBonesUpToDateDuringTick = 0;
 
-			if(Pin->GetNode()->IsSelected())
-			{
-				TArray<UObject*> Objects;
-				Objects.Add(Pin->GetNode());
-				SetDetailObjects(Objects);
-			}
-			break;
-		}
-		case ERigVMGraphNotifType::NodeRemoved:
-		{
-			if (URigVMCollapseNode* CollapseNode = Cast<URigVMCollapseNode>(InSubject))
-			{
-				if (UEdGraph* EdGraph = ControlRigBlueprint->GetEdGraph(CollapseNode->GetContainedGraph()))
-				{
-					CloseDocumentTab(EdGraph);
-					ClearDetailObject();
-				}
-			}
-			else if(URigVMFunctionReferenceNode* FunctionRefNode = Cast<URigVMFunctionReferenceNode>(InSubject))
-			{
-				ClearDetailObject();
-			}
-				
-			// fall through next case since we want to refresh the name lists
-			// both for removing or adding an event
-		}
-		case ERigVMGraphNotifType::NodeAdded:
-		{
-			if(URigVMUnitNode* UnitNode = Cast<URigVMUnitNode>(InSubject))
-			{
-				if(UnitNode->IsEvent())
-				{
-					FBlueprintEditorUtils::MarkBlueprintAsStructurallyModified(GetBlueprintObj());
-					CacheNameLists();
-				}
-			}
-			break;
-		}
-		default:
+		UControlRigLayerInstance* AnimInstance = Cast<UControlRigLayerInstance>(EditorSkelComp->GetAnimInstance());
+		if (AnimInstance && GetControlRig())
 		{
-			break;
+			// update control rig data to anim instance since animation system has been reinitialized
+			FInputBlendPose Filter;
+			AnimInstance->ResetControlRigTracks();
+			AnimInstance->AddControlRigTrack(0, GetControlRig());
+			AnimInstance->UpdateControlRigTrack(0, 1.0f, FControlRigIOSettings::MakeEnabled(), bExecutionControlRig);
 		}
 	}
 }
 
-void FControlRigEditor::HandleVMCompiledEvent(UObject* InCompiledObject, URigVM* InVM)
+void FControlRigEditor::HandleVMExecutedEvent(URigVMHost* InHost, const FName& InEventName)
 {
-	if(UControlRigBlueprint* RigBlueprint = Cast<UControlRigBlueprint>(InCompiledObject))
-	{
-		CompilerResultsListing->ClearMessages();
-		CompilerResultsListing->AddMessages(RigBlueprint->GetCompileLog().Messages);
-		RigBlueprint->GetCompileLog().Messages.Reset();
-		RigBlueprint->GetCompileLog().NumErrors = RigBlueprint->GetCompileLog().NumWarnings = 0;
-	}
-
-	RefreshDetailView();
-	
-	TArray<FName> TabIds;
-	TabIds.Add(*FString::Printf(TEXT("ControlRigMemoryDetails_%d"), (int32)ERigVMMemoryType::Literal));
-	TabIds.Add(*FString::Printf(TEXT("ControlRigMemoryDetails_%d"), (int32)ERigVMMemoryType::Work));
-	TabIds.Add(*FString::Printf(TEXT("ControlRigMemoryDetails_%d"), (int32)ERigVMMemoryType::Debug));
-
-	for (const FName& TabId : TabIds)
-	{
-		TSharedPtr<SDockTab> ActiveTab = GetTabManager()->FindExistingLiveTab(TabId);
-		if(ActiveTab)
-		{
-			if(ActiveTab->GetMetaData<FMemoryTypeMetaData>().IsValid())
-			{
-				ERigVMMemoryType MemoryType = ActiveTab->GetMetaData<FMemoryTypeMetaData>()->MemoryType;			
-				TSharedRef<IDetailsView> DetailsView = StaticCastSharedRef<IDetailsView>(ActiveTab->GetContent());
-				DetailsView->SetObject(InVM->GetMemoryByType(MemoryType));
-			}
-		}
-	}
-
-	UpdateGraphCompilerErrors();
-}
+	FRigVMEditor::HandleVMExecutedEvent(InHost, InEventName);
 
-void FControlRigEditor::HandleControlRigExecutedEvent(URigVMHost* InControlRig, const FName& InEventName)
-{
 	if (UControlRigBlueprint* ControlRigBP = GetControlRigBlueprint())
 	{
-		UControlRig* DebuggedControlRig = Cast<UControlRig>(ControlRigBP->GetObjectBeingDebugged());
-		if (DebuggedControlRig == nullptr)
-		{
-			DebuggedControlRig = GetControlRig();
-		}
-
 		URigHierarchy* Hierarchy = GetHierarchyBeingDebugged(); 
 
 		TArray< TWeakObjectPtr<UObject> > SelectedObjects = Inspector->GetSelectedObjects();
-		for (TWeakObjectPtr<UObject> SelectedObject : SelectedObjects)
-		{
-			if (SelectedObject.IsValid())
-			{
-				if(URigVMDetailsViewWrapperObject* WrapperObject = Cast<URigVMDetailsViewWrapperObject>(SelectedObject.Get()))
-				{
-					if (const UScriptStruct* Struct = WrapperObject->GetWrappedStruct())
-					{
-						if(Struct->IsChildOf(FRigBaseElement::StaticStruct()))
-						{
-							const FRigElementKey Key = WrapperObject->GetContent<FRigBaseElement>().GetKey();
-
-							FRigBaseElement* Element = Hierarchy->Find(Key);
-							if(Element == nullptr)
-							{
-								ClearDetailObject();
-								break;
-							}
-
-							if(FRigControlElement* ControlElement = Cast<FRigControlElement>(Element))
-							{
-								// compute all transforms
-								Hierarchy->GetTransform(ControlElement, ERigTransformType::CurrentGlobal);
-								Hierarchy->GetTransform(ControlElement, ERigTransformType::CurrentLocal);
-								Hierarchy->GetTransform(ControlElement, ERigTransformType::InitialGlobal);
-								Hierarchy->GetTransform(ControlElement, ERigTransformType::InitialLocal);
-								Hierarchy->GetControlOffsetTransform(ControlElement, ERigTransformType::CurrentGlobal);
-								Hierarchy->GetControlOffsetTransform(ControlElement, ERigTransformType::CurrentLocal);
-								Hierarchy->GetControlOffsetTransform(ControlElement, ERigTransformType::InitialGlobal);
-								Hierarchy->GetControlOffsetTransform(ControlElement, ERigTransformType::InitialLocal);
-								Hierarchy->GetControlShapeTransform(ControlElement, ERigTransformType::CurrentGlobal);
-								Hierarchy->GetControlShapeTransform(ControlElement, ERigTransformType::CurrentLocal);
-								Hierarchy->GetControlShapeTransform(ControlElement, ERigTransformType::InitialGlobal);
-								Hierarchy->GetControlShapeTransform(ControlElement, ERigTransformType::InitialLocal);
-
-								WrapperObject->SetContent<FRigControlElement>(*ControlElement);
-							}
-							else if(FRigTransformElement* TransformElement = Cast<FRigTransformElement>(Element))
-							{
-								// compute all transforms
-								Hierarchy->GetTransform(TransformElement, ERigTransformType::CurrentGlobal);
-								Hierarchy->GetTransform(TransformElement, ERigTransformType::CurrentLocal);
-								Hierarchy->GetTransform(TransformElement, ERigTransformType::InitialGlobal);
-								Hierarchy->GetTransform(TransformElement, ERigTransformType::InitialLocal);
-
-								WrapperObject->SetContent<FRigTransformElement>(*TransformElement);
-							}
-							else
-							{
-								WrapperObject->SetContent<FRigBaseElement>(*Element);
-							}
-						}
-					}
-				}
-			}
-		}
-
-		if(ControlRigBP->RigGraphDisplaySettings.NodeRunLimit > 1)
-		{
-			if(DebuggedControlRig)
-			{
-				if(URigVM* VM = DebuggedControlRig->GetVM())
-				{
-					bool bFoundLimitWarnings = false;
-					
-					const FRigVMByteCode& ByteCode = VM->GetByteCode();
-					for(int32 InstructionIndex = 0; InstructionIndex < ByteCode.GetNumInstructions(); InstructionIndex++)
-					{
-						const int32 Count = VM->GetInstructionVisitedCount(DebuggedControlRig->GetExtendedExecuteContext(), InstructionIndex);
-						if(Count > ControlRigBP->RigGraphDisplaySettings.NodeRunLimit)
-						{
-							bFoundLimitWarnings = true;
-
-							const FString CallPath = VM->GetByteCode().GetCallPathForInstruction(InstructionIndex); 
-							if(!KnownInstructionLimitWarnings.Contains(CallPath))
-							{
-								const FString Message = FString::Printf(
-                                    TEXT("Instruction has hit the NodeRunLimit\n(ran %d times, limit is %d)\n\nYou can increase the limit in the class settings."),
-                                    Count,
-                                    ControlRigBP->RigGraphDisplaySettings.NodeRunLimit
-                                );
-
-								if(DebuggedControlRig->RigVMLog)
-								{
-									DebuggedControlRig->RigVMLog->Entries.Add(
-										FRigVMLog::FLogEntry(EMessageSeverity::Warning, InEventName, InstructionIndex, Message
-									));
-								}
-
-								if(URigVMNode* Subject = Cast<URigVMNode>(VM->GetByteCode().GetSubjectForInstruction(InstructionIndex)))
-								{
-									FNotificationInfo Info(FText::FromString(Message));
-									Info.bFireAndForget = true;
-									Info.FadeOutDuration = 1.0f;
-									Info.ExpireDuration = 5.0f;
-
-									if(UControlRigGraph* EdGraph = Cast<UControlRigGraph>(GetControlRigBlueprint()->GetEdGraph(Subject->GetGraph())))
-									{
-										if(UEdGraphNode* Node = EdGraph->FindNodeForModelNodeName(Subject->GetFName()))
-										{
-											Info.Hyperlink = FSimpleDelegate::CreateLambda([this, Node] ()
-	                                        {
-	                                            JumpToHyperlink(Node, false);
-	                                        });
-									
-											Info.HyperlinkText = FText::FromString(Subject->GetName());
-										}
-									}
-
-									TSharedPtr<SNotificationItem> NotificationPtr = FSlateNotificationManager::Get().AddNotification(Info);
-									NotificationPtr->SetCompletionState(SNotificationItem::CS_Fail);
-								}
-
-								KnownInstructionLimitWarnings.Add(CallPath, Message);
-							}
-						}
-					}
-
-					if(!bFoundLimitWarnings)
-					{
-						KnownInstructionLimitWarnings.Reset();
-					}
-				}
-			}
-		}
-
-		if(ControlRigBP->VMRuntimeSettings.bEnableProfiling)
-		{
-			if(DebuggedControlRig)
-			{
-				ControlRigBP->RigGraphDisplaySettings.TotalMicroSeconds = DebuggedControlRig->GetExtendedExecuteContext().LastExecutionMicroSeconds;
-			}
-
-			if(ControlRigBP->RigGraphDisplaySettings.bAutoDetermineRange)
-			{
-				if(ControlRigBP->RigGraphDisplaySettings.LastMaxMicroSeconds < 0.0)
-				{
-					ControlRigBP->RigGraphDisplaySettings.LastMinMicroSeconds = ControlRigBP->RigGraphDisplaySettings.MinMicroSeconds; 
-					ControlRigBP->RigGraphDisplaySettings.LastMaxMicroSeconds = ControlRigBP->RigGraphDisplaySettings.MaxMicroSeconds;
-				}
-				else if(ControlRigBP->RigGraphDisplaySettings.MaxMicroSeconds >= 0.0)
-				{
-					const double T = 0.05;
-					ControlRigBP->RigGraphDisplaySettings.LastMinMicroSeconds = FMath::Lerp<double>(ControlRigBP->RigGraphDisplaySettings.LastMinMicroSeconds, ControlRigBP->RigGraphDisplaySettings.MinMicroSeconds, T); 
-					ControlRigBP->RigGraphDisplaySettings.LastMaxMicroSeconds = FMath::Lerp<double>(ControlRigBP->RigGraphDisplaySettings.LastMaxMicroSeconds, ControlRigBP->RigGraphDisplaySettings.MaxMicroSeconds, T); 
-				}
-
-				ControlRigBP->RigGraphDisplaySettings.MinMicroSeconds = DBL_MAX; 
-				ControlRigBP->RigGraphDisplaySettings.MaxMicroSeconds = (double)INDEX_NONE;
-			}
-			else
-			{
-				ControlRigBP->RigGraphDisplaySettings.LastMinMicroSeconds = ControlRigBP->RigGraphDisplaySettings.MinMicroSeconds; 
-				ControlRigBP->RigGraphDisplaySettings.LastMaxMicroSeconds = ControlRigBP->RigGraphDisplaySettings.MaxMicroSeconds;
-			}
-		}
-	}
-
-	UpdateGraphCompilerErrors();
-}
-
-void FControlRigEditor::HandleControlRigExecutionHalted(const int32 InstructionIndex, UObject* InNodeObject, const FName& InEntryName)
-{
-	if (HaltedAtNode == InNodeObject)
-	{
-		return;
-	}
-		
-	if (URigVMNode* InNode = Cast<URigVMNode>(InNodeObject))
-	{
-		SetHaltedNode(InNode);
-		
-		if (UControlRigBlueprint* Blueprint = GetControlRigBlueprint())
+		for (TWeakObjectPtr<UObject> SelectedObject : SelectedObjects)
 		{
-			if (Blueprint->GetAllModels().Contains(InNode->GetGraph()))
+			if (SelectedObject.IsValid())
 			{
-				if(UControlRigGraph* EdGraph = Cast<UControlRigGraph>(Blueprint->GetEdGraph(InNode->GetGraph())))
+				if(URigVMDetailsViewWrapperObject* WrapperObject = Cast<URigVMDetailsViewWrapperObject>(SelectedObject.Get()))
 				{
-					if(UEdGraphNode* EdNode = EdGraph->FindNodeForModelNodeName(InNode->GetFName()))
+					if (const UScriptStruct* Struct = WrapperObject->GetWrappedStruct())
 					{
-						JumpToHyperlink(EdNode, false);
+						if(Struct->IsChildOf(FRigBaseElement::StaticStruct()))
+						{
+							const FRigElementKey Key = WrapperObject->GetContent<FRigBaseElement>().GetKey();
+
+							FRigBaseElement* Element = Hierarchy->Find(Key);
+							if(Element == nullptr)
+							{
+								ClearDetailObject();
+								break;
+							}
+
+							if(FRigControlElement* ControlElement = Cast<FRigControlElement>(Element))
+							{
+								// compute all transforms
+								Hierarchy->GetTransform(ControlElement, ERigTransformType::CurrentGlobal);
+								Hierarchy->GetTransform(ControlElement, ERigTransformType::CurrentLocal);
+								Hierarchy->GetTransform(ControlElement, ERigTransformType::InitialGlobal);
+								Hierarchy->GetTransform(ControlElement, ERigTransformType::InitialLocal);
+								Hierarchy->GetControlOffsetTransform(ControlElement, ERigTransformType::CurrentGlobal);
+								Hierarchy->GetControlOffsetTransform(ControlElement, ERigTransformType::CurrentLocal);
+								Hierarchy->GetControlOffsetTransform(ControlElement, ERigTransformType::InitialGlobal);
+								Hierarchy->GetControlOffsetTransform(ControlElement, ERigTransformType::InitialLocal);
+								Hierarchy->GetControlShapeTransform(ControlElement, ERigTransformType::CurrentGlobal);
+								Hierarchy->GetControlShapeTransform(ControlElement, ERigTransformType::CurrentLocal);
+								Hierarchy->GetControlShapeTransform(ControlElement, ERigTransformType::InitialGlobal);
+								Hierarchy->GetControlShapeTransform(ControlElement, ERigTransformType::InitialLocal);
+
+								WrapperObject->SetContent<FRigControlElement>(*ControlElement);
+							}
+							else if(FRigTransformElement* TransformElement = Cast<FRigTransformElement>(Element))
+							{
+								// compute all transforms
+								Hierarchy->GetTransform(TransformElement, ERigTransformType::CurrentGlobal);
+								Hierarchy->GetTransform(TransformElement, ERigTransformType::CurrentLocal);
+								Hierarchy->GetTransform(TransformElement, ERigTransformType::InitialGlobal);
+								Hierarchy->GetTransform(TransformElement, ERigTransformType::InitialLocal);
+
+								WrapperObject->SetContent<FRigTransformElement>(*TransformElement);
+							}
+							else
+							{
+								WrapperObject->SetContent<FRigBaseElement>(*Element);
+							}
+						}
 					}
 				}
 			}
 		}
 	}
-	else 
-	{
-		if (InEntryName == GetControlRig()->GetEventQueue().Last())
-		{
-			SetHaltedNode(nullptr);
-		}
-	}
-}
-
-void FControlRigEditor::SetHaltedNode(URigVMNode* Node)
-{
-	if (HaltedAtNode)
-	{
-		HaltedAtNode->SetExecutionIsHaltedAtThisNode(false);
-	}
-	HaltedAtNode = Node;
-	if (HaltedAtNode)
-	{
-		HaltedAtNode->SetExecutionIsHaltedAtThisNode(true);
-	}
 }
 
 void FControlRigEditor::CreateEditorModeManager()
@@ -3424,7 +1287,7 @@ void FControlRigEditor::CreateEditorModeManager()
 
 void FControlRigEditor::Tick(float DeltaTime)
 {
-	FBlueprintEditor::Tick(DeltaTime);
+	FRigVMEditor::Tick(DeltaTime);
 
 	bool bDrawHierarchyBones = false;
 
@@ -3460,21 +1323,6 @@ void FControlRigEditor::Tick(float DeltaTime)
 			ControlRig->Evaluate_AnyThread();
 			bDrawHierarchyBones = true;
 		}
-
-		if (LastDebuggedRig != GetCustomDebugObjectLabel(Blueprint->GetObjectBeingDebugged()))
-		{
-			TArray<FCustomDebugObject> DebugList;
-			GetCustomDebugObjects(DebugList);
-
-			for (const FCustomDebugObject& DebugObject : DebugList)
-			{
-				if (DebugObject.NameOverride == LastDebuggedRig)
-				{
-					GetBlueprintObj()->SetObjectBeingDebugged(DebugObject.Object);
-					break;
-				}
-			}
-		}
 	}
 
 	if (FControlRigEditorEditMode* EditMode = GetEditMode())
@@ -3494,163 +1342,6 @@ void FControlRigEditor::Tick(float DeltaTime)
 	}
 }
 
-bool FControlRigEditor::IsEditable(UEdGraph* InGraph) const
-{
-	if(!IsGraphInCurrentBlueprint(InGraph))
-	{
-		return false;
-	}
-	
-	if(UControlRigBlueprint* ControlRigBlueprint = CastChecked<UControlRigBlueprint>(GetBlueprintObj()))
-	{
-		// aggregate graphs are always read only
-		if(const URigVMGraph* Model = ControlRigBlueprint->GetModel(InGraph))
-		{
-			if(Model->GetOuter()->IsA<URigVMAggregateNode>())
-			{
-				return false;
-			}
-		}
-
-		const UControlRig* ControlRig = GetControlRig();
-		if(ControlRig && ControlRig->VM)
-		{
-			const bool bIsReadOnly = ControlRig->VM->IsNativized();
-			const bool bIsEditable = !bIsReadOnly;
-			InGraph->bEditable = bIsEditable ? 1 : 0;
-			return bIsEditable;
-		}
-	}
-
-	return IControlRigEditor::IsEditable(InGraph);
-}
-
-bool FControlRigEditor::IsCompilingEnabled() const
-{
-	return true;
-}
-
-FText FControlRigEditor::GetGraphDecorationString(UEdGraph* InGraph) const
-{
-	return FText::GetEmpty();
-}
-
-TStatId FControlRigEditor::GetStatId() const
-{
-	RETURN_QUICK_DECLARE_CYCLE_STAT(FControlRigEditor, STATGROUP_Tickables);
-}
-
-void FControlRigEditor::OnSelectedNodesChangedImpl(const TSet<class UObject*>& NewSelection)
-{
-	DECLARE_SCOPE_HIERARCHICAL_COUNTER_FUNC()
-
-	UControlRigGraph* RigGraph = Cast<UControlRigGraph>(GetFocusedGraph());
-	if (RigGraph == nullptr)
-	{
-		return;
-	}
-
-	if (RigGraph->bIsSelecting || GIsTransacting)
-	{
-		return;
-	}
-
-	TGuardValue<bool> SelectGuard(RigGraph->bIsSelecting, true);
-
-	UControlRigBlueprint* ControlRigBlueprint = CastChecked<UControlRigBlueprint>(GetBlueprintObj());
-	if (ControlRigBlueprint)
-	{
-		TArray<FName> NodeNamesToSelect;
-		for (UObject* Object : NewSelection)
-		{
-			if (UControlRigGraphNode* ControlRigGraphNode = Cast<UControlRigGraphNode>(Object))
-			{
-				NodeNamesToSelect.Add(ControlRigGraphNode->GetModelNodeName());
-			}
-			else if(UEdGraphNode* Node = Cast<UEdGraphNode>(Object))
-			{
-				NodeNamesToSelect.Add(Node->GetFName());
-			}
-		}
-		GetFocusedController()->SetNodeSelection(NodeNamesToSelect, true, true);
-	}
-}
-
-void FControlRigEditor::HandleHideItem()
-{
-	DECLARE_SCOPE_HIERARCHICAL_COUNTER_FUNC()
-
-	UControlRigBlueprint* ControlRigBlueprint = CastChecked<UControlRigBlueprint>(GetBlueprintObj());
-
-	TSet<UObject*> SelectedNodes = GetSelectedNodes();
-	if(SelectedNodes.Num() > 0)
-	{
-		FScopedTransaction Transaction(LOCTEXT("HideRigItem", "Hide rig item"));
-
-		ControlRigBlueprint->Modify();
-
-		for(UObject* SelectedNodeObject : SelectedNodes)
-		{
-			if(UControlRigGraphNode* SelectedNode = Cast<UControlRigGraphNode>(SelectedNodeObject))
-			{
-				FBlueprintEditorUtils::RemoveNode(ControlRigBlueprint, SelectedNode, true);
-			}
-		}
-	}
-}
-
-bool FControlRigEditor::CanHideItem() const
-{
-	return GetNumberOfSelectedNodes() > 0;
-}
-
-void FControlRigEditor::OnBlueprintChangedImpl(UBlueprint* InBlueprint, bool bIsJustBeingCompiled)
-{
-	DECLARE_SCOPE_HIERARCHICAL_COUNTER_FUNC()
-
-	if (!bControlRigEditorInitialized)
-	{
-		return;
-	}
-
-	FBlueprintEditor::OnBlueprintChangedImpl(InBlueprint, bIsJustBeingCompiled);
-
-	if(InBlueprint == GetBlueprintObj())
-	{
-		if(bIsJustBeingCompiled)
-		{
-			UpdateControlRig();
-
-			if (!LastDebuggedRig.IsEmpty())
-			{
-				TArray<FCustomDebugObject> DebugList;
-				GetCustomDebugObjects(DebugList);
-
-				for (const FCustomDebugObject& DebugObject : DebugList)
-				{
-					if (DebugObject.NameOverride == LastDebuggedRig)
-					{
-						GetBlueprintObj()->SetObjectBeingDebugged(DebugObject.Object);
-						LastDebuggedRig.Empty();
-						break;
-					}
-				}
-			}
-		}
-	}
-}
-
-void FControlRigEditor::RefreshEditors(ERefreshBlueprintEditorReason::Type Reason)
-{
-	if(Reason == ERefreshBlueprintEditorReason::UnknownReason)
-	{
-		// we mark the reason as just compiled since we don't want to
-		// update the graph(s) all the time during compilation
-		Reason = ERefreshBlueprintEditorReason::BlueprintCompiled;
-	}
-	IControlRigEditor::RefreshEditors(Reason);
-}
-
 void FControlRigEditor::HandleViewportCreated(const TSharedRef<class IPersonaViewport>& InViewport)
 {
 	DECLARE_SCOPE_HIERARCHICAL_COUNTER_FUNC()
@@ -4283,59 +1974,24 @@ void FControlRigEditor::HandlePreviewSceneCreated(const TSharedRef<IPersonaPrevi
 	InPersonaPreviewScene->UnregisterForUndo();
 }
 
-void FControlRigEditor::UpdateControlRig()
+void FControlRigEditor::UpdateRigVMHost()
 {
 	DECLARE_SCOPE_HIERARCHICAL_COUNTER_FUNC()
 
+	FRigVMEditor::UpdateRigVMHost();
+
 	UControlRigBlueprint* Blueprint = Cast<UControlRigBlueprint>(GetBlueprintObj());
 	if(UClass* Class = Blueprint->GeneratedClass)
 	{
 		UControlRigSkeletalMeshComponent* EditorSkelComp = Cast<UControlRigSkeletalMeshComponent>(GetPersonaToolkit()->GetPreviewScene()->GetPreviewMeshComponent());
 		UControlRigLayerInstance* AnimInstance = Cast<UControlRigLayerInstance>(EditorSkelComp->GetAnimInstance());
+		UControlRig* ControlRig = GetControlRig();
 
-		if (AnimInstance)
+		if (AnimInstance && ControlRig)
 		{
-			if (const URigVMHost* CurrentHost = GetRigVMHost())
-			{
-				if(TestDataStrongPtr.IsValid())
-				{
-					TestDataStrongPtr->ReleaseReplay();
-				}
-
-				if(!IsValid(CurrentHost))
-				{
-					SetHost(nullptr);
-				}
-				
-				// if this control rig is from a temporary step,
-				// for example the reinstancing class, clear it 
-				// and create a new one!
-				if (CurrentHost->GetClass() != Class)
-				{
-					SetHost(nullptr);
-				}
-			}
-
-			UControlRig* ControlRig = GetControlRig();
-			if (ControlRig == nullptr)
-			{
-				ControlRig = NewObject<UControlRig>(EditorSkelComp, Class);
-				SetHost(ControlRig);
-				
-				// this is editing time rig
-				ControlRig->ExecutionType = ERigExecutionType::Editing;
-				ControlRig->RigVMLog = &ControlRigLog;
-
-				ControlRig->Initialize(true);
- 			}
-
  			PreviewInstance = Cast<UAnimPreviewInstance>(AnimInstance->GetSourceAnimInstance());
 			ControlRig->PreviewInstance = PreviewInstance;
 
-#if WITH_EDITOR
-			ControlRig->SetIsInDebugMode(ExecutionMode == EControlRigExecutionModeType_Debug);
-#endif
-
 			if (UControlRig* CDO = Cast<UControlRig>(Class->GetDefaultObject()))
 			{
 				CDO->ShapeLibraries = GetControlRigBlueprint()->ShapeLibraries;
@@ -4348,9 +2004,6 @@ void FControlRigEditor::UpdateControlRig()
 			{
 				ControlRig->SetObjectBinding(MakeShared<FControlRigObjectBinding>());
 			}
-			
-			// Make sure the object being debugged is the preview instance
-			GetBlueprintObj()->SetObjectBeingDebugged(ControlRig);
 
 			// initialize is moved post reinstance
 			AnimInstance->ResetControlRigTracks();
@@ -4365,26 +2018,25 @@ void FControlRigEditor::UpdateControlRig()
 				EditMode->SetObjects(ControlRig, EditorSkelComp,nullptr);
 			}
 
-			if(!bIsCompilingThroughUI)
-			{
-				Blueprint->SetFlags(RF_Transient);
-				Blueprint->RecompileVM();
-				Blueprint->ClearFlags(RF_Transient);
-			}
-
-			ControlRig->OnInitialized_AnyThread().AddSP(this, &FControlRigEditor::HandleControlRigExecutedEvent);
-			ControlRig->OnExecuted_AnyThread().AddSP(this, &FControlRigEditor::HandleControlRigExecutedEvent);
-			ControlRig->RequestInit();
 			ControlRig->ControlModified().AddSP(this, &FControlRigEditor::HandleOnControlModified);
-			ControlRig->GetExtendedExecuteContext().ExecutionHalted().AddSP(this, &FControlRigEditor::HandleControlRigExecutionHalted);
 		}
 	}
 }
 
+void FControlRigEditor::UpdateRigVMHost_PreClearOldHost(URigVMHost* InPreviousHost)
+{
+	if(TestDataStrongPtr.IsValid())
+	{
+		TestDataStrongPtr->ReleaseReplay();
+	}
+}
+
 void FControlRigEditor::CacheNameLists()
 {
 	DECLARE_SCOPE_HIERARCHICAL_COUNTER_FUNC()
 
+	FRigVMEditor::CacheNameLists();
+
 	if (UControlRigBlueprint* ControlRigBP = GetControlRigBlueprint())
 	{
 		TArray<UEdGraph*> EdGraphs;
@@ -4394,165 +2046,103 @@ void FControlRigEditor::CacheNameLists()
 		for (UEdGraph* Graph : EdGraphs)
 		{
 			UControlRigGraph* RigGraph = Cast<UControlRigGraph>(Graph);
-			if (RigGraph == nullptr)
-			{
-				continue;
-			}
-
-			const TArray<TSoftObjectPtr<UControlRigShapeLibrary>>* ShapeLibraries = &ControlRigBP->ShapeLibraries;
-			if(const UControlRig* DebuggedControlRig = Hierarchy->GetTypedOuter<UControlRig>())
-			{
-				ShapeLibraries = &DebuggedControlRig->GetShapeLibraries();
-			}
-			RigGraph->CacheEntryNameList();
-			RigGraph->CacheNameLists(Hierarchy, &ControlRigBP->DrawContainer, *ShapeLibraries);
-		}
-	}
-}
-
-void FControlRigEditor::AddReferencedObjects( FReferenceCollector& Collector )
-{
-	FBlueprintEditor::AddReferencedObjects(Collector);
-
-	TWeakObjectPtr<UControlRig> ControlRig(GetControlRig());
-	Collector.AddReferencedObject(ControlRig);
-}
-
-void FControlRigEditor::HandlePreviewMeshChanged(USkeletalMesh* InOldSkeletalMesh, USkeletalMesh* InNewSkeletalMesh)
-{
-	RebindToSkeletalMeshComponent();
-
-	if (GetObjectsCurrentlyBeingEdited()->Num() > 0)
-	{
-		if (UControlRigBlueprint* ControlRigBP = GetControlRigBlueprint())
-		{
-			ControlRigBP->SetPreviewMesh(InNewSkeletalMesh);
-			UpdateControlRig();
-			
-			if(UControlRig* DebuggedControlRig = Cast<UControlRig>(GetBlueprintObj()->GetObjectBeingDebugged()))
-			{
-				DebuggedControlRig->GetHierarchy()->Notify(ERigHierarchyNotification::HierarchyReset, nullptr);
-				DebuggedControlRig->Initialize(true);
-			}
-		}
-
-		Compile();
-	}
-}
-
-void FControlRigEditor::RebindToSkeletalMeshComponent()
-{
-	DECLARE_SCOPE_HIERARCHICAL_COUNTER_FUNC()
-
-	UDebugSkelMeshComponent* MeshComponent = GetPersonaToolkit()->GetPreviewScene()->GetPreviewMeshComponent();
-	if (MeshComponent)
-	{
-		bool bWasCreated = false;
-		FAnimCustomInstanceHelper::BindToSkeletalMeshComponent<UControlRigLayerInstance>(MeshComponent , bWasCreated);
-	}
-}
-
-void FControlRigEditor::UpdateStaleWatchedPins()
-{
-	UControlRigBlueprint* ControlRigBP = GetControlRigBlueprint();
-	if (ControlRigBP == nullptr)
-	{
-		return;
-	}
-
-	TSet<UEdGraphPin*> AllPins;
-	uint16 WatchCount;
-
-	// Find all unique pins being watched
-	FKismetDebugUtilities::ForeachPinWatch(
-		ControlRigBP,
-		[&AllPins, &WatchCount](UEdGraphPin* Pin)
-		{
-			++WatchCount;
-			if (Pin == nullptr)
-			{
-				return; // ~continue
-			}
-
-			UEdGraphNode* OwningNode = Pin->GetOwningNode();
-			// during node reconstruction, dead pins get moved to the transient 
-			// package (so just in case this blueprint got saved with dead pin watches)
-			if (OwningNode == nullptr)
+			if (RigGraph == nullptr)
 			{
-				return; // ~continue
+				continue;
 			}
 
-			if (!OwningNode->Pins.Contains(Pin))
+			const TArray<TSoftObjectPtr<UControlRigShapeLibrary>>* ShapeLibraries = &ControlRigBP->ShapeLibraries;
+			if(const UControlRig* DebuggedControlRig = Hierarchy->GetTypedOuter<UControlRig>())
 			{
-				return; // ~continue
+				ShapeLibraries = &DebuggedControlRig->GetShapeLibraries();
 			}
-
-			AllPins.Add(Pin);
+			RigGraph->CacheNameLists(Hierarchy, &ControlRigBP->DrawContainer, *ShapeLibraries);
 		}
-	);
-
-	// Refresh watched pins with unique pins (throw away null or duplicate watches)
-	if (WatchCount != AllPins.Num())
-	{
-		ControlRigBP->Status = BS_Dirty;
 	}
+}
 
-	FKismetDebugUtilities::ClearPinWatches(ControlRigBP);
+void FControlRigEditor::HandlePreviewMeshChanged(USkeletalMesh* InOldSkeletalMesh, USkeletalMesh* InNewSkeletalMesh)
+{
+	RebindToSkeletalMeshComponent();
 
-	TArray<URigVMGraph*> Models = ControlRigBP->GetAllModels();
-	for (URigVMGraph* Model : Models)
+	if (GetObjectsCurrentlyBeingEdited()->Num() > 0)
 	{
-		for (URigVMNode* ModelNode : Model->GetNodes())
+		if (UControlRigBlueprint* ControlRigBP = GetControlRigBlueprint())
 		{
-			TArray<URigVMPin*> ModelPins = ModelNode->GetAllPinsRecursively();
-			for (URigVMPin* ModelPin : ModelPins)
+			ControlRigBP->SetPreviewMesh(InNewSkeletalMesh);
+			UpdateRigVMHost();
+			
+			if(UControlRig* DebuggedControlRig = Cast<UControlRig>(GetBlueprintObj()->GetObjectBeingDebugged()))
 			{
-				if (ModelPin->RequiresWatch())
-				{
-					ControlRigBP->GetController(Model)->SetPinIsWatched(ModelPin->GetPinPath(), false, false);
-				}
+				DebuggedControlRig->GetHierarchy()->Notify(ERigHierarchyNotification::HierarchyReset, nullptr);
+				DebuggedControlRig->Initialize(true);
 			}
 		}
-	}
-	for (UEdGraphPin* Pin : AllPins)
-	{
-		FKismetDebugUtilities::AddPinWatch(ControlRigBP, FBlueprintWatchedPin(Pin));
-		UEdGraph* EdGraph = Pin->GetOwningNode()->GetGraph();
-		ControlRigBP->GetController(EdGraph)->SetPinIsWatched(Pin->GetName(), true, false);
+
+		Compile();
 	}
 }
 
-void FControlRigEditor::SetupGraphEditorEvents(UEdGraph* InGraph, SGraphEditor::FGraphEditorEvents& InEvents)
+void FControlRigEditor::RebindToSkeletalMeshComponent()
 {
 	DECLARE_SCOPE_HIERARCHICAL_COUNTER_FUNC()
 
-	FBlueprintEditor::SetupGraphEditorEvents(InGraph, InEvents);
-
-	InEvents.OnCreateActionMenu = SGraphEditor::FOnCreateActionMenu::CreateSP(this, &FControlRigEditor::HandleCreateGraphActionMenu);
-	InEvents.OnTextCommitted = FOnNodeTextCommitted::CreateSP(this, &FControlRigEditor::OnNodeTitleCommitted);
-}
-
-void FControlRigEditor::FocusInspectorOnGraphSelection(const TSet<UObject*>& NewSelection, bool bForceRefresh)
-{
-	// nothing to do here for control rig
+	UDebugSkelMeshComponent* MeshComponent = GetPersonaToolkit()->GetPreviewScene()->GetPreviewMeshComponent();
+	if (MeshComponent)
+	{
+		bool bWasCreated = false;
+		FAnimCustomInstanceHelper::BindToSkeletalMeshComponent<UControlRigLayerInstance>(MeshComponent , bWasCreated);
+	}
 }
 
-FActionMenuContent FControlRigEditor::HandleCreateGraphActionMenu(UEdGraph* InGraph, const FVector2D& InNodePosition, const TArray<UEdGraphPin*>& InDraggedPins, bool bAutoExpand, SGraphEditor::FActionMenuClosed InOnMenuClosed)
+void FControlRigEditor::GenerateEventQueueMenuContent(FMenuBuilder& MenuBuilder)
 {
-	return FBlueprintEditor::OnCreateGraphActionMenu(InGraph, InNodePosition, InDraggedPins, bAutoExpand, InOnMenuClosed);
-}
+	MenuBuilder.BeginSection(TEXT("Events"));
+    MenuBuilder.AddMenuEntry(FControlRigEditorCommands::Get().ConstructionEvent, TEXT("Setup"), TAttribute<FText>(), TAttribute<FText>(), GetEventQueueIcon(ConstructionEventQueue));
+    MenuBuilder.AddMenuEntry(FControlRigEditorCommands::Get().ForwardsSolveEvent, TEXT("Forwards Solve"), TAttribute<FText>(), TAttribute<FText>(), GetEventQueueIcon(ForwardsSolveEventQueue));
+    MenuBuilder.AddMenuEntry(FControlRigEditorCommands::Get().BackwardsSolveEvent, TEXT("Backwards Solve"), TAttribute<FText>(), TAttribute<FText>(), GetEventQueueIcon(BackwardsSolveEventQueue));
+    MenuBuilder.EndSection();
 
-void FControlRigEditor::OnNodeTitleCommitted(const FText& NewText, ETextCommit::Type CommitInfo, UEdGraphNode* NodeBeingChanged)
-{
-	DECLARE_SCOPE_HIERARCHICAL_COUNTER_FUNC()
+    MenuBuilder.BeginSection(TEXT("Validation"));
+    MenuBuilder.AddMenuEntry(FControlRigEditorCommands::Get().BackwardsAndForwardsSolveEvent, TEXT("BackwardsAndForwards"), TAttribute<FText>(), TAttribute<FText>(), GetEventQueueIcon(BackwardsAndForwardsSolveEventQueue));
+    MenuBuilder.EndSection();
 
-	if (UEdGraphNode_Comment* CommentBeingChanged = Cast<UEdGraphNode_Comment>(NodeBeingChanged))
-	{
-		if (UControlRigBlueprint* ControlRigBP = GetControlRigBlueprint())
-		{
-			GetFocusedController()->SetCommentTextByName(CommentBeingChanged->GetFName(), NewText.ToString(), CommentBeingChanged->FontSize, CommentBeingChanged->bCommentBubbleVisible, CommentBeingChanged->bColorCommentBubble, true, true);
-		}
-	}
+    if (const UControlRigBlueprint* RigBlueprint = Cast<UControlRigBlueprint>(GetBlueprintObj()))
+    {
+    	URigVMEdGraphSchema* Schema = CastChecked<URigVMEdGraphSchema>(RigBlueprint->GetRigVMEdGraphSchemaClass()->GetDefaultObject());
+    	
+    	bool bFoundUserDefinedEvent = false;
+    	const TArray<FName> EntryNames = RigBlueprint->GetRigVMClient()->GetEntryNames();
+    	for(const FName& EntryName : EntryNames)
+    	{
+    		if(Schema->IsRigVMDefaultEvent(EntryName))
+    		{
+    			continue;
+    		}
+
+    		if(!bFoundUserDefinedEvent)
+    		{
+    			MenuBuilder.AddSeparator();
+    			bFoundUserDefinedEvent = true;
+    		}
+
+    		FString EventNameStr = EntryName.ToString();
+    		if(!EventNameStr.EndsWith(TEXT("Event")))
+    		{
+    			EventNameStr += TEXT(" Event");
+    		}
+
+    		MenuBuilder.AddMenuEntry(
+    			FText::FromString(EventNameStr),
+    			FText::FromString(FString::Printf(TEXT("Runs the user defined %s"), *EventNameStr)),
+    			GetEventQueueIcon({EntryName}),
+    			FUIAction(
+    				FExecuteAction::CreateSP(this, &FRigVMEditor::SetEventQueue, TArray<FName>({EntryName})),
+    				FCanExecuteAction()
+    			)
+    		);
+    	}
+    }
 }
 
 FTransform FControlRigEditor::GetRigElementTransform(const FRigElementKey& InElement, bool bLocal, bool bOnDebugInstance) const
@@ -4667,101 +2257,36 @@ void FControlRigEditor::SetRigElementTransform(const FRigElementKey& InElement,
 	}
 }
 
-void FControlRigEditor::NotifyPreChange(FProperty* PropertyAboutToChange)
-{
-	FBlueprintEditor::NotifyPreChange(PropertyAboutToChange);
-
-	if (UControlRigBlueprint* ControlRigBP = GetControlRigBlueprint())
-	{
-		ControlRigBP->Modify();
-	}
-}
-
-void FControlRigEditor::NotifyPostChange(const FPropertyChangedEvent& PropertyChangedEvent, FProperty* PropertyThatChanged)
-{
-	FBlueprintEditor::NotifyPostChange(PropertyChangedEvent, PropertyThatChanged);
-
-	// we need to listen to changes for variables on the blueprint here since
-	// OnFinishedChangingProperties is called only for top level property changes.
-	// changes on a lower level property like transform under a user defined struct
-	// only go through this.
-	UControlRigBlueprint* ControlRigBP = GetControlRigBlueprint();
-	if(GetControlRig() && ControlRigBP)
-	{
-		bool bUseCDO = false; 
-		if(PropertyChangedEvent.GetNumObjectsBeingEdited() == 1)
-		{
-			bUseCDO = PropertyChangedEvent.GetObjectBeingEdited(0)->HasAnyFlags(RF_ClassDefaultObject);
-		}
-		
-		const FName VarName = PropertyChangedEvent.MemberProperty->GetFName();
-		for(FBPVariableDescription& NewVariable : ControlRigBP->NewVariables)
-		{
-			if(NewVariable.VarName == VarName)
-			{
-				UpdateDefaultValueForVariable(NewVariable, bUseCDO);
-				break;
-			}
-		}
-	}
-}
-
 void FControlRigEditor::OnFinishedChangingProperties(const FPropertyChangedEvent& PropertyChangedEvent)
 {
+	FRigVMEditor::OnFinishedChangingProperties(PropertyChangedEvent);
+	
 	DECLARE_SCOPE_HIERARCHICAL_COUNTER_FUNC()
 
 	UControlRigBlueprint* ControlRigBP = GetControlRigBlueprint();
 
 	if (ControlRigBP)
 	{
-		if (PropertyChangedEvent.MemberProperty->GetNameCPP() == GET_MEMBER_NAME_STRING_CHECKED(UControlRigBlueprint, VMCompileSettings))
-		{
-			ControlRigBP->RecompileVM();
-			return;
-		}
-
-		if (PropertyChangedEvent.MemberProperty->GetNameCPP() == GET_MEMBER_NAME_STRING_CHECKED(UControlRigBlueprint, VMRuntimeSettings))
-		{
-			ControlRigBP->VMRuntimeSettings.Validate();
-			ControlRigBP->PropagateRuntimeSettingsFromBPToInstances();
-			return;
-		}
-
 		if (PropertyChangedEvent.MemberProperty->GetNameCPP() == GET_MEMBER_NAME_STRING_CHECKED(UControlRigBlueprint, HierarchySettings))
 		{
 			ControlRigBP->PropagateHierarchyFromBPToInstances();
-			return;
 		}
 
-		if (PropertyChangedEvent.MemberProperty->GetNameCPP() == GET_MEMBER_NAME_STRING_CHECKED(UControlRigBlueprint, DrawContainer))
+		else if (PropertyChangedEvent.MemberProperty->GetNameCPP() == GET_MEMBER_NAME_STRING_CHECKED(UControlRigBlueprint, DrawContainer))
 		{
 			ControlRigBP->PropagateDrawInstructionsFromBPToInstances();
-			return;
-		}
-	}
-}
-
-void FControlRigEditor::OnPropertyChanged(UObject* InObject, FPropertyChangedEvent& InEvent)
-{
-	UControlRigBlueprint* ControlRigBP = GetControlRigBlueprint();
-
-	if (ControlRigBP && InObject == ControlRigBP)
-	{
-		// if the models have changed - we may need to close a document
-		if(InEvent.MemberProperty == ControlRigBP->GetClass()->FindPropertyByName(GET_MEMBER_NAME_CHECKED(UControlRigBlueprint, RigVMClient)) ||
-			InEvent.MemberProperty == ControlRigBP->GetClass()->FindPropertyByName(GET_MEMBER_NAME_CHECKED(UControlRigBlueprint, UbergraphPages)))
-		{
-			DocumentManager->CleanInvalidTabs();
 		}
 	}
 }
 
 void FControlRigEditor::OnWrappedPropertyChangedChainEvent(URigVMDetailsViewWrapperObject* InWrapperObject, const FString& InPropertyPath, FPropertyChangedChainEvent& InPropertyChangedChainEvent)
 {
+	FRigVMEditor::OnWrappedPropertyChangedChainEvent(InWrapperObject, InPropertyPath, InPropertyChangedChainEvent);
+	
 	check(InWrapperObject);
-	check(!WrapperObjects.IsEmpty());
+	check(!GetWrapperObjects().IsEmpty());
 
-	TGuardValue<bool> SuspendDetailsPanelRefresh(bSuspendDetailsPanelRefresh, true);
+	TGuardValue<bool> SuspendDetailsPanelRefresh(GetSuspendDetailsPanelRefreshFlag(), true);
 
 	UControlRigBlueprint* ControlRigBP = GetControlRigBlueprint();
 
@@ -4770,11 +2295,11 @@ void FControlRigEditor::OnWrappedPropertyChangedChainEvent(URigVMDetailsViewWrap
 	{
 		if(WrappedStruct->IsChildOf(FRigBaseElement::StaticStruct()))
 		{
-			check(WrappedStruct == WrapperObjects[0]->GetWrappedStruct());
+			check(WrappedStruct == GetWrapperObjects()[0]->GetWrappedStruct());
 
 			URigHierarchy* Hierarchy = CastChecked<URigHierarchy>(InWrapperObject->GetSubject());
 			const FRigBaseElement WrappedElement = InWrapperObject->GetContent<FRigBaseElement>();
-			const FRigBaseElement FirstWrappedElement = WrapperObjects[0]->GetContent<FRigBaseElement>();
+			const FRigBaseElement FirstWrappedElement = GetWrapperObjects()[0]->GetContent<FRigBaseElement>();
 			const FRigElementKey& Key = WrappedElement.GetKey();
 			if(!Hierarchy->Contains(Key))
 			{
@@ -4870,442 +2395,71 @@ void FControlRigEditor::OnWrappedPropertyChangedChainEvent(URigVMDetailsViewWrap
 				ERigTransformType::Type TransformType = Local::GetTransformTypeFromPath(PropertyPath);
 				bIsInitial = bIsInitial || ERigTransformType::IsInitial(TransformType);
 
-				const FTransform Transform = WrapperObjects[0]->GetContent<FRigControlElement>().Offset.Get(TransformType);
+				const FTransform Transform = GetWrapperObjects()[0]->GetContent<FRigControlElement>().Offset.Get(TransformType);
 				
 				ControlRigBP->Hierarchy->SetControlOffsetTransform(ControlElement, Transform, ERigTransformType::MakeInitial(TransformType), true, true, false, true);
-			}
-			else if(PropertyPath.RemoveFromStart(ShapeString))
-			{
-				FRigControlElement* ControlElement = ControlRigBP->Hierarchy->Find<FRigControlElement>(WrappedElement.GetKey());
-				if(ControlElement == nullptr)
-				{
-					return;
-				}
-
-				ERigTransformType::Type TransformType = Local::GetTransformTypeFromPath(PropertyPath);
-				bIsInitial = bIsInitial || ERigTransformType::IsInitial(TransformType);
-
-				const FTransform Transform = WrapperObjects[0]->GetContent<FRigControlElement>().Shape.Get(TransformType);
-				
-				ControlRigBP->Hierarchy->SetControlShapeTransform(ControlElement, Transform, ERigTransformType::MakeInitial(TransformType), true, false, true);
-			}
-			else if(PropertyPath.RemoveFromStart(SettingsString))
-			{
-				const FRigControlSettings Settings  = InWrapperObject->GetContent<FRigControlElement>().Settings;
-
-				FRigControlElement* ControlElement = ControlRigBP->Hierarchy->Find<FRigControlElement>(WrappedElement.GetKey());
-				if(ControlElement == nullptr)
-				{
-					return;
-				}
-
-				ControlRigBP->Hierarchy->SetControlSettings(ControlElement, Settings, true, false, true);
-			}
-
-			if(IsConstructionModeEnabled() || bIsInitial)
-			{
-				ControlRigBP->PropagatePoseFromBPToInstances();
-				ControlRigBP->Modify();
-				ControlRigBP->MarkPackageDirty();
-			}
-		}
-		else if(WrappedStruct->IsChildOf(FRigVMGraphVariableDescription::StaticStruct()))
-		{
-			check(WrappedStruct == WrapperObjects[0]->GetWrappedStruct());
-			
-			const FRigVMGraphVariableDescription VariableDescription = InWrapperObject->GetContent<FRigVMGraphVariableDescription>();
-			URigVMGraph* Graph = CastChecked<URigVMGraph>(InWrapperObject->GetSubject());
-			URigVMController* Controller = ControlRigBP->GetController(Graph);
-			if (PropertyPath == TEXT("Name") && MyBlueprintWidget.IsValid())
-			{
-				if (FEdGraphSchemaAction_BlueprintVariableBase* VariableAcion = MyBlueprintWidget->SelectionAsBlueprintVariable())
-				{
-					const FName OldVariableName = VariableAcion->GetVariableName();
-					if (!OldVariableName.IsNone())
-					{
-						for (FRigVMGraphVariableDescription& Variable : Graph->GetLocalVariables())
-						{
-							if (Variable.Name == OldVariableName)
-							{
-								Controller->RenameLocalVariable(OldVariableName, VariableDescription.Name);
-								break;
-							}
-						}
-					}
-				}
-				RefreshMyBlueprint();
-				GetControlRigBlueprint()->RequestAutoVMRecompilation();
-			}
-			else if (PropertyPath == TEXT("CPPType") || PropertyPath == TEXT("CPPTypeObject"))
-			{			
-				for (FRigVMGraphVariableDescription& Variable : Graph->GetLocalVariables())
-				{
-					if (Variable.Name == VariableDescription.Name)
-					{
-						Controller->SetLocalVariableType(Variable.Name, VariableDescription.CPPType, VariableDescription.CPPTypeObject);
-						break;
-					}
-				}
-				GetControlRigBlueprint()->RequestAutoVMRecompilation();
-			}
-			else if (PropertyPath == TEXT("DefaultValue"))
-			{
-				FRigVMControllerNotifGuard NotifGuard(Controller, true);
-				for (FRigVMGraphVariableDescription& Variable : Graph->GetLocalVariables())
-				{
-					if (Variable.Name == VariableDescription.Name)
-					{
-						Controller->SetLocalVariableDefaultValue(Variable.Name, VariableDescription.DefaultValue, true, true);
-						break;
-					}
-				}
-
-				// Do not recompile now! That destroys the object that is currently being displayed (the literal memory storage), and can cause a crash.
-				// The user has to manually trigger the recompilation.
-			}		
-		}
-	}
-	else if(!InWrapperObject->GetWrappedNodeNotation().IsEmpty())
-	{
-		FName RootPinName = InPropertyChangedChainEvent.PropertyChain.GetHead()->GetValue()->GetFName();
-		FProperty* TargetProperty = WrapperObjects[0]->GetClass()->FindPropertyByName(RootPinName);
-		uint8* FirstPropertyStorage = TargetProperty->ContainerPtrToValuePtr<uint8>(WrapperObjects[0].Get());
-
-		URigVMNode* Node = CastChecked<URigVMNode>(InWrapperObject->GetSubject());
-
-		FString DefaultValue = FRigVMStruct::ExportToFullyQualifiedText(TargetProperty, FirstPropertyStorage);
-
-		if(TargetProperty->IsA<FStrProperty>() || TargetProperty->IsA<FNameProperty>())
-		{
-			DefaultValue.TrimCharInline(TEXT('\"'), nullptr);
-		}
-		
-		URigVMController* Controller = GetControlRigBlueprint()->GetController(Node->GetGraph());
-
-		if (!DefaultValue.IsEmpty())
-		{
-			FString PinPath = FString::Printf(TEXT("%s.%s"), *Node->GetName(), *RootPinName.ToString());
-			const bool bInteractive = InPropertyChangedChainEvent.ChangeType == EPropertyChangeType::Interactive;
-			Controller->SetPinDefaultValue(PinPath, DefaultValue, true, !bInteractive, true, !bInteractive);
-		}
-	}
-}
-
-void FControlRigEditor::OnRequestLocalizeFunctionDialog(FRigVMGraphFunctionIdentifier& InFunction, URigVMBlueprint* InTargetBlueprint, bool bForce)
-{
-	check(InTargetBlueprint);
-
-	if(InTargetBlueprint != GetControlRigBlueprint())
-	{
-		return;
-	}
-	
-	if(URigVMController* TargetController = InTargetBlueprint->GetController(InTargetBlueprint->GetDefaultModel()))
-	{
-		bool bIsPublic;
-		if (FRigVMGraphFunctionData::FindFunctionData(InFunction, &bIsPublic))
-		{
-			if (bForce || bIsPublic)
-			{
-				TSharedRef<SRigVMGraphFunctionLocalizationDialog> LocalizationDialog = SNew(SRigVMGraphFunctionLocalizationDialog)
-							.Function(InFunction)
-							.TargetBlueprint(InTargetBlueprint);
-
-				if (LocalizationDialog->ShowModal() != EAppReturnType::Cancel)
-				{
-					TargetController->LocalizeFunctions(LocalizationDialog->GetFunctionsToLocalize(), true, true, true);
-				}
-			}
-		}
-	}
-}
-
-FRigVMController_BulkEditResult FControlRigEditor::OnRequestBulkEditDialog(URigVMBlueprint* InBlueprint, URigVMController* InController,
-	URigVMLibraryNode* InFunction, ERigVMControllerBulkEditType InEditType)
-{
-	if (bAllowBulkEdits)
-	{
-		FRigVMController_BulkEditResult Result;
-		Result.bCanceled = false; 
-		Result.bSetupUndoRedo = false;
-		return Result;
-	}
-	
-	const TArray<FAssetData> FirstLevelReferenceAssets = InController->GetAffectedAssets(InEditType, false);
-	if(FirstLevelReferenceAssets.Num() == 0)
-	{
-		return FRigVMController_BulkEditResult();
-	}
-	
-	TSharedRef<SRigVMGraphFunctionBulkEditDialog> BulkEditDialog = SNew(SRigVMGraphFunctionBulkEditDialog)
-	.Blueprint(InBlueprint)
-	.Controller(InController)
-	.Function(InFunction)
-	.EditType(InEditType);
-
-	FRigVMController_BulkEditResult Result;
-	Result.bCanceled = BulkEditDialog->ShowModal() == EAppReturnType::Cancel; 
-	Result.bSetupUndoRedo = false;
-
-	if (!Result.bCanceled)
-	{
-		bAllowBulkEdits = true;
-	}
-	
-	return Result;
-}
-
-bool FControlRigEditor::OnRequestBreakLinksDialog(TArray<URigVMLink*> InLinks)
-{
-	if(InLinks.Num() == 0)
-	{
-		return true;
-	}
-
-	TSharedRef<SRigVMGraphBreakLinksDialog> BreakLinksDialog = SNew(SRigVMGraphBreakLinksDialog)
-	.Links(InLinks)
-	.OnFocusOnLink(FRigVMOnFocusOnLinkRequestedDelegate::CreateLambda([&](URigVMLink* InLink)
-	{
-		HandleJumpToHyperlink(InLink);
-	}));
-
-	return BreakLinksDialog->ShowModal() == EAppReturnType::Ok; 
-}
-
-TRigVMTypeIndex FControlRigEditor::OnRequestPinTypeSelectionDialog(const TArray<TRigVMTypeIndex>& InTypes)
-{
-	if(InTypes.Num() == 0)
-	{
-		return true;
-	}
-
-	TRigVMTypeIndex Answer = INDEX_NONE;
-
-	const FRigVMRegistry& Registry = FRigVMRegistry::Get();
-
-	TArray<TSharedPtr<FName>> TypeNames;
-	TMap<FName, uint8> TypeNameToIndex;
-	TypeNames.Reserve(InTypes.Num());
-	for (int32 i=0; i<InTypes.Num(); ++i)
-	{
-		const TRigVMTypeIndex& TypeIndex = InTypes[i];
-		TRigVMTypeIndex FinalType = TypeIndex;
-		if (FinalType == RigVMTypeUtils::TypeIndex::Float)
-		{
-			FinalType = RigVMTypeUtils::TypeIndex::Double;
-		}
-		if (FinalType == RigVMTypeUtils::TypeIndex::FloatArray)
-		{
-			FinalType = RigVMTypeUtils::TypeIndex::DoubleArray;
-		}
-
-		const FRigVMTemplateArgumentType& ArgumentType = Registry.GetType(FinalType);
-		if (!TypeNames.ContainsByPredicate([&ArgumentType](const TSharedPtr<FName>& InName)
-		{
-			return *InName.Get() == ArgumentType.CPPType;
-		}))
-		{
-			TypeNames.AddUnique(MakeShared<FName>(ArgumentType.CPPType));
-			TypeNameToIndex.Add(ArgumentType.CPPType, i);
-		}
-	}
-	TSharedPtr< SWindow > Window = SNew(SWindow)
-		.Title(LOCTEXT("SelectPinType", "Select Pin Type"))
-		.ScreenPosition(FSlateApplication::Get().GetCursorPos())
-		.SizingRule(ESizingRule::Autosized)
-		.AutoCenter(EAutoCenter::None)
-		.SupportsMaximize(false)
-		.SupportsMinimize(false)
-		[
-			SNew( SBorder )
-			.Padding( 4.f )
-			.BorderImage( FAppStyle::GetBrush( "ToolPanel.GroupBorder" ) )
-			[
-				SNew(SBox)
-				.MaxDesiredHeight(static_cast<float>(300))
-				[
-					SNew(SVerticalBox)
-					+SVerticalBox::Slot()
-					[
-						SNew(SBox)
-						.MaxDesiredHeight(static_cast<float>(300))
-						[
-							SNew(SScrollBox)
-							+SScrollBox::Slot()
-							[
-								SNew(SListView<TSharedPtr<FName>>)
-									.ListItemsSource(&TypeNames)
-									.OnGenerateRow_Lambda([&Registry, &TypeNameToIndex, &InTypes](const TSharedPtr<FName> InItem, const TSharedRef<STableViewBase>& Owner)
-									{
-										TRigVMTypeIndex TypeIndex = InTypes[TypeNameToIndex.FindChecked(*InItem.Get())];
-										const FRigVMTemplateArgumentType Type = FRigVMRegistry::Get().GetType(TypeIndex);
-										const bool bIsArray = Type.IsArray();
-										static const FName TypeIcon(TEXT("Kismet.VariableList.TypeIcon"));
-										static const FName ArrayTypeIcon(TEXT("Kismet.VariableList.ArrayTypeIcon"));
-
-										const FEdGraphPinType PinType = RigVMTypeUtils::PinTypeFromTypeIndex(TypeIndex);
-										const UControlRigGraphSchema* Schema = GetDefault<UControlRigGraphSchema>();
-										const FLinearColor Color = Schema->GetPinTypeColor(PinType);
-										
-										return SNew(STableRow<TSharedPtr<FString>>, Owner)
-												.Padding(FMargin(16, 4, 16, 4))
-												[
-													
-													SNew(SHorizontalBox)
-													+ SHorizontalBox::Slot()
-													.AutoWidth()
-													.VAlign(VAlign_Center)
-													[
-														SNew(SBox)
-														.HeightOverride(16.0f)
-														[
-															SNew(SImage)
-															.Image(bIsArray ? FAppStyle::GetBrush(ArrayTypeIcon) : FAppStyle::GetBrush(TypeIcon))
-															.ColorAndOpacity(Color)
-														]
-													]
-
-													+ SHorizontalBox::Slot()
-													[
-														SNew(STextBlock).Text(FText::FromName(*InItem.Get()))
-													]
-												];
-									})
-									.OnSelectionChanged_Lambda([&Answer, &TypeNames, &TypeNameToIndex, &InTypes](const TSharedPtr<FName> InName, ESelectInfo::Type)
-									{
-										Answer = InTypes[TypeNameToIndex.FindChecked(*InName.Get())];
-										FSlateApplication::Get().GetActiveModalWindow()->RequestDestroyWindow();
-									})
-							]
-						]
-					]
-				]
-			]
-		];
-
-	GEditor->EditorAddModalWindow(Window.ToSharedRef());
-	return Answer;
-}
-
-void FControlRigEditor::HandleJumpToHyperlink(const UObject* InSubject)
-{
-	UControlRigBlueprint* RigBlueprint = GetControlRigBlueprint();
-	if(RigBlueprint == nullptr)
-	{
-		return;
-	}
-
-	const URigVMGraph* GraphToJumpTo = nullptr;
-	const URigVMNode* NodeToJumpTo = nullptr;
-	const URigVMPin* PinToJumpTo = nullptr;
-	if(const URigVMNode* Node = Cast<URigVMNode>(InSubject))
-	{
-		GraphToJumpTo = Node->GetGraph();
-		NodeToJumpTo = Node;
-	}
-	else if(const URigVMPin* Pin = Cast<URigVMPin>(InSubject))
-	{
-		GraphToJumpTo = Pin->GetGraph();
-		NodeToJumpTo = Pin->GetNode();
-		PinToJumpTo = Pin;
-	}
-	else if(const URigVMLink* Link = Cast<URigVMLink>(InSubject))
-	{
-		GraphToJumpTo = Link->GetGraph();
-		if(const URigVMPin* TargetPin = ((URigVMLink*)Link)->GetTargetPin())
-		{
-			NodeToJumpTo = TargetPin->GetNode();
-			PinToJumpTo = TargetPin;
-		}
-	}
-
-	if (GraphToJumpTo && NodeToJumpTo)
-	{
-		if(UControlRigGraph* EdGraph = Cast<UControlRigGraph>(RigBlueprint->GetEdGraph(NodeToJumpTo->GetGraph())))
-		{
-			if(const UControlRigGraphNode* EdGraphNode = Cast<UControlRigGraphNode>(EdGraph->FindNodeForModelNodeName(NodeToJumpTo->GetFName())))
+			}
+			else if(PropertyPath.RemoveFromStart(ShapeString))
 			{
-				if(PinToJumpTo)
+				FRigControlElement* ControlElement = ControlRigBP->Hierarchy->Find<FRigControlElement>(WrappedElement.GetKey());
+				if(ControlElement == nullptr)
 				{
-					if(const UEdGraphPin* EdGraphPin = EdGraphNode->FindPin(PinToJumpTo->GetSegmentPath(true)))
-					{
-						JumpToPin(EdGraphPin);
-						return;
-					}
+					return;
 				}
-				
-				JumpToNode(EdGraphNode);
-				return;
-			}
-			
-			JumpToHyperlink(EdGraph);
-		}
-	}
-}
 
-bool FControlRigEditor::UpdateDefaultValueForVariable(FBPVariableDescription& InVariable, bool bUseCDO)
-{
-	bool bAnyValueChanged = false;
-	if (UControlRigBlueprint* ControlRigBP = GetControlRigBlueprint())
-	{
-		UClass* GeneratedClass = ControlRigBP->GeneratedClass;
-		UObject* ObjectContainer = bUseCDO ? GeneratedClass->GetDefaultObject() : ControlRigBP->GetObjectBeingDebugged();
-		if(ObjectContainer)
-		{
-			FProperty* TargetProperty = FindFProperty<FProperty>(GeneratedClass, InVariable.VarName);
+				ERigTransformType::Type TransformType = Local::GetTransformTypeFromPath(PropertyPath);
+				bIsInitial = bIsInitial || ERigTransformType::IsInitial(TransformType);
 
-			if (TargetProperty)
+				const FTransform Transform = GetWrapperObjects()[0]->GetContent<FRigControlElement>().Shape.Get(TransformType);
+				
+				ControlRigBP->Hierarchy->SetControlShapeTransform(ControlElement, Transform, ERigTransformType::MakeInitial(TransformType), true, false, true);
+			}
+			else if(PropertyPath.RemoveFromStart(SettingsString))
 			{
-				FString NewDefaultValue;
-				const uint8* Container = (const uint8*)ObjectContainer;
-				FBlueprintEditorUtils::PropertyValueToString(TargetProperty, Container, NewDefaultValue, nullptr);
-				if (InVariable.DefaultValue != NewDefaultValue)
+				const FRigControlSettings Settings  = InWrapperObject->GetContent<FRigControlElement>().Settings;
+
+				FRigControlElement* ControlElement = ControlRigBP->Hierarchy->Find<FRigControlElement>(WrappedElement.GetKey());
+				if(ControlElement == nullptr)
 				{
-					InVariable.DefaultValue = NewDefaultValue;
-					bAnyValueChanged = true;
+					return;
 				}
+
+				ControlRigBP->Hierarchy->SetControlSettings(ControlElement, Settings, true, false, true);
 			}
-		}
-	}
-	return bAnyValueChanged;
-}
 
-bool FControlRigEditor::SelectLocalVariable(const UEdGraph* Graph, const FName& VariableName)
-{
-	if (const UControlRigGraph* ControlRigGraph = Cast<UControlRigGraph>(Graph))
-	{
-		if (URigVMGraph* RigVMGraph = ControlRigGraph->GetModel())
-		{
-			for (FRigVMGraphVariableDescription& Variable : RigVMGraph->GetLocalVariables())
+			if(IsConstructionModeEnabled() || bIsInitial)
 			{
-				if (Variable.Name == VariableName)
-				{
-					URigVMDetailsViewWrapperObject* WrapperObject = URigVMDetailsViewWrapperObject::MakeInstance(
-						GetDetailWrapperClass(), GetBlueprintObj(), Variable.StaticStruct(), (uint8*)&Variable, RigVMGraph);
-					WrapperObject->GetWrappedPropertyChangedChainEvent().AddSP(this, &FControlRigEditor::OnWrappedPropertyChangedChainEvent);
-					WrapperObject->AddToRoot();
-
-					TArray<UObject*> Objects = {WrapperObject};
-					SetDetailObjects(Objects, false);
-					return true;
-				}
+				ControlRigBP->PropagatePoseFromBPToInstances();
+				ControlRigBP->Modify();
+				ControlRigBP->MarkPackageDirty();
 			}
 		}
 	}
-	return false;
 }
 
-void FControlRigEditor::OnCreateComment()
+void FControlRigEditor::BindCommands()
 {
-	TSharedPtr<SGraphEditor> GraphEditor = FocusedGraphEdPtr.Pin();
-	if (GraphEditor.IsValid())
-	{
-		if (UEdGraph* Graph = GraphEditor->GetCurrentGraph())
-		{
-			FEdGraphSchemaAction_K2AddComment CommentAction;
-			CommentAction.PerformAction(Graph, NULL, GraphEditor->GetPasteLocation());
-		}
-	}
+	FRigVMEditor::BindCommands();
+
+	GetToolkitCommands()->MapAction(
+		FControlRigEditorCommands::Get().ConstructionEvent,
+		FExecuteAction::CreateSP(this, &FRigVMEditor::SetEventQueue, TArray<FName>(ConstructionEventQueue)),
+		FCanExecuteAction());
+
+	GetToolkitCommands()->MapAction(
+		FControlRigEditorCommands::Get().ForwardsSolveEvent,
+		FExecuteAction::CreateSP(this, &FRigVMEditor::SetEventQueue, TArray<FName>(ForwardsSolveEventQueue)),
+		FCanExecuteAction());
+
+	GetToolkitCommands()->MapAction(
+		FControlRigEditorCommands::Get().BackwardsSolveEvent,
+		FExecuteAction::CreateSP(this, &FRigVMEditor::SetEventQueue, TArray<FName>(BackwardsSolveEventQueue)),
+		FCanExecuteAction());
+
+	GetToolkitCommands()->MapAction(
+		FControlRigEditorCommands::Get().BackwardsAndForwardsSolveEvent,
+		FExecuteAction::CreateSP(this, &FRigVMEditor::SetEventQueue, TArray<FName>(BackwardsAndForwardsSolveEventQueue)),
+		FCanExecuteAction());
 }
 
 void FControlRigEditor::OnHierarchyChanged()
@@ -5786,7 +2940,7 @@ void FControlRigEditor::OnCurveContainerChanged()
 	{
 		// restart animation 
 		EditorSkelComp->InitAnim(true);
-		UpdateControlRig();
+		UpdateRigVMHost();
 	}
 	CacheNameLists();
 
@@ -6779,289 +3933,7 @@ void FControlRigEditor::HandleOnControlModified(UControlRig* Subject, FRigContro
 void FControlRigEditor::HandleRefreshEditorFromBlueprint(URigVMBlueprint* InBlueprint)
 {
 	OnHierarchyChanged();
-	Compile();
-}
-
-void FControlRigEditor::HandleVariableDroppedFromBlueprint(UObject* InSubject, FProperty* InVariableToDrop, const FVector2D& InDropPosition, const FVector2D& InScreenPosition)
-{
-	UControlRigBlueprint* Blueprint = Cast<UControlRigBlueprint>(GetBlueprintObj());
-	if (Blueprint == nullptr)
-	{
-		return;
-	}
-
-	URigVMController* Controller = GetFocusedController();
-	check(Controller);
-
-	FRigVMExternalVariable ExternalVariable = FRigVMExternalVariable::Make(InVariableToDrop, nullptr);
-	if (!ExternalVariable.IsValid(true /* allow null ptr */))
-	{
-		return;
-	}
-
-	FMenuBuilder MenuBuilder(true, NULL);
-	const FText SectionText = FText::FromString(FString::Printf(TEXT("Variable %s"), *ExternalVariable.Name.ToString()));
-
-	MenuBuilder.BeginSection("VariableDropped", SectionText);
-
-	MenuBuilder.AddMenuEntry(
-		FText::FromString(FString::Printf(TEXT("Get %s"), *ExternalVariable.Name.ToString())),
-		FText::FromString(FString::Printf(TEXT("Adds a getter node for variable %s"), *ExternalVariable.Name.ToString())),
-		FSlateIcon(),
-		FUIAction(
-			FExecuteAction::CreateLambda([ExternalVariable, Controller, InDropPosition] {
-
-				Controller->AddVariableNode(ExternalVariable.Name, ExternalVariable.TypeName.ToString(), ExternalVariable.TypeObject, true, FString(), InDropPosition, FString(), true, true);
-
-			}),
-			FCanExecuteAction()
-		)
-	);
-
-	MenuBuilder.AddMenuEntry(
-		FText::FromString(FString::Printf(TEXT("Set %s"), *ExternalVariable.Name.ToString())),
-		FText::FromString(FString::Printf(TEXT("Adds a setter node for variable %s"), *ExternalVariable.Name.ToString())),
-		FSlateIcon(),
-		FUIAction(
-			FExecuteAction::CreateLambda([ExternalVariable, Controller, InDropPosition] {
-
-				Controller->AddVariableNode(ExternalVariable.Name, ExternalVariable.TypeName.ToString(), ExternalVariable.TypeObject, false, FString(), InDropPosition, FString(), true, true);
-
-			}),
-			FCanExecuteAction()
-		)
-	);
-
-	MenuBuilder.EndSection();
-
-	TSharedRef<SWidget> GraphEditorPanel = FocusedGraphEdPtr.Pin().ToSharedRef();
-
-	// Show dialog to choose getter vs setter
-	FSlateApplication::Get().PushMenu(
-		GraphEditorPanel,
-		FWidgetPath(),
-		MenuBuilder.MakeWidget(),
-		InScreenPosition,
-		FPopupTransitionEffect(FPopupTransitionEffect::ContextMenu)
-	);
-}
-
-void FControlRigEditor::HandleBreakpointAdded()
-{
-	SetExecutionMode(EControlRigExecutionModeType_Debug);
-}
-
-void FControlRigEditor::OnGraphNodeClicked(URigVMEdGraphNode* InNode)
-{
-	if (InNode)
-	{
-		if (InNode->IsSelectedInEditor())
-		{
-			SetDetailViewForGraph(InNode->GetModel());
-		}
-	}
-}
-
-void FControlRigEditor::OnNodeDoubleClicked(URigVMBlueprint* InBlueprint, URigVMNode* InNode)
-{
-	ensure(GetControlRigBlueprint() == InBlueprint);
-
-	if (URigVMLibraryNode* LibraryNode = Cast<URigVMLibraryNode>(InNode))
-	{
-		URigVMGraph* ContainedGraph = LibraryNode->GetContainedGraph();
-		if (URigVMFunctionReferenceNode* FunctionReferenceNode = Cast<URigVMFunctionReferenceNode>(LibraryNode))
-		{
-			if (URigVMLibraryNode* ReferencedNode = FunctionReferenceNode->LoadReferencedNode())
-			{
-				ContainedGraph = ReferencedNode->GetContainedGraph();
-			}
-		}
-		if(ContainedGraph)
-		{
-			if (UEdGraph* EdGraph = InBlueprint->GetEdGraph(ContainedGraph))
-			{
-				OpenGraphAndBringToFront(EdGraph, true);
-			}
-			else
-			{
-				if(URigVMCollapseNode* FunctionLibraryNode = Cast<URigVMCollapseNode>(ContainedGraph->GetOuter()))
-				{
-					if(URigVMFunctionLibrary* FunctionLibrary = FunctionLibraryNode->GetLibrary())
-					{
-						if(UControlRigBlueprint* FunctionBlueprint = Cast<UControlRigBlueprint>(FunctionLibrary->GetOuter()))
-						{
-							if (UEdGraph* FunctionEdGraph = FunctionBlueprint->GetEdGraph(ContainedGraph))
-							{
-								FKismetEditorUtilities::BringKismetToFocusAttentionOnObject(FunctionEdGraph);
-							}
-						}
-					}
-				}
-			}
-		}
-	}
-}
-
-void FControlRigEditor::OnGraphImported(UEdGraph* InEdGraph)
-{
-	check(InEdGraph);
-
-	OpenDocument(InEdGraph, FDocumentTracker::OpenNewDocument);
-	RenameNewlyAddedAction(InEdGraph->GetFName());
-}
-
-bool FControlRigEditor::OnActionMatchesName(FEdGraphSchemaAction* InAction, const FName& InName) const
-{
-	if (InAction->GetMenuDescription().ToString() == InName.ToString())
-	{
-		return true;
-	}
-	return false;
-}
-
-void FControlRigEditor::HandleShowCurrentStatement()
-{
-	if(HaltedAtNode)
-	{
-		if(UControlRigBlueprint* Blueprint = Cast<UControlRigBlueprint>(GetBlueprintObj()))
-		{
-			if(UControlRigGraph* EdGraph = Cast<UControlRigGraph>(Blueprint->GetEdGraph(HaltedAtNode->GetGraph())))
-			{
-				if(UEdGraphNode* EdNode = EdGraph->FindNodeForModelNodeName(HaltedAtNode->GetFName()))
-				{
-					JumpToHyperlink(EdNode, false);
-				}
-			}
-		}
-	}
-}
-
-void FControlRigEditor::HandleBreakpointActionRequested(const ERigVMBreakpointAction BreakpointAction)
-{
-	if (UControlRig* ControlRig = GetControlRig())
-	{
-		ControlRig->ExecuteBreakpointAction(BreakpointAction);
-	}
-}
-
-bool FControlRigEditor::IsHaltedAtBreakpoint() const
-{
-	return HaltedAtNode != nullptr;
-}
-
-void FControlRigEditor::FrameSelection()
-{
-	if (SGraphEditor* GraphEd = FocusedGraphEdPtr.Pin().Get())
-	{
-		if(URigVMGraph* Model = GetFocusedModel())
-		{
-			const bool bFrameAll = Model->GetSelectNodes().Num() == 0;
-			GraphEd->ZoomToFit(!bFrameAll);
-		}
-	}
-}
-
-void FControlRigEditor::UpdateGraphCompilerErrors()
-{
-	DECLARE_SCOPE_HIERARCHICAL_COUNTER_FUNC()
-
-	UControlRigBlueprint* Blueprint = Cast<UControlRigBlueprint>(GetBlueprintObj());
-	UControlRig* ControlRig = GetControlRig();
-	if (Blueprint && ControlRig && ControlRig->VM)
-	{
-		if (ControlRigLog.Entries.Num() == 0 && !bAnyErrorsLeft)
-		{
-			return;
-		}
-
-		const FRigVMByteCode& ByteCode = ControlRig->VM->GetByteCode();
-
-		TArray<UEdGraph*> EdGraphs;
-		Blueprint->GetAllGraphs(EdGraphs);
-
-		for (UEdGraph* Graph : EdGraphs)
-		{
-			UControlRigGraph* RigGraph = Cast<UControlRigGraph>(Graph);
-			if (RigGraph == nullptr)
-			{
-				continue;
-			}
-
-			// reset all nodes and store them in the map
-			bool bFoundWarning = false;
-			bool bFoundError = false;
-			
-			for (UEdGraphNode* GraphNode : Graph->Nodes)
-			{
-				if (UControlRigGraphNode* ControlRigGraphNode = Cast<UControlRigGraphNode>(GraphNode))
-				{
-					bFoundError = bFoundError || ControlRigGraphNode->ErrorType <= (int32)EMessageSeverity::Error;
-					bFoundWarning = bFoundWarning || ControlRigGraphNode->ErrorType <= (int32)EMessageSeverity::Warning;
-
-					if(ControlRigGraphNode->ErrorType <= (int32)EMessageSeverity::Warning)
-					{
-						if(!ControlRig->VM->WasInstructionVisitedDuringLastRun(ControlRig->GetExtendedExecuteContext(), ControlRigGraphNode->GetInstructionIndex(true)) &&
-							!ControlRig->VM->WasInstructionVisitedDuringLastRun(ControlRig->GetExtendedExecuteContext(), ControlRigGraphNode->GetInstructionIndex(false)))
-						{
-							continue;
-						}
-					}
-				}
-
-				GraphNode->ErrorType = int32(EMessageSeverity::Info) + 1;
-			}
-
-			// update the nodes' error messages
-			for (const FRigVMLog::FLogEntry& Entry : ControlRigLog.Entries)
-			{
-				URigVMNode* ModelNode = Cast<URigVMNode>(ByteCode.GetSubjectForInstruction(Entry.InstructionIndex));
-				if (ModelNode == nullptr)
-				{
-					continue;
-				}
-
-				UEdGraphNode* GraphNode = RigGraph->FindNodeForModelNodeName(ModelNode->GetFName());
-				if (GraphNode == nullptr)
-				{
-					continue;
-				}
-
-				bFoundError = bFoundError || Entry.Severity <= EMessageSeverity::Error;
-				bFoundWarning = bFoundWarning || Entry.Severity <= EMessageSeverity::Warning;
-
-				int32 ErrorType = (int32)Entry.Severity;
-				if (GraphNode->ErrorType < ErrorType)
-				{
-					continue;
-				}
-				else if (GraphNode->ErrorType == ErrorType)
-				{
-					GraphNode->ErrorMsg = FString::Printf(TEXT("%s\n%s"), *GraphNode->ErrorMsg, *Entry.Message);
-				}
-				else
-				{
-					GraphNode->ErrorMsg = Entry.Message;
-					GraphNode->ErrorType = ErrorType;
-				}
-			}
-
-			bAnyErrorsLeft = false;
-			for (UEdGraphNode* GraphNode : Graph->Nodes)
-			{
-				GraphNode->bHasCompilerMessage = GraphNode->ErrorType <= int32(EMessageSeverity::Info);
-				bAnyErrorsLeft = bAnyErrorsLeft || GraphNode->bHasCompilerMessage; 
-			}
-
-			if (bFoundError)
-			{
-				Blueprint->Status = BS_Error;
-				Blueprint->MarkPackageDirty();
-			}
-		}
-
-		//Stack
-	}
-
+	FRigVMEditor::HandleRefreshEditorFromBlueprint(InBlueprint);
 }
 
 UToolMenu* FControlRigEditor::HandleOnGetViewportContextMenuDelegate()
@@ -7082,45 +3954,6 @@ TSharedPtr<FUICommandList> FControlRigEditor::HandleOnViewportContextMenuCommand
 	return TSharedPtr<FUICommandList>();
 }
 
-URigVMGraph* FControlRigEditor::GetFocusedModel() const
-{
-	UControlRigBlueprint* Blueprint = Cast<UControlRigBlueprint>(GetBlueprintObj());
-	if (Blueprint == nullptr)
-	{
-		return nullptr;
-	}
-
-	UControlRigGraph* EdGraph = Cast<UControlRigGraph>(GetFocusedGraph());
-	return Blueprint->GetModel(EdGraph);
-}
-
-URigVMController* FControlRigEditor::GetFocusedController() const
-{
-	UControlRigBlueprint* Blueprint = Cast<UControlRigBlueprint>(GetBlueprintObj());
-	if (Blueprint == nullptr)
-	{
-		return nullptr;
-	}
-	return Blueprint->GetOrCreateController(GetFocusedModel());
-}
-
-TSharedPtr<SGraphEditor> FControlRigEditor::GetGraphEditor(UEdGraph* InEdGraph) const
-{
-	TArray< TSharedPtr<SDockTab> > GraphEditorTabs;
-	DocumentManager->FindAllTabsForFactory(GraphEditorTabFactoryPtr, /*out*/ GraphEditorTabs);
-
-	for (TSharedPtr<SDockTab>& GraphEditorTab : GraphEditorTabs)
-	{
-		TSharedRef<SGraphEditor> Editor = StaticCastSharedRef<SGraphEditor>((GraphEditorTab)->GetContent());
-		if (Editor->GetCurrentGraph() == InEdGraph)
-		{
-			return Editor;
-		}
-	}
-
-	return TSharedPtr<SGraphEditor>();
-}
-
 void FControlRigEditor::OnPreConstruction_AnyThread(UControlRig* InRig, const FName& InEventName)
 {
 	bIsConstructionEventRunning = true;
