RigVM: Move most of the functionality from ControlRigEditor to RigVMEditor

diff --git a/Engine/Plugins/Animation/ControlRig/Source/ControlRigEditor/Private/Editor/ControlRigEditor.h b/Engine/Plugins/Animation/ControlRig/Source/ControlRigEditor/Private/Editor/ControlRigEditor.h
--- a/Engine/Plugins/Animation/ControlRig/Source/ControlRigEditor/Private/Editor/ControlRigEditor.h
+++ b/Engine/Plugins/Animation/ControlRig/Source/ControlRigEditor/Private/Editor/ControlRigEditor.h
@@ -32,8 +32,6 @@ class FStructOnScope;
 class UToolMenu;
 class FControlRigEditor;
 
-DECLARE_MULTICAST_DELEGATE_TwoParams(FControlRigEditorClosed, const FControlRigEditor*, URigVMBlueprint*);
-
 struct FControlRigEditorModes
 {
 	// Mode constants
@@ -59,96 +57,48 @@ private:
 class FControlRigEditor : public IControlRigEditor
 {
 public:
-	/**
-	 * Edits the specified character asset(s)
-	 *
-	 * @param	Mode					Mode that this editor should operate in
-	 * @param	InitToolkitHost			When Mode is WorldCentric, this is the level editor instance to spawn this editor within
-	 * @param	InControlRigBlueprint	The blueprint object to start editing.
-	 */
-	void InitControlRigEditor(const EToolkitMode::Type Mode, const TSharedPtr< class IToolkitHost >& InitToolkitHost, class UControlRigBlueprint* InControlRigBlueprint);
+
+	virtual void InitRigVMEditor(const EToolkitMode::Type Mode, const TSharedPtr< class IToolkitHost >& InitToolkitHost, class URigVMBlueprint* InRigVMBlueprint) override;
+
+	virtual const FName GetEditorAppName() const override;
+	virtual const FName GetEditorModeName() const override;
+	virtual TSharedPtr<FApplicationMode> CreateEditorMode() override;
 
 public:
 	FControlRigEditor();
 	virtual ~FControlRigEditor();
 
+	virtual UObject* GetOuterForHost() const override;
+
 	// FRigVMEditor interface
 	virtual UClass* GetDetailWrapperClass() const;
+	virtual void Compile() override;
 
-	// FBlueprintEditor interface
-	virtual UBlueprint* GetBlueprintObj() const override;
-	virtual TSubclassOf<UEdGraphSchema> GetDefaultSchemaClass() const override;
+	// allows the editor to fill an empty graph
+	virtual void CreateEmptyGraphContent(URigVMController* InController) override;
 
 	int32 GetRigHierarchyTabCount() const { return RigHierarchyTabCount; }
 
-	FControlRigEditorClosed& OnControlRigEditorClosed() { return ControlRigEditorClosedDelegate; }
-
-private:
-
-	FControlRigEditorClosed ControlRigEditorClosedDelegate;
-
-
 public:
 	
 	// IToolkit Interface
 	virtual FName GetToolkitFName() const override;
-	virtual FName GetToolkitContextFName() const override;
 	virtual FText GetBaseToolkitName() const override;
-	virtual FText GetToolkitToolTipText() const override;
 	virtual FString GetWorldCentricTabPrefix() const override;
-	virtual FLinearColor GetWorldCentricTabColorScale() const override;	
 	virtual FString GetDocumentationLink() const override
 	{
 		return FString(TEXT("Engine/Animation/ControlRig"));
 	}
-	virtual void InitToolMenuContext(FToolMenuContext& MenuContext) override;
 
 	// BlueprintEditor interface
-	virtual bool TransactionObjectAffectsBlueprint(UObject* InTransactedObject) override;
-	virtual bool CanAddNewLocalVariable() const override;
-	virtual void OnAddNewLocalVariable() override;
-	virtual void OnPasteNewLocalVariable(const FBPVariableDescription& VariableDescription) override;
-	
-	virtual void DeleteSelectedNodes() override;
-	virtual bool CanDeleteNodes() const override;
-
-	virtual void CopySelectedNodes() override;
-	virtual bool CanCopyNodes() const override;
-
-	virtual void PasteNodes() override;
-	virtual bool CanPasteNodes() const override;
-
-	virtual bool IsNativeParentClassCodeLinkEnabled() const override { return false; }
-	virtual bool ReparentBlueprint_IsVisible() const override { return false; }
-
 	virtual FReply OnSpawnGraphNodeByShortcut(FInputChord InChord, const FVector2D& InPosition, UEdGraph* InGraph) override;
-
-	// Control Rig BP does not use regular BP function libraries, no need to load them.
-	virtual bool ShouldLoadBPLibrariesFromAssetRegistry() override { return false; }
-
-	// FEditorUndoClient Interface
-	virtual void PostUndo(bool bSuccess) override;
-	virtual void PostRedo(bool bSuccess) override;
-
-	virtual void PostTransaction(bool bSuccess, const FTransaction* Transaction, bool bIsRedo);
-
-	virtual void JumpToHyperlink(const UObject* ObjectReference, bool bRequestRename = false) override;
+	
+	virtual void PostTransaction(bool bSuccess, const FTransaction* Transaction, bool bIsRedo) override;
 
 	void EnsureValidRigElementsInDetailPanel();
 
-	virtual void OnStartWatchingPin();
-	virtual bool CanStartWatchingPin() const;
-
-	virtual void OnStopWatchingPin();
-	virtual bool CanStopWatchingPin() const;
-
-	// IToolkitHost Interface
-	virtual void OnToolkitHostingStarted(const TSharedRef<class IToolkit>& Toolkit) override;
-	virtual void OnToolkitHostingFinished(const TSharedRef<class IToolkit>& Toolkit) override;
-
 	//  FTickableEditorObject Interface
 	virtual void Tick(float DeltaTime) override;
-	virtual TStatId GetStatId() const override;
 
 	// Gets the Control Rig Blueprint being edited/viewed
 	UControlRigBlueprint* GetControlRigBlueprint() const;
@@ -159,31 +109,18 @@ public:
 	// returns the hierarchy being debugged
 	URigHierarchy* GetHierarchyBeingDebugged() const;
 
-	void SetDetailObjects(const TArray<UObject*>& InObjects);
-	void SetDetailObjects(const TArray<UObject*>& InObjects, bool bChangeUISelectionState);
 	void SetDetailViewForRigElements();
-	void SetDetailViewForGraph(URigVMGraph* InGraph);
-	void SetDetailViewForFocusedGraph();
-	void SetDetailViewForLocalVariable();
-	void RefreshDetailView();
 	bool DetailViewShowsAnyRigElement() const;
-	bool DetailViewShowsAnyRigUnit() const;
-	bool DetailViewShowsLocalVariable() const;
-	bool DetailViewShowsStruct(UScriptStruct* InStruct) const;
 	bool DetailViewShowsRigElement(FRigElementKey InKey) const;
+	virtual void RefreshDetailView() override;
 
-	void ClearDetailObject(bool bChangeUISelectionState = true);
-private:
+	void CreatePersonaToolKitIfRequired();
 
-	void ClearDetailsViewWrapperObjects();
 public:
 
 	/** Get the persona toolkit */
 	TSharedRef<IPersonaToolkit> GetPersonaToolkit() const { return PersonaToolkit.ToSharedRef(); }
 
-	/** Get the toolbox hosting widget */
-	TSharedRef<SBorder> GetToolbox() { return Toolbox.ToSharedRef(); }
-
 	/** Get the edit mode */
 	FControlRigEditorEditMode* GetEditMode() const;
 
@@ -200,22 +137,20 @@ public:
 	void OnGraphNodeDropToPerform(TSharedPtr<FGraphNodeDragDropOp> InDragDropOp, UEdGraph* InGraph, const FVector2D& InNodePosition, const FVector2D& InScreenPosition);
 
 	FPersonaViewportKeyDownDelegate& GetKeyDownDelegate() { return OnKeyDownDelegate; }
+
 	FOnGetContextMenu& OnGetViewportContextMenu() { return OnGetViewportContextMenuDelegate; }
 	FNewMenuCommandsDelegate& OnViewportContextMenuCommands() { return OnViewportContextMenuCommandsDelegate; }
 
-	virtual bool SelectLocalVariable(const UEdGraph* Graph, const FName& VariableName) override;
-
-private:
-
-	virtual void OnCreateComment() override;
-
 protected:
 
+	virtual void BindCommands() override;
+
 	void OnHierarchyChanged();
-	void OnControlsSettingsChanged();
 
 	void SynchronizeViewportBoneSelection();
-	
+
+	virtual void SaveAsset_Execute() override;
+	virtual void SaveAssetAs_Execute() override;
 
 	// update the cached modification value
 	void UpdateBoneModification(FName BoneName, const FTransform& Transform);
@@ -226,87 +161,21 @@ protected:
 	// reset all bone modification across all instance
 	void ResetAllBoneModification();
 
-	// FBlueprintEditor Interface
-	virtual void CreateDefaultCommands() override;
-	virtual void OnCreateGraphEditorCommands(TSharedPtr<FUICommandList> GraphEditorCommandsList);
-	virtual TSharedRef<SGraphEditor> CreateGraphEditorWidget(TSharedRef<FTabInfo> InTabInfo, UEdGraph* InGraph) override;
-	virtual void Compile() override;
-	virtual void SaveAsset_Execute() override;
-	virtual void SaveAssetAs_Execute() override;
-	virtual bool IsInAScriptingMode() const override { return true; }
-	virtual void CreateDefaultTabContents(const TArray<UBlueprint*>& InBlueprints) override;
-	virtual void NewDocument_OnClicked(ECreatedDocumentType GraphType) override;
-	virtual bool IsSectionVisible(NodeSectionID::Type InSectionID) const override;
-	virtual FGraphAppearanceInfo GetGraphAppearance(class UEdGraph* InGraph) const override;
-	virtual bool IsEditable(UEdGraph* InGraph) const override;
-	virtual bool IsCompilingEnabled() const override;
-	virtual FText GetGraphDecorationString(UEdGraph* InGraph) const override;
-	virtual void OnActiveTabChanged( TSharedPtr<SDockTab> PreviouslyActive, TSharedPtr<SDockTab> NewlyActivated ) override;
-	virtual void OnSelectedNodesChangedImpl(const TSet<class UObject*>& NewSelection) override;
-	virtual void OnBlueprintChangedImpl(UBlueprint* InBlueprint, bool bIsJustBeingCompiled) override;
-	virtual void RefreshEditors(ERefreshBlueprintEditorReason::Type Reason = ERefreshBlueprintEditorReason::UnknownReason) override;
-	virtual void SetupGraphEditorEvents(UEdGraph* InGraph, SGraphEditor::FGraphEditorEvents& InEvents) override;
-	virtual void FocusInspectorOnGraphSelection(const TSet<class UObject*>& NewSelection, bool bForceRefresh = false) override;
-
-	virtual void HandleModifiedEvent(ERigVMGraphNotifType InNotifType, URigVMGraph* InGraph, UObject* InSubject) override;
-	void HandleVMCompiledEvent(UObject* InCompiledObject, URigVM* InVM);
-	void HandleControlRigExecutedEvent(URigVMHost* InControlRig, const FName& InEventName);
-	void HandleControlRigExecutionHalted(const int32 InstructionIndex, UObject* InNode, const FName& InEntryName);
-	void SetHaltedNode(URigVMNode* Node);
+	virtual void HandleVMExecutedEvent(URigVMHost* InHost, const FName& InEventName) override;
 
 	// FBaseToolKit overrides
-	void CreateEditorModeManager() override;
-
-	// FGCObject Interface
-	virtual void AddReferencedObjects( FReferenceCollector& Collector ) override;
-
-	void BindCommands();
-
-protected:
-	/** Undo Action**/
-	void UndoAction();
-
-	/** Redo Action **/
-	void RedoAction();
+	virtual void CreateEditorModeManager() override;
 
 private:
-	/** Extend menu */
-	void ExtendMenu();
-
-	/** Extend toolbar */
-	void ExtendToolbar();
-
 	/** Fill the toolbar with content */
-	void FillToolbar(FToolBarBuilder& ToolbarBuilder);
-
-	TArray<FName> GetEventQueue() const;
-	void SetEventQueue(TArray<FName> InEventQueue);
-	void SetEventQueue(TArray<FName> InEventQueue, bool bCompile);
-	int32 GetEventQueueComboValue() const;
-	FText GetEventQueueLabel() const;
-	static FSlateIcon GetEventQueueIcon(const TArray<FName>& InEventQueue);
-	FSlateIcon GetEventQueueIcon() const;
+	virtual void FillToolbar(FToolBarBuilder& ToolbarBuilder, bool bEndSection = true) override;
 
-	enum EControlRigExecutionModeType
-	{
-		EControlRigExecutionModeType_Release,
-		EControlRigExecutionModeType_Debug
-	};
-
-	void SetExecutionMode(const EControlRigExecutionModeType InExecutionMode);
-	int32 GetExecutionModeComboValue() const;
-	FText GetExecutionModeLabel() const;
-	static FSlateIcon GetExecutionModeIcon(const EControlRigExecutionModeType InExecutionMode);
-	FSlateIcon GetExecutionModeIcon() const;
-	
-	virtual void GetCustomDebugObjects(TArray<FCustomDebugObject>& DebugList) const override;
-	virtual bool OnlyShowCustomDebugObjects() const override { return true; }
+	virtual TArray<FName> GetDefaultEventQueue() const override;
+	virtual void SetEventQueue(TArray<FName> InEventQueue, bool bCompile) override;
+	virtual int32 GetEventQueueComboValue() const override;
+	virtual FText GetEventQueueLabel() const override;
+	virtual FSlateIcon GetEventQueueIcon(const TArray<FName>& InEventQueue) const override;
 	virtual void HandleSetObjectBeingDebugged(UObject* InObject) override;
-	virtual FString GetCustomDebugObjectLabel(UObject* ObjectBeingDebugged) const override;
-
-	/** Handle hiding items in the graph */
-	void HandleHideItem();
-	bool CanHideItem() const;
 
 	/** Handle preview scene setup */
 	void HandlePreviewSceneCreated(const TSharedRef<IPersonaPreviewScene>& InPersonaPreviewScene);
@@ -323,30 +192,16 @@ private:
 	void HandlePreviewMeshChanged(USkeletalMesh* InOldSkeletalMesh, USkeletalMesh* InNewSkeletalMesh);
 
 	/** Push a newly compiled/opened control rig to the edit mode */
-	void UpdateControlRig();
+	virtual void UpdateRigVMHost() override;
+	virtual void UpdateRigVMHost_PreClearOldHost(URigVMHost* InPreviousHost) override;
 
 	/** Update the name lists for use in name combo boxes */
-	void CacheNameLists();
+	virtual void CacheNameLists() override;
 
 	/** Rebind our anim instance to the preview's skeletal mesh component */
 	void RebindToSkeletalMeshComponent();
 
-	/** Update stale watch pins */
-	void UpdateStaleWatchedPins();
-
-	/** Wraps the normal blueprint editor's action menu creation callback */
-	FActionMenuContent HandleCreateGraphActionMenu(UEdGraph* InGraph, const FVector2D& InNodePosition, const TArray<UEdGraphPin*>& InDraggedPins, bool bAutoExpand, SGraphEditor::FActionMenuClosed InOnMenuClosed);
-	void OnNodeTitleCommitted(const FText& NewText, ETextCommit::Type CommitInfo, UEdGraphNode* NodeBeingChanged);
-
-	void ToggleExecuteGraph();
-	bool IsExecuteGraphOn() const;
-	void ToggleAutoCompileGraph();
-	bool IsAutoCompileGraphOn() const;
-	bool CanAutoCompileGraph() const { return true; }
-	void ToggleEventQueue();
-	void ToggleExecutionMode();
-	TSharedRef<SWidget> GenerateEventQueueMenuContent();
-	TSharedRef<SWidget> GenerateExecutionModeMenuContent();
+	virtual void GenerateEventQueueMenuContent(FMenuBuilder& MenuBuilder) override;
 
 	enum ERigElementGetterSetterType
 	{
@@ -363,27 +218,7 @@ private:
 
 	 void HandleOnControlModified(UControlRig* Subject, FRigControlElement* ControlElement, const FRigControlModifiedContext& Context);
 
-	 void HandleRefreshEditorFromBlueprint(URigVMBlueprint* InBlueprint);
-
-	 void HandleVariableDroppedFromBlueprint(UObject* InSubject, FProperty* InVariableToDrop, const FVector2D& InDropPosition, const FVector2D& InScreenPosition);
-
-	 void HandleBreakpointAdded();
-
-	 void OnGraphNodeClicked(URigVMEdGraphNode* InNode);
-
-	 void OnNodeDoubleClicked(URigVMBlueprint* InBlueprint, URigVMNode* InNode);
-
-	 void OnGraphImported(UEdGraph* InEdGraph);
-
-	 virtual bool OnActionMatchesName(FEdGraphSchemaAction* InAction, const FName& InName) const override;
-
-	void HandleShowCurrentStatement();
-	
-	void HandleBreakpointActionRequested(const ERigVMBreakpointAction BreakpointAction);
-
-	bool IsHaltedAtBreakpoint() const;
-
-	void FrameSelection();
+	virtual void HandleRefreshEditorFromBlueprint(URigVMBlueprint* InBlueprint) override;
 
 	FText GetTestAssetName() const;
 	FText GetTestAssetTooltip() const;
@@ -395,18 +230,9 @@ private:
 
 protected:
 
-	/** Toolbox hosting widget */
-	TSharedPtr<SBorder> Toolbox;
-
 	/** Persona toolkit used to support skeletal mesh preview */
 	TSharedPtr<IPersonaToolkit> PersonaToolkit;
 
-	/** The extender to pass to the level editor to extend it's window menu */
-	TSharedPtr<FExtender> MenuExtender;
-
-	/** Toolbar extender */
-	TSharedPtr<FExtender> ToolbarExtender;
-
 	/** Preview instance inspector widget */
 	TSharedPtr<SWidget> PreviewEditor;
 
@@ -429,41 +255,17 @@ protected:
 	FTransform GetRigElementTransform(const FRigElementKey& InElement, bool bLocal, bool bOnDebugInstance) const;
 	void SetRigElementTransform(const FRigElementKey& InElement, const FTransform& InTransform, bool bLocal);
 	
-	// FNotifyHook Interface
-	virtual void NotifyPreChange(FProperty* PropertyAboutToChange) override;
-	virtual void NotifyPostChange(const FPropertyChangedEvent& PropertyChangedEvent, FProperty* PropertyThatChanged) override;
 	/** delegate for changing property */
 	virtual void OnFinishedChangingProperties(const FPropertyChangedEvent& PropertyChangedEvent) override;
-	void OnPropertyChanged(UObject* InObject, FPropertyChangedEvent& InEvent);
-	void OnWrappedPropertyChangedChainEvent(URigVMDetailsViewWrapperObject* InWrapperObject, const FString& InPropertyPath, FPropertyChangedChainEvent& InPropertyChangedChainEvent);
-	void OnRequestLocalizeFunctionDialog(FRigVMGraphFunctionIdentifier& InFunction, URigVMBlueprint* InTargetBlueprint, bool bForce);
-	FRigVMController_BulkEditResult OnRequestBulkEditDialog(URigVMBlueprint* InBlueprint, URigVMController* InController, URigVMLibraryNode* InFunction, ERigVMControllerBulkEditType InEditType);
-	bool OnRequestBreakLinksDialog(TArray<URigVMLink*> InLinks);
-	TRigVMTypeIndex OnRequestPinTypeSelectionDialog(const TArray<TRigVMTypeIndex>& InTypes);
-	void HandleJumpToHyperlink(const UObject* InSubject);
-	bool UpdateDefaultValueForVariable(FBPVariableDescription& InVariable, bool bUseCDO);
-
-	URigVMGraph* GetFocusedModel() const;
-	URigVMController* GetFocusedController() const;
-	TSharedPtr<SGraphEditor> GetGraphEditor(UEdGraph* InEdGraph) const;
-	
+	virtual void OnWrappedPropertyChangedChainEvent(URigVMDetailsViewWrapperObject* InWrapperObject, const FString& InPropertyPath, FPropertyChangedChainEvent& InPropertyChangedChainEvent) override;
+
 	URigVMController* ActiveController;
-	bool bControlRigEditorInitialized;
-	bool bIsSettingObjectBeingDebugged;
 
 	/** Currently executing ControlRig or not - later maybe this will change to enum for whatever different mode*/
 	bool bExecutionControlRig;
 
-	/** The log to use for errors resulting from the init phase of the units */
-	FRigVMLog ControlRigLog;
-	/** Once the log is collected update the graph */
-	void UpdateGraphCompilerErrors();
-
 	void OnAnimInitialized();
 
-	/** Are we currently compiling through the user interface */
-	bool bIsCompilingThroughUI;
-
 	TSharedPtr<SRigVMGraphPinNameListValueWidget> PinControlNameList;
 	bool IsPinControlNameListEnabled() const;
 	TSharedRef<SWidget> MakePinControlNameListItemWidget(TSharedPtr<FString> InItem);
@@ -474,25 +276,9 @@ protected:
 	void OnPinControlNameListComboBox(const TArray<TSharedPtr<FString>>* InNameList);
 	bool IsConstructionModeEnabled() const;
 
-	bool bAnyErrorsLeft;
-
-	TArray<FName> LastEventQueue;
-	EControlRigExecutionModeType ExecutionMode;
-	FString LastDebuggedRig;
 	int32 RigHierarchyTabCount;
-	TMap<FString, FString> KnownInstructionLimitWarnings;
-
-	URigVMNode* HaltedAtNode;
-
-	bool bSuspendDetailsPanelRefresh;
-
-	bool bAllowBulkEdits;
-
-	TArray<TStrongObjectPtr<URigVMDetailsViewWrapperObject>> WrapperObjects;
 	TWeakObjectPtr<AStaticMeshActor> WeakGroundActorPtr;
 
-	FDelegateHandle PropertyChangedHandle;
-
 	void OnPreConstruction_AnyThread(UControlRig* InRig, const FName& InEventName);
 	void OnPostConstruction_AnyThread(UControlRig* InRig, const FName& InEventName);
 
diff --git a/Engine/Plugins/Animation/ControlRig/Source/ControlRigEditor/Private/Editor/SRigCurveContainer.cpp b/Engine/Plugins/Animation/ControlRig/Source/ControlRigEditor/Private/Editor/SRigCurveContainer.cpp
--- a/Engine/Plugins/Animation/ControlRig/Source/ControlRigEditor/Private/Editor/SRigCurveContainer.cpp
+++ b/Engine/Plugins/Animation/ControlRig/Source/ControlRigEditor/Private/Editor/SRigCurveContainer.cpp
@@ -199,7 +199,7 @@ void SRigCurveContainer::Construct(const FArguments& InArgs, TSharedRef<FControl
 
 	if(ControlRigEditor.IsValid())
 	{
-		ControlRigEditor.Pin()->OnControlRigEditorClosed().AddSP(this, &SRigCurveContainer::OnEditorClose);
+		ControlRigEditor.Pin()->OnEditorClosed().AddSP(this, &SRigCurveContainer::OnEditorClose);
 	}
 }
 
@@ -304,7 +304,7 @@ TSharedPtr<SWidget> SRigCurveContainer::OnGetContextMenuContent() const
 	return MenuBuilder.MakeWidget();
 }
 
-void SRigCurveContainer::OnEditorClose(const FControlRigEditor* InEditor, URigVMBlueprint* InBlueprint)
+void SRigCurveContainer::OnEditorClose(const FRigVMEditor* InEditor, URigVMBlueprint* InBlueprint)
 {
 	if (UControlRigBlueprint* BP = Cast<UControlRigBlueprint>(InBlueprint))
 	{
diff --git a/Engine/Plugins/Animation/ControlRig/Source/ControlRigEditor/Private/Editor/SRigCurveContainer.h b/Engine/Plugins/Animation/ControlRig/Source/ControlRigEditor/Private/Editor/SRigCurveContainer.h
--- a/Engine/Plugins/Animation/ControlRig/Source/ControlRigEditor/Private/Editor/SRigCurveContainer.h
+++ b/Engine/Plugins/Animation/ControlRig/Source/ControlRigEditor/Private/Editor/SRigCurveContainer.h
@@ -207,7 +207,7 @@ private:
 	/** Handler for context menus */
 	TSharedPtr<SWidget> OnGetContextMenuContent() const;
 
-	void OnEditorClose(const FControlRigEditor* InEditor, URigVMBlueprint* InBlueprint);
+	void OnEditorClose(const FRigVMEditor* InEditor, URigVMBlueprint* InBlueprint);
 
 	/**
 	* Clears and rebuilds the table, according to an optional search string
diff --git a/Engine/Plugins/Animation/ControlRig/Source/ControlRigEditor/Private/Editor/SRigHierarchy.cpp b/Engine/Plugins/Animation/ControlRig/Source/ControlRigEditor/Private/Editor/SRigHierarchy.cpp
--- a/Engine/Plugins/Animation/ControlRig/Source/ControlRigEditor/Private/Editor/SRigHierarchy.cpp
+++ b/Engine/Plugins/Animation/ControlRig/Source/ControlRigEditor/Private/Editor/SRigHierarchy.cpp
@@ -264,14 +264,14 @@ void SRigHierarchy::Construct(const FArguments& InArgs, TSharedRef<FControlRigEd
 		});
 		ControlRigEditor.Pin()->OnGetViewportContextMenu().BindSP(this, &SRigHierarchy::GetContextMenu);
 		ControlRigEditor.Pin()->OnViewportContextMenuCommands().BindSP(this, &SRigHierarchy::GetContextMenuCommands);
-		ControlRigEditor.Pin()->OnControlRigEditorClosed().AddSP(this, &SRigHierarchy::OnEditorClose);
+		ControlRigEditor.Pin()->OnEditorClosed().AddSP(this, &SRigHierarchy::OnEditorClose);
 	}
 	
 	CreateContextMenu();
 	CreateDragDropMenu();
 }
 
-void SRigHierarchy::OnEditorClose(const FControlRigEditor* InEditor, URigVMBlueprint* InBlueprint)
+void SRigHierarchy::OnEditorClose(const FRigVMEditor* InEditor, URigVMBlueprint* InBlueprint)
 {
 	if (InEditor)
 	{
@@ -523,7 +523,7 @@ void SRigHierarchy::RefreshTreeView(bool bRebuildContent)
 	bool* SuspensionFlagPtr = &bDummySuspensionFlag;
 	if (ControlRigEditor.IsValid())
 	{
-		SuspensionFlagPtr = &ControlRigEditor.Pin()->bSuspendDetailsPanelRefresh;
+		SuspensionFlagPtr = &ControlRigEditor.Pin()->GetSuspendDetailsPanelRefreshFlag();
 	}
 	TGuardValue<bool> SuspendDetailsPanelRefreshGuard(*SuspensionFlagPtr, true);
 	TGuardValue<bool> GuardRigHierarchyChanges(bIsChangingRigHierarchy, true);
@@ -571,7 +571,7 @@ void SRigHierarchy::OnSelectionChanged(TSharedPtr<FRigTreeElement> Selection, ES
 		bool* SuspensionFlagPtr = &bDummySuspensionFlag;
 		if (ControlRigEditor.IsValid())
 		{
-			SuspensionFlagPtr = &ControlRigEditor.Pin()->bSuspendDetailsPanelRefresh;
+			SuspensionFlagPtr = &ControlRigEditor.Pin()->GetSuspendDetailsPanelRefreshFlag();
 		}
 
 		TGuardValue<bool> SuspendDetailsPanelRefreshGuard(*SuspensionFlagPtr, true);
diff --git a/Engine/Plugins/Animation/ControlRig/Source/ControlRigEditor/Private/Editor/SRigHierarchy.h b/Engine/Plugins/Animation/ControlRig/Source/ControlRigEditor/Private/Editor/SRigHierarchy.h
--- a/Engine/Plugins/Animation/ControlRig/Source/ControlRigEditor/Private/Editor/SRigHierarchy.h
+++ b/Engine/Plugins/Animation/ControlRig/Source/ControlRigEditor/Private/Editor/SRigHierarchy.h
@@ -9,6 +9,7 @@
 #include "Editor/SRigHierarchyTreeView.h"
 #include "Units/RigUnitContext.h"
 #include "ControlRigBlueprint.h"
+#include "Editor/RigVMEditor.h"
 #include "SRigHierarchy.generated.h"
 
 class SRigHierarchy;
@@ -86,7 +87,7 @@ public:
 
 private:
 
-	void OnEditorClose(const FControlRigEditor* InEditor, URigVMBlueprint* InBlueprint);
+	void OnEditorClose(const FRigVMEditor* InEditor, URigVMBlueprint* InBlueprint);
 
 	/** Bind commands that this widget handles */
 	void BindCommands();
diff --git a/Engine/Plugins/Runtime/RigVM/Source/RigVM/Public/RigVMCore/RigVM.h b/Engine/Plugins/Runtime/RigVM/Source/RigVM/Public/RigVMCore/RigVM.h
--- a/Engine/Plugins/Runtime/RigVM/Source/RigVM/Public/RigVMCore/RigVM.h
+++ b/Engine/Plugins/Runtime/RigVM/Source/RigVM/Public/RigVMCore/RigVM.h
@@ -731,13 +731,13 @@ public:
 
 
 #if WITH_EDITOR
-	// returns the instructions as text, OperandFormatFunction is an optional argument that allows you to override how operands are displayed, for example, see SControlRigStackView::PopulateStackView 
+	// returns the instructions as text, OperandFormatFunction is an optional argument that allows you to override how operands are displayed, for example, see SRigVMExecutionStackView::PopulateStackView 
 	TArray<FString> DumpByteCodeAsTextArray(const TArray<int32> & InInstructionOrder = TArray<int32>(), bool bIncludeLineNumbers = true, TFunction<FString(const FString& RegisterName, const FString& RegisterOffsetName)> OperandFormatFunction = nullptr);
 	FString DumpByteCodeAsText(const TArray<int32>& InInstructionOrder = TArray<int32>(), bool bIncludeLineNumbers = true);
 #endif
 
 #if WITH_EDITOR
-	// FormatFunction is an optional argument that allows you to override how operands are displayed, for example, see SControlRigStackView::PopulateStackView
+	// FormatFunction is an optional argument that allows you to override how operands are displayed, for example, see SRigVMExecutionStackView::PopulateStackView
 	FString GetOperandLabel(const FRigVMOperand & InOperand, TFunction<FString(const FString& RegisterName, const FString& RegisterOffsetName)> FormatFunction = nullptr);
 #endif
 
diff --git a/Engine/Plugins/Runtime/RigVM/Source/RigVM/Public/RigVMCore/RigVMExecuteContext.h b/Engine/Plugins/Runtime/RigVM/Source/RigVM/Public/RigVMCore/RigVMExecuteContext.h
--- a/Engine/Plugins/Runtime/RigVM/Source/RigVM/Public/RigVMCore/RigVMExecuteContext.h
+++ b/Engine/Plugins/Runtime/RigVM/Source/RigVM/Public/RigVMCore/RigVMExecuteContext.h
@@ -194,7 +194,7 @@ struct RIGVM_API FRigVMRuntimeSettings
 	GENERATED_BODY()
 
 	/**
-	 * The largest allowed size for arrays within the Control Rig VM.
+	 * The largest allowed size for arrays within the RigVM.
 	 * Accessing or creating larger arrays will cause runtime errors in the rig.
 	 */
 	UPROPERTY(EditAnywhere, Category = "VM")
@@ -748,7 +748,7 @@ struct RIGVM_API FRigVMExtendedExecuteContext
 	const void SetFirstEntryEventInEventQueue(const FName& InFirstEventName) { FirstEntryEventInQueue = InFirstEventName; }
 	const FName& GetFirstEntryEventInEventQueue() { return FirstEntryEventInQueue; }
 
-	// Control Rig can run multiple events per evaluation, such as the Backward&Forward Solve Mode,
+	// A RigVMHost can run multiple events per evaluation, such as the Backward&Forward Solve Mode,
 	// store the first event such that we know when to reset data for a new round of rig evaluation
 	FName FirstEntryEventInQueue = NAME_None;
 
diff --git a/Engine/Plugins/Runtime/RigVM/Source/RigVM/Public/RigVMCore/RigVMGraphFunctionDefinition.h b/Engine/Plugins/Runtime/RigVM/Source/RigVM/Public/RigVMCore/RigVMGraphFunctionDefinition.h
--- a/Engine/Plugins/Runtime/RigVM/Source/RigVM/Public/RigVMCore/RigVMGraphFunctionDefinition.h
+++ b/Engine/Plugins/Runtime/RigVM/Source/RigVM/Public/RigVMCore/RigVMGraphFunctionDefinition.h
@@ -305,7 +305,7 @@ struct RIGVM_API FRigVMGraphFunctionIdentifier
 	UPROPERTY()
 	FSoftObjectPath LibraryNode;
 
-	// A path to the IRigVMGraphFunctionHost that stores the function information, and compilation data (e.g. ControlRigBlueprintGeneratedClass)
+	// A path to the IRigVMGraphFunctionHost that stores the function information, and compilation data (e.g. RigVMBlueprintGeneratedClass)
 	UPROPERTY()
 	FSoftObjectPath HostObject;
 
diff --git a/Engine/Plugins/Runtime/RigVM/Source/RigVM/Public/RigVMHost.h b/Engine/Plugins/Runtime/RigVM/Source/RigVM/Public/RigVMHost.h
--- a/Engine/Plugins/Runtime/RigVM/Source/RigVM/Public/RigVMHost.h
+++ b/Engine/Plugins/Runtime/RigVM/Source/RigVM/Public/RigVMHost.h
@@ -143,11 +143,11 @@ public:
 
 #if WITH_EDITOR
 
-	/** Bindable event for external objects to be notified that a control rig is fully end-loaded*/
+	/** Bindable event for external objects to be notified that a RigVMHost is fully end-loaded*/
 	DECLARE_EVENT_OneParam(URigVMHost, FOnEndLoadPackage, URigVMHost*);
 
 	// these are needed so that sequencer can have a chance to update its 
-	// ControlRig instances after the package is fully end-loaded
+	// RigVMHost instances after the package is fully end-loaded
 	void BroadCastEndLoadPackage() { EndLoadPackageEvent.Broadcast(this); }
 	FOnEndLoadPackage& OnEndLoadPackage() { return EndLoadPackageEvent; }
 
@@ -211,6 +211,8 @@ public:
 
 #if WITH_EDITOR
 	const FRigVMLog* GetLog() const { return RigVMLog; }
+	FRigVMLog* GetLog() { return RigVMLog; }
+	void SetLog(FRigVMLog* InLog) { RigVMLog = InLog; }
 #endif
 
 	DECLARE_EVENT_TwoParams(URigVM, FRigVMExecutedEvent, class URigVMHost*, const FName&);
@@ -320,10 +322,10 @@ public:
 	/** Copy the default values of external variables from the default object */
 	void CopyExternalVariableDefaultValuesFromCDO();
 	
-	/** Broadcasts a notification whenever the controlrig's memory is initialized. */
+	/** Broadcasts a notification whenever the RigVMHost's memory is initialized. */
 	FRigVMExecutedEvent InitializedEvent;
 
-	/** Broadcasts a notification whenever the controlrig is executed / updated. */
+	/** Broadcasts a notification whenever the RigVMHost is executed / updated. */
 	FRigVMExecutedEvent ExecutedEvent;
 
 protected: 
diff --git a/Engine/Plugins/Runtime/RigVM/Source/RigVM/Public/RigVMTypeUtils.h b/Engine/Plugins/Runtime/RigVM/Source/RigVM/Public/RigVMTypeUtils.h
--- a/Engine/Plugins/Runtime/RigVM/Source/RigVM/Public/RigVMTypeUtils.h
+++ b/Engine/Plugins/Runtime/RigVM/Source/RigVM/Public/RigVMTypeUtils.h
@@ -249,7 +249,7 @@ namespace RigVMTypeUtils
 			OldObjectName,
 			NewObjectName,
 			nullptr,
-			ECoreRedirectMatchFlags::AllowPartialMatch); // AllowPartialMatch to allow redirects from one package to another (see /Script/ControlRig.CRFourPointBezier -> /Script/RigVM.RigVMFourPointBezier)
+			ECoreRedirectMatchFlags::AllowPartialMatch); // AllowPartialMatch to allow redirects from one package to another
 
 		if (!bFoundRedirect)
 		{
diff --git a/Engine/Plugins/Runtime/RigVM/Source/RigVMDeveloper/Private/EdGraph/RigVMEdGraphNode.cpp b/Engine/Plugins/Runtime/RigVM/Source/RigVMDeveloper/Private/EdGraph/RigVMEdGraphNode.cpp
--- a/Engine/Plugins/Runtime/RigVM/Source/RigVMDeveloper/Private/EdGraph/RigVMEdGraphNode.cpp
+++ b/Engine/Plugins/Runtime/RigVM/Source/RigVMDeveloper/Private/EdGraph/RigVMEdGraphNode.cpp
@@ -14,6 +14,7 @@
 #include "RigVMBlueprint.h"
 #include "PropertyPathHelpers.h"
 #include "RigVMBlueprintUtils.h"
+#include "../../../RigVMEditor/Public/RigVMEditorModule.h"
 #include "RigVMCore/RigVMExecuteContext.h"
 #include "RigVMModel/Nodes/RigVMAggregateNode.h"
 #include "RigVMModel/Nodes/RigVMFunctionReferenceNode.h"
@@ -1393,9 +1394,8 @@ FSlateIcon URigVMEdGraphNode::GetIconAndTint(FLinearColor& OutColor) const
 void URigVMEdGraphNode::GetNodeContextMenuActions(class UToolMenu* Menu, class UGraphNodeContextMenuContext* Context) const
 {
 #if WITH_EDITOR
-	// toodoooo
-	// const URigVMEdGraphSchema* Schema = Cast<URigVMEdGraphSchema>(GetSchema());
-	// IControlRigEditorModule::Get().GetContextMenuActions(Schema, Menu, Context);
+	const URigVMEdGraphSchema* Schema = Cast<URigVMEdGraphSchema>(GetSchema());
+	IRigVMEditorModule::Get().GetContextMenuActions(Schema, Menu, Context);
 #endif
 }
 
diff --git a/Engine/Plugins/Runtime/RigVM/Source/RigVMDeveloper/Private/EdGraph/RigVMEdGraphSchema.cpp b/Engine/Plugins/Runtime/RigVM/Source/RigVMDeveloper/Private/EdGraph/RigVMEdGraphSchema.cpp
--- a/Engine/Plugins/Runtime/RigVM/Source/RigVMDeveloper/Private/EdGraph/RigVMEdGraphSchema.cpp
+++ b/Engine/Plugins/Runtime/RigVM/Source/RigVMDeveloper/Private/EdGraph/RigVMEdGraphSchema.cpp
@@ -30,19 +30,9 @@
 #include "RigVMModel/Nodes/RigVMDispatchNode.h"
 #include "BlueprintEditor.h"
 
-/*
- * todooooo
-#include "Units/Execution/RigUnit_BeginExecution.h"
-#include "Units/Execution/RigUnit_PrepareForExecution.h"
-#include "Units/Execution/RigUnit_InverseExecution.h"
-#include "Units/Execution/RigUnit_InteractionExecution.h"
-*/
-
 #include UE_INLINE_GENERATED_CPP_BY_NAME(RigVMEdGraphSchema)
 
 #if WITH_EDITOR
-// todooooo
-//#include "Editor/SControlRigFunctionLocalizationWidget.h"
 #include "RigVMEditorModule.h"
 #include "Misc/MessageDialog.h"
 #include "Editor/Transactor.h"
diff --git a/Engine/Plugins/Runtime/RigVM/Source/RigVMDeveloper/Private/RigVMBlueprint.cpp b/Engine/Plugins/Runtime/RigVM/Source/RigVMDeveloper/Private/RigVMBlueprint.cpp
--- a/Engine/Plugins/Runtime/RigVM/Source/RigVMDeveloper/Private/RigVMBlueprint.cpp
+++ b/Engine/Plugins/Runtime/RigVM/Source/RigVMDeveloper/Private/RigVMBlueprint.cpp
@@ -34,6 +34,7 @@
 #include "Editor/UnrealEdEngine.h"
 #include "Editor/Transactor.h"
 #include "CookOnTheSide/CookOnTheFlyServer.h"
+#include "RigVMEditorModule.h"
 #endif//WITH_EDITOR
 
 #define LOCTEXT_NAMESPACE "RigVMBlueprint"
@@ -101,6 +102,7 @@ bool FRigVMOldPublicFunctionData::IsMutable() const
 
 FSoftObjectPath URigVMBlueprint::PreDuplicateAssetPath;
 FSoftObjectPath URigVMBlueprint::PreDuplicateHostPath;
+TArray<URigVMBlueprint*> URigVMBlueprint::sCurrentlyOpenedRigVMBlueprints;
 
 URigVMBlueprint::URigVMBlueprint()
 {
@@ -678,6 +680,13 @@ URigVMEditorSettings* URigVMBlueprint::GetRigVMEditorSettings() const
 	return GetMutableDefault<URigVMEditorSettings>(GetRigVMEditorSettingsClass());
 }
 
+#if WITH_EDITOR
+IRigVMEditorModule* URigVMBlueprint::GetEditorModule() const
+{
+	return &IRigVMEditorModule::Get();
+}
+#endif
+
 void URigVMBlueprint::Serialize(FArchive& Ar)
 {
 	RigVMClient.SetOuterClientHost(this, GET_MEMBER_NAME_CHECKED(URigVMBlueprint, RigVMClient));
@@ -2332,16 +2341,18 @@ void URigVMBlueprint::GetTypeActions(FBlueprintActionDatabaseRegistrar& ActionRe
 {
 	DECLARE_SCOPE_HIERARCHICAL_COUNTER_FUNC()
 
-	// todooooo
-	// IControlRigEditorModule::Get().GetTypeActions((URigVMBlueprint*)this, ActionRegistrar);
+#if WITH_EDITOR
+	GetEditorModule()->GetTypeActions((URigVMBlueprint*)this, ActionRegistrar);
+#endif
 }
 
 void URigVMBlueprint::GetInstanceActions(FBlueprintActionDatabaseRegistrar& ActionRegistrar) const
 {
 	DECLARE_SCOPE_HIERARCHICAL_COUNTER_FUNC()
 
-	// todooooo
-	// IControlRigEditorModule::Get().GetInstanceActions((URigVMBlueprint*)this, ActionRegistrar);
+#if WITH_EDITOR
+	GetEditorModule()->GetInstanceActions((URigVMBlueprint*)this, ActionRegistrar);
+#endif
 }
 
 void URigVMBlueprint::SetObjectBeingDebugged(UObject* NewObject)
diff --git a/Engine/Plugins/Runtime/RigVM/Source/RigVMDeveloper/Private/RigVMModel/RigVMController.cpp b/Engine/Plugins/Runtime/RigVM/Source/RigVMDeveloper/Private/RigVMModel/RigVMController.cpp
--- a/Engine/Plugins/Runtime/RigVM/Source/RigVMDeveloper/Private/RigVMModel/RigVMController.cpp
+++ b/Engine/Plugins/Runtime/RigVM/Source/RigVMDeveloper/Private/RigVMModel/RigVMController.cpp
@@ -45,7 +45,7 @@
 #include "AssetRegistry/AssetRegistryModule.h"
 #endif
 
-TMap<URigVMController::FControlRigStructPinRedirectorKey, FString> URigVMController::PinPathCoreRedirectors;
+TMap<URigVMController::FRigVMStructPinRedirectorKey, FString> URigVMController::PinPathCoreRedirectors;
 
 FRigVMControllerCompileBracketScope::FRigVMControllerCompileBracketScope(URigVMController* InController)
 : Graph(nullptr), bSuspendNotifications(InController->bSuspendNotifications)
@@ -14609,7 +14609,7 @@ bool URigVMController::ShouldRedirectPin(UScriptStruct* InOwningStruct, const FS
 		return false;
 	}
 	
-	FControlRigStructPinRedirectorKey RedirectorKey(InOwningStruct, InOldRelativePinPath);
+	FRigVMStructPinRedirectorKey RedirectorKey(InOwningStruct, InOldRelativePinPath);
 	if (const FString* RedirectedPinPath = PinPathCoreRedirectors.Find(RedirectorKey))
 	{
 		InOutNewRelativePinPath = *RedirectedPinPath;
@@ -14698,7 +14698,7 @@ bool URigVMController::ShouldRedirectPin(UScriptStruct* InOwningStruct, const FS
 
 					// this is also going to cache paths which haven't been redirected.
 					// consumers of the table have to still compare old != new
-					FControlRigStructPinRedirectorKey SubRedirectorKey(InOwningStruct, OldPath);
+					FRigVMStructPinRedirectorKey SubRedirectorKey(InOwningStruct, OldPath);
 					if (!PinPathCoreRedirectors.Contains(SubRedirectorKey))
 					{
 						PinPathCoreRedirectors.Add(SubRedirectorKey, NewPath);
diff --git a/Engine/Plugins/Runtime/RigVM/Source/RigVMDeveloper/Public/EdGraph/RigVMEdGraph.h b/Engine/Plugins/Runtime/RigVM/Source/RigVMDeveloper/Public/EdGraph/RigVMEdGraph.h
--- a/Engine/Plugins/Runtime/RigVM/Source/RigVMDeveloper/Public/EdGraph/RigVMEdGraph.h
+++ b/Engine/Plugins/Runtime/RigVM/Source/RigVMDeveloper/Public/EdGraph/RigVMEdGraph.h
@@ -104,8 +104,7 @@ private:
 #endif
 	friend class URigVMEdGraphNode;
 	friend class URigVMBlueprint;
-	friend class UControlRigGraphNode;
-	friend class FControlRigEditor;
+	friend class FRigVMEditor;
 	friend class SRigVMGraphNode;
 	friend class URigVMBlueprint;
 
diff --git a/Engine/Plugins/Runtime/RigVM/Source/RigVMDeveloper/Public/EdGraph/RigVMEdGraphNode.h b/Engine/Plugins/Runtime/RigVM/Source/RigVMDeveloper/Public/EdGraph/RigVMEdGraphNode.h
--- a/Engine/Plugins/Runtime/RigVM/Source/RigVMDeveloper/Public/EdGraph/RigVMEdGraphNode.h
+++ b/Engine/Plugins/Runtime/RigVM/Source/RigVMDeveloper/Public/EdGraph/RigVMEdGraphNode.h
@@ -277,10 +277,4 @@ private:
 	friend class URigVMEdGraphArrayNodeSpawner;
 	friend class URigVMEdGraphIfNodeSpawner;
 	friend class URigVMEdGraphSelectNodeSpawner;
-
-	// todo
-	friend class FControlRigEditor;
-	friend class UControlRigBlueprint;
-	friend class SControlRigGraphNode;
-	friend class UControlRigGraphSchema;
 };
diff --git a/Engine/Plugins/Runtime/RigVM/Source/RigVMDeveloper/Public/EdGraph/RigVMEdGraphSchema.h b/Engine/Plugins/Runtime/RigVM/Source/RigVMDeveloper/Public/EdGraph/RigVMEdGraphSchema.h
--- a/Engine/Plugins/Runtime/RigVM/Source/RigVMDeveloper/Public/EdGraph/RigVMEdGraphSchema.h
+++ b/Engine/Plugins/Runtime/RigVM/Source/RigVMDeveloper/Public/EdGraph/RigVMEdGraphSchema.h
@@ -314,7 +314,7 @@ public:
 	/** Helper function to recursively reset the pin defaults */
 	virtual void ResetPinDefaultsRecursive(UEdGraphPin* InPin) const;
 
-	/** Returns all of the applicable pin types for variables within a control rig */
+	/** Returns all of the applicable pin types for variables within a rigvm host */
 	virtual void GetVariablePinTypes(TArray<FEdGraphPinType>& PinTypes) const;
 
 	void StartGraphNodeInteraction(UEdGraphNode* InNode) const;
diff --git a/Engine/Plugins/Runtime/RigVM/Source/RigVMDeveloper/Public/RigVMBlueprint.h b/Engine/Plugins/Runtime/RigVM/Source/RigVMDeveloper/Public/RigVMBlueprint.h
--- a/Engine/Plugins/Runtime/RigVM/Source/RigVMDeveloper/Public/RigVMBlueprint.h
+++ b/Engine/Plugins/Runtime/RigVM/Source/RigVMDeveloper/Public/RigVMBlueprint.h
@@ -17,6 +17,10 @@
 #include "RigVMBlueprint.generated.h"
 
 class URigVMBlueprintGeneratedClass;
+
+#if WITH_EDITOR
+class IRigVMEditorModule;
+#endif
 struct FEndLoadPackageContext;
 
 DECLARE_EVENT_TwoParams(URigVMBlueprint, FOnRigVMCompiledEvent, UObject*, URigVM*);
@@ -233,6 +237,11 @@ public:
 	/** Returns the settings defaults for this blueprint */
 	URigVMEditorSettings* GetRigVMEditorSettings() const;
 
+#if WITH_EDITOR
+	/** Returns the editor module to be used for this blueprint */
+	virtual IRigVMEditorModule* GetEditorModule() const;
+#endif
+
 	virtual void Serialize(FArchive& Ar) override;
 
 #if WITH_EDITOR
@@ -319,7 +328,7 @@ public:
 	void DecrementVMRecompileBracket();
 
 	// this is needed since even after load
-	// model data can change while the Control Rig BP is not opened
+	// model data can change while the RigVM BP is not opened
 	// for example, if a user defined struct changed after BP load,
 	// any pin that references the struct needs to be regenerated
 	void RefreshAllModels(ERigVMBlueprintLoadType InLoadType = ERigVMBlueprintLoadType::PostLoad);
@@ -703,6 +712,7 @@ protected:
 
 	static FSoftObjectPath PreDuplicateAssetPath;
 	static FSoftObjectPath PreDuplicateHostPath;
+	static TArray<URigVMBlueprint*> sCurrentlyOpenedRigVMBlueprints;
 
 	void MarkDirtyDuringLoad() { bDirtyDuringLoad = true; }
 	
@@ -719,13 +729,10 @@ private:
 	bool bMarkBlueprintAsStructurallyModifiedPending;
 
 	friend class FRigVMBlueprintCompilerContext;
+	friend class FRigVMEditor;
 	friend class FRigVMEditorModule;
 	friend class URigVMEdGraphSchema;
 	friend struct FRigVMEdGraphSchemaAction_PromoteToVariable;
-
-	// todo
-	friend class FControlRigEditorModule;
-	friend class FControlRigBlueprintCompilerContext;
 };
 
 class RIGVMDEVELOPER_API FRigVMBlueprintCompileScope
diff --git a/Engine/Plugins/Runtime/RigVM/Source/RigVMDeveloper/Public/RigVMModel/Nodes/RigVMFunctionReferenceNode.h b/Engine/Plugins/Runtime/RigVM/Source/RigVMDeveloper/Public/RigVMModel/Nodes/RigVMFunctionReferenceNode.h
--- a/Engine/Plugins/Runtime/RigVM/Source/RigVMDeveloper/Public/RigVMModel/Nodes/RigVMFunctionReferenceNode.h
+++ b/Engine/Plugins/Runtime/RigVM/Source/RigVMDeveloper/Public/RigVMModel/Nodes/RigVMFunctionReferenceNode.h
@@ -70,7 +70,6 @@ private:
 
 	friend class URigVMController;
 	friend class FRigVMParserAST;
-	friend class UControlRigBlueprint;
 	friend class URigVMBlueprint;
 	friend struct FRigVMClient;
 	friend struct EngineTestRigVMFramework;
diff --git a/Engine/Plugins/Runtime/RigVM/Source/RigVMDeveloper/Public/RigVMModel/Nodes/RigVMInvokeEntryNode.h b/Engine/Plugins/Runtime/RigVM/Source/RigVMDeveloper/Public/RigVMModel/Nodes/RigVMInvokeEntryNode.h
--- a/Engine/Plugins/Runtime/RigVM/Source/RigVMDeveloper/Public/RigVMModel/Nodes/RigVMInvokeEntryNode.h
+++ b/Engine/Plugins/Runtime/RigVM/Source/RigVMDeveloper/Public/RigVMModel/Nodes/RigVMInvokeEntryNode.h
@@ -38,7 +38,6 @@ private:
 	
 
 	friend class URigVMController;
-	friend class UControlRigBlueprint;
 	friend struct FRigVMRemoveNodeAction;
 	friend class URigVMCompiler;
 	friend class FRigVMParserAST;
diff --git a/Engine/Plugins/Runtime/RigVM/Source/RigVMDeveloper/Public/RigVMModel/Nodes/RigVMTemplateNode.h b/Engine/Plugins/Runtime/RigVM/Source/RigVMDeveloper/Public/RigVMModel/Nodes/RigVMTemplateNode.h
--- a/Engine/Plugins/Runtime/RigVM/Source/RigVMDeveloper/Public/RigVMModel/Nodes/RigVMTemplateNode.h
+++ b/Engine/Plugins/Runtime/RigVM/Source/RigVMDeveloper/Public/RigVMModel/Nodes/RigVMTemplateNode.h
@@ -59,7 +59,6 @@ protected:
 
 	friend class URigVMTemplateNode;
 	friend class URigVMController;
-	friend class UControlRigBlueprint;
 };
 
 /**
@@ -169,7 +168,6 @@ protected:
 	mutable int32 ResolvedPermutation;
 
 	friend class URigVMController;
-	friend class UControlRigBlueprint;
 	friend class URigVMBlueprint;
 	friend struct FRigVMSetTemplateFilteredPermutationsAction;
 	friend struct FRigVMSetPreferredTemplatePermutationsAction;
diff --git a/Engine/Plugins/Runtime/RigVM/Source/RigVMDeveloper/Public/RigVMModel/RigVMController.h b/Engine/Plugins/Runtime/RigVM/Source/RigVMDeveloper/Public/RigVMModel/RigVMController.h
--- a/Engine/Plugins/Runtime/RigVM/Source/RigVMDeveloper/Public/RigVMModel/RigVMController.h
+++ b/Engine/Plugins/Runtime/RigVM/Source/RigVMDeveloper/Public/RigVMModel/RigVMController.h
@@ -974,7 +974,7 @@ public:
 	FRigVMUnitNodeCreatedContext& GetUnitNodeCreatedContext() { return UnitNodeCreatedContext; }
 
 	// Wires the unit node delegates to the default controller delegates.
-	// this is used only within the Control Rig Editor currently.
+	// this is used only within the RigVM Editor currently.
 	void SetupDefaultUnitNodeDelegates(TDelegate<FName(FRigVMExternalVariable, FString)> InCreateExternalVariableDelegate);
 	void ResetUnitNodeDelegates();
 
@@ -1386,29 +1386,29 @@ private:
 	TMap<FString, FString> InputPinRedirectors;
 	TMap<FString, FString> OutputPinRedirectors;
 
-	struct FControlRigStructPinRedirectorKey
+	struct FRigVMStructPinRedirectorKey
 	{
-		FControlRigStructPinRedirectorKey()
+		FRigVMStructPinRedirectorKey()
 		{
 		}
 
-		FControlRigStructPinRedirectorKey(UScriptStruct* InScriptStruct, const FString& InPinPathInNode)
+		FRigVMStructPinRedirectorKey(UScriptStruct* InScriptStruct, const FString& InPinPathInNode)
 		: Struct(InScriptStruct)
 		, PinPathInNode(InPinPathInNode)
 		{
 		}
 
-		friend uint32 GetTypeHash(const FControlRigStructPinRedirectorKey& Cache)
+		friend uint32 GetTypeHash(const FRigVMStructPinRedirectorKey& Cache)
 		{
 			return HashCombine(GetTypeHash(Cache.Struct), GetTypeHash(Cache.PinPathInNode));
 		}
 
-		bool operator ==(const FControlRigStructPinRedirectorKey& Other) const
+		bool operator ==(const FRigVMStructPinRedirectorKey& Other) const
 		{
 			return Struct == Other.Struct && PinPathInNode == Other.PinPathInNode;
 		}
 
-		bool operator !=(const FControlRigStructPinRedirectorKey& Other) const
+		bool operator !=(const FRigVMStructPinRedirectorKey& Other) const
 		{
 			return Struct != Other.Struct || PinPathInNode != Other.PinPathInNode;
 		}
@@ -1417,7 +1417,7 @@ private:
 		FString PinPathInNode;
 	};
 
-	static TMap<FControlRigStructPinRedirectorKey, FString> PinPathCoreRedirectors;
+	static TMap<FRigVMStructPinRedirectorKey, FString> PinPathCoreRedirectors;
 	FCriticalSection PinPathCoreRedirectorsLock;
 
 	FRigVMUnitNodeCreatedContext UnitNodeCreatedContext;
diff --git a/Engine/Plugins/Runtime/RigVM/Source/RigVMDeveloper/RigVMDeveloper.Build.cs b/Engine/Plugins/Runtime/RigVM/Source/RigVMDeveloper/RigVMDeveloper.Build.cs
--- a/Engine/Plugins/Runtime/RigVM/Source/RigVMDeveloper/RigVMDeveloper.Build.cs
+++ b/Engine/Plugins/Runtime/RigVM/Source/RigVMDeveloper/RigVMDeveloper.Build.cs
@@ -33,8 +33,9 @@ public class RigVMDeveloper : ModuleRules
 					"Kismet",
                 }
 			);
-
+            
             PrivateIncludePathModuleNames.Add("RigVMEditor");
+            DynamicallyLoadedModuleNames.Add("RigVMEditor");
         }
     }
 }
diff --git a/Engine/Plugins/Runtime/RigVM/Source/RigVMEditor/Private/EdGraph/NodeSpawners/RigVMEdGraphFunctionRefNodeSpawner.cpp b/Engine/Plugins/Runtime/RigVM/Source/RigVMEditor/Private/EdGraph/NodeSpawners/RigVMEdGraphFunctionRefNodeSpawner.cpp
--- a/Engine/Plugins/Runtime/RigVM/Source/RigVMEditor/Private/EdGraph/NodeSpawners/RigVMEdGraphFunctionRefNodeSpawner.cpp
+++ b/Engine/Plugins/Runtime/RigVM/Source/RigVMEditor/Private/EdGraph/NodeSpawners/RigVMEdGraphFunctionRefNodeSpawner.cpp
@@ -411,6 +411,11 @@ URigVMEdGraphNode* URigVMEdGraphFunctionRefNodeSpawner::SpawnNode(UEdGraph* Pare
 
 bool URigVMEdGraphFunctionRefNodeSpawner::IsTemplateNodeFilteredOut(FBlueprintActionFilter const& Filter) const
 {
+	if(URigVMEdGraphNodeSpawner::IsTemplateNodeFilteredOut(Filter))
+	{
+		return true;
+	}
+	
 	if(bIsLocalFunction)
 	{
 		if(ReferencedPublicFunctionHeader.IsValid())
diff --git a/Engine/Plugins/Runtime/RigVM/Source/RigVMEditor/Private/EdGraph/NodeSpawners/RigVMEdGraphNodeSpawner.cpp b/Engine/Plugins/Runtime/RigVM/Source/RigVMEditor/Private/EdGraph/NodeSpawners/RigVMEdGraphNodeSpawner.cpp
new file mode 100644
--- /dev/null
+++ b/Engine/Plugins/Runtime/RigVM/Source/RigVMEditor/Private/EdGraph/NodeSpawners/RigVMEdGraphNodeSpawner.cpp
@@ -0,0 +1,30 @@
+// Copyright Epic Games, Inc. All Rights Reserved.
+
+#include "EdGraph/NodeSpawners/RigVMEdGraphNodeSpawner.h"
+
+#include UE_INLINE_GENERATED_CPP_BY_NAME(RigVMEdGraphNodeSpawner)
+
+#define LOCTEXT_NAMESPACE "RigVMEdGraphNodeSpawner"
+
+bool URigVMEdGraphNodeSpawner::IsTemplateNodeFilteredOut(FBlueprintActionFilter const& Filter) const
+{
+	check(RelatedBlueprintClass);
+
+	for(const UBlueprint* Blueprint : Filter.Context.Blueprints)
+	{
+		if(Blueprint->GetClass() != RelatedBlueprintClass)
+		{
+			return true;
+		}
+	}
+
+	return false;
+}
+
+void URigVMEdGraphNodeSpawner::SetRelatedBlueprintClass(TSubclassOf<URigVMBlueprint> InClass)
+{
+	RelatedBlueprintClass = InClass;
+}
+
+#undef LOCTEXT_NAMESPACE
+
diff --git a/Engine/Plugins/Runtime/RigVM/Source/RigVMEditor/Private/EdGraph/NodeSpawners/RigVMEdGraphTemplateNodeSpawner.cpp b/Engine/Plugins/Runtime/RigVM/Source/RigVMEditor/Private/EdGraph/NodeSpawners/RigVMEdGraphTemplateNodeSpawner.cpp
--- a/Engine/Plugins/Runtime/RigVM/Source/RigVMEditor/Private/EdGraph/NodeSpawners/RigVMEdGraphTemplateNodeSpawner.cpp
+++ b/Engine/Plugins/Runtime/RigVM/Source/RigVMEditor/Private/EdGraph/NodeSpawners/RigVMEdGraphTemplateNodeSpawner.cpp
@@ -75,6 +75,11 @@ FBlueprintNodeSignature URigVMEdGraphTemplateNodeSpawner::GetSpawnerSignature()
 
 bool URigVMEdGraphTemplateNodeSpawner::IsTemplateNodeFilteredOut(FBlueprintActionFilter const& Filter) const
 {
+	if(URigVMEdGraphNodeSpawner::IsTemplateNodeFilteredOut(Filter))
+	{
+		return true;
+	}
+
 	bool bFilteredOut = true;
 	for (UEdGraphPin* Pin : Filter.Context.Pins)
 	{
@@ -133,8 +138,8 @@ bool URigVMEdGraphTemplateNodeSpawner::IsTemplateNodeFilteredOut(FBlueprintActio
 			}
 		}
 	}
-	
-	return Super::IsTemplateNodeFilteredOut(Filter);
+
+	return false;
 }
 
 FBlueprintActionUiSpec URigVMEdGraphTemplateNodeSpawner::GetUiSpec(FBlueprintActionContext const& Context, FBindingSet const& Bindings) const
diff --git a/Engine/Plugins/Runtime/RigVM/Source/RigVMEditor/Private/EdGraph/NodeSpawners/RigVMEdGraphVariableNodeSpawner.cpp b/Engine/Plugins/Runtime/RigVM/Source/RigVMEditor/Private/EdGraph/NodeSpawners/RigVMEdGraphVariableNodeSpawner.cpp
--- a/Engine/Plugins/Runtime/RigVM/Source/RigVMEditor/Private/EdGraph/NodeSpawners/RigVMEdGraphVariableNodeSpawner.cpp
+++ b/Engine/Plugins/Runtime/RigVM/Source/RigVMEditor/Private/EdGraph/NodeSpawners/RigVMEdGraphVariableNodeSpawner.cpp
@@ -62,6 +62,11 @@ void URigVMEdGraphVariableNodeSpawner::Prime()
 
 bool URigVMEdGraphVariableNodeSpawner::IsTemplateNodeFilteredOut(FBlueprintActionFilter const& Filter) const
 {
+	if(URigVMEdGraphNodeSpawner::IsTemplateNodeFilteredOut(Filter))
+	{
+		return true;
+	}
+
 	if (Blueprint.IsValid())
 	{
 		if (!Filter.Context.Blueprints.Contains(Blueprint.Get()))
@@ -89,12 +94,12 @@ bool URigVMEdGraphVariableNodeSpawner::IsTemplateNodeFilteredOut(FBlueprintActio
 			}
 		}
 	}
-	return Super::IsTemplateNodeFilteredOut(Filter);
+	return false;
 }
 
 FBlueprintNodeSignature URigVMEdGraphVariableNodeSpawner::GetSpawnerSignature() const
 {
-	return FBlueprintNodeSignature(FString("RigUnit=" + ExternalVariable.Name.ToString()));
+	return FBlueprintNodeSignature(FString("ExternalVariable=" + ExternalVariable.Name.ToString()));
 }
 
 FBlueprintActionUiSpec URigVMEdGraphVariableNodeSpawner::GetUiSpec(FBlueprintActionContext const& Context, FBindingSet const& Bindings) const
