Control Rig: Fix issue with use of stack variable after free as reported by address sanitizer.

#rnx

diff --git a/Engine/Plugins/Animation/ControlRig/Source/ControlRig/Private/Rigs/RigHierarchy.cpp b/Engine/Plugins/Animation/ControlRig/Source/ControlRig/Private/Rigs/RigHierarchy.cpp
--- a/Engine/Plugins/Animation/ControlRig/Source/ControlRig/Private/Rigs/RigHierarchy.cpp
+++ b/Engine/Plugins/Animation/ControlRig/Source/ControlRig/Private/Rigs/RigHierarchy.cpp
@@ -1937,9 +1937,10 @@ bool URigHierarchy::SwitchToParent(FRigBaseElement* InChild, FRigBaseElement* In
 	const TElementDependencyMap* DependencyMapPtr = &InDependencyMap;
 
 #if WITH_EDITOR
+	// Keep this in function scope, since we might be taking a pointer to it.
+	TElementDependencyMap DependencyMapFromVM;
 	{
 		FScopeLock Lock(&ExecuteContextLock);
-		TElementDependencyMap DependencyMapFromVM;
 		if(ExecuteContext != nullptr && DependencyMapPtr->IsEmpty())
 		{
 			if(ExecuteContext->VM)
