Control Rigs: Fix broken Control Rigs that aren't completely removed.

diff --git a/Engine/Plugins/Animation/ControlRig/Source/ControlRigEditor/Private/Sequencer/ControlRigParameterTrackEditor.cpp b/Engine/Plugins/Animation/ControlRig/Source/ControlRigEditor/Private/Sequencer/ControlRigParameterTrackEditor.cpp
--- a/Engine/Plugins/Animation/ControlRig/Source/ControlRigEditor/Private/Sequencer/ControlRigParameterTrackEditor.cpp
+++ b/Engine/Plugins/Animation/ControlRig/Source/ControlRigEditor/Private/Sequencer/ControlRigParameterTrackEditor.cpp
@@ -478,8 +478,9 @@ void FControlRigParameterTrackEditor::UnbindControlRig(UControlRig* ControlRig)
 			Binding->OnControlRigBind().RemoveAll(this);
 		}
 		ControlRig->ControlRigBound().RemoveAll(this);
-		ControlRig->Rename(nullptr, GetTransientPackage(), REN_ForceNoResetLoaders | REN_DoNotDirty | REN_DontCreateRedirectors | REN_NonTransactional);
-		ControlRig->MarkAsGarbage();
+		//for sara was added for additive but breaks reloeading
+		//ControlRig->Rename(nullptr, GetTransientPackage(), REN_ForceNoResetLoaders | REN_DoNotDirty | REN_DontCreateRedirectors | REN_NonTransactional);
+		//ControlRig->MarkAsGarbage();
 		
 		BoundControlRigs.Remove(ControlRig);
 		ClearOutAllSpaceAndConstraintDelegates(ControlRig);
