Control Rig: Fix additive rotations applying twice
#rnx

diff --git a/Engine/Plugins/Animation/ControlRig/Source/ControlRig/Private/ControlRig.cpp b/Engine/Plugins/Animation/ControlRig/Source/ControlRig/Private/ControlRig.cpp
--- a/Engine/Plugins/Animation/ControlRig/Source/ControlRig/Private/ControlRig.cpp
+++ b/Engine/Plugins/Animation/ControlRig/Source/ControlRig/Private/ControlRig.cpp
@@ -332,37 +332,39 @@ void UControlRig::Evaluate_AnyThread()
 		// and we don't want other systems to see that.
 		FScopeLock EvaluateLock(&GetEvaluateMutex());
 
+		URigHierarchy* Hierarchy = GetHierarchy();
+
 		if (PoseBeforeBackwardsSolve.Num() == 0)
 		{
 			// If the pose is empty, this is an indication that a new pose is coming in
-			PoseBeforeBackwardsSolve = GetHierarchy()->GetPose(false, ERigElementType::Bone, TArrayView<const FRigElementKey>());
+			PoseBeforeBackwardsSolve = Hierarchy->GetPose(false, ERigElementType::Bone, TArrayView<const FRigElementKey>());
 		}
 		else
 		{
 			// Restore the pose from the anim sequence
-			DynamicHierarchy->SetPose(PoseBeforeBackwardsSolve);
+			Hierarchy->SetPose(PoseBeforeBackwardsSolve);
 		}
 		
 		// Backwards solve
 		{
-			TGuardValue<bool> UpdatePreferredAngles(DynamicHierarchy->bUpdatePreferedEulerAngleWhenSettingTransform, false);
+			TGuardValue<bool> UpdatePreferredAngles(Hierarchy->bUpdatePreferedEulerAngleWhenSettingTransform, false);
 			Execute(FRigUnit_InverseExecution::EventName);
 		}
 
 		// Store control pose after backwards solve to figure out additive local transforms based on animation
-		ControlsAfterBackwardsSolve = GetHierarchy()->GetPose(false, ERigElementType::Control, TArrayView<const FRigElementKey>());
+		ControlsAfterBackwardsSolve = Hierarchy->GetPose(false, ERigElementType::Control, TArrayView<const FRigElementKey>());
 
 		// Apply additive controls
 		for (TPair<FRigElementKey, FRigSetControlValueInfo>& Value : ControlValues)
 		{
-			if (FRigBaseElement* Element = GetHierarchy()->Find(Value.Key))
+			if (FRigBaseElement* Element = Hierarchy->Find(Value.Key))
 			{
 				if (FRigControlElement* Control = Cast<FRigControlElement>(Element))
 				{
 					FRigSetControlValueInfo& Info = Value.Value;
 
 					// Transform from animation
-					const FRigControlValue PreviousValue = GetHierarchy()->GetControlValue(Control, ERigControlValueType::Current);
+					const FRigControlValue PreviousValue = Hierarchy->GetControlValue(Control, ERigControlValueType::Current, false);
 					const FTransform PreviousTransform = PreviousValue.GetAsTransform(Control->Settings.ControlType, Control->Settings.PrimaryAxis);
 
 					// Additive transform from controls
@@ -375,8 +377,8 @@ void UControlRig::Evaluate_AnyThread()
 					FRigControlValue FinalValue;
 					FinalValue.SetFromTransform(FinalTransform, Control->Settings.ControlType, Control->Settings.PrimaryAxis);
 					
-					GetHierarchy()->SetControlValue(Control, FinalValue, ERigControlValueType::Current, Info.bSetupUndo, false, Info.bPrintPythonCommnds, false);
-					GetHierarchy()->SetPreferredEulerAnglesFromValue(Control, AdditiveValue, ERigControlValueType::Current, Info.bFixEulerFlips);
+					Hierarchy->SetControlValue(Control, FinalValue, ERigControlValueType::Current, Info.bSetupUndo, false, Info.bPrintPythonCommnds, false);
+					Hierarchy->SetPreferredEulerAnglesFromValue(Control, AdditiveValue, ERigControlValueType::Current, Info.bFixEulerFlips);
 					
 					if (Info.bNotify && OnControlModified.IsBound())
 					{
@@ -389,7 +391,7 @@ void UControlRig::Evaluate_AnyThread()
 		
 		// Forward solve
 		{
-			TGuardValue<bool> UpdatePreferredAngles(DynamicHierarchy->bUpdatePreferedEulerAngleWhenSettingTransform, false);
+			TGuardValue<bool> UpdatePreferredAngles(Hierarchy->bUpdatePreferedEulerAngleWhenSettingTransform, false);
 			Execute(FRigUnit_BeginExecution::EventName);
 		}
 	}
diff --git a/Engine/Plugins/Animation/ControlRig/Source/ControlRig/Private/Rigs/RigHierarchy.cpp b/Engine/Plugins/Animation/ControlRig/Source/ControlRig/Private/Rigs/RigHierarchy.cpp
--- a/Engine/Plugins/Animation/ControlRig/Source/ControlRig/Private/Rigs/RigHierarchy.cpp
+++ b/Engine/Plugins/Animation/ControlRig/Source/ControlRig/Private/Rigs/RigHierarchy.cpp
@@ -3470,7 +3470,7 @@ FTransform URigHierarchy::GetParentTransform(FRigBaseElement* InElement, const E
 	return FTransform::Identity;
 }
 
-FRigControlValue URigHierarchy::GetControlValue(FRigControlElement* InControlElement, ERigControlValueType InValueType) const
+FRigControlValue URigHierarchy::GetControlValue(FRigControlElement* InControlElement, ERigControlValueType InValueType, bool bUsePreferredAngles) const
 {
 	LLM_SCOPE_BYNAME(TEXT("Animation/ControlRig"));
 	using namespace ERigTransformType;
@@ -3479,8 +3479,13 @@ FRigControlValue URigHierarchy::GetControlValue(FRigControlElement* InControlEle
 
 	if(InControlElement != nullptr)
 	{
-		auto GetValueFromPreferredEulerAngles = [this, InControlElement, &Value, InValueType]() -> bool
+		auto GetValueFromPreferredEulerAngles = [this, InControlElement, &Value, InValueType, bUsePreferredAngles]() -> bool
 		{
+			if (!bUsePreferredAngles)
+			{
+				return false;
+			}
+			
 			const bool bInitial = InValueType == ERigControlValueType::Initial;
 			switch(InControlElement->Settings.ControlType)
 			{
diff --git a/Engine/Plugins/Animation/ControlRig/Source/ControlRig/Public/Rigs/RigHierarchy.h b/Engine/Plugins/Animation/ControlRig/Source/ControlRig/Public/Rigs/RigHierarchy.h
--- a/Engine/Plugins/Animation/ControlRig/Source/ControlRig/Public/Rigs/RigHierarchy.h
+++ b/Engine/Plugins/Animation/ControlRig/Source/ControlRig/Public/Rigs/RigHierarchy.h
@@ -3622,9 +3622,10 @@ public:
 	 * Returns a control's current value
 	 * @param InControlElement The element to retrieve the current value for
 	 * @param InValueType The type of value to return
+	 * @param bUsePreferredAngles When true will use euler preferred angles to compute the value
 	 * @return Returns the current value of the control
 	 */
-	FRigControlValue GetControlValue(FRigControlElement* InControlElement, ERigControlValueType InValueType) const;
+	FRigControlValue GetControlValue(FRigControlElement* InControlElement, ERigControlValueType InValueType, bool bUsePreferredAngles = true) const;
 
 	void SetPreferredEulerAnglesFromValue(FRigControlElement* InControlElement, const FRigControlValue& InValue,
 	                                      const ERigControlValueType& InValueType, bool bFixEulerFlips);
