Control Rig: Fix crash when switching levels if an additive control rig was present in the old level
#rnx

diff --git a/Engine/Plugins/Animation/ControlRig/Source/ControlRigEditor/Private/Sequencer/ControlRigParameterTrackEditor.cpp b/Engine/Plugins/Animation/ControlRig/Source/ControlRigEditor/Private/Sequencer/ControlRigParameterTrackEditor.cpp
--- a/Engine/Plugins/Animation/ControlRig/Source/ControlRigEditor/Private/Sequencer/ControlRigParameterTrackEditor.cpp
+++ b/Engine/Plugins/Animation/ControlRig/Source/ControlRigEditor/Private/Sequencer/ControlRigParameterTrackEditor.cpp
@@ -478,9 +478,6 @@ void FControlRigParameterTrackEditor::UnbindControlRig(UControlRig* ControlRig)
 			Binding->OnControlRigBind().RemoveAll(this);
 		}
 		ControlRig->ControlRigBound().RemoveAll(this);
-		//for sara was added for additive but breaks reloeading
-		//ControlRig->Rename(nullptr, GetTransientPackage(), REN_ForceNoResetLoaders | REN_DoNotDirty | REN_DontCreateRedirectors | REN_NonTransactional);
-		//ControlRig->MarkAsGarbage();
 		
 		BoundControlRigs.Remove(ControlRig);
 		ClearOutAllSpaceAndConstraintDelegates(ControlRig);
@@ -1138,7 +1135,8 @@ void FControlRigParameterTrackEditor::AddAdditiveControlRig(FGuid ObjectBinding,
 				ControlRig->SetBoneInitialTransformsFromSkeletalMeshComponent(SkelMeshComp, true);
 				ControlRig->Evaluate_AnyThread();
 
-				const bool bSequencerOwnsControlRig = true;
+				// The control rig is owned by the ControlRigComponent, which is owned by the actor
+				const bool bSequencerOwnsControlRig = false;
 				UMovieSceneSection* NewSection = Track->CreateControlRigSection(0, ControlRig, bSequencerOwnsControlRig);
 				UMovieSceneControlRigParameterSection* ParamSection = Cast<UMovieSceneControlRigParameterSection>(NewSection);
 
