Sequencer : Control Rig: Fix backing to control rig if doing so outside of section range so that spawnable doesn't exist.
#preflght na
#trivial

diff --git a/Engine/Plugins/Animation/ControlRig/Source/ControlRigEditor/Private/EditMode/ControlRigEditMode.cpp b/Engine/Plugins/Animation/ControlRig/Source/ControlRigEditor/Private/EditMode/ControlRigEditMode.cpp
--- a/Engine/Plugins/Animation/ControlRig/Source/ControlRigEditor/Private/EditMode/ControlRigEditMode.cpp
+++ b/Engine/Plugins/Animation/ControlRig/Source/ControlRigEditor/Private/EditMode/ControlRigEditMode.cpp
@@ -4423,6 +4423,10 @@ TArray<const UControlRig*> FControlRigEditMode::GetControlRigsArray(bool bIsVisi
 
 void FControlRigEditMode::RemoveControlRig(UControlRig* InControlRig)
 {
+	if (InControlRig == nullptr)
+	{
+		return;
+	}
 	InControlRig->ControlModified().RemoveAll(this);
 	InControlRig->GetHierarchy()->OnModified().RemoveAll(this);
 	InControlRig->OnPreConstructionForUI_AnyThread().RemoveAll(this);
diff --git a/Engine/Source/Editor/MovieSceneTools/Private/MovieSceneToolHelpers.cpp b/Engine/Source/Editor/MovieSceneTools/Private/MovieSceneToolHelpers.cpp
--- a/Engine/Source/Editor/MovieSceneTools/Private/MovieSceneToolHelpers.cpp
+++ b/Engine/Source/Editor/MovieSceneTools/Private/MovieSceneToolHelpers.cpp
@@ -3981,7 +3981,7 @@ bool MovieSceneToolHelpers::ExportToAnimSequence(UAnimSequence* AnimSequence, UA
 	MovieSceneToolHelpers::BakeToSkelMeshToCallbacks(MovieScene,Player,
 		SkelMeshComp, Template, RootToLocalTransform, ExportOptions,
 		InitCallback, StartCallback, TickCallback, EndCallback);
-	return true;
+	return AnimSequence->GetDataModel()->HasBeenPopulated();
 }
 
 FSpawnableRestoreState::FSpawnableRestoreState(UMovieScene* MovieScene)
