Sequencer: Control Rig:  Don't snap if there are no keys and don't create a transform when looking for keys.
#trivial

diff --git a/Engine/Plugins/Animation/ControlRig/Source/ControlRigEditor/Private/Tools/ControlRigSnapper.cpp b/Engine/Plugins/Animation/ControlRig/Source/ControlRigEditor/Private/Tools/ControlRigSnapper.cpp
--- a/Engine/Plugins/Animation/ControlRig/Source/ControlRigEditor/Private/Tools/ControlRigSnapper.cpp
+++ b/Engine/Plugins/Animation/ControlRig/Source/ControlRigEditor/Private/Tools/ControlRigSnapper.cpp
@@ -473,13 +473,9 @@ bool FControlRigSnapper::SnapIt(FFrameNumber StartFrame, FFrameNumber EndFrame,c
 	TWeakPtr<ISequencer> InSequencer = GetSequencer();
 	if (InSequencer.IsValid() && InSequencer.Pin()->GetFocusedMovieSceneSequence() && ActorToSnap.IsValid())
 	{
-		FScopedTransaction ScopedTransaction(LOCTEXT("SnapAnimation", "Snap Animation"));
-
 		ISequencer* Sequencer = InSequencer.Pin().Get();
-		Sequencer->ForceEvaluate(); // force an evaluate so any control rig get's binding setup
 		UMovieScene* MovieScene = Sequencer->GetFocusedMovieSceneSequence()->GetMovieScene();
-		MovieScene->Modify();
-		
+
 		TArray<FFrameNumber> Frames;
 		if (SnapSettings && SnapSettings->BakingKeySettings == EBakingKeySettings::AllFrames)
 		{
@@ -493,6 +489,12 @@ bool FControlRigSnapper::SnapIt(FFrameNumber StartFrame, FFrameNumber EndFrame,c
 			FrameMap.GenerateKeyArray(Frames);
 
 		}
+		if (Frames.Num() == 0)
+		{
+			UE_LOG(LogControlRig, Error, TEXT("No Frames To Snap, probably no keys within time range"));
+			return false;
+		}
+		Sequencer->ForceEvaluate(); // force an evaluate so any control rig get's binding setup
 		TArray<FTransform> WorldTransformToSnap;
 		bool bSnapToFirstFrameNotParents = !CalculateWorldTransformsFromParents(Sequencer, ParentToSnap, Frames, WorldTransformToSnap);
 		if (Frames.Num() != WorldTransformToSnap.Num())
@@ -501,6 +503,11 @@ bool FControlRigSnapper::SnapIt(FFrameNumber StartFrame, FFrameNumber EndFrame,c
 			return false;
 		}
 
+		FScopedTransaction ScopedTransaction(LOCTEXT("SnapAnimation", "Snap Animation"));
+
+		
+		MovieScene->Modify();
+
 		//need to make sure binding is setup
 		Sequencer->ForceEvaluate();
 		
diff --git a/Engine/Source/Editor/MovieSceneTools/Private/MovieSceneToolHelpers.cpp b/Engine/Source/Editor/MovieSceneTools/Private/MovieSceneToolHelpers.cpp
--- a/Engine/Source/Editor/MovieSceneTools/Private/MovieSceneToolHelpers.cpp
+++ b/Engine/Source/Editor/MovieSceneTools/Private/MovieSceneToolHelpers.cpp
@@ -5170,8 +5170,7 @@ UMovieScene3DTransformSection* MovieSceneToolHelpers::GetTransformSection(
 	UMovieScene3DTransformTrack* TransformTrack = MovieScene->FindTrack<UMovieScene3DTransformTrack>(InGuid);
 	if (!TransformTrack)
 	{
-		MovieScene->Modify();
-		TransformTrack = MovieScene->AddTrack<UMovieScene3DTransformTrack>(InGuid);
+		return nullptr;
 	}
 	TransformTrack->Modify();
 
