Fix Control Display actor always shows hidden controls

diff --git a/Engine/Plugins/Animation/ControlRig/Source/ControlRig/Private/ControlRigControlActor.cpp b/Engine/Plugins/Animation/ControlRig/Source/ControlRig/Private/ControlRigControlActor.cpp
--- a/Engine/Plugins/Animation/ControlRig/Source/ControlRig/Private/ControlRigControlActor.cpp
+++ b/Engine/Plugins/Animation/ControlRig/Source/ControlRig/Private/ControlRigControlActor.cpp
@@ -266,16 +266,22 @@ void AControlRigControlActor::Refresh()
 	{
 		URigHierarchy* Hierarchy = ControlRig->GetHierarchy();
 
+		UStaticMeshComponent* Component = Components[GizmoIndex];
 		const FRigElementKey ControlKey(ControlNames[GizmoIndex], ERigElementType::Control);
 		const FRigControlElement* ControlElement = Hierarchy->Find<FRigControlElement>(ControlKey);
-		if(ControlElement == nullptr)
+		const bool bVisible = ControlElement && ControlElement->Settings.IsVisible();
+
+		if (Component && Component->IsVisible() != bVisible)
+		{
+			Component->SetVisibility(bVisible);
+		}
+		if (Component == nullptr || ControlElement == nullptr)
 		{
-			Components[GizmoIndex]->SetVisibility(false);
 			continue;
 		}
 
 		FTransform ControlTransform = ControlRig->GetControlGlobalTransform(ControlNames[GizmoIndex]);
-		Components[GizmoIndex]->SetRelativeTransform(ShapeTransforms[GizmoIndex] * ControlTransform);
+		Component->SetRelativeTransform(ShapeTransforms[GizmoIndex] * ControlTransform);
 		Materials[GizmoIndex]->SetVectorParameterValue(ColorParameterName, FVector(ControlElement->Settings.ShapeColor));
 	}
 }
