Control Rig: Fix constructor for FRigPoseElement (initialize all members)

diff --git a/Engine/Plugins/Animation/ControlRig/Source/ControlRig/Public/Rigs/RigHierarchyPose.h b/Engine/Plugins/Animation/ControlRig/Source/ControlRig/Public/Rigs/RigHierarchyPose.h
--- a/Engine/Plugins/Animation/ControlRig/Source/ControlRig/Public/Rigs/RigHierarchyPose.h
+++ b/Engine/Plugins/Animation/ControlRig/Source/ControlRig/Public/Rigs/RigHierarchyPose.h
@@ -18,6 +18,7 @@ public:
 	: Index()
 	, GlobalTransform(FTransform::Identity)
 	, LocalTransform(FTransform::Identity)
+	, PreferredEulerAngle(FVector::ZeroVector)
 	, CurveValue(0.f)
 	{
 	}
