ControlRig: Fix crash on close

diff --git a/Engine/Plugins/Animation/ControlRig/Source/ControlRigEditor/Private/Editor/ControlRigEditor.cpp b/Engine/Plugins/Animation/ControlRig/Source/ControlRigEditor/Private/Editor/ControlRigEditor.cpp
--- a/Engine/Plugins/Animation/ControlRig/Source/ControlRigEditor/Private/Editor/ControlRigEditor.cpp
+++ b/Engine/Plugins/Animation/ControlRig/Source/ControlRigEditor/Private/Editor/ControlRigEditor.cpp
@@ -145,12 +145,6 @@ FControlRigEditor::~FControlRigEditor()
 		}
 	}
 
-	if (UWorld* PreviewWorld = GetPersonaToolkit()->GetPreviewScene()->GetWorld())
-	{
-		PreviewWorld->MarkObjectsPendingKill();
-		PreviewWorld->MarkAsGarbage();
-	}
-
 	if (PersonaToolkit.IsValid())
 	{
 		constexpr bool bSetPreviewMeshInAsset = false;
