RigVM: fixed variable/parameter patch (lost during refactoring)

#rnx

diff --git a/Engine/Plugins/Animation/ControlRig/Source/ControlRigDeveloper/Private/ControlRigBlueprint.cpp b/Engine/Plugins/Animation/ControlRig/Source/ControlRigDeveloper/Private/ControlRigBlueprint.cpp
--- a/Engine/Plugins/Animation/ControlRig/Source/ControlRigDeveloper/Private/ControlRigBlueprint.cpp
+++ b/Engine/Plugins/Animation/ControlRig/Source/ControlRigDeveloper/Private/ControlRigBlueprint.cpp
@@ -939,6 +939,216 @@ void UControlRigBlueprint::PatchFunctionsOnLoad()
 	}
 }
 
+void UControlRigBlueprint::CreateMemberVariablesOnLoad()
+{
+#if WITH_EDITOR
+
+	const int32 LinkerVersion = GetLinkerCustomVersion(FControlRigObjectVersion::GUID);
+	if (LinkerVersion < FControlRigObjectVersion::SwitchedToRigVM)
+	{
+		// ignore errors during the first potential compile of the VM
+		// since that this point variable nodes may still be ill-formed.
+		TGuardValue<FRigVMReportDelegate> SuspendReportDelegate(VMCompileSettings.ASTSettings.ReportDelegate,
+			FRigVMReportDelegate::CreateLambda([](EMessageSeverity::Type,  UObject*, const FString&)
+			{
+				// do nothing
+			})
+		);
+		InitializeModelIfRequired();
+	}
+
+	AddedMemberVariableMap.Reset();
+
+	for (int32 VariableIndex = 0; VariableIndex < NewVariables.Num(); VariableIndex++)
+	{
+		AddedMemberVariableMap.Add(NewVariables[VariableIndex].VarName, VariableIndex);
+	}
+
+	if (RigVMClient.Num() == 0)
+	{
+		return;
+	}
+
+	// setup variables on the blueprint based on the previous "parameters"
+	if (GetLinkerCustomVersion(FControlRigObjectVersion::GUID) < FControlRigObjectVersion::BlueprintVariableSupport)
+	{
+		TSharedPtr<FKismetNameValidator> NameValidator = MakeShareable(new FKismetNameValidator(this, NAME_None, nullptr));
+
+		auto CreateVariable = [this, NameValidator](const URigVMVariableNode* InVariableNode)
+		{
+			if (!InVariableNode)
+			{
+				return;
+			}
+			
+			static const FString VariableString = TEXT("Variable");
+			if (URigVMPin* VariablePin = InVariableNode->FindPin(VariableString))
+			{
+				if (VariablePin->GetDirection() != ERigVMPinDirection::Visible)
+				{
+					return;
+				}
+			}
+
+			const FRigVMGraphVariableDescription Description = InVariableNode->GetVariableDescription();
+			if (AddedMemberVariableMap.Contains(Description.Name))
+			{
+				return;
+			}
+
+			const FEdGraphPinType PinType = RigVMTypeUtils::PinTypeFromExternalVariable(Description.ToExternalVariable());
+			if (!PinType.PinCategory.IsValid())
+			{
+				return;
+			}
+
+			const FName VarName = FindHostMemberVariableUniqueName(NameValidator, Description.Name.ToString());
+			const int32 VariableIndex = AddHostMemberVariable(this, VarName, PinType, false, false, FString());
+			if (VariableIndex != INDEX_NONE)
+			{
+				AddedMemberVariableMap.Add(Description.Name, VariableIndex);
+				MarkDirtyDuringLoad();
+			}
+		};
+
+		auto CreateParameter = [this, NameValidator](const URigVMParameterNode* InParameterNode)
+		{
+			if (!InParameterNode)
+			{
+				return;
+			}
+
+			static const FString ParameterString = TEXT("Parameter");
+			if (const URigVMPin* ParameterPin = InParameterNode->FindPin(ParameterString))
+			{
+				if (ParameterPin->GetDirection() != ERigVMPinDirection::Visible)
+				{
+					return;
+				}
+			}
+
+			const FRigVMGraphParameterDescription Description = InParameterNode->GetParameterDescription();
+			if (AddedMemberVariableMap.Contains(Description.Name))
+			{
+				return;
+			}
+
+			const FEdGraphPinType PinType = RigVMTypeUtils::PinTypeFromExternalVariable(Description.ToExternalVariable());
+			if (!PinType.PinCategory.IsValid())
+			{
+				return;
+			}
+
+			const FName VarName = FindHostMemberVariableUniqueName(NameValidator, Description.Name.ToString());
+			const int32 VariableIndex = AddHostMemberVariable(this, VarName, PinType, true, !Description.bIsInput, FString());
+			
+			if (VariableIndex != INDEX_NONE)
+			{
+				AddedMemberVariableMap.Add(Description.Name, VariableIndex);
+				MarkDirtyDuringLoad();
+			}
+		};
+		
+		for (const URigVMGraph* Model : RigVMClient)
+		{
+			const TArray<URigVMNode*>& Nodes = Model->GetNodes();
+			for (const URigVMNode* Node : Nodes)
+			{
+				if (const URigVMVariableNode* VariableNode = Cast<URigVMVariableNode>(Node))
+				{
+					CreateVariable(VariableNode);
+				}
+
+				// Leaving this for backwards compatibility, even though we don't support parameters anymore
+				// When a parameter node is found, we will create a variable
+				else if (const URigVMParameterNode* ParameterNode = Cast<URigVMParameterNode>(Node))
+				{
+					CreateParameter(ParameterNode);
+				}
+			}
+		}
+	}
+
+#endif
+}
+
+void UControlRigBlueprint::PatchVariableNodesOnLoad()
+{
+#if WITH_EDITOR
+
+	// setup variables on the blueprint based on the previous "parameters"
+	if (GetLinkerCustomVersion(FControlRigObjectVersion::GUID) < FControlRigObjectVersion::BlueprintVariableSupport)
+	{
+		TGuardValue<bool> GuardNotifsSelf(bSuspendModelNotificationsForSelf, true);
+		
+		check(GetDefaultModel());
+
+		auto PatchVariableNode = [this](const URigVMVariableNode* InVariableNode)
+		{
+			if (!InVariableNode)
+			{
+				return;
+			}
+
+			const FRigVMGraphVariableDescription Description = InVariableNode->GetVariableDescription();
+			if (!AddedMemberVariableMap.Contains(Description.Name))
+			{
+				return;
+			}
+			
+			const int32 VariableIndex = AddedMemberVariableMap.FindChecked(Description.Name);
+			const FName VarName = NewVariables[VariableIndex].VarName;
+			
+			GetOrCreateController()->RefreshVariableNode(
+				InVariableNode->GetFName(), VarName, Description.CPPType, Description.CPPTypeObject, false);
+			
+			MarkDirtyDuringLoad();			
+		};
+
+		auto PatchParameterNode = [this](const URigVMParameterNode* InParameterNode)
+		{
+			if (!InParameterNode)
+			{
+				return;
+			}
+			
+			const FRigVMGraphParameterDescription Description = InParameterNode->GetParameterDescription();
+			if (!AddedMemberVariableMap.Contains(Description.Name))
+			{
+				return;
+			}
+
+			const int32 VariableIndex = AddedMemberVariableMap.FindChecked(Description.Name);
+			const FName VarName = NewVariables[VariableIndex].VarName;
+			
+			GetOrCreateController()->ReplaceParameterNodeWithVariable(
+				InParameterNode->GetFName(), VarName, Description.CPPType, Description.CPPTypeObject, false);
+
+			MarkDirtyDuringLoad();	
+		};
+		
+		for(const URigVMGraph* Model : RigVMClient)
+		{
+			TArray<URigVMNode*> Nodes = Model->GetNodes();
+			for (URigVMNode* Node : Nodes)
+			{
+				if (const URigVMVariableNode* VariableNode = Cast<URigVMVariableNode>(Node))
+				{
+					PatchVariableNode(VariableNode);
+				}
+				else if (const URigVMParameterNode* ParameterNode = Cast<URigVMParameterNode>(Node))
+				{
+					PatchParameterNode(ParameterNode);
+				}
+			}
+		}
+	}
+
+#endif
+
+	Super::PatchVariableNodesOnLoad();
+}
+
 void UControlRigBlueprint::PropagatePoseFromInstanceToBP(UControlRig* InControlRig) const
 {
 	check(InControlRig);
diff --git a/Engine/Plugins/Animation/ControlRig/Source/ControlRigDeveloper/Public/ControlRigBlueprint.h b/Engine/Plugins/Animation/ControlRig/Source/ControlRigDeveloper/Public/ControlRigBlueprint.h
--- a/Engine/Plugins/Animation/ControlRig/Source/ControlRigDeveloper/Public/ControlRigBlueprint.h
+++ b/Engine/Plugins/Animation/ControlRig/Source/ControlRigDeveloper/Public/ControlRigBlueprint.h
@@ -189,6 +189,10 @@ private:
 	void PatchRigElementKeyCacheOnLoad();
 	void PatchPropagateToChildren();
 
+protected:
+	virtual void CreateMemberVariablesOnLoad() override;
+	virtual void PatchVariableNodesOnLoad() override;
+
 public:
 	void PropagatePoseFromInstanceToBP(UControlRig* InControlRig) const;
 	void PropagatePoseFromBPToInstances() const;
diff --git a/Engine/Plugins/Runtime/RigVM/Source/RigVMDeveloper/Public/RigVMBlueprint.h b/Engine/Plugins/Runtime/RigVM/Source/RigVMDeveloper/Public/RigVMBlueprint.h
--- a/Engine/Plugins/Runtime/RigVM/Source/RigVMDeveloper/Public/RigVMBlueprint.h
+++ b/Engine/Plugins/Runtime/RigVM/Source/RigVMDeveloper/Public/RigVMBlueprint.h
@@ -549,15 +549,7 @@ private:
 
 	FOnRigVMCompiledEvent VMCompiledEvent;
 
-	virtual void CreateMemberVariablesOnLoad();
-#if WITH_EDITOR
-	static FName FindHostMemberVariableUniqueName(TSharedPtr<FKismetNameValidator> InNameValidator, const FString& InBaseName);
-	static int32 AddHostMemberVariable(URigVMBlueprint* InBlueprint, const FName& InVarName, FEdGraphPinType InVarType, bool bIsPublic, bool bIsReadOnly, FString InDefaultValue);
-	FName AddHostMemberVariableFromExternal(FRigVMExternalVariable InVariableToCreate, FString InDefaultValue = FString());
-#endif
-
 	virtual void PatchFunctionReferencesOnLoad();
-	virtual void PatchVariableNodesOnLoad();
 	virtual void PathDomainSpecificContentOnLoad() {}
 	virtual void PatchBoundVariables();
 	virtual void PatchVariableNodesWithIncorrectType();
@@ -565,6 +557,16 @@ private:
 	virtual void PatchLinksWithCast();
 	virtual void PatchFunctionsOnLoad();
 
+protected:
+
+#if WITH_EDITOR
+	static FName FindHostMemberVariableUniqueName(TSharedPtr<FKismetNameValidator> InNameValidator, const FString& InBaseName);
+	static int32 AddHostMemberVariable(URigVMBlueprint* InBlueprint, const FName& InVarName, FEdGraphPinType InVarType, bool bIsPublic, bool bIsReadOnly, FString InDefaultValue);
+	FName AddHostMemberVariableFromExternal(FRigVMExternalVariable InVariableToCreate, FString InDefaultValue = FString());
+#endif
+	
+	virtual void CreateMemberVariablesOnLoad();
+	virtual void PatchVariableNodesOnLoad();
 	TMap<FName, int32> AddedMemberVariableMap;
 	TArray<FBPVariableDescription> LastNewVariables;
 
