Control Rig: Fix edit with FK not working in sequencer

diff --git a/Engine/Plugins/Animation/ControlRig/Source/ControlRig/Private/Rigs/FKControlRig.cpp b/Engine/Plugins/Animation/ControlRig/Source/ControlRig/Private/Rigs/FKControlRig.cpp
--- a/Engine/Plugins/Animation/ControlRig/Source/ControlRig/Private/Rigs/FKControlRig.cpp
+++ b/Engine/Plugins/Animation/ControlRig/Source/ControlRig/Private/Rigs/FKControlRig.cpp
@@ -217,7 +217,7 @@ bool UFKControlRig::Execute_Internal(const FName& InEventName)
 					FTransform Transform = Current.GetRelativeTransform(Offset);
 					Transform.NormalizeRotation();
 
-					SetControlValue(ControlName, FRigControlValue::Make(FEulerTransform(Transform)), bNotify, Context, bSetupUndo);
+					SetControlValue(ControlName, FRigControlValue::Make(FEulerTransform(Transform)), bNotify, Context, bSetupUndo, false, true);
 					break;
 				}
 
@@ -225,7 +225,7 @@ bool UFKControlRig::Execute_Internal(const FName& InEventName)
 				{
 					FTransform Transform = GetHierarchy()->GetTransform(BoneElement, ERigTransformType::CurrentLocal);
 					Transform.NormalizeRotation();
-					SetControlValue(ControlName, FRigControlValue::Make(FEulerTransform(Transform)), bNotify, Context, bSetupUndo);
+					SetControlValue(ControlName, FRigControlValue::Make(FEulerTransform(Transform)), bNotify, Context, bSetupUndo, false, true);
 
 					break;
 				}
