Control Rig: Support preferred rotation order on spawned controls

diff --git a/Engine/Plugins/Animation/ControlRig/Source/ControlRig/Private/Units/Execution/RigUnit_DynamicHierarchy.cpp b/Engine/Plugins/Animation/ControlRig/Source/ControlRig/Private/Units/Execution/RigUnit_DynamicHierarchy.cpp
--- a/Engine/Plugins/Animation/ControlRig/Source/ControlRig/Private/Units/Execution/RigUnit_DynamicHierarchy.cpp
+++ b/Engine/Plugins/Animation/ControlRig/Source/ControlRig/Private/Units/Execution/RigUnit_DynamicHierarchy.cpp
@@ -690,6 +690,8 @@ void FRigUnit_HierarchyAddControlTransform_Settings::Configure(FRigControlSettin
 	
 	OutSettings.ControlType = ERigControlType::EulerTransform;
 	OutSettings.FilteredChannels = FilteredChannels;
+	OutSettings.bUsePreferredRotationOrder = bUsePreferredRotationOrder;
+	OutSettings.PreferredRotationOrder = PreferredRotationOrder;
 
 	Proxy.Configure(OutSettings);
 	Shape.Configure(OutSettings);
diff --git a/Engine/Plugins/Animation/ControlRig/Source/ControlRig/Private/Units/Execution/RigUnit_DynamicHierarchy.h b/Engine/Plugins/Animation/ControlRig/Source/ControlRig/Private/Units/Execution/RigUnit_DynamicHierarchy.h
--- a/Engine/Plugins/Animation/ControlRig/Source/ControlRig/Private/Units/Execution/RigUnit_DynamicHierarchy.h
+++ b/Engine/Plugins/Animation/ControlRig/Source/ControlRig/Private/Units/Execution/RigUnit_DynamicHierarchy.h
@@ -989,6 +989,8 @@ struct CONTROLRIG_API FRigUnit_HierarchyAddControlTransform_Settings : public FR
 	FRigUnit_HierarchyAddControlTransform_Settings()
 		: FRigUnit_HierarchyAddControl_Settings()
 		, InitialSpace(ERigVMTransformSpace::LocalSpace)
+		, bUsePreferredRotationOrder(false)
+		, PreferredRotationOrder(EEulerRotationOrder::YZX) 
 	{
 		FilteredChannels.Reset();
 	}
@@ -999,6 +1001,13 @@ struct CONTROLRIG_API FRigUnit_HierarchyAddControlTransform_Settings : public FR
 	UPROPERTY(EditAnywhere, BlueprintReadWrite, Category = "Settings")
 	ERigVMTransformSpace InitialSpace;
 
+	// Enables overriding the preferred rotation order
+	UPROPERTY(EditAnywhere, BlueprintReadWrite, Category = "Settings")
+	bool bUsePreferredRotationOrder;
+
+	UPROPERTY(EditAnywhere, BlueprintReadWrite, Category = "Settings")
+	EEulerRotationOrder PreferredRotationOrder;
+
 	UPROPERTY(EditAnywhere, BlueprintReadWrite, Category = "Settings")
 	FRigUnit_HierarchyAddControl_ShapeSettings Shape;
 
