Control Rig: Fix additive controls not working when adding multiple controls on same subhierarchy

diff --git a/Engine/Plugins/Animation/ControlRig/Source/ControlRig/Private/ControlRig.cpp b/Engine/Plugins/Animation/ControlRig/Source/ControlRig/Private/ControlRig.cpp
--- a/Engine/Plugins/Animation/ControlRig/Source/ControlRig/Private/ControlRig.cpp
+++ b/Engine/Plugins/Animation/ControlRig/Source/ControlRig/Private/ControlRig.cpp
@@ -1708,8 +1708,17 @@ FRigControlValue UControlRig::GetControlValueFromGlobalTransform(const FName& In
 				const int32 ControlIndex = ControlsAfterBackwardsSolve.GetIndex(ControlElement->GetKey());
 				if (ControlIndex != INDEX_NONE)
 				{
-					const FTransform& AnimGlobalTransform = ControlsAfterBackwardsSolve[ControlIndex].GlobalTransform;
-					const FTransform AdditiveTransform = InGlobalTransform.GetRelativeTransform(AnimGlobalTransform);					
+					// ParentGlobalTransform = ParentAnimationGlobalTransform + ParentAdditiveTransform
+					const FTransform& ParentGlobalTransform = DynamicHierarchy->GetParentTransform(ControlElement, ERigTransformType::CurrentGlobal);
+
+					// LocalTransform = InGlobal - ParentGlobalTransform - OffsetLocal
+					const FTransform OffsetLocalTransform = InGlobalTransform.GetRelativeTransform(ParentGlobalTransform);
+					const FTransform LocalTransform = OffsetLocalTransform.GetRelativeTransform(DynamicHierarchy->GetControlOffsetTransform(ControlElement, ERigTransformType::CurrentLocal));
+
+					// Additive = LocalTransform - AnimLocalTransform
+					const FTransform& AnimLocalTransform = ControlsAfterBackwardsSolve[ControlIndex].LocalTransform;
+					const FTransform AdditiveTransform = LocalTransform.GetRelativeTransform(AnimLocalTransform);
+
 					Value.SetFromTransform(AdditiveTransform, ControlElement->Settings.ControlType, ControlElement->Settings.PrimaryAxis);
 				}
 			}
