Control Rig: Fix undo/redo not working on additive control rigs

diff --git a/Engine/Plugins/Animation/ControlRig/Source/ControlRig/Private/ControlRig.cpp b/Engine/Plugins/Animation/ControlRig/Source/ControlRig/Private/ControlRig.cpp
--- a/Engine/Plugins/Animation/ControlRig/Source/ControlRig/Private/ControlRig.cpp
+++ b/Engine/Plugins/Animation/ControlRig/Source/ControlRig/Private/ControlRig.cpp
@@ -381,8 +381,9 @@ void UControlRig::Evaluate_AnyThread()
 
 					FRigControlValue FinalValue;
 					FinalValue.SetFromTransform(FinalTransform, Control->Settings.ControlType, Control->Settings.PrimaryAxis);
-					
-					Hierarchy->SetControlValue(Control, FinalValue, ERigControlValueType::Current, Info.bSetupUndo, false, Info.bPrintPythonCommnds, false);
+
+					const bool bSetupUndo = false; // Rely on the sequencer track to handle undo/redo
+					Hierarchy->SetControlValue(Control, FinalValue, ERigControlValueType::Current, bSetupUndo, false, Info.bPrintPythonCommnds, false);
 					Hierarchy->SetPreferredEulerAnglesFromValue(Control, AdditiveValue, ERigControlValueType::Current, Info.bFixEulerFlips);
 					
 					if (Info.bNotify && OnControlModified.IsBound())
