Control Rig: Fix user data asset in skeletal mesh component not available in control rig
#rnx

diff --git a/Engine/Plugins/Animation/ControlRig/Source/ControlRig/Private/ControlRig.cpp b/Engine/Plugins/Animation/ControlRig/Source/ControlRig/Private/ControlRig.cpp
--- a/Engine/Plugins/Animation/ControlRig/Source/ControlRig/Private/ControlRig.cpp
+++ b/Engine/Plugins/Animation/ControlRig/Source/ControlRig/Private/ControlRig.cpp
@@ -2582,6 +2582,13 @@ const TArray<UAssetUserData*>* UControlRig::GetAssetUserDataArray() const
 					}
 				}
 			}
+			if (const TArray<UAssetUserData*>* ActorComponentUserDataArray = SkeletalMeshComponent->GetAssetUserDataArray())
+			{
+				for(UAssetUserData* ActorComponentUserData : *ActorComponentUserDataArray)
+				{
+					CombinedAssetUserData.Add(ActorComponentUserData);
+				}
+			}
 		}
 	}
 
