Control Rig: Additive control rig allows rigs which don't match the skeleton exactly
#rnx

diff --git a/Engine/Plugins/Animation/ControlRig/Source/ControlRigEditor/Private/Sequencer/ControlRigParameterTrackEditor.cpp b/Engine/Plugins/Animation/ControlRig/Source/ControlRigEditor/Private/Sequencer/ControlRigParameterTrackEditor.cpp
--- a/Engine/Plugins/Animation/ControlRig/Source/ControlRigEditor/Private/Sequencer/ControlRigParameterTrackEditor.cpp
+++ b/Engine/Plugins/Animation/ControlRig/Source/ControlRigEditor/Private/Sequencer/ControlRigParameterTrackEditor.cpp
@@ -822,11 +822,13 @@ class SAddAdditiveControlRigOptionsWindow : public SCompoundWidget
 {
 public:
 	SLATE_BEGIN_ARGS(SAddAdditiveControlRigOptionsWindow)
-		: _ControlRigClass(nullptr)
+		: _bFilterBySkeleton(true)
+		, _ControlRigClass(nullptr)
 		, _Skeleton(nullptr)
 		, _WidgetWindow()
 	{}
 
+		SLATE_ARGUMENT(bool, bFilterBySkeleton)
 		SLATE_ARGUMENT(TSubclassOf<UControlRig>*, ControlRigClass)
 		SLATE_ARGUMENT(USkeleton*, Skeleton)
 		SLATE_ARGUMENT(TSharedPtr<SWindow>, WidgetWindow)
@@ -854,20 +856,50 @@ public:
 		return FReply::Handled();
 	}
 
+	ECheckBoxState IsFilterBySkeletonChecked() const
+	{
+		return bFilterBySkeleton ? ECheckBoxState::Checked : ECheckBoxState::Unchecked;
+	}
+
+	void HandleFilterBySkeletonStateChanged( ECheckBoxState NewState )
+	{
+		bFilterBySkeleton = NewState == ECheckBoxState::Checked;
+		UpdateControlRigClasses();
+	}
+
+	void UpdateControlRigClasses()
+	{
+		FClassViewerInitializationOptions Options;
+		Options.bShowUnloadedBlueprints = true;
+		Options.NameTypeToDisplay = EClassViewerNameTypeToDisplay::DisplayName;
+
+		TSharedPtr<FControlRigClassFilter> ClassFilter = MakeShareable(new FControlRigClassFilter(bFilterBySkeleton, false, true, Skeleton));
+		Options.ClassFilters.Add(ClassFilter.ToSharedRef());
+		Options.bShowNoneOption = false;
+
+		FClassViewerModule& ClassViewerModule = FModuleManager::LoadModuleChecked<FClassViewerModule>("ClassViewer");
+		TSharedRef<SWidget> ClassViewer = ClassViewerModule.CreateClassViewer(Options, FOnClassPicked::CreateRaw(this, &SAddAdditiveControlRigOptionsWindow::OnAdd));
+		InspectorBox->SetContent(ClassViewer->AsShared());
+	}
+
 	SAddAdditiveControlRigOptionsWindow()
-		: ControlRigClass(nullptr)
+		: bFilterBySkeleton(true)
+		, ControlRigClass(nullptr)
 		, Skeleton(nullptr)
 	{}
 
 private:
+	bool bFilterBySkeleton;
 	TSubclassOf<UControlRig>* ControlRigClass;
 	USkeleton* Skeleton;
+	TSharedPtr<SBox> InspectorBox;
 	TWeakPtr< SWindow > WidgetWindow;
 };
 
 
 void SAddAdditiveControlRigOptionsWindow::Construct(const FArguments& InArgs)
 {
+	bFilterBySkeleton = InArgs._bFilterBySkeleton;
 	ControlRigClass = InArgs._ControlRigClass;
 	Skeleton = InArgs._Skeleton;
 	WidgetWindow = InArgs._WidgetWindow;
@@ -877,7 +909,6 @@ void SAddAdditiveControlRigOptionsWindow::Construct(const FArguments& InArgs)
 
 	TSharedPtr<SBox> HeaderToolBox;
 	TSharedPtr<SHorizontalBox> AnimHeaderButtons;
-	TSharedPtr<SBox> InspectorBox;
 	this->ChildSlot
 	[
 		SNew(SBox)
@@ -900,6 +931,18 @@ void SAddAdditiveControlRigOptionsWindow::Construct(const FArguments& InArgs)
 					SNew(SHorizontalBox)
 					+ SHorizontalBox::Slot()
 					.AutoWidth()
+					[
+						SNew(SCheckBox)
+						.IsChecked(this, &SAddAdditiveControlRigOptionsWindow::IsFilterBySkeletonChecked)
+						.OnCheckStateChanged(this, &SAddAdditiveControlRigOptionsWindow::HandleFilterBySkeletonStateChanged)
+						.Padding(FMargin(6.0, 2.0))
+						[
+							SNew(STextBlock)
+								.Text(LOCTEXT("SkeletonFilterLabel", "Filter by Skeleton"))
+						]
+					]
+					+ SHorizontalBox::Slot()
+					.AutoWidth()
 					[
 						SNew(STextBlock)
 						.Font(FAppStyle::GetFontStyle("CurveEd.LabelFont"))
@@ -1036,17 +1079,7 @@ void SAddAdditiveControlRigOptionsWindow::Construct(const FArguments& InArgs)
 
 	};
 
-	FClassViewerInitializationOptions Options;
-	Options.bShowUnloadedBlueprints = true;
-	Options.NameTypeToDisplay = EClassViewerNameTypeToDisplay::DisplayName;
-
-	TSharedPtr<FControlRigClassFilter> ClassFilter = MakeShareable(new FControlRigClassFilter(true, false, true, Skeleton));
-	Options.ClassFilters.Add(ClassFilter.ToSharedRef());
-	Options.bShowNoneOption = false;
-
-	FClassViewerModule& ClassViewerModule = FModuleManager::LoadModuleChecked<FClassViewerModule>("ClassViewer");
-	TSharedRef<SWidget> ClassViewer = ClassViewerModule.CreateClassViewer(Options, FOnClassPicked::CreateRaw(this, &SAddAdditiveControlRigOptionsWindow::OnAdd));
-	InspectorBox->SetContent(ClassViewer->AsShared());
+	UpdateControlRigClasses();
 }
 
 void FControlRigParameterTrackEditor::AddAdditiveControlRig(FGuid ObjectBinding, UObject* BoundActor, USkeletalMeshComponent* SkelMeshComp, USkeleton* Skeleton)
