Control Rig: Allow adding an additive control rig even when other rigs are present in sequencer
#rnx

diff --git a/Engine/Plugins/Animation/ControlRig/Source/ControlRigEditor/Private/Sequencer/ControlRigParameterTrackEditor.cpp b/Engine/Plugins/Animation/ControlRig/Source/ControlRigEditor/Private/Sequencer/ControlRigParameterTrackEditor.cpp
--- a/Engine/Plugins/Animation/ControlRig/Source/ControlRigEditor/Private/Sequencer/ControlRigParameterTrackEditor.cpp
+++ b/Engine/Plugins/Animation/ControlRig/Source/ControlRigEditor/Private/Sequencer/ControlRigParameterTrackEditor.cpp
@@ -1176,10 +1176,10 @@ void FControlRigParameterTrackEditor::AddAdditiveControlRig(FGuid ObjectBinding,
 
 bool FControlRigParameterTrackEditor::CanAddAdditiveConrolRig()
 {
-	if (FControlRigEditMode* ControlRigEditMode = GetEditMode())
-	{
-		return ControlRigEditMode->GetControlRigsArray(false).IsEmpty();
-	}
+	// if (FControlRigEditMode* ControlRigEditMode = GetEditMode())
+	// {
+	// 	return ControlRigEditMode->GetControlRigsArray(false).IsEmpty();
+	// }
 	return true;
 }
 
