Control Rig: Fix ensure when undoing selection of controls
#rnx

diff --git a/Engine/Plugins/Animation/ControlRig/Source/ControlRigEditor/Private/EditMode/SControlRigOutliner.cpp b/Engine/Plugins/Animation/ControlRig/Source/ControlRigEditor/Private/EditMode/SControlRigOutliner.cpp
--- a/Engine/Plugins/Animation/ControlRig/Source/ControlRigEditor/Private/EditMode/SControlRigOutliner.cpp
+++ b/Engine/Plugins/Animation/ControlRig/Source/ControlRigEditor/Private/EditMode/SControlRigOutliner.cpp
@@ -1071,7 +1071,7 @@ void SControlRigOutliner::HandleSelectionChanged(TSharedPtr<FMultiRigTreeElement
 	TGuardValue<bool> GuardRigHierarchyChanges(bIsChangingRigHierarchy, true);
 	FControlRigEditMode* EditMode = static_cast<FControlRigEditMode*>(ModeTools->GetActiveMode(FControlRigEditMode::ModeName));
 	bool bEndTransaction = false;
-	if (GEditor && EditMode && EditMode->IsInLevelEditor())
+	if (GEditor && !GIsTransacting && EditMode && EditMode->IsInLevelEditor())
 	{
 		GEditor->BeginTransaction(LOCTEXT("SelectControl", "Select Control"));
 		bEndTransaction = true;
