Constraints: handle deprecated CR transform types

#rnx

diff --git a/Engine/Plugins/Animation/ControlRig/Source/ControlRig/Private/Sequencer/MovieSceneControlRigParameterSection.cpp b/Engine/Plugins/Animation/ControlRig/Source/ControlRig/Private/Sequencer/MovieSceneControlRigParameterSection.cpp
--- a/Engine/Plugins/Animation/ControlRig/Source/ControlRig/Private/Sequencer/MovieSceneControlRigParameterSection.cpp
+++ b/Engine/Plugins/Animation/ControlRig/Source/ControlRig/Private/Sequencer/MovieSceneControlRigParameterSection.cpp
@@ -570,19 +570,33 @@ struct FParameterTransformChannelEditorData
 			FRigControlElement* ControlElement = ControlRig->FindControl(ParameterName);
 			if (ControlElement)
 			{
+				auto GetTranslationFromTransform = [&](const FVector& InTranslation)
+				{
+					// switch translation to constraint space if needed
+					const uint32 ControlHash = UTransformableControlHandle::ComputeHash(ControlRig, ControlElement->GetName());
+					TOptional<FTransform> ConstraintSpaceTransform = FTransformConstraintUtils::GetRelativeTransform(ControlRig->GetWorld(), ControlHash);
+					if (ConstraintSpaceTransform)
+					{
+						return ConstraintSpaceTransform->GetTranslation();
+					}
+					return InTranslation;
+				};
+				
 				if (ControlElement->Settings.ControlType == ERigControlType::Transform)
 				{
 					const FRigControlValue::FTransform_Float Transform = 
 						ControlRig
 						->GetControlValue(ControlElement, ERigControlValueType::Current).Get<FRigControlValue::FTransform_Float>();
-					return FVector(Transform.GetTranslation());
+
+					return GetTranslationFromTransform( FVector(Transform.GetTranslation()) );
 				}
 				else if  (ControlElement->Settings.ControlType == ERigControlType::TransformNoScale)
 				{
 					const FRigControlValue::FTransformNoScale_Float Transform = 
 						ControlRig
 						->GetControlValue(ControlElement, ERigControlValueType::Current).Get<FRigControlValue::FTransformNoScale_Float>();
-					return FVector(Transform.GetTranslation());
+
+					return GetTranslationFromTransform( FVector(Transform.GetTranslation()) );
 				}
 				else if (ControlElement->Settings.ControlType == ERigControlType::EulerTransform)
 				{
@@ -590,15 +604,7 @@ struct FParameterTransformChannelEditorData
 						ControlRig
 						->GetControlValue(ControlElement, ERigControlValueType::Current).Get<FRigControlValue::FEulerTransform_Float>();
 
-					// switch translation to constraint space if needed
-					const uint32 ControlHash = UTransformableControlHandle::ComputeHash(ControlRig, ControlElement->GetName());
-					TOptional<FTransform> ConstraintSpaceTransform = FTransformConstraintUtils::GetRelativeTransform(ControlRig->GetWorld(), ControlHash);
-					if (ConstraintSpaceTransform)
-					{
-						return ConstraintSpaceTransform->GetTranslation();
-					}
-					
-					return FVector(Euler.GetTranslation());
+					return GetTranslationFromTransform( FVector(Euler.GetTranslation()) );
 				}
 				else if(ControlElement->Settings.ControlType == ERigControlType::Position)
 				{
@@ -644,12 +650,25 @@ struct FParameterTransformChannelEditorData
 			FRigControlElement* ControlElement = ControlRig->FindControl(ParameterName);
 			if (ControlElement)
 			{
+				auto GetScaleFromTransform = [&](const FVector& InScale3D)
+				{
+					// switch scale to constraint space if needed
+					const uint32 ControlHash = UTransformableControlHandle::ComputeHash(ControlRig, ControlElement->GetName());
+					TOptional<FTransform> ConstraintSpaceTransform = FTransformConstraintUtils::GetRelativeTransform(ControlRig->GetWorld(), ControlHash);
+					if (ConstraintSpaceTransform)
+					{
+						return ConstraintSpaceTransform->GetScale3D();
+					}
+					return InScale3D;
+				};
+				
 				if (ControlElement->Settings.ControlType == ERigControlType::Transform)
 				{
 					const FRigControlValue::FTransform_Float Transform = 
 						ControlRig
 						->GetControlValue(ControlElement, ERigControlValueType::Current).Get<FRigControlValue::FTransform_Float>();
-					return FVector(Transform.GetScale3D());
+
+					return GetScaleFromTransform( FVector(Transform.GetScale3D()) );
 				}
 				else if (ControlElement->Settings.ControlType == ERigControlType::EulerTransform)
 				{
@@ -657,15 +676,7 @@ struct FParameterTransformChannelEditorData
 						ControlRig
 						->GetControlValue(ControlElement, ERigControlValueType::Current).Get<FRigControlValue::FEulerTransform_Float>();
 
-					// switch scale to constraint space if needed
-					const uint32 ControlHash = UTransformableControlHandle::ComputeHash(ControlRig, ControlElement->GetName());
-					TOptional<FTransform> ConstraintSpaceTransform = FTransformConstraintUtils::GetRelativeTransform(ControlRig->GetWorld(), ControlHash);
-					if (ConstraintSpaceTransform)
-					{
-						return ConstraintSpaceTransform->GetScale3D();
-					}
-					
-					return FVector(Transform.GetScale3D());
+					return GetScaleFromTransform( FVector(Transform.GetScale3D()) );
 				}
 				else if (ControlElement->Settings.ControlType == ERigControlType::Scale)
 				{
