Removed bone name restriction not allowing bone names to start with a number at URigHierarchy::SanitizeName, in order to fix fbx import issues

diff --git a/Engine/Plugins/Animation/ControlRig/Source/ControlRig/Private/Rigs/RigHierarchy.cpp b/Engine/Plugins/Animation/ControlRig/Source/ControlRig/Private/Rigs/RigHierarchy.cpp
--- a/Engine/Plugins/Animation/ControlRig/Source/ControlRig/Private/Rigs/RigHierarchy.cpp
+++ b/Engine/Plugins/Animation/ControlRig/Source/ControlRig/Private/Rigs/RigHierarchy.cpp
@@ -953,7 +953,7 @@ void URigHierarchy::SanitizeName(FString& InOutName)
 
 		const bool bGoodChar = FChar::IsAlpha(C) ||				// Any letter
 			(C == '_') || (C == '-') || (C == '.') ||			// _  - . anytime
-			((i > 0) && (FChar::IsDigit(C))) ||					// 0-9 after the first character
+			(FChar::IsDigit(C)) ||								// 0-9 anytime
 			((i > 0) && (C== ' '));								// Space after the first character to support virtual bones
 
 		if (!bGoodChar)
