Control Rig: Direct manipulation on nodes

* Added functionality to URigVMUnitNode to enumerate properties + to create a live instance (an instance mapped to the actual debugged host)
* Changed FRigUnit virtuals over to allow the node to advertise direct manipulation targets
* Adapted control rig / module / editor to interact with backend API to offer manipulation using node context and hotkeys
* Added viewport notification to allow editing of direct manipulation targets directly in the viewport

diff --git a/Engine/Plugins/Animation/ControlRig/Source/ControlRig/Private/ControlRig.cpp b/Engine/Plugins/Animation/ControlRig/Source/ControlRig/Private/ControlRig.cpp
--- a/Engine/Plugins/Animation/ControlRig/Source/ControlRig/Private/ControlRig.cpp
+++ b/Engine/Plugins/Animation/ControlRig/Source/ControlRig/Private/ControlRig.cpp
@@ -28,6 +28,8 @@
 #include "Kismet2/BlueprintEditorUtils.h"
 #include "RigVMDeveloperTypeUtils.h"
 #include "Editor.h"
+#include "RigVMBlueprint.h"
+#include "RigVMCompiler/RigVMCompiler.h"
 #endif// WITH_EDITOR
 #include "ControlRigComponent.h"
 #include "Constraints/ControlRigTransformableHandle.h"
@@ -1833,6 +1835,12 @@ TArray<FName> UControlRig::CurrentControlSelection() const
 		{
 			SelectedControlNames.Add(SelectedControl->GetName());
 		}
+#if WITH_EDITOR
+		for(const TSharedPtr<FRigDirectManipulationInfo>& ManipulationInfo : RigUnitManipulationInfos)
+		{
+			SelectedControlNames.Add(ManipulationInfo->ControlKey.Name);
+		}
+#endif
 	}
 	return SelectedControlNames;
 }
@@ -1857,12 +1865,14 @@ void UControlRig::HandleHierarchyModified(ERigHierarchyNotification InNotificati
 		case ERigHierarchyNotification::ElementSelected:
 		case ERigHierarchyNotification::ElementDeselected:
 		{
+			bool bClearTransientControls = true;
 			if(FRigControlElement* ControlElement = Cast<FRigControlElement>((FRigBaseElement*)InElement))
 			{
 				const bool bSelected = InNotification == ERigHierarchyNotification::ElementSelected;
 				ControlSelected().Broadcast(this, ControlElement, bSelected);
 
 				OnControlSelected_BP.Broadcast(this, *ControlElement, bSelected);
+				bClearTransientControls = !ControlElement->Settings.bIsTransientControl;
 			}
 
 #if WITH_EDITOR
@@ -1879,11 +1889,10 @@ void UControlRig::HandleHierarchyModified(ERigHierarchyNotification InNotificati
 						SelectionPoseForConstructionMode.Remove(InElement->GetKey());
 					}
 				}
-
-				if(InNotification == ERigHierarchyNotification::ElementDeselected)
-				{
-					ClearTransientControls();
-				}
+			}
+			if(bClearTransientControls)
+			{
+				ClearTransientControls();
 			}
 #endif
 			break;
@@ -1906,21 +1915,67 @@ void UControlRig::HandleHierarchyModified(ERigHierarchyNotification InNotificati
 
 #if WITH_EDITOR
 
-FName UControlRig::AddTransientControl(URigVMPin* InPin, FRigElementKey SpaceKey, FTransform OffsetTransform)
+bool UControlRig::CanAddTransientControl(const URigVMUnitNode* InNode, const FRigDirectManipulationTarget& InTarget, FString* OutFailureReason)
 {
-	if ((InPin == nullptr) || (DynamicHierarchy == nullptr))
+	if (InNode == nullptr)
 	{
-		return NAME_None;
+		if(OutFailureReason)
+		{
+			static const FString Reason = TEXT("Provided node is nullptr.");
+			*OutFailureReason = Reason;
+		}
+		return false;
+	}
+	if (DynamicHierarchy == nullptr)
+	{
+		if(OutFailureReason)
+		{
+			static const FString Reason = TEXT("The rig does not contain a hierarchy.");
+			*OutFailureReason = Reason;
+		}
+		return false;
 	}
 
-	if (InPin->GetCPPType() != TEXT("FVector") &&
-		InPin->GetCPPType() != TEXT("FQuat") &&
-		InPin->GetCPPType() != TEXT("FTransform"))
+	const UScriptStruct* UnitStruct = InNode->GetScriptStruct();
+	if(UnitStruct == nullptr)
 	{
-		return NAME_None;
+		if(OutFailureReason)
+		{
+			static const FString Reason = TEXT("The node is not resolved.");
+			*OutFailureReason = Reason;
+		}
+		return false;
 	}
 
-	RemoveTransientControl(InPin);
+	const TSharedPtr<FStructOnScope> NodeInstance = InNode->ConstructLiveStructInstance(this);
+	if(!NodeInstance.IsValid() || !NodeInstance->IsValid())
+	{
+		if(OutFailureReason)
+		{
+			static const FString Reason = TEXT("Unexpected error: Node instance could not be constructed.");
+			*OutFailureReason = Reason;
+		}
+		return false;
+	}
+
+	const FRigUnit* UnitInstance = GetRigUnitInstanceFromScope(NodeInstance);
+	check(UnitInstance);
+
+	TArray<FRigDirectManipulationTarget> Targets;
+	if(UnitInstance->GetDirectManipulationTargets(InNode, NodeInstance, DynamicHierarchy, Targets, OutFailureReason))
+	{
+		return Targets.Contains(InTarget);
+	}
+	return false;
+}
+
+FName UControlRig::AddTransientControl(const URigVMUnitNode* InNode, const FRigDirectManipulationTarget& InTarget)
+{
+	if(!CanAddTransientControl(InNode, InTarget, nullptr))
+	{
+		return NAME_None;
+	}
+	RemoveTransientControl(InNode, InTarget);
 
 	URigHierarchyController* Controller = DynamicHierarchy->GetController(true);
 	if(Controller == nullptr)
@@ -1928,89 +1983,94 @@ FName UControlRig::AddTransientControl(URigVMPin* InPin, FRigElementKey SpaceKey
 		return NAME_None;
 	}
 
-	URigVMPin* PinForLink = InPin->GetPinForLink();
+	const UScriptStruct* UnitStruct = InNode->GetScriptStruct();
+	check(UnitStruct);
+
+	const TSharedPtr<FStructOnScope> NodeInstance = InNode->ConstructLiveStructInstance(this);
+	check(NodeInstance.IsValid() && NodeInstance->IsValid());
+
+	const FRigUnit* UnitInstance = GetRigUnitInstanceFromScope(NodeInstance);
+	check(UnitInstance);
 
-	const FName ControlName = GetNameForTransientControl(InPin);
+	const FName ControlName = GetNameForTransientControl(InNode, InTarget);
 	FTransform ShapeTransform = FTransform::Identity;
 	ShapeTransform.SetScale3D(FVector::ZeroVector);
 
 	FRigControlSettings Settings;
 	Settings.ControlType = ERigControlType::Transform;
-	if (URigVMPin* ColorPin = PinForLink->GetNode()->FindPin(TEXT("Color")))
-	{
-		if (ColorPin->GetCPPType() == TEXT("FLinearColor"))
-		{
-			FRigControlValue Value;
-			Settings.ShapeColor = Value.SetFromString<FLinearColor>(ColorPin->GetDefaultValue());
-		}
-	}
 	Settings.bIsTransientControl = true;
 	Settings.DisplayName = TEXT("Temporary Control");
 
+	TSharedPtr<FRigDirectManipulationInfo> Info = MakeShareable(new FRigDirectManipulationInfo());
+	Info->Target = InTarget;
+	Info->Node = TWeakObjectPtr<const URigVMUnitNode>(InNode);
+
+	FRigControlValue Value = FRigControlValue::Make(FTransform::Identity);
+	UnitInstance->ConfigureDirectManipulationControl(InNode, Info, Settings, Value);
+
 	Controller->ClearSelection();
 
     const FRigElementKey ControlKey = Controller->AddControl(
     	ControlName,
-    	SpaceKey,
+    	FRigElementKey(),
     	Settings,
-    	FRigControlValue::Make(FTransform::Identity),
-    	OffsetTransform,
+    	Value,
+    	FTransform::Identity,
     	ShapeTransform, false);
 
-	SetTransientControlValue(InPin);
+	Info->ControlKey = ControlKey;
+	RigUnitManipulationInfos.Add(Info);
+	SetTransientControlValue(InNode, Info);
+	Info->bInitialized = true;
 
-	return ControlName;
+	return Info->ControlKey.Name;
 }
 
-bool UControlRig::SetTransientControlValue(URigVMPin* InPin)
+bool UControlRig::SetTransientControlValue(const URigVMUnitNode* InNode, TSharedPtr<FRigDirectManipulationInfo> InInfo)
 {
-	const FName ControlName = GetNameForTransientControl(InPin);
-	if (FRigControlElement* ControlElement = FindControl(ControlName))
+	check(InNode);
+	check(DynamicHierarchy);
+	
+	const TSharedPtr<FStructOnScope> NodeInstance = InNode->ConstructLiveStructInstance(this);
+	check(NodeInstance.IsValid());
+	check(NodeInstance->IsValid());
+
+	const UScriptStruct* UnitStruct = InNode->GetScriptStruct();
+	check(UnitStruct);
+
+	FRigUnit* UnitInstance = GetRigUnitInstanceFromScope(NodeInstance);
+	check(UnitInstance);
+
+	const FName ControlName = GetNameForTransientControl(InNode, InInfo->Target);
+	const FRigControlElement* ControlElement = DynamicHierarchy->Find<FRigControlElement>({ ControlName, ERigElementType::Control });
+	if(ControlElement == nullptr)
 	{
-		FString DefaultValue = InPin->GetPinForLink()->GetDefaultValue();
-		if (!DefaultValue.IsEmpty())
-		{
-			if (InPin->GetCPPType() == TEXT("FVector"))
-			{
-				ControlElement->Settings.ControlType = ERigControlType::Position;
-				FRigControlValue Value;
-				Value.SetFromString<FVector>(DefaultValue);
-				DynamicHierarchy->SetControlValue(ControlElement, Value, ERigControlValueType::Current, false);
-			}
-			else if (InPin->GetCPPType() == TEXT("FQuat"))
-			{
-				ControlElement->Settings.ControlType = ERigControlType::Rotator;
-				FRigControlValue Value;
-				Value.SetFromString<FRotator>(DefaultValue);
-				DynamicHierarchy->SetControlValue(ControlElement, Value, ERigControlValueType::Current, false);
-			}
-			else
-			{
-				ControlElement->Settings.ControlType = ERigControlType::Transform;
-				FRigControlValue Value;
-				Value.SetFromString<FTransform>(DefaultValue);
-				DynamicHierarchy->SetControlValue(ControlElement, Value, ERigControlValueType::Current, false);
-			}
-		}
-		return true;
+		return false;
 	}
-	return false;
+
+	FControlRigExecuteContext& PublicContext = GetExtendedExecuteContext().GetPublicDataSafe<FControlRigExecuteContext>();
+	return UnitInstance->UpdateHierarchyForDirectManipulation(InNode, NodeInstance, PublicContext, InInfo);
 }
 
-FName UControlRig::RemoveTransientControl(URigVMPin* InPin)
+FName UControlRig::RemoveTransientControl(const URigVMUnitNode* InNode, const FRigDirectManipulationTarget& InTarget)
 {
-	if ((InPin == nullptr) || (DynamicHierarchy == nullptr))
+	if ((InNode == nullptr) || (DynamicHierarchy == nullptr))
 	{
 		return NAME_None;
 	}
 
+	RigUnitManipulationInfos.RemoveAll([InTarget](const TSharedPtr<FRigDirectManipulationInfo>& Info) -> bool
+	{
+		return Info->Target == InTarget;
+	});
+
 	URigHierarchyController* Controller = DynamicHierarchy->GetController(true);
 	if(Controller == nullptr)
 	{
 		return NAME_None;
 	}
 
-	const FName ControlName = GetNameForTransientControl(InPin);
+	const FName ControlName = GetNameForTransientControl(InNode, InTarget);
 	if(FRigControlElement* ControlElement = FindControl(ControlName))
 	{
 		DynamicHierarchy->Notify(ERigHierarchyNotification::ElementDeselected, ControlElement);
@@ -2121,7 +2181,7 @@ FName UControlRig::AddTransientControl(const FRigElementKey& InElement)
 
 	SetTransientControlValue(InElement);
 
-	return ControlName;
+	return ControlKey.Name;
 }
 
 bool UControlRig::SetTransientControlValue(const FRigElementKey& InElement)
@@ -2227,25 +2287,62 @@ FName UControlRig::RemoveTransientControl(const FRigElementKey& InElement)
 	return NAME_None;
 }
 
-FName UControlRig::GetNameForTransientControl(URigVMPin* InPin) const
+FName UControlRig::GetNameForTransientControl(const URigVMUnitNode* InNode, const FRigDirectManipulationTarget& InTarget) const
 {
-	check(InPin);
+	check(InNode);
 	check(DynamicHierarchy);
 	
-	const FString OriginalPinPath = InPin->GetOriginalPinFromInjectedNode()->GetPinPath();
-	return DynamicHierarchy->GetSanitizedName(FString::Printf(TEXT("ControlForPin_%s"), *OriginalPinPath));
+	const FString NodeName = InNode->GetName();
+	return DynamicHierarchy->GetSanitizedName(FString::Printf(TEXT("ControlForNode|%s|%s"), *NodeName, *InTarget.Name));
+}
+
+FString UControlRig::GetNodeNameFromTransientControl(const FRigElementKey& InKey)
+{
+	FString Name = InKey.Name.ToString();
+	if(Name.StartsWith(TEXT("ControlForNode|")))
+	{
+		Name.RightChopInline(15);
+		Name.LeftInline(Name.Find(TEXT("|")));
+	}
+	else
+	{
+		return FString();
+	}
+	return Name;
 }
 
-FString UControlRig::GetPinNameFromTransientControl(const FRigElementKey& InKey)
+FString UControlRig::GetTargetFromTransientControl(const FRigElementKey& InKey)
 {
 	FString Name = InKey.Name.ToString();
-	if(Name.StartsWith(TEXT("ControlForPin_")))
+	if(Name.StartsWith(TEXT("ControlForNode|")))
 	{
-		Name.RightChopInline(14);
+		Name.RightChopInline(15);
+		Name.RightChopInline(Name.Find(TEXT("|")) + 1);
+	}
+	else
+	{
+		return FString();
 	}
 	return Name;
 }
 
+TSharedPtr<FRigDirectManipulationInfo> UControlRig::GetRigUnitManipulationInfoForTransientControl(
+	const FRigElementKey& InKey)
+{
+	const TSharedPtr<FRigDirectManipulationInfo>* InfoPtr = RigUnitManipulationInfos.FindByPredicate(
+		[InKey](const TSharedPtr<FRigDirectManipulationInfo>& Info) -> bool
+		{
+			return Info->ControlKey == InKey;
+		}) ;
+
+	if(InfoPtr)
+	{
+		return *InfoPtr;
+	}
+
+	return TSharedPtr<FRigDirectManipulationInfo>();
+}
+
 FName UControlRig::GetNameForTransientControl(const FRigElementKey& InElement)
 {
 	if (InElement.Type == ERigElementType::Control)
@@ -2309,15 +2406,15 @@ void UControlRig::ClearTransientControls()
 	}
 	TGuardValue<bool> ReEntryGuard(bIsClearingTransientControls, true);
 
+	RigUnitManipulationInfos.Reset();
+
 	const TArray<FRigControlElement*> ControlsToRemove = DynamicHierarchy->GetTransientControls();
 	for (FRigControlElement* ControlToRemove : ControlsToRemove)
 	{
 		const FRigElementKey KeyToRemove = ControlToRemove->GetKey();
 		if(Controller->RemoveElement(ControlToRemove))
 		{
-#if WITH_EDITOR
 			SelectionPoseForConstructionMode.Remove(KeyToRemove);
-#endif
 		}
 	}
 }
@@ -2540,6 +2637,22 @@ void UControlRig::PostEditChangeProperty(FPropertyChangedEvent& PropertyChangedE
 }
 #endif
 
+FRigUnit* UControlRig::GetRigUnitInstanceFromScope(TSharedPtr<FStructOnScope> InScope)
+{
+	if(InScope.IsValid())
+	{
+		if(InScope->IsValid())
+		{
+			if(InScope->GetStruct()->IsChildOf(FRigUnit::StaticStruct()))
+			{
+				return (FRigUnit*)InScope->GetStructMemory(); 
+			}
+		}
+	}
+	static FStructOnScope DefaultRigUnitInstance(FRigUnit::StaticStruct());
+	return (FRigUnit*)DefaultRigUnitInstance.GetStructMemory();
+}
+
 const TArray<UAssetUserData*>* UControlRig::GetAssetUserDataArray() const
 {
 	if(HasAnyFlags(RF_ClassDefaultObject))
diff --git a/Engine/Plugins/Animation/ControlRig/Source/ControlRig/Private/Rigs/RigHierarchy.cpp b/Engine/Plugins/Animation/ControlRig/Source/ControlRig/Private/Rigs/RigHierarchy.cpp
--- a/Engine/Plugins/Animation/ControlRig/Source/ControlRig/Private/Rigs/RigHierarchy.cpp
+++ b/Engine/Plugins/Animation/ControlRig/Source/ControlRig/Private/Rigs/RigHierarchy.cpp
@@ -951,10 +951,10 @@ void URigHierarchy::SanitizeName(FString& InOutName)
 	{
 		TCHAR& C = InOutName[i];
 
-		const bool bGoodChar = FChar::IsAlpha(C) ||				// Any letter
-			(C == '_') || (C == '-') || (C == '.') ||			// _  - . anytime
-			(FChar::IsDigit(C)) ||								// 0-9 anytime
-			((i > 0) && (C== ' '));								// Space after the first character to support virtual bones
+		const bool bGoodChar = FChar::IsAlpha(C) ||					 // Any letter
+			(C == '_') || (C == '-') || (C == '.') || (C == '|') ||	 // _  - .  | anytime
+			(FChar::IsDigit(C)) ||									 // 0-9 anytime
+			((i > 0) && (C== ' '));									 // Space after the first character to support virtual bones
 
 		if (!bGoodChar)
 		{
diff --git a/Engine/Plugins/Animation/ControlRig/Source/ControlRig/Private/Units/Debug/RigUnit_VisualDebug.h b/Engine/Plugins/Animation/ControlRig/Source/ControlRig/Private/Units/Debug/RigUnit_VisualDebug.h
--- a/Engine/Plugins/Animation/ControlRig/Source/ControlRig/Private/Units/Debug/RigUnit_VisualDebug.h
+++ b/Engine/Plugins/Animation/ControlRig/Source/ControlRig/Private/Units/Debug/RigUnit_VisualDebug.h
@@ -69,11 +69,6 @@ struct CONTROLRIG_API FRigUnit_VisualDebugVectorItemSpace : public FRigUnit_Debu
 		Space = FRigElementKey(NAME_None, ERigElementType::Bone);
 	}
 
-	virtual FRigElementKey DetermineSpaceForPin(const FString& InPinPath, void* InUserContext) const override
-	{
-		return Space;
-	}
-
 	RIGVM_METHOD()
 	virtual void Execute() override;
 
@@ -155,11 +150,6 @@ struct CONTROLRIG_API FRigUnit_VisualDebugQuatItemSpace : public FRigUnit_DebugB
 		Space = FRigElementKey(NAME_None, ERigElementType::Bone);
 	}
 
-	virtual FRigElementKey DetermineSpaceForPin(const FString& InPinPath, void* InUserContext) const override
-	{
-		return Space;
-	}
-
 	RIGVM_METHOD()
 	virtual void Execute() override;
 
@@ -235,11 +225,6 @@ struct CONTROLRIG_API FRigUnit_VisualDebugTransformItemSpace : public FRigUnit_D
 		Space = FRigElementKey(NAME_None, ERigElementType::Bone);
 	}
 
-	virtual FRigElementKey DetermineSpaceForPin(const FString& InPinPath, void* InUserContext) const override
-	{
-		return Space;
-	}
-
 	RIGVM_METHOD()
 	virtual void Execute() override;
 
diff --git a/Engine/Plugins/Animation/ControlRig/Source/ControlRig/Private/Units/Hierarchy/RigUnit_GetBoneTransform.h b/Engine/Plugins/Animation/ControlRig/Source/ControlRig/Private/Units/Hierarchy/RigUnit_GetBoneTransform.h
--- a/Engine/Plugins/Animation/ControlRig/Source/ControlRig/Private/Units/Hierarchy/RigUnit_GetBoneTransform.h
+++ b/Engine/Plugins/Animation/ControlRig/Source/ControlRig/Private/Units/Hierarchy/RigUnit_GetBoneTransform.h
@@ -18,18 +18,6 @@ struct CONTROLRIG_API FRigUnit_GetBoneTransform : public FRigUnit
 		, CachedBone()
 	{}
 
-	virtual FRigElementKey DetermineSpaceForPin(const FString& InPinPath, void* InUserContext) const override
-	{
-		if (InPinPath.StartsWith(TEXT("Transform")) && Space == ERigVMTransformSpace::LocalSpace)
-		{
-			if (const URigHierarchy* Hierarchy = (const URigHierarchy*)InUserContext)
-			{
-				return Hierarchy->GetFirstParent(FRigElementKey(Bone, ERigElementType::Bone));
-			}
-		}
-		return FRigElementKey();
-	}
-
 	RIGVM_METHOD()
 	virtual void Execute() override;
 
diff --git a/Engine/Plugins/Animation/ControlRig/Source/ControlRig/Private/Units/Hierarchy/RigUnit_GetInitialBoneTransform.h b/Engine/Plugins/Animation/ControlRig/Source/ControlRig/Private/Units/Hierarchy/RigUnit_GetInitialBoneTransform.h
--- a/Engine/Plugins/Animation/ControlRig/Source/ControlRig/Private/Units/Hierarchy/RigUnit_GetInitialBoneTransform.h
+++ b/Engine/Plugins/Animation/ControlRig/Source/ControlRig/Private/Units/Hierarchy/RigUnit_GetInitialBoneTransform.h
@@ -18,18 +18,6 @@ struct CONTROLRIG_API FRigUnit_GetInitialBoneTransform : public FRigUnit
 		, CachedBone(FCachedRigElement())
 	{}
 
-	virtual FRigElementKey DetermineSpaceForPin(const FString& InPinPath, void* InUserContext) const override
-	{
-		if (InPinPath.StartsWith(TEXT("Transform")) && Space == ERigVMTransformSpace::LocalSpace)
-		{
-			if (const URigHierarchy* Hierarchy = (const URigHierarchy*)InUserContext)
-			{
-				return Hierarchy->GetFirstParent(FRigElementKey(Bone, ERigElementType::Bone));
-			}
-		}
-		return FRigElementKey();
-	}
-
 	RIGVM_METHOD()
 	virtual void Execute() override;
 
diff --git a/Engine/Plugins/Animation/ControlRig/Source/ControlRig/Private/Units/Hierarchy/RigUnit_GetRelativeBoneTransform.h b/Engine/Plugins/Animation/ControlRig/Source/ControlRig/Private/Units/Hierarchy/RigUnit_GetRelativeBoneTransform.h
--- a/Engine/Plugins/Animation/ControlRig/Source/ControlRig/Private/Units/Hierarchy/RigUnit_GetRelativeBoneTransform.h
+++ b/Engine/Plugins/Animation/ControlRig/Source/ControlRig/Private/Units/Hierarchy/RigUnit_GetRelativeBoneTransform.h
@@ -18,15 +18,6 @@ struct CONTROLRIG_API FRigUnit_GetRelativeBoneTransform : public FRigUnit
 		, CachedSpace(FCachedRigElement())
 	{}
 
-	virtual FRigElementKey DetermineSpaceForPin(const FString& InPinPath, void* InUserContext) const override
-	{
-		if (InPinPath.StartsWith(TEXT("Transform")))
-		{
-			return FRigElementKey(Space, ERigElementType::Bone);
-		}
-		return FRigElementKey();
-	}
-
 	RIGVM_METHOD()
 	virtual void Execute() override;
 
diff --git a/Engine/Plugins/Animation/ControlRig/Source/ControlRig/Private/Units/Hierarchy/RigUnit_GetRelativeTransform.h b/Engine/Plugins/Animation/ControlRig/Source/ControlRig/Private/Units/Hierarchy/RigUnit_GetRelativeTransform.h
--- a/Engine/Plugins/Animation/ControlRig/Source/ControlRig/Private/Units/Hierarchy/RigUnit_GetRelativeTransform.h
+++ b/Engine/Plugins/Animation/ControlRig/Source/ControlRig/Private/Units/Hierarchy/RigUnit_GetRelativeTransform.h
@@ -23,11 +23,6 @@ struct CONTROLRIG_API FRigUnit_GetRelativeTransformForItem : public FRigUnit
 		, CachedParent()
 	{}
 
-	virtual FRigElementKey DetermineSpaceForPin(const FString& InPinPath, void* InUserContext) const override
-	{
-		return Parent;
-	}
-
 	RIGVM_METHOD()
 	virtual void Execute() override;
 
diff --git a/Engine/Plugins/Animation/ControlRig/Source/ControlRig/Private/Units/Hierarchy/RigUnit_GetTransform.h b/Engine/Plugins/Animation/ControlRig/Source/ControlRig/Private/Units/Hierarchy/RigUnit_GetTransform.h
--- a/Engine/Plugins/Animation/ControlRig/Source/ControlRig/Private/Units/Hierarchy/RigUnit_GetTransform.h
+++ b/Engine/Plugins/Animation/ControlRig/Source/ControlRig/Private/Units/Hierarchy/RigUnit_GetTransform.h
@@ -23,18 +23,6 @@ struct CONTROLRIG_API FRigUnit_GetTransform : public FRigUnit
 
 	virtual FString GetUnitLabel() const override;
 
-	virtual FRigElementKey DetermineSpaceForPin(const FString& InPinPath, void* InUserContext) const override
-	{
-		if(Space == ERigVMTransformSpace::LocalSpace)
-		{
-			if (const URigHierarchy* Hierarchy = (const URigHierarchy*)InUserContext)
-			{
-				return Hierarchy->GetFirstParent(Item);
-			}
-		}
-		return FRigElementKey();
-	}
-
 	RIGVM_METHOD()
 	virtual void Execute() override;
 
diff --git a/Engine/Plugins/Animation/ControlRig/Source/ControlRig/Private/Units/Hierarchy/RigUnit_OffsetTransform.cpp b/Engine/Plugins/Animation/ControlRig/Source/ControlRig/Private/Units/Hierarchy/RigUnit_OffsetTransform.cpp
--- a/Engine/Plugins/Animation/ControlRig/Source/ControlRig/Private/Units/Hierarchy/RigUnit_OffsetTransform.cpp
+++ b/Engine/Plugins/Animation/ControlRig/Source/ControlRig/Private/Units/Hierarchy/RigUnit_OffsetTransform.cpp
@@ -24,3 +24,53 @@ FRigUnit_OffsetTransformForItem_Execute()
 	FRigVMFunction_MathTransformMakeAbsolute::StaticExecute(ExecuteContext, OffsetTransform, PreviousTransform, GlobalTransform);
 	FRigUnit_SetTransform::StaticExecute(ExecuteContext, Item, ERigVMTransformSpace::GlobalSpace, false, GlobalTransform, Weight, bPropagateToChildren, CachedIndex);
 }
+
+#if WITH_EDITOR
+
+bool FRigUnit_OffsetTransformForItem::UpdateHierarchyForDirectManipulation(const URigVMUnitNode* InNode, TSharedPtr<FStructOnScope> InInstance, FControlRigExecuteContext& InContext, TSharedPtr<FRigDirectManipulationInfo> InInfo)
+{
+	URigHierarchy* Hierarchy = InContext.Hierarchy;
+	if (Hierarchy == nullptr)
+	{
+		return false;
+	}
+	
+	if(InInfo->Target.Name == GET_MEMBER_NAME_STRING_CHECKED(FRigUnit_OffsetTransformForItem, OffsetTransform))
+	{
+		check(InInfo.IsValid());
+
+		if(!InInfo->bInitialized)
+		{
+			InInfo->Reset();
+			const FTransform ParentTransform = Hierarchy->GetParentTransform(Item, false);
+			InInfo->OffsetTransform = OffsetTransform.Inverse() * Hierarchy->GetLocalTransform(Item) * ParentTransform;
+		}
+
+		Hierarchy->SetControlOffsetTransform(InInfo->ControlKey, InInfo->OffsetTransform, false);
+		Hierarchy->SetLocalTransform(InInfo->ControlKey, OffsetTransform, false);
+		if(!InInfo->bInitialized)
+		{
+			Hierarchy->SetLocalTransform(InInfo->ControlKey, OffsetTransform, true);
+		}
+		return true;
+	}
+	return FRigUnitMutable::UpdateHierarchyForDirectManipulation(InNode, InInstance, InContext, InInfo);
+}
+
+bool FRigUnit_OffsetTransformForItem::UpdateDirectManipulationFromHierarchy(const URigVMUnitNode* InNode, TSharedPtr<FStructOnScope> InInstance, FControlRigExecuteContext& InContext, TSharedPtr<FRigDirectManipulationInfo> InInfo)
+{
+	URigHierarchy* Hierarchy = InContext.Hierarchy;
+	if (Hierarchy == nullptr)
+	{
+		return false;
+	}
+	
+	if(InInfo->Target.Name == GET_MEMBER_NAME_STRING_CHECKED(FRigUnit_OffsetTransformForItem, OffsetTransform))
+	{
+		OffsetTransform = Hierarchy->GetLocalTransform(InInfo->ControlKey, false);
+		return true;
+	}
+	return FRigUnitMutable::UpdateDirectManipulationFromHierarchy(InNode, InInstance, InContext, InInfo);
+}
+
+#endif
\ No newline at end of file
diff --git a/Engine/Plugins/Animation/ControlRig/Source/ControlRig/Private/Units/Hierarchy/RigUnit_OffsetTransform.h b/Engine/Plugins/Animation/ControlRig/Source/ControlRig/Private/Units/Hierarchy/RigUnit_OffsetTransform.h
--- a/Engine/Plugins/Animation/ControlRig/Source/ControlRig/Private/Units/Hierarchy/RigUnit_OffsetTransform.h
+++ b/Engine/Plugins/Animation/ControlRig/Source/ControlRig/Private/Units/Hierarchy/RigUnit_OffsetTransform.h
@@ -21,26 +21,11 @@ struct CONTROLRIG_API FRigUnit_OffsetTransformForItem : public FRigUnitMutable
 		, CachedIndex()
 	{}
 
-	virtual FRigElementKey DetermineSpaceForPin(const FString& InPinPath, void* InUserContext) const override
-	{
-		if (const URigHierarchy* Hierarchy = (const URigHierarchy*)InUserContext)
-		{
-			return Hierarchy->GetFirstParent(Item);
-		}
-		return FRigElementKey();
-	}
-
-	virtual FTransform DetermineOffsetTransformForPin(const FString& InPinPath, void* InUserContext) const override
-	{
-		if (const URigHierarchy* Hierarchy = (const URigHierarchy*)InUserContext)
-		{
-			// this is similar to RigUnit_ModifyTransform
-			return OffsetTransform.Inverse() * Hierarchy->GetLocalTransform(Item);
-		}
-		
-		return FTransform::Identity;
-	}
-
+#if WITH_EDITOR
+	virtual bool UpdateHierarchyForDirectManipulation(const URigVMUnitNode* InNode, TSharedPtr<FStructOnScope> InInstance, FControlRigExecuteContext& InContext, TSharedPtr<FRigDirectManipulationInfo> InInfo) override;
+	virtual bool UpdateDirectManipulationFromHierarchy(const URigVMUnitNode* InNode, TSharedPtr<FStructOnScope> InInstance, FControlRigExecuteContext& InContext, TSharedPtr<FRigDirectManipulationInfo> InInfo) override;
+#endif
+	
 	RIGVM_METHOD()
 	virtual void Execute() override;
 
diff --git a/Engine/Plugins/Animation/ControlRig/Source/ControlRig/Private/Units/Hierarchy/RigUnit_PropagateTransform.h b/Engine/Plugins/Animation/ControlRig/Source/ControlRig/Private/Units/Hierarchy/RigUnit_PropagateTransform.h
--- a/Engine/Plugins/Animation/ControlRig/Source/ControlRig/Private/Units/Hierarchy/RigUnit_PropagateTransform.h
+++ b/Engine/Plugins/Animation/ControlRig/Source/ControlRig/Private/Units/Hierarchy/RigUnit_PropagateTransform.h
@@ -22,11 +22,6 @@ struct CONTROLRIG_API FRigUnit_PropagateTransform : public FRigUnitMutable
 		, CachedIndex()
 	{}
 
-	virtual FRigElementKey DetermineSpaceForPin(const FString& InPinPath, void* InUserContext) const override
-	{
-		return Item;
-	}
-
 	RIGVM_METHOD()
 	virtual void Execute() override;
 
diff --git a/Engine/Plugins/Animation/ControlRig/Source/ControlRig/Private/Units/Hierarchy/RigUnit_SetBoneRotation.h b/Engine/Plugins/Animation/ControlRig/Source/ControlRig/Private/Units/Hierarchy/RigUnit_SetBoneRotation.h
--- a/Engine/Plugins/Animation/ControlRig/Source/ControlRig/Private/Units/Hierarchy/RigUnit_SetBoneRotation.h
+++ b/Engine/Plugins/Animation/ControlRig/Source/ControlRig/Private/Units/Hierarchy/RigUnit_SetBoneRotation.h
@@ -22,18 +22,6 @@ struct CONTROLRIG_API FRigUnit_SetBoneRotation : public FRigUnitMutable
 		, CachedBone(FCachedRigElement())
 	{}
 
-	virtual FRigElementKey DetermineSpaceForPin(const FString& InPinPath, void* InUserContext) const override
-	{
-		if (InPinPath.StartsWith(TEXT("Rotation")) && Space == ERigVMTransformSpace::LocalSpace)
-		{
-			if (const URigHierarchy* Hierarchy = (const URigHierarchy*)InUserContext)
-			{
-				return Hierarchy->GetFirstParent(FRigElementKey(Bone, ERigElementType::Bone));
-			}
-		}
-		return FRigElementKey();
-	}
-
 	RIGVM_METHOD()
 	virtual void Execute() override;
 
diff --git a/Engine/Plugins/Animation/ControlRig/Source/ControlRig/Private/Units/Hierarchy/RigUnit_SetBoneTransform.h b/Engine/Plugins/Animation/ControlRig/Source/ControlRig/Private/Units/Hierarchy/RigUnit_SetBoneTransform.h
--- a/Engine/Plugins/Animation/ControlRig/Source/ControlRig/Private/Units/Hierarchy/RigUnit_SetBoneTransform.h
+++ b/Engine/Plugins/Animation/ControlRig/Source/ControlRig/Private/Units/Hierarchy/RigUnit_SetBoneTransform.h
@@ -21,18 +21,6 @@ struct CONTROLRIG_API FRigUnit_SetBoneTransform : public FRigUnitMutable
 		, CachedBone(FCachedRigElement())
 	{}
 
-	virtual FRigElementKey DetermineSpaceForPin(const FString& InPinPath, void* InUserContext) const override
-	{
-		if (InPinPath.StartsWith(TEXT("Transform")) && Space == ERigVMTransformSpace::LocalSpace)
-		{
-			if (const URigHierarchy* Hierarchy = (const URigHierarchy*)InUserContext)
-			{
-				return Hierarchy->GetFirstParent(FRigElementKey(Bone, ERigElementType::Bone));
-			}
-		}
-		return FRigElementKey();
-	}
-
 	RIGVM_METHOD()
 	virtual void Execute() override;
 
diff --git a/Engine/Plugins/Animation/ControlRig/Source/ControlRig/Private/Units/Hierarchy/RigUnit_SetBoneTranslation.h b/Engine/Plugins/Animation/ControlRig/Source/ControlRig/Private/Units/Hierarchy/RigUnit_SetBoneTranslation.h
--- a/Engine/Plugins/Animation/ControlRig/Source/ControlRig/Private/Units/Hierarchy/RigUnit_SetBoneTranslation.h
+++ b/Engine/Plugins/Animation/ControlRig/Source/ControlRig/Private/Units/Hierarchy/RigUnit_SetBoneTranslation.h
@@ -22,18 +22,6 @@ struct CONTROLRIG_API FRigUnit_SetBoneTranslation : public FRigUnitMutable
 		, CachedBone(FCachedRigElement())
 	{}
 
-	virtual FRigElementKey DetermineSpaceForPin(const FString& InPinPath, void* InUserContext) const override
-	{
-		if (InPinPath.StartsWith(TEXT("Translation")) && Space == ERigVMTransformSpace::LocalSpace)
-		{
-			if (const URigHierarchy* Hierarchy = (const URigHierarchy*)InUserContext)
-			{
-				return Hierarchy->GetFirstParent(FRigElementKey(Bone, ERigElementType::Bone));
-			}
-		}
-		return FRigElementKey();
-	}
-
 	RIGVM_METHOD()
 	virtual void Execute() override;
 
diff --git a/Engine/Plugins/Animation/ControlRig/Source/ControlRig/Private/Units/Hierarchy/RigUnit_SetRelativeBoneTransform.h b/Engine/Plugins/Animation/ControlRig/Source/ControlRig/Private/Units/Hierarchy/RigUnit_SetRelativeBoneTransform.h
--- a/Engine/Plugins/Animation/ControlRig/Source/ControlRig/Private/Units/Hierarchy/RigUnit_SetRelativeBoneTransform.h
+++ b/Engine/Plugins/Animation/ControlRig/Source/ControlRig/Private/Units/Hierarchy/RigUnit_SetRelativeBoneTransform.h
@@ -21,15 +21,6 @@ struct CONTROLRIG_API FRigUnit_SetRelativeBoneTransform : public FRigUnitMutable
 		, CachedSpaceIndex(FCachedRigElement())
 	{}
 
-	virtual FRigElementKey DetermineSpaceForPin(const FString& InPinPath, void* InUserContext) const override
-	{
-		if (InPinPath.StartsWith(TEXT("Transform")))
-		{
-			return FRigElementKey(Space, ERigElementType::Bone);
-		}
-		return FRigElementKey();
-	}
-
 	RIGVM_METHOD()
 	virtual void Execute() override;
 
diff --git a/Engine/Plugins/Animation/ControlRig/Source/ControlRig/Private/Units/Hierarchy/RigUnit_SetRelativeTransform.cpp b/Engine/Plugins/Animation/ControlRig/Source/ControlRig/Private/Units/Hierarchy/RigUnit_SetRelativeTransform.cpp
--- a/Engine/Plugins/Animation/ControlRig/Source/ControlRig/Private/Units/Hierarchy/RigUnit_SetRelativeTransform.cpp
+++ b/Engine/Plugins/Animation/ControlRig/Source/ControlRig/Private/Units/Hierarchy/RigUnit_SetRelativeTransform.cpp
@@ -67,3 +67,121 @@ FRigUnit_SetRelativeRotationForItem_Execute()
 	FRigUnit_SetTransform::StaticExecute(ExecuteContext, Child, ERigVMTransformSpace::GlobalSpace, false, GlobalTransform, Weight, bPropagateToChildren, CachedChild);
 }
 
+#if WITH_EDITOR
+
+bool FRigUnit_SetRelativeTransformForItem::UpdateHierarchyForDirectManipulation(const URigVMUnitNode* InNode, TSharedPtr<FStructOnScope> InInstance, FControlRigExecuteContext& InContext, TSharedPtr<FRigDirectManipulationInfo> InInfo)
+{
+	URigHierarchy* Hierarchy = InContext.Hierarchy;
+	if (Hierarchy == nullptr)
+	{
+		return false;
+	}
+	
+	if(InInfo->Target.Name == GET_MEMBER_NAME_STRING_CHECKED(FRigUnit_SetRelativeTransformForItem, Value))
+	{
+		const FTransform ParentTransform = Hierarchy->GetGlobalTransform(Parent, false);
+		Hierarchy->SetControlOffsetTransform(InInfo->ControlKey, ParentTransform, false);
+		Hierarchy->SetLocalTransform(InInfo->ControlKey, Value, false);
+		if(!InInfo->bInitialized)
+		{
+			Hierarchy->SetLocalTransform(InInfo->ControlKey, Value, true);
+		}
+		return true;
+	}
+	return FRigUnitMutable::UpdateHierarchyForDirectManipulation(InNode, InInstance, InContext, InInfo);
+}
+
+bool FRigUnit_SetRelativeTransformForItem::UpdateDirectManipulationFromHierarchy(const URigVMUnitNode* InNode, TSharedPtr<FStructOnScope> InInstance, FControlRigExecuteContext& InContext, TSharedPtr<FRigDirectManipulationInfo> InInfo)
+{
+	URigHierarchy* Hierarchy = InContext.Hierarchy;
+	if (Hierarchy == nullptr)
+	{
+		return false;
+	}
+
+	if(InInfo->Target.Name == GET_MEMBER_NAME_STRING_CHECKED(FRigUnit_SetRelativeTransformForItem, Value))
+	{
+		Value = Hierarchy->GetLocalTransform(InInfo->ControlKey, false);
+		return true;
+	}
+	return FRigUnitMutable::UpdateDirectManipulationFromHierarchy(InNode, InInstance, InContext, InInfo);
+}
+
+bool FRigUnit_SetRelativeTranslationForItem::UpdateHierarchyForDirectManipulation(const URigVMUnitNode* InNode, TSharedPtr<FStructOnScope> InInstance, FControlRigExecuteContext& InContext, TSharedPtr<FRigDirectManipulationInfo> InInfo)
+{
+	URigHierarchy* Hierarchy = InContext.Hierarchy;
+	if (Hierarchy == nullptr)
+	{
+		return false;
+	}
+
+	if(InInfo->Target.Name == GET_MEMBER_NAME_STRING_CHECKED(FRigUnit_SetRelativeTranslationForItem, Value))
+	{
+		const FTransform ParentTransform = Hierarchy->GetGlobalTransform(Parent, false);
+		Hierarchy->SetControlOffsetTransform(InInfo->ControlKey, ParentTransform, false);
+		Hierarchy->SetLocalTransform(InInfo->ControlKey, FTransform(Value), false);
+		if(!InInfo->bInitialized)
+		{
+			Hierarchy->SetLocalTransform(InInfo->ControlKey, FTransform(Value), true);
+		}
+		return true;
+	}
+	return FRigUnitMutable::UpdateHierarchyForDirectManipulation(InNode, InInstance, InContext, InInfo);
+}
+
+bool FRigUnit_SetRelativeTranslationForItem::UpdateDirectManipulationFromHierarchy(const URigVMUnitNode* InNode, TSharedPtr<FStructOnScope> InInstance, FControlRigExecuteContext& InContext, TSharedPtr<FRigDirectManipulationInfo> InInfo)
+{
+	URigHierarchy* Hierarchy = InContext.Hierarchy;
+	if (Hierarchy == nullptr)
+	{
+		return false;
+	}
+
+	if(InInfo->Target.Name == GET_MEMBER_NAME_STRING_CHECKED(FRigUnit_SetRelativeTranslationForItem, Value))
+	{
+		Value = Hierarchy->GetLocalTransform(InInfo->ControlKey, false).GetTranslation();
+		return true;
+	}
+	return FRigUnitMutable::UpdateDirectManipulationFromHierarchy(InNode, InInstance, InContext, InInfo);
+}
+
+bool FRigUnit_SetRelativeRotationForItem::UpdateHierarchyForDirectManipulation(const URigVMUnitNode* InNode, TSharedPtr<FStructOnScope> InInstance, FControlRigExecuteContext& InContext, TSharedPtr<FRigDirectManipulationInfo> InInfo)
+{
+	URigHierarchy* Hierarchy = InContext.Hierarchy;
+	if (Hierarchy == nullptr)
+	{
+		return false;
+	}
+
+	if(InInfo->Target.Name == GET_MEMBER_NAME_STRING_CHECKED(FRigUnit_SetRelativeRotationForItem, Value))
+	{
+		FTransform ParentTransform = Hierarchy->GetGlobalTransform(Parent, false);
+		ParentTransform.SetTranslation(Hierarchy->GetGlobalTransform(Child).GetTranslation());
+		Hierarchy->SetControlOffsetTransform(InInfo->ControlKey, ParentTransform, false);
+		Hierarchy->SetLocalTransform(InInfo->ControlKey, FTransform(Value), false);
+		if(!InInfo->bInitialized)
+		{
+			Hierarchy->SetLocalTransform(InInfo->ControlKey, FTransform(Value), true);
+		}
+		return true;
+	}
+	return FRigUnitMutable::UpdateHierarchyForDirectManipulation(InNode, InInstance, InContext, InInfo);
+}
+
+bool FRigUnit_SetRelativeRotationForItem::UpdateDirectManipulationFromHierarchy(const URigVMUnitNode* InNode, TSharedPtr<FStructOnScope> InInstance, FControlRigExecuteContext& InContext, TSharedPtr<FRigDirectManipulationInfo> InInfo)
+{
+	URigHierarchy* Hierarchy = InContext.Hierarchy;
+	if (Hierarchy == nullptr)
+	{
+		return false;
+	}
+
+	if(InInfo->Target.Name == GET_MEMBER_NAME_STRING_CHECKED(FRigUnit_SetRelativeRotationForItem, Value))
+	{
+		Value = Hierarchy->GetLocalTransform(InInfo->ControlKey, false).GetRotation();
+		return true;
+	}
+	return FRigUnitMutable::UpdateDirectManipulationFromHierarchy(InNode, InInstance, InContext, InInfo);
+}
+
+#endif
diff --git a/Engine/Plugins/Animation/ControlRig/Source/ControlRig/Private/Units/Hierarchy/RigUnit_SetRelativeTransform.h b/Engine/Plugins/Animation/ControlRig/Source/ControlRig/Private/Units/Hierarchy/RigUnit_SetRelativeTransform.h
--- a/Engine/Plugins/Animation/ControlRig/Source/ControlRig/Private/Units/Hierarchy/RigUnit_SetRelativeTransform.h
+++ b/Engine/Plugins/Animation/ControlRig/Source/ControlRig/Private/Units/Hierarchy/RigUnit_SetRelativeTransform.h
@@ -24,11 +24,11 @@ struct CONTROLRIG_API FRigUnit_SetRelativeTransformForItem : public FRigUnitMuta
 		, CachedParent()
 	{}
 
-	virtual FRigElementKey DetermineSpaceForPin(const FString& InPinPath, void* InUserContext) const override
-	{
-		return Parent;
-	}
-
+#if WITH_EDITOR
+	virtual bool UpdateHierarchyForDirectManipulation(const URigVMUnitNode* InNode, TSharedPtr<FStructOnScope> InInstance, FControlRigExecuteContext& InContext, TSharedPtr<FRigDirectManipulationInfo> InInfo) override;
+	virtual bool UpdateDirectManipulationFromHierarchy(const URigVMUnitNode* InNode, TSharedPtr<FStructOnScope> InInstance, FControlRigExecuteContext& InContext, TSharedPtr<FRigDirectManipulationInfo> InInfo) override;
+#endif
+	
 	RIGVM_METHOD()
 	virtual void Execute() override;
 
@@ -93,11 +93,11 @@ struct CONTROLRIG_API FRigUnit_SetRelativeTranslationForItem : public FRigUnitMu
 		, CachedParent()
 	{}
 
-	virtual FRigElementKey DetermineSpaceForPin(const FString& InPinPath, void* InUserContext) const override
-	{
-		return Parent;
-	}
-
+#if WITH_EDITOR
+	virtual bool UpdateHierarchyForDirectManipulation(const URigVMUnitNode* InNode, TSharedPtr<FStructOnScope> InInstance, FControlRigExecuteContext& InContext, TSharedPtr<FRigDirectManipulationInfo> InInfo) override;
+	virtual bool UpdateDirectManipulationFromHierarchy(const URigVMUnitNode* InNode, TSharedPtr<FStructOnScope> InInstance, FControlRigExecuteContext& InContext, TSharedPtr<FRigDirectManipulationInfo> InInfo) override;
+#endif
+	
 	RIGVM_METHOD()
 	virtual void Execute() override;
 
@@ -162,11 +162,11 @@ struct CONTROLRIG_API FRigUnit_SetRelativeRotationForItem : public FRigUnitMutab
 		, CachedParent()
 	{}
 
-	virtual FRigElementKey DetermineSpaceForPin(const FString& InPinPath, void* InUserContext) const override
-	{
-		return Parent;
-	}
-
+#if WITH_EDITOR
+	virtual bool UpdateHierarchyForDirectManipulation(const URigVMUnitNode* InNode, TSharedPtr<FStructOnScope> InInstance, FControlRigExecuteContext& InContext, TSharedPtr<FRigDirectManipulationInfo> InInfo) override;
+	virtual bool UpdateDirectManipulationFromHierarchy(const URigVMUnitNode* InNode, TSharedPtr<FStructOnScope> InInstance, FControlRigExecuteContext& InContext, TSharedPtr<FRigDirectManipulationInfo> InInfo) override;
+#endif
+	
 	RIGVM_METHOD()
 	virtual void Execute() override;
 
diff --git a/Engine/Plugins/Animation/ControlRig/Source/ControlRig/Private/Units/Hierarchy/RigUnit_SetTransform.cpp b/Engine/Plugins/Animation/ControlRig/Source/ControlRig/Private/Units/Hierarchy/RigUnit_SetTransform.cpp
--- a/Engine/Plugins/Animation/ControlRig/Source/ControlRig/Private/Units/Hierarchy/RigUnit_SetTransform.cpp
+++ b/Engine/Plugins/Animation/ControlRig/Source/ControlRig/Private/Units/Hierarchy/RigUnit_SetTransform.cpp
@@ -17,6 +17,62 @@ FString FRigUnit_SetTransform::GetUnitLabel() const
 	return FString::Printf(TEXT("Set Transform - %s%s"), *Type, *Initial);
 }
 
+#if WITH_EDITOR
+
+bool FRigUnit_SetTransform::UpdateHierarchyForDirectManipulation(const URigVMUnitNode* InNode, TSharedPtr<FStructOnScope> InInstance, FControlRigExecuteContext& InContext, TSharedPtr<FRigDirectManipulationInfo> InInfo)
+{
+	URigHierarchy* Hierarchy = InContext.Hierarchy;
+	if (Hierarchy == nullptr)
+	{
+		return false;
+	}
+	
+	if(InInfo->Target.Name == GET_MEMBER_NAME_STRING_CHECKED(FRigUnit_SetTransform, Value))
+	{
+		if(Space == ERigVMTransformSpace::LocalSpace)
+		{
+			const FTransform ParentTransform = Hierarchy->GetParentTransform(Item, false);
+			Hierarchy->SetControlOffsetTransform(InInfo->ControlKey, ParentTransform, false);
+		}
+		else
+		{
+			Hierarchy->SetControlOffsetTransform(InInfo->ControlKey, FTransform::Identity, false);
+		}
+		Hierarchy->SetLocalTransform(InInfo->ControlKey, Value, false);
+		if(!InInfo->bInitialized)
+		{
+			Hierarchy->SetLocalTransform(InInfo->ControlKey, Value, true);
+		}
+		return true;
+	}
+	return FRigUnitMutable::UpdateHierarchyForDirectManipulation(InNode, InInstance, InContext, InInfo);
+}
+
+bool FRigUnit_SetTransform::UpdateDirectManipulationFromHierarchy(const URigVMUnitNode* InNode, TSharedPtr<FStructOnScope> InInstance, FControlRigExecuteContext& InContext, TSharedPtr<FRigDirectManipulationInfo> InInfo)
+{
+	URigHierarchy* Hierarchy = InContext.Hierarchy;
+	if (Hierarchy == nullptr)
+	{
+		return false;
+	}
+	
+	if(InInfo->Target.Name == GET_MEMBER_NAME_STRING_CHECKED(FRigUnit_SetTransform, Value))
+	{
+		if(Space == ERigVMTransformSpace::LocalSpace)
+		{
+			Value = Hierarchy->GetLocalTransform(InInfo->ControlKey, false);
+		}
+		else
+		{
+			Value = Hierarchy->GetGlobalTransform(InInfo->ControlKey, false);
+		}
+		return true;
+	}
+	return FRigUnitMutable::UpdateDirectManipulationFromHierarchy(InNode, InInstance, InContext, InInfo);
+}
+
+#endif
+
 FRigUnit_SetTransform_Execute()
 {
     DECLARE_SCOPE_HIERARCHICAL_COUNTER_RIGUNIT()
@@ -120,6 +176,64 @@ FString FRigUnit_SetTranslation::GetUnitLabel() const
 	return FString::Printf(TEXT("Set Translation - %s"), *Type);
 }
 
+#if WITH_EDITOR
+
+bool FRigUnit_SetTranslation::UpdateHierarchyForDirectManipulation(const URigVMUnitNode* InNode, TSharedPtr<FStructOnScope> InInstance, FControlRigExecuteContext& InContext, TSharedPtr<FRigDirectManipulationInfo> InInfo)
+{
+	URigHierarchy* Hierarchy = InContext.Hierarchy;
+	if (Hierarchy == nullptr)
+	{
+		return false;
+	}
+	
+	if(InInfo->Target.Name == GET_MEMBER_NAME_STRING_CHECKED(FRigUnit_SetTranslation, Value))
+	{
+		if(Space == ERigVMTransformSpace::LocalSpace)
+		{
+			const FTransform ParentTransform = Hierarchy->GetParentTransform(Item, false);
+			Hierarchy->SetControlOffsetTransform(InInfo->ControlKey, ParentTransform, false);
+		}
+		else
+		{
+			Hierarchy->SetControlOffsetTransform(InInfo->ControlKey, FTransform::Identity, false);
+		}
+		FTransform Transform = FTransform::Identity;
+		Transform.SetTranslation(Value);
+		Hierarchy->SetLocalTransform(InInfo->ControlKey, Transform, false);
+		if(!InInfo->bInitialized)
+		{
+			Hierarchy->SetLocalTransform(InInfo->ControlKey, Transform, true);
+		}
+		return true;
+	}
+	return FRigUnitMutable::UpdateHierarchyForDirectManipulation(InNode, InInstance, InContext, InInfo);
+}
+
+bool FRigUnit_SetTranslation::UpdateDirectManipulationFromHierarchy(const URigVMUnitNode* InNode, TSharedPtr<FStructOnScope> InInstance, FControlRigExecuteContext& InContext, TSharedPtr<FRigDirectManipulationInfo> InInfo)
+{
+	URigHierarchy* Hierarchy = InContext.Hierarchy;
+	if (Hierarchy == nullptr)
+	{
+		return false;
+	}
+	
+	if(InInfo->Target.Name == GET_MEMBER_NAME_STRING_CHECKED(FRigUnit_SetTranslation, Value))
+	{
+		if(Space == ERigVMTransformSpace::LocalSpace)
+		{
+			Value = Hierarchy->GetLocalTransform(InInfo->ControlKey, false).GetTranslation();
+		}
+		else
+		{
+			Value = Hierarchy->GetGlobalTransform(InInfo->ControlKey, false).GetTranslation();
+		}
+		return true;
+	}
+	return FRigUnitMutable::UpdateDirectManipulationFromHierarchy(InNode, InInstance, InContext, InInfo);
+}
+
+#endif
+
 FRigUnit_SetTranslation_Execute()
 {
 	FTransform Transform = FTransform::Identity;
@@ -134,6 +248,66 @@ FString FRigUnit_SetRotation::GetUnitLabel() const
 	return FString::Printf(TEXT("Set Rotation - %s"), *Type);
 }
 
+#if WITH_EDITOR
+
+bool FRigUnit_SetRotation::UpdateHierarchyForDirectManipulation(const URigVMUnitNode* InNode, TSharedPtr<FStructOnScope> InInstance, FControlRigExecuteContext& InContext, TSharedPtr<FRigDirectManipulationInfo> InInfo)
+{
+	URigHierarchy* Hierarchy = InContext.Hierarchy;
+	if (Hierarchy == nullptr)
+	{
+		return false;
+	}
+	
+	if(InInfo->Target.Name == GET_MEMBER_NAME_STRING_CHECKED(FRigUnit_SetRotation, Value))
+	{
+		if(Space == ERigVMTransformSpace::LocalSpace)
+		{
+			const FTransform ParentTransform = Hierarchy->GetParentTransform(Item, false);
+			Hierarchy->SetControlOffsetTransform(InInfo->ControlKey, ParentTransform, false);
+		}
+		else
+		{
+			FTransform OffsetTransform = FTransform::Identity;
+			OffsetTransform.SetTranslation(Hierarchy->GetGlobalTransform(Item, false).GetTranslation());
+			Hierarchy->SetControlOffsetTransform(InInfo->ControlKey, OffsetTransform, false);
+		}
+		FTransform Transform = FTransform::Identity;
+		Transform.SetRotation(Value);
+		Hierarchy->SetLocalTransform(InInfo->ControlKey, Transform, false);
+		if(!InInfo->bInitialized)
+		{
+			Hierarchy->SetLocalTransform(InInfo->ControlKey, Transform, true);
+		}
+		return true;
+	}
+	return FRigUnitMutable::UpdateHierarchyForDirectManipulation(InNode, InInstance, InContext, InInfo);
+}
+
+bool FRigUnit_SetRotation::UpdateDirectManipulationFromHierarchy(const URigVMUnitNode* InNode, TSharedPtr<FStructOnScope> InInstance, FControlRigExecuteContext& InContext, TSharedPtr<FRigDirectManipulationInfo> InInfo)
+{
+	URigHierarchy* Hierarchy = InContext.Hierarchy;
+	if (Hierarchy == nullptr)
+	{
+		return false;
+	}
+	
+	if(InInfo->Target.Name == GET_MEMBER_NAME_STRING_CHECKED(FRigUnit_SetRotation, Value))
+	{
+		if(Space == ERigVMTransformSpace::LocalSpace)
+		{
+			Value = Hierarchy->GetLocalTransform(InInfo->ControlKey, false).GetRotation();
+		}
+		else
+		{
+			Value = Hierarchy->GetGlobalTransform(InInfo->ControlKey, false).GetRotation();
+		}
+		return true;
+	}
+	return FRigUnitMutable::UpdateDirectManipulationFromHierarchy(InNode, InInstance, InContext, InInfo);
+}
+
+#endif
+
 FRigUnit_SetRotation_Execute()
 {
 	FTransform Transform = FTransform::Identity;
@@ -148,6 +322,66 @@ FString FRigUnit_SetScale::GetUnitLabel() const
 	return FString::Printf(TEXT("Set Scale - %s"), *Type);
 }
 
+#if WITH_EDITOR
+
+bool FRigUnit_SetScale::UpdateHierarchyForDirectManipulation(const URigVMUnitNode* InNode, TSharedPtr<FStructOnScope> InInstance, FControlRigExecuteContext& InContext, TSharedPtr<FRigDirectManipulationInfo> InInfo)
+{
+	URigHierarchy* Hierarchy = InContext.Hierarchy;
+	if (Hierarchy == nullptr)
+	{
+		return false;
+	}
+	
+	if(InInfo->Target.Name == GET_MEMBER_NAME_STRING_CHECKED(FRigUnit_SetScale, Scale))
+	{
+		if(Space == ERigVMTransformSpace::LocalSpace)
+		{
+			const FTransform ParentTransform = Hierarchy->GetParentTransform(Item, false);
+			Hierarchy->SetControlOffsetTransform(InInfo->ControlKey, ParentTransform, false);
+		}
+		else
+		{
+			FTransform OffsetTransform = FTransform::Identity;
+			OffsetTransform.SetTranslation(Hierarchy->GetGlobalTransform(Item, false).GetTranslation());
+			Hierarchy->SetControlOffsetTransform(InInfo->ControlKey, OffsetTransform, false);
+		}
+		FTransform Transform = FTransform::Identity;
+		Transform.SetScale3D(Scale);
+		Hierarchy->SetLocalTransform(InInfo->ControlKey, Transform, false);
+		if(!InInfo->bInitialized)
+		{
+			Hierarchy->SetLocalTransform(InInfo->ControlKey, Transform, true);
+		}
+		return true;
+	}
+	return FRigUnitMutable::UpdateHierarchyForDirectManipulation(InNode, InInstance, InContext, InInfo);
+}
+
+bool FRigUnit_SetScale::UpdateDirectManipulationFromHierarchy(const URigVMUnitNode* InNode, TSharedPtr<FStructOnScope> InInstance, FControlRigExecuteContext& InContext, TSharedPtr<FRigDirectManipulationInfo> InInfo)
+{
+	URigHierarchy* Hierarchy = InContext.Hierarchy;
+	if (Hierarchy == nullptr)
+	{
+		return false;
+	}
+	
+	if(InInfo->Target.Name == GET_MEMBER_NAME_STRING_CHECKED(FRigUnit_SetScale, Scale))
+	{
+		if(Space == ERigVMTransformSpace::LocalSpace)
+		{
+			Scale = Hierarchy->GetLocalTransform(InInfo->ControlKey, false).GetScale3D();
+		}
+		else
+		{
+			Scale = Hierarchy->GetGlobalTransform(InInfo->ControlKey, false).GetScale3D();
+		}
+		return true;
+	}
+	return FRigUnitMutable::UpdateDirectManipulationFromHierarchy(InNode, InInstance, InContext, InInfo);
+}
+
+#endif
+
 FRigUnit_SetScale_Execute()
 {
 	FTransform Transform = FTransform::Identity;
diff --git a/Engine/Plugins/Animation/ControlRig/Source/ControlRig/Private/Units/Highlevel/Hierarchy/RigUnit_AimBone.cpp b/Engine/Plugins/Animation/ControlRig/Source/ControlRig/Private/Units/Highlevel/Hierarchy/RigUnit_AimBone.cpp
--- a/Engine/Plugins/Animation/ControlRig/Source/ControlRig/Private/Units/Highlevel/Hierarchy/RigUnit_AimBone.cpp
+++ b/Engine/Plugins/Animation/ControlRig/Source/ControlRig/Private/Units/Highlevel/Hierarchy/RigUnit_AimBone.cpp
@@ -211,6 +211,149 @@ FRigVMStructUpgradeInfo FRigUnit_AimBone::GetUpgradeInfo() const
 	return Info;
 }
 
+#if WITH_EDITOR
+
+bool FRigUnit_AimItem::GetDirectManipulationTargets(const URigVMUnitNode* InNode, TSharedPtr<FStructOnScope> InInstance, URigHierarchy* InHierarchy, TArray<FRigDirectManipulationTarget>& InOutTargets, FString* OutFailureReason) const
+{
+	static TArray<FRigDirectManipulationTarget> Targets = {
+		{ TEXT("Primary Target"), ERigControlType::Position },
+		{ TEXT("Secondary Target"), ERigControlType::Position },
+		{ TEXT("Resulting Transform"), ERigControlType::EulerTransform }
+	};
+	InOutTargets.Append(Targets);
+	return true;
+}
+
+TArray<const URigVMPin*> FRigUnit_AimItem::GetPinsForDirectManipulation(const URigVMUnitNode* InNode, const FRigDirectManipulationTarget& InTarget) const
+{
+	if(InTarget.Name == TEXT("Primary Target"))
+	{
+		return {
+			InNode->FindPin(GET_MEMBER_NAME_STRING_CHECKED(FRigUnit_AimItem, Primary))->FindSubPin(GET_MEMBER_NAME_STRING_CHECKED(FRigUnit_AimItem_Target, Kind)),
+			InNode->FindPin(GET_MEMBER_NAME_STRING_CHECKED(FRigUnit_AimItem, Primary))->FindSubPin(GET_MEMBER_NAME_STRING_CHECKED(FRigUnit_AimItem_Target, Target)),
+		};
+	}
+	if(InTarget.Name == TEXT("Secondary Target"))
+	{
+		return {
+			InNode->FindPin(GET_MEMBER_NAME_STRING_CHECKED(FRigUnit_AimItem, Secondary))->FindSubPin(GET_MEMBER_NAME_STRING_CHECKED(FRigUnit_AimItem_Target, Kind)),
+			InNode->FindPin(GET_MEMBER_NAME_STRING_CHECKED(FRigUnit_AimItem, Secondary))->FindSubPin(GET_MEMBER_NAME_STRING_CHECKED(FRigUnit_AimItem_Target, Target)),
+		};
+	}
+	if(InTarget.Name == TEXT("Resulting Transform"))
+	{
+		return {
+			InNode->FindPin(GET_MEMBER_NAME_STRING_CHECKED(FRigUnit_AimItem, Primary))->FindSubPin(GET_MEMBER_NAME_STRING_CHECKED(FRigUnit_AimItem_Target, Target)),
+			InNode->FindPin(GET_MEMBER_NAME_STRING_CHECKED(FRigUnit_AimItem, Secondary))->FindSubPin(GET_MEMBER_NAME_STRING_CHECKED(FRigUnit_AimItem_Target, Target)),
+			InNode->FindPin(GET_MEMBER_NAME_STRING_CHECKED(FRigUnit_AimItem, Primary))->FindSubPin(GET_MEMBER_NAME_STRING_CHECKED(FRigUnit_AimItem_Target, Kind)),
+			InNode->FindPin(GET_MEMBER_NAME_STRING_CHECKED(FRigUnit_AimItem, Secondary))->FindSubPin(GET_MEMBER_NAME_STRING_CHECKED(FRigUnit_AimItem_Target, Kind)),
+		};
+	}
+	return FRigUnit_HighlevelBaseMutable::GetPinsForDirectManipulation(InNode, InTarget);
+}
+
+bool FRigUnit_AimItem::UpdateHierarchyForDirectManipulation(const URigVMUnitNode* InNode, TSharedPtr<FStructOnScope> InInstance, FControlRigExecuteContext& InContext, TSharedPtr<FRigDirectManipulationInfo> InInfo)
+{
+	URigHierarchy* Hierarchy = InContext.Hierarchy;
+	if(Hierarchy == nullptr)
+	{
+		return false;
+	}
+	
+	if(InInfo->Target.Name == TEXT("Primary Target"))
+	{
+		const FTransform ParentTransform = Hierarchy->GetGlobalTransform(Primary.Space, false);
+		Hierarchy->SetControlOffsetTransform(InInfo->ControlKey, ParentTransform, false);
+		Hierarchy->SetLocalTransform(InInfo->ControlKey, FTransform(Primary.Target), false);
+		if(!InInfo->bInitialized)
+		{
+			Hierarchy->SetLocalTransform(InInfo->ControlKey, FTransform(Primary.Target), true);
+		}
+		return true;
+	}
+	if(InInfo->Target.Name == TEXT("Secondary Target"))
+	{
+		const FTransform ParentTransform = Hierarchy->GetGlobalTransform(Secondary.Space, false);
+		Hierarchy->SetControlOffsetTransform(InInfo->ControlKey, ParentTransform, false);
+		Hierarchy->SetLocalTransform(InInfo->ControlKey, FTransform(Secondary.Target), false);
+		if(!InInfo->bInitialized)
+		{
+			Hierarchy->SetLocalTransform(InInfo->ControlKey, FTransform(Primary.Target), true);
+		}
+		return true;
+	}
+	if(InInfo->Target.Name == TEXT("Resulting Transform"))
+	{
+		// execute the aim
+		Execute(InContext);
+		
+		const FTransform ItemTransform = Hierarchy->GetGlobalTransform(CachedItem);
+		
+		if(!InInfo->bInitialized)
+		{
+			InInfo->OffsetTransform = ItemTransform;
+			
+			const FTransform PrimaryParentTransform = Hierarchy->GetGlobalTransform(Primary.Space, false);
+			const FTransform SecondaryParentTransform = Hierarchy->GetGlobalTransform(Secondary.Space, false);
+			InInfo->Transforms = { FTransform(Primary.Target), FTransform(Secondary.Target) };
+			InInfo->Transforms[0] = InInfo->Transforms[0] * PrimaryParentTransform;
+			InInfo->Transforms[1] = InInfo->Transforms[1] * SecondaryParentTransform;
+			InInfo->Transforms[0] = InInfo->Transforms[0].GetRelativeTransform(ItemTransform);
+			InInfo->Transforms[1] = InInfo->Transforms[1].GetRelativeTransform(ItemTransform);
+			InInfo->Transforms[0].NormalizeRotation();
+			InInfo->Transforms[1].NormalizeRotation();
+
+			Hierarchy->SetControlOffsetTransform(InInfo->ControlKey, InInfo->OffsetTransform, false);
+			Hierarchy->SetGlobalTransform(InInfo->ControlKey, FTransform::Identity, false);
+			Hierarchy->SetGlobalTransform(InInfo->ControlKey, FTransform::Identity, true);
+		}
+		else
+		{
+			Hierarchy->SetGlobalTransform(InInfo->ControlKey, ItemTransform, false);
+		}
+		return true;
+	}
+	return FRigUnit_HighlevelBaseMutable::UpdateHierarchyForDirectManipulation(InNode, InInstance, InContext, InInfo);
+}
+
+bool FRigUnit_AimItem::UpdateDirectManipulationFromHierarchy(const URigVMUnitNode* InNode, TSharedPtr<FStructOnScope> InInstance, FControlRigExecuteContext& InContext, TSharedPtr<FRigDirectManipulationInfo> InInfo)
+{
+	URigHierarchy* Hierarchy = InContext.Hierarchy;
+	if(Hierarchy == nullptr)
+	{
+		return false;
+	}
+
+	if(InInfo->Target.Name == TEXT("Primary Target"))
+	{
+		Primary.Kind = EControlRigVectorKind::Location;
+		Primary.Target = Hierarchy->GetLocalTransform(InInfo->ControlKey, false).GetTranslation();
+		return true;
+	}
+	if(InInfo->Target.Name == TEXT("Secondary Target"))
+	{
+		Primary.Kind = EControlRigVectorKind::Location;
+		Secondary.Target = Hierarchy->GetLocalTransform(InInfo->ControlKey, false).GetTranslation();
+		return true;
+	}
+	if(InInfo->Target.Name == TEXT("Resulting Transform"))
+	{
+		Primary.Kind = EControlRigVectorKind::Location;
+		Secondary.Kind = EControlRigVectorKind::Location;
+		const FTransform ControlTransform = Hierarchy->GetGlobalTransform(InInfo->ControlKey);
+		const FTransform GlobalPrimaryTransform = InInfo->Transforms[0] * ControlTransform;
+		const FTransform GlobalSecondaryTransform = InInfo->Transforms[1] * ControlTransform;
+		const FTransform LocalPrimaryTransform = GlobalPrimaryTransform.GetRelativeTransform(Hierarchy->GetGlobalTransform(Primary.Space));
+		const FTransform LocalSecondaryTransform = GlobalSecondaryTransform.GetRelativeTransform(Hierarchy->GetGlobalTransform(Secondary.Space));
+		Primary.Target = LocalPrimaryTransform.GetTranslation();
+		Secondary.Target = LocalSecondaryTransform.GetTranslation();
+		return true;
+	}
+	return FRigUnit_HighlevelBaseMutable::UpdateDirectManipulationFromHierarchy(InNode, InInstance, InContext, InInfo);
+}
+
+#endif
+
 FRigUnit_AimItem_Execute()
 {
     DECLARE_SCOPE_HIERARCHICAL_COUNTER_RIGUNIT()
diff --git a/Engine/Plugins/Animation/ControlRig/Source/ControlRig/Private/Units/Highlevel/Hierarchy/RigUnit_AimBone.h b/Engine/Plugins/Animation/ControlRig/Source/ControlRig/Private/Units/Highlevel/Hierarchy/RigUnit_AimBone.h
--- a/Engine/Plugins/Animation/ControlRig/Source/ControlRig/Private/Units/Highlevel/Hierarchy/RigUnit_AimBone.h
+++ b/Engine/Plugins/Animation/ControlRig/Source/ControlRig/Private/Units/Highlevel/Hierarchy/RigUnit_AimBone.h
@@ -224,19 +224,6 @@ struct CONTROLRIG_API FRigUnit_AimBone : public FRigUnit_HighlevelBaseMutable
 		bIsInitialized = false;
 	}
 
-	virtual FRigElementKey DetermineSpaceForPin(const FString& InPinPath, void* InUserContext) const override
-	{
-		if (InPinPath.StartsWith(TEXT("Primary.Target")))
-		{
-			return FRigElementKey(Primary.Space, ERigElementType::Bone);
-		}
-		if (InPinPath.StartsWith(TEXT("Secondary.Target")))
-		{
-			return FRigElementKey(Secondary.Space, ERigElementType::Bone);
-		}
-		return FRigElementKey();
-	}
-
 	RIGVM_METHOD()
 	virtual void Execute() override;
 
@@ -316,19 +303,13 @@ struct CONTROLRIG_API FRigUnit_AimItem: public FRigUnit_HighlevelBaseMutable
 		bIsInitialized = false;
 	}
 
-	virtual FRigElementKey DetermineSpaceForPin(const FString& InPinPath, void* InUserContext) const override
-	{
-		if (InPinPath.StartsWith(TEXT("Primary.Target")))
-		{
-			return Primary.Space;
-		}
-		if (InPinPath.StartsWith(TEXT("Secondary.Target")))
-		{
-			return Secondary.Space;
-		}
-		return FRigElementKey();
-	}
-
+#if WITH_EDITOR
+	virtual bool GetDirectManipulationTargets(const URigVMUnitNode* InNode, TSharedPtr<FStructOnScope> InInstance, URigHierarchy* InHierarchy, TArray<FRigDirectManipulationTarget>& InOutTargets, FString* OutFailureReason) const override;
+	virtual TArray<const URigVMPin*> GetPinsForDirectManipulation(const URigVMUnitNode* InNode, const FRigDirectManipulationTarget& InTarget) const override;
+	virtual bool UpdateHierarchyForDirectManipulation(const URigVMUnitNode* InNode, TSharedPtr<FStructOnScope> InInstance, FControlRigExecuteContext& InContext, TSharedPtr<FRigDirectManipulationInfo> InInfo) override;
+	virtual bool UpdateDirectManipulationFromHierarchy(const URigVMUnitNode* InNode, TSharedPtr<FStructOnScope> InInstance, FControlRigExecuteContext& InContext, TSharedPtr<FRigDirectManipulationInfo> InInfo) override;
+#endif
+	
 	RIGVM_METHOD()
 	virtual void Execute() override;
 
diff --git a/Engine/Plugins/Animation/ControlRig/Source/ControlRig/Private/Units/Highlevel/Hierarchy/RigUnit_DistributeRotation.h b/Engine/Plugins/Animation/ControlRig/Source/ControlRig/Private/Units/Highlevel/Hierarchy/RigUnit_DistributeRotation.h
--- a/Engine/Plugins/Animation/ControlRig/Source/ControlRig/Private/Units/Highlevel/Hierarchy/RigUnit_DistributeRotation.h
+++ b/Engine/Plugins/Animation/ControlRig/Source/ControlRig/Private/Units/Highlevel/Hierarchy/RigUnit_DistributeRotation.h
@@ -69,15 +69,6 @@ struct CONTROLRIG_API FRigUnit_DistributeRotation : public FRigUnit_HighlevelBas
 		bPropagateToChildren = true;
 	}
 
-	virtual FRigElementKey DetermineSpaceForPin(const FString& InPinPath, void* InUserContext) const override
-	{
-		if (InPinPath.StartsWith(TEXT("Rotations")))
-		{
-			return FRigElementKey(StartBone, ERigElementType::Bone);
-		}
-		return FRigElementKey();
-	}
-
 	RIGVM_METHOD()
 	virtual void Execute() override;
 
diff --git a/Engine/Plugins/Animation/ControlRig/Source/ControlRig/Private/Units/Highlevel/Hierarchy/RigUnit_ModifyBoneTransforms.h b/Engine/Plugins/Animation/ControlRig/Source/ControlRig/Private/Units/Highlevel/Hierarchy/RigUnit_ModifyBoneTransforms.h
--- a/Engine/Plugins/Animation/ControlRig/Source/ControlRig/Private/Units/Highlevel/Hierarchy/RigUnit_ModifyBoneTransforms.h
+++ b/Engine/Plugins/Animation/ControlRig/Source/ControlRig/Private/Units/Highlevel/Hierarchy/RigUnit_ModifyBoneTransforms.h
@@ -53,39 +53,6 @@ struct CONTROLRIG_API FRigUnit_ModifyBoneTransforms : public FRigUnit_HighlevelB
 		BoneToModify.Add(FRigUnit_ModifyBoneTransforms_PerBone());
 	}
 
-	virtual FRigElementKey DetermineSpaceForPin(const FString& InPinPath, void* InUserContext) const override
-	{
-		if (InPinPath.StartsWith(TEXT("BoneToModify")))
-		{
-			int32 Index = INDEX_NONE;
-			FString Left, Middle, Right;
-			if (InPinPath.Replace(TEXT("["), TEXT(".")).Split(TEXT("."), &Left, &Middle))
-			{
-				if (Middle.Replace(TEXT("]"), TEXT(".")).Split(TEXT("."), &Left, &Right))
-				{
-					Index = FCString::Atoi(*Left);
-				}
-			}
-
-			if (BoneToModify.IsValidIndex(Index))
-			{
-				if (Mode == EControlRigModifyBoneMode::AdditiveLocal)
-				{
-					return FRigElementKey(BoneToModify[Index].Bone, ERigElementType::Bone);
-				}
-
-				if (Mode == EControlRigModifyBoneMode::OverrideLocal)
-				{
-					if (const URigHierarchy* Hierarchy = (const URigHierarchy*)InUserContext)
-					{
-						return Hierarchy->GetFirstParent(FRigElementKey(BoneToModify[Index].Bone, ERigElementType::Bone));
-					}
-				}
-			}
-		}
-		return FRigElementKey();
-	}
-
 	RIGVM_METHOD()
 	virtual void Execute() override;
 
diff --git a/Engine/Plugins/Animation/ControlRig/Source/ControlRig/Private/Units/Highlevel/Hierarchy/RigUnit_ModifyTransforms.cpp b/Engine/Plugins/Animation/ControlRig/Source/ControlRig/Private/Units/Highlevel/Hierarchy/RigUnit_ModifyTransforms.cpp
--- a/Engine/Plugins/Animation/ControlRig/Source/ControlRig/Private/Units/Highlevel/Hierarchy/RigUnit_ModifyTransforms.cpp
+++ b/Engine/Plugins/Animation/ControlRig/Source/ControlRig/Private/Units/Highlevel/Hierarchy/RigUnit_ModifyTransforms.cpp
@@ -95,3 +95,139 @@ FRigUnit_ModifyTransforms_Execute()
 	}
 }
 
+#if WITH_EDITOR
+
+bool FRigUnit_ModifyTransforms::GetDirectManipulationTargets(const URigVMUnitNode* InNode, TSharedPtr<FStructOnScope> InInstance, URigHierarchy* InHierarchy, TArray<FRigDirectManipulationTarget>& InOutTargets, FString* OutFailureReason) const
+{
+	for(const FRigUnit_ModifyTransforms_PerItem& ItemInfo : ItemToModify)
+	{
+		if(ItemInfo.Item.IsValid())
+		{
+			static const UEnum* TypeEnum = StaticEnum<ERigElementType>();
+			const FString Prefix = TypeEnum->GetDisplayNameTextByValue((int64)ItemInfo.Item.Type).ToString();
+			InOutTargets.Emplace(FString::Printf(TEXT("%s %s"), *Prefix, *ItemInfo.Item.Name.ToString()), ERigControlType::EulerTransform);
+		}
+	}
+	return !InOutTargets.IsEmpty();
+}
+
+bool FRigUnit_ModifyTransforms::UpdateHierarchyForDirectManipulation(const URigVMUnitNode* InNode, TSharedPtr<FStructOnScope> InInstance, FControlRigExecuteContext& InContext, TSharedPtr<FRigDirectManipulationInfo> InInfo)
+{
+	URigHierarchy* Hierarchy = InContext.Hierarchy;
+	if (Hierarchy == nullptr)
+	{
+		return false;
+	}
+
+	const int32 Index = GetIndexFromTarget(InInfo->Target.Name);
+	if (ItemToModify.IsValidIndex(Index))
+	{
+		if(!InInfo->bInitialized)
+		{
+			InInfo->OffsetTransform = FTransform::Identity;
+
+			const FRigElementKey FirstParent = Hierarchy->GetFirstParent(ItemToModify[Index].Item);
+
+			switch(Mode)
+			{
+				case EControlRigModifyBoneMode::AdditiveLocal:
+				case EControlRigModifyBoneMode::AdditiveGlobal:
+				{
+					InInfo->OffsetTransform = Hierarchy->GetGlobalTransform(ItemToModify[Index].Item);
+					break;
+				}
+				case EControlRigModifyBoneMode::OverrideLocal:
+				{
+					InInfo->OffsetTransform = Hierarchy->GetGlobalTransform(FirstParent);
+					break;
+				}
+				case EControlRigModifyBoneMode::OverrideGlobal:
+				default:
+				{
+					break;
+				}
+			}
+
+		}
+
+		Hierarchy->SetControlOffsetTransform(InInfo->ControlKey, InInfo->OffsetTransform, false);
+		Hierarchy->SetLocalTransform(InInfo->ControlKey, ItemToModify[Index].Transform, false);
+		if(!InInfo->bInitialized)
+		{
+			Hierarchy->SetLocalTransform(InInfo->ControlKey, ItemToModify[Index].Transform, true);
+		}
+		return true;
+	}
+	return false;
+}
+
+bool FRigUnit_ModifyTransforms::UpdateDirectManipulationFromHierarchy(const URigVMUnitNode* InNode, TSharedPtr<FStructOnScope> InInstance, FControlRigExecuteContext& InContext, TSharedPtr<FRigDirectManipulationInfo> InInfo)
+{
+	URigHierarchy* Hierarchy = InContext.Hierarchy;
+	if (Hierarchy == nullptr)
+	{
+		return false;
+	}
+
+	const int32 Index = GetIndexFromTarget(InInfo->Target.Name);
+	if (ItemToModify.IsValidIndex(Index))
+	{
+		ItemToModify[Index].Transform = Hierarchy->GetLocalTransform(InInfo->ControlKey);;
+		return true;
+	}
+	return false;
+}
+
+TArray<const URigVMPin*> FRigUnit_ModifyTransforms::GetPinsForDirectManipulation(const URigVMUnitNode* InNode, const FRigDirectManipulationTarget& InTarget) const
+{
+	TArray<const URigVMPin*> AffectedPins;
+	const int32 Index = GetIndexFromTarget(InTarget.Name);
+	if (ItemToModify.IsValidIndex(Index))
+	{
+		if(const URigVMPin* ItemToModifyArrayPin = InNode->FindPin(GET_MEMBER_NAME_STRING_CHECKED(FRigUnit_ModifyTransforms, ItemToModify)))
+		{
+			if(const URigVMPin* ItemToModifyElementPin = ItemToModifyArrayPin->FindSubPin(FString::FromInt(Index)))
+			{
+				if(const URigVMPin* TransformPin = ItemToModifyElementPin->FindSubPin(GET_MEMBER_NAME_STRING_CHECKED(FRigUnit_ModifyTransforms_PerItem, Transform)))
+				{
+					AffectedPins.Add(TransformPin);
+				}
+			}
+		}
+	}
+	return AffectedPins;
+}
+
+int32 FRigUnit_ModifyTransforms::GetIndexFromTarget(const FString& InTarget) const
+{
+	FString Left, Right;
+	if(InTarget.Split(TEXT(" "), &Left, &Right))
+	{
+		static const UEnum* TypeEnum = StaticEnum<ERigElementType>();
+		static TArray<FString> DisplayNames;
+		if(DisplayNames.IsEmpty())
+		{
+			for(int64 Index = 0; Index < (int64)ERigElementType::Last; Index++)
+			{
+				DisplayNames.Add(TypeEnum->GetDisplayNameTextByValue(Index).ToString());
+			}
+		}
+
+		const int32 TypeIndex = DisplayNames.Find(Left);
+		if(TypeIndex != INDEX_NONE)
+		{
+			const ERigElementType ElementType = (ERigElementType)TypeIndex;
+			const FName ElementName(*Right);
+			for(int32 Index = 0; Index < ItemToModify.Num(); Index++)
+			{
+				if(ItemToModify[Index].Item == FRigElementKey(ElementName, ElementType))
+				{
+					return Index;
+				}
+			}
+		}
+	}
+	return INDEX_NONE;
+}
+
+#endif
\ No newline at end of file
diff --git a/Engine/Plugins/Animation/ControlRig/Source/ControlRig/Private/Units/Highlevel/Hierarchy/RigUnit_ModifyTransforms.h b/Engine/Plugins/Animation/ControlRig/Source/ControlRig/Private/Units/Highlevel/Hierarchy/RigUnit_ModifyTransforms.h
--- a/Engine/Plugins/Animation/ControlRig/Source/ControlRig/Private/Units/Highlevel/Hierarchy/RigUnit_ModifyTransforms.h
+++ b/Engine/Plugins/Animation/ControlRig/Source/ControlRig/Private/Units/Highlevel/Hierarchy/RigUnit_ModifyTransforms.h
@@ -81,74 +81,14 @@ struct CONTROLRIG_API FRigUnit_ModifyTransforms : public FRigUnit_HighlevelBaseM
 		ItemToModify[0].Item = FRigElementKey(NAME_None, ERigElementType::Bone);
 	}
 
-	virtual FRigElementKey DetermineSpaceForPin(const FString& InPinPath, void* InUserContext) const override
-	{
-		if (InPinPath.StartsWith(TEXT("ItemToModify")))
-		{
-			int32 Index = INDEX_NONE;
-			FString Left, Middle, Right;
-			if (InPinPath.Replace(TEXT("["), TEXT(".")).Split(TEXT("."), &Left, &Middle))
-			{
-				if (Middle.Replace(TEXT("]"), TEXT(".")).Split(TEXT("."), &Left, &Right))
-				{
-					Index = FCString::Atoi(*Left);
-				}
-			}
-
-			if (ItemToModify.IsValidIndex(Index))
-			{
-				if (Mode == EControlRigModifyBoneMode::AdditiveLocal || Mode == EControlRigModifyBoneMode::OverrideLocal)
-				{
-					if (const URigHierarchy* Hierarchy = (const URigHierarchy*)InUserContext)
-					{
-						return Hierarchy->GetFirstParent(ItemToModify[Index].Item);
-					}
-				}
-			}
-		}
-		return FRigElementKey();
-	}
-
-	virtual FTransform DetermineOffsetTransformForPin(const FString& InPinPath, void* InUserContext) const override
-	{
-		if (InPinPath.StartsWith(TEXT("ItemToModify")))
-		{
-			int32 Index = INDEX_NONE;
-			FString Left, Middle, Right;
-			if (InPinPath.Replace(TEXT("["), TEXT(".")).Split(TEXT("."), &Left, &Middle))
-			{
-				if (Middle.Replace(TEXT("]"), TEXT(".")).Split(TEXT("."), &Left, &Right))
-				{
-					Index = FCString::Atoi(*Left);
-				}
-			}
-
-			if (ItemToModify.IsValidIndex(Index))
-			{
-				if (ItemToModify[Index].Item.IsValid())
-				{
-					if (const URigHierarchy* Hierarchy = (const URigHierarchy*)InUserContext)
-					{
-						if (Mode == EControlRigModifyBoneMode::AdditiveLocal)
-						{
-							// in the case of AdditiveLocal, we want the transient control's global transform to match the item's global transform
-							// we know: transient control's global = transform value * offset transform * parent transform
-							// we know: item's global transform = item's local transform from rig graph evaluation * its parent transform
-							// we ensure: the item and the transient control share the same parent
-							// so from our goal, transient control's global == item's global transform
-							// we get: transform value * offset transform * parent transform == item's local transform from rig graph evaluation * its parent transform
-							// we get: transform value * offset transform == item's local transform from rig graph evaluation
-							// we get: offset transform == transform value.inverse * item's local transform from rig graph evaluation
-							return ItemToModify[Index].Transform.Inverse() * Hierarchy->GetLocalTransform(ItemToModify[Index].Item);
-						}
-					}
-				}
-			}
-		}
-
-		return FTransform::Identity;
-	}
-
+#if WITH_EDITOR
+	virtual bool GetDirectManipulationTargets(const URigVMUnitNode* InNode, TSharedPtr<FStructOnScope> InInstance, URigHierarchy* InHierarchy, TArray<FRigDirectManipulationTarget>& InOutTargets, FString* OutFailureReason) const override;
+	virtual bool UpdateHierarchyForDirectManipulation(const URigVMUnitNode* InNode, TSharedPtr<FStructOnScope> InInstance, FControlRigExecuteContext& InContext, TSharedPtr<FRigDirectManipulationInfo> InInfo) override;
+	virtual bool UpdateDirectManipulationFromHierarchy(const URigVMUnitNode* InNode, TSharedPtr<FStructOnScope> InInstance, FControlRigExecuteContext& InContext, TSharedPtr<FRigDirectManipulationInfo> InInfo) override;
+	virtual TArray<const URigVMPin*> GetPinsForDirectManipulation(const URigVMUnitNode* InNode, const FRigDirectManipulationTarget& InTarget) const override;
+	int32 GetIndexFromTarget(const FString& InTarget) const;
+#endif
+	
 	RIGVM_METHOD()
 	virtual void Execute() override;
 
@@ -187,3 +127,4 @@ struct CONTROLRIG_API FRigUnit_ModifyTransforms : public FRigUnit_HighlevelBaseM
 	UPROPERTY(transient)
 	FRigUnit_ModifyTransforms_WorkData WorkData;
 };
+
diff --git a/Engine/Plugins/Animation/ControlRig/Source/ControlRig/Private/Units/Highlevel/Hierarchy/RigUnit_SpringIK.cpp b/Engine/Plugins/Animation/ControlRig/Source/ControlRig/Private/Units/Highlevel/Hierarchy/RigUnit_SpringIK.cpp
--- a/Engine/Plugins/Animation/ControlRig/Source/ControlRig/Private/Units/Highlevel/Hierarchy/RigUnit_SpringIK.cpp
+++ b/Engine/Plugins/Animation/ControlRig/Source/ControlRig/Private/Units/Highlevel/Hierarchy/RigUnit_SpringIK.cpp
@@ -5,6 +5,49 @@
 
 #include UE_INLINE_GENERATED_CPP_BY_NAME(RigUnit_SpringIK)
 
+#if WITH_EDITOR
+
+bool FRigUnit_SpringIK::UpdateHierarchyForDirectManipulation(const URigVMUnitNode* InNode, TSharedPtr<FStructOnScope> InInstance, FControlRigExecuteContext& InContext, TSharedPtr<FRigDirectManipulationInfo> InInfo)
+{
+	URigHierarchy* Hierarchy = InContext.Hierarchy;
+	if (Hierarchy == nullptr)
+	{
+		return false;
+	}
+
+	if(InInfo->Target.Name == GET_MEMBER_NAME_STRING_CHECKED(FRigUnit_SpringIK, PoleVector))
+	{
+		const FRigElementKey PoleVectorSpaceKey(PoleVectorSpace, ERigElementType::Bone);
+		const FTransform ParentTransform = Hierarchy->GetGlobalTransform(PoleVectorSpaceKey, false);
+		Hierarchy->SetControlOffsetTransform(InInfo->ControlKey, ParentTransform, false);
+		Hierarchy->SetLocalTransform(InInfo->ControlKey, FTransform(PoleVector), false);
+		if(!InInfo->bInitialized)
+		{
+			Hierarchy->SetLocalTransform(InInfo->ControlKey, FTransform(PoleVector), true);
+		}
+		return true;
+	}
+	return FRigUnit_HighlevelBaseMutable::UpdateHierarchyForDirectManipulation(InNode, InInstance, InContext, InInfo);
+}
+
+bool FRigUnit_SpringIK::UpdateDirectManipulationFromHierarchy(const URigVMUnitNode* InNode, TSharedPtr<FStructOnScope> InInstance, FControlRigExecuteContext& InContext, TSharedPtr<FRigDirectManipulationInfo> InInfo)
+{
+	URigHierarchy* Hierarchy = InContext.Hierarchy;
+	if (Hierarchy == nullptr)
+	{
+		return false;
+	}
+
+	if(InInfo->Target.Name == GET_MEMBER_NAME_STRING_CHECKED(FRigUnit_SpringIK, PoleVector))
+	{
+		PoleVector = Hierarchy->GetLocalTransform(InInfo->ControlKey, false).GetTranslation();
+		return true;
+	}
+	return FRigUnit_HighlevelBaseMutable::UpdateDirectManipulationFromHierarchy(InNode, InInstance, InContext, InInfo);
+}
+
+#endif
+
 FRigUnit_SpringIK_Execute()
 {
     DECLARE_SCOPE_HIERARCHICAL_COUNTER_RIGUNIT()
diff --git a/Engine/Plugins/Animation/ControlRig/Source/ControlRig/Private/Units/Highlevel/Hierarchy/RigUnit_SpringIK.h b/Engine/Plugins/Animation/ControlRig/Source/ControlRig/Private/Units/Highlevel/Hierarchy/RigUnit_SpringIK.h
--- a/Engine/Plugins/Animation/ControlRig/Source/ControlRig/Private/Units/Highlevel/Hierarchy/RigUnit_SpringIK.h
+++ b/Engine/Plugins/Animation/ControlRig/Source/ControlRig/Private/Units/Highlevel/Hierarchy/RigUnit_SpringIK.h
@@ -96,16 +96,11 @@ struct CONTROLRIG_API FRigUnit_SpringIK : public FRigUnit_HighlevelBaseMutable
 		DebugSettings = FRigUnit_SpringIK_DebugSettings();
 	}
 
-
-	virtual FRigElementKey DetermineSpaceForPin(const FString& InPinPath, void* InUserContext) const override
-	{
-		if (InPinPath.StartsWith(TEXT("PoleVector")))
-		{
-			return FRigElementKey(PoleVectorSpace, ERigElementType::Bone);
-		}
-		return FRigElementKey();
-	}
-
+#if WITH_EDITOR
+	virtual bool UpdateHierarchyForDirectManipulation(const URigVMUnitNode* InNode, TSharedPtr<FStructOnScope> InInstance, FControlRigExecuteContext& InContext, TSharedPtr<FRigDirectManipulationInfo> InInfo) override;
+	virtual bool UpdateDirectManipulationFromHierarchy(const URigVMUnitNode* InNode, TSharedPtr<FStructOnScope> InInstance, FControlRigExecuteContext& InContext, TSharedPtr<FRigDirectManipulationInfo> InInfo) override;
+#endif
+	
 	RIGVM_METHOD()
 	virtual void Execute() override;
 
diff --git a/Engine/Plugins/Animation/ControlRig/Source/ControlRig/Private/Units/Highlevel/Hierarchy/RigUnit_TransformConstraint.h b/Engine/Plugins/Animation/ControlRig/Source/ControlRig/Private/Units/Highlevel/Hierarchy/RigUnit_TransformConstraint.h
--- a/Engine/Plugins/Animation/ControlRig/Source/ControlRig/Private/Units/Highlevel/Hierarchy/RigUnit_TransformConstraint.h
+++ b/Engine/Plugins/Animation/ControlRig/Source/ControlRig/Private/Units/Highlevel/Hierarchy/RigUnit_TransformConstraint.h
@@ -58,26 +58,6 @@ struct CONTROLRIG_API FRigUnit_TransformConstraint : public FRigUnit_HighlevelBa
 		, bUseInitialTransforms(true)
 	{}
 
-	virtual FRigElementKey DetermineSpaceForPin(const FString& InPinPath, void* InUserContext) const override
-	{
-		if (InPinPath.StartsWith(TEXT("Targets")))
-		{
-			if (BaseTransformSpace == ETransformSpaceMode::BaseJoint)
-			{
-				return FRigElementKey(BaseBone, ERigElementType::Bone);
-			}
-
-			if (BaseTransformSpace == ETransformSpaceMode::LocalSpace)
-			{
-				if (const URigHierarchy* Hierarchy = (const URigHierarchy*)InUserContext)
-				{
-					return Hierarchy->GetFirstParent(FRigElementKey(Bone, ERigElementType::Bone));
-				}
-			}
-		}
-		return FRigElementKey();
-	}
-
 	RIGVM_METHOD()
 	virtual void Execute() override;
 
@@ -128,26 +108,6 @@ struct CONTROLRIG_API FRigUnit_TransformConstraintPerItem : public FRigUnit_High
 		Item = BaseItem = FRigElementKey(NAME_None, ERigElementType::Bone);
 	}
 
-	virtual FRigElementKey DetermineSpaceForPin(const FString& InPinPath, void* InUserContext) const override
-	{
-		if (InPinPath.StartsWith(TEXT("Targets")))
-		{
-			if (BaseTransformSpace == ETransformSpaceMode::BaseJoint)
-			{
-				return BaseItem;
-			}
-
-			if (BaseTransformSpace == ETransformSpaceMode::LocalSpace)
-			{
-				if (const URigHierarchy* Hierarchy = (const URigHierarchy*)InUserContext)
-				{
-					return Hierarchy->GetFirstParent(Item);
-				}
-			}
-		}
-		return FRigElementKey();
-	}
-
 	RIGVM_METHOD()
 	virtual void Execute() override;
 
diff --git a/Engine/Plugins/Animation/ControlRig/Source/ControlRig/Private/Units/Highlevel/Hierarchy/RigUnit_TwistBones.h b/Engine/Plugins/Animation/ControlRig/Source/ControlRig/Private/Units/Highlevel/Hierarchy/RigUnit_TwistBones.h
--- a/Engine/Plugins/Animation/ControlRig/Source/ControlRig/Private/Units/Highlevel/Hierarchy/RigUnit_TwistBones.h
+++ b/Engine/Plugins/Animation/ControlRig/Source/ControlRig/Private/Units/Highlevel/Hierarchy/RigUnit_TwistBones.h
@@ -40,11 +40,6 @@ struct CONTROLRIG_API FRigUnit_TwistBones : public FRigUnit_HighlevelBaseMutable
 		bPropagateToChildren = true;
 	}
 
-	virtual FRigElementKey DetermineSpaceForPin(const FString& InPinPath, void* InUserContext) const override
-	{
-		return FRigElementKey(StartBone, ERigElementType::Bone);
-	}
-
 	RIGVM_METHOD()
 	virtual void Execute() override;
 
diff --git a/Engine/Plugins/Animation/ControlRig/Source/ControlRig/Private/Units/Highlevel/Hierarchy/RigUnit_TwoBoneIKSimple.cpp b/Engine/Plugins/Animation/ControlRig/Source/ControlRig/Private/Units/Highlevel/Hierarchy/RigUnit_TwoBoneIKSimple.cpp
--- a/Engine/Plugins/Animation/ControlRig/Source/ControlRig/Private/Units/Highlevel/Hierarchy/RigUnit_TwoBoneIKSimple.cpp
+++ b/Engine/Plugins/Animation/ControlRig/Source/ControlRig/Private/Units/Highlevel/Hierarchy/RigUnit_TwoBoneIKSimple.cpp
@@ -66,6 +66,48 @@ FRigVMStructUpgradeInfo FRigUnit_TwoBoneIKSimple::GetUpgradeInfo() const
 	return Info;
 }
 
+#if WITH_EDITOR
+
+bool FRigUnit_TwoBoneIKSimplePerItem::UpdateHierarchyForDirectManipulation(const URigVMUnitNode* InNode, TSharedPtr<FStructOnScope> InInstance, FControlRigExecuteContext& InContext, TSharedPtr<FRigDirectManipulationInfo> InInfo)
+{
+	URigHierarchy* Hierarchy = InContext.Hierarchy;
+	if (Hierarchy == nullptr)
+	{
+		return false;
+	}
+
+	if(InInfo->Target.Name == GET_MEMBER_NAME_STRING_CHECKED(FRigUnit_TwoBoneIKSimplePerItem, PoleVector))
+	{
+		const FTransform ParentTransform = Hierarchy->GetParentTransform(PoleVectorSpace, false);
+		Hierarchy->SetControlOffsetTransform(InInfo->ControlKey, ParentTransform, false);
+		Hierarchy->SetLocalTransform(InInfo->ControlKey, FTransform(PoleVector), false);
+		if(!InInfo->bInitialized)
+		{
+			Hierarchy->SetLocalTransform(InInfo->ControlKey, FTransform(PoleVector), true);
+		}
+		return true;
+	}
+	return FRigUnit_HighlevelBaseMutable::UpdateHierarchyForDirectManipulation(InNode, InInstance, InContext, InInfo);
+}
+
+bool FRigUnit_TwoBoneIKSimplePerItem::UpdateDirectManipulationFromHierarchy(const URigVMUnitNode* InNode, TSharedPtr<FStructOnScope> InInstance, FControlRigExecuteContext& InContext, TSharedPtr<FRigDirectManipulationInfo> InInfo)
+{
+	URigHierarchy* Hierarchy = InContext.Hierarchy;
+	if (Hierarchy == nullptr)
+	{
+		return false;
+	}
+
+	if(InInfo->Target.Name == GET_MEMBER_NAME_STRING_CHECKED(FRigUnit_TwoBoneIKSimplePerItem, PoleVector))
+	{
+		PoleVector = Hierarchy->GetLocalTransform(InInfo->ControlKey, false).GetTranslation();
+		return true;
+	}
+	return FRigUnit_HighlevelBaseMutable::UpdateDirectManipulationFromHierarchy(InNode, InInstance, InContext, InInfo);
+}
+
+#endif
+
 FRigUnit_TwoBoneIKSimplePerItem_Execute()
 {
     DECLARE_SCOPE_HIERARCHICAL_COUNTER_RIGUNIT()
diff --git a/Engine/Plugins/Animation/ControlRig/Source/ControlRig/Private/Units/Highlevel/Hierarchy/RigUnit_TwoBoneIKSimple.h b/Engine/Plugins/Animation/ControlRig/Source/ControlRig/Private/Units/Highlevel/Hierarchy/RigUnit_TwoBoneIKSimple.h
--- a/Engine/Plugins/Animation/ControlRig/Source/ControlRig/Private/Units/Highlevel/Hierarchy/RigUnit_TwoBoneIKSimple.h
+++ b/Engine/Plugins/Animation/ControlRig/Source/ControlRig/Private/Units/Highlevel/Hierarchy/RigUnit_TwoBoneIKSimple.h
@@ -68,15 +68,6 @@ struct CONTROLRIG_API FRigUnit_TwoBoneIKSimple : public FRigUnit_HighlevelBaseMu
 		CachedPoleVectorSpaceIndex = FCachedRigElement();
 	}
 
-	virtual FRigElementKey DetermineSpaceForPin(const FString& InPinPath, void* InUserContext) const override
-	{
-		if (InPinPath.StartsWith(TEXT("PoleVector")))
-		{
-			return FRigElementKey(PoleVectorSpace, ERigElementType::Bone);
-		}
-		return FRigElementKey();
-	}
-
 	RIGVM_METHOD()
 	virtual void Execute() override;
 
@@ -238,14 +229,10 @@ struct CONTROLRIG_API FRigUnit_TwoBoneIKSimplePerItem : public FRigUnit_Highleve
 		CachedPoleVectorSpaceIndex = FCachedRigElement();
 	}
 
-	virtual FRigElementKey DetermineSpaceForPin(const FString& InPinPath, void* InUserContext) const override
-	{
-		if (InPinPath.StartsWith(TEXT("PoleVector")))
-		{
-			return PoleVectorSpace;
-		}
-		return FRigElementKey();
-	}
+#if WITH_EDITOR
+	virtual bool UpdateHierarchyForDirectManipulation(const URigVMUnitNode* InNode, TSharedPtr<FStructOnScope> InInstance, FControlRigExecuteContext& InContext, TSharedPtr<FRigDirectManipulationInfo> InInfo) override;
+	virtual bool UpdateDirectManipulationFromHierarchy(const URigVMUnitNode* InNode, TSharedPtr<FStructOnScope> InInstance, FControlRigExecuteContext& InContext, TSharedPtr<FRigDirectManipulationInfo> InInfo) override;
+#endif
 
 	RIGVM_METHOD()
 	virtual void Execute() override;
diff --git a/Engine/Plugins/Animation/ControlRig/Source/ControlRig/Private/Units/RigUnit.cpp b/Engine/Plugins/Animation/ControlRig/Source/ControlRig/Private/Units/RigUnit.cpp
--- a/Engine/Plugins/Animation/ControlRig/Source/ControlRig/Private/Units/RigUnit.cpp
+++ b/Engine/Plugins/Animation/ControlRig/Source/ControlRig/Private/Units/RigUnit.cpp
@@ -6,4 +6,277 @@
 
 DEFINE_LOG_CATEGORY_STATIC(LogRigUnit, Log, All);
 
+#if WITH_EDITOR
 
+bool FRigUnit::GetDirectManipulationTargets(const URigVMUnitNode* InNode, TSharedPtr<FStructOnScope> InInstance, URigHierarchy* InHierarchy, TArray<FRigDirectManipulationTarget>& InOutTargets, FString* OutFailureReason) const
+{
+	const UScriptStruct* ScriptStruct = InNode->GetScriptStruct();
+	if(ScriptStruct == nullptr)
+	{
+		if(OutFailureReason)
+		{
+			*OutFailureReason = TEXT("Node is not resolved yet.");
+		}
+		return false;
+	}
+	
+	for(const URigVMPin* Pin : InNode->GetPins())
+	{
+		(void)AddDirectManipulationTarget_Internal(InOutTargets, Pin, ScriptStruct);
+	}
+
+	return true;
+}
+
+bool FRigUnit::AddDirectManipulationTarget_Internal(TArray<FRigDirectManipulationTarget>& InOutTargets, const URigVMPin* InPin, const UScriptStruct* InScriptStruct)
+{
+	check(InPin);
+
+	if(InPin->GetDirection() == ERigVMPinDirection::Input ||
+		InPin->GetDirection() == ERigVMPinDirection::IO ||
+		InPin->GetDirection() == ERigVMPinDirection::Visible)
+	{
+		if(InPin->GetCPPTypeObject() == TBaseStructure<FTransform>::Get() ||
+			InPin->GetCPPTypeObject() == TBaseStructure<FEulerTransform>::Get() ||
+			InPin->GetCPPTypeObject() == TBaseStructure<FVector>::Get() ||
+			InPin->GetCPPTypeObject() == TBaseStructure<FQuat>::Get())
+		{
+			if(const URigVMPin* ParentPin = InPin->GetParentPin())
+			{
+				if(ParentPin->GetCPPTypeObject() == TBaseStructure<FTransform>::Get() ||
+					ParentPin->GetCPPTypeObject() == TBaseStructure<FEulerTransform>::Get())
+				{
+					return false;
+				}
+			}
+
+			ERigControlType ControlType = ERigControlType::EulerTransform;
+			if(InPin->GetCPPTypeObject() == TBaseStructure<FVector>::Get())
+			{
+				ControlType = ERigControlType::Position;
+			}
+			else if(InPin->GetCPPTypeObject() == TBaseStructure<FQuat>::Get())
+			{
+				ControlType = ERigControlType::Rotator;
+			}
+			InOutTargets.AddUnique({InPin->GetSegmentPath(true), ControlType});
+		}
+
+		for(const URigVMPin* SubPin : InPin->GetSubPins())
+		{
+			AddDirectManipulationTarget_Internal(InOutTargets, SubPin, InScriptStruct);
+		}
+	}
+
+	return false;
+}
+
+TTuple<const FProperty*, FRigVMPropertyPath> FRigUnit::FindPropertyFromPinPath(const UScriptStruct* InStruct, const FString& InPinPath)
+{
+	FString PinPath = InPinPath;
+	PinPath = PinPath.Replace(TEXT("["), TEXT("."));
+	PinPath = PinPath.Replace(TEXT("]"), TEXT("."));
+	PinPath = PinPath.Replace(TEXT(".."), TEXT("."));
+	PinPath.TrimCharInline('.', nullptr);
+
+	FString Left, Right;
+	if(!PinPath.Split(TEXT("."), &Left, &Right))
+	{
+		Left = PinPath;
+		Right.Reset();
+	}
+
+	const FProperty* Property = InStruct->FindPropertyByName(*Left);
+	if(Property == nullptr)
+	{
+		return {nullptr, FRigVMPropertyPath() };
+	}
+
+	if(Right.IsEmpty())
+	{
+		return {Property, FRigVMPropertyPath() };
+	}
+
+	return {Property, FRigVMPropertyPath(Property, Right) };
+}
+
+void FRigUnit::ConfigureDirectManipulationControl(const URigVMUnitNode* InNode, TSharedPtr<FRigDirectManipulationInfo> InInfo, FRigControlSettings& InOutSettings, FRigControlValue& InOutValue) const
+{
+	if(InInfo->Target.ControlType == ERigControlType::Transform)
+	{
+		InOutSettings.ControlType = ERigControlType::EulerTransform;
+		InOutValue = FRigControlValue::Make<FEulerTransform>(FEulerTransform::Identity);
+	}
+	else if(InInfo->Target.ControlType == ERigControlType::Position)
+	{
+		InOutSettings.ControlType = ERigControlType::Position;
+		InOutValue = FRigControlValue::Make<FVector>(FVector::ZeroVector);
+	}
+	else if(InInfo->Target.ControlType == ERigControlType::Rotator)
+	{
+		InOutSettings.ControlType = ERigControlType::Rotator;
+		InOutValue = FRigControlValue::Make<FRotator>(FRotator::ZeroRotator);
+	}
+}
+
+bool FRigUnit::UpdateHierarchyForDirectManipulation(const URigVMUnitNode* InNode, TSharedPtr<FStructOnScope> InInstance, FControlRigExecuteContext& InContext, TSharedPtr<FRigDirectManipulationInfo> InInfo)
+{
+	const UScriptStruct* Struct = InNode->GetScriptStruct();
+	check(Struct);
+	check(InInstance.IsValid() && InInstance->IsValid());
+	check(InInstance->GetStruct() == Struct);
+
+	if(!InInfo.IsValid())
+	{
+		return false;
+	}
+
+	auto PropertyAndPath = FindPropertyFromPinPath(Struct, InInfo->Target.Name);
+	const FStructProperty* StructProperty = CastField<FStructProperty>(PropertyAndPath.Get<0>());
+	if(StructProperty == nullptr)
+	{
+		return false;
+	}
+
+	const uint8* Memory = StructProperty->ContainerPtrToValuePtr<uint8>(InInstance->GetStructMemory());
+	if(PropertyAndPath.Get<1>().IsValid())
+	{
+		Memory = PropertyAndPath.Get<1>().GetData<uint8>((uint8*)Memory, StructProperty);
+	}
+
+	URigHierarchy* Hierarchy = InContext.Hierarchy;
+	if (Hierarchy == nullptr)
+	{
+		return false;
+	}
+
+	FTransform Transform = Hierarchy->GetGlobalTransform(InInfo->ControlKey, false);
+
+	if(StructProperty->Struct == TBaseStructure<FTransform>::Get())
+	{
+		const FTransform& Result = *(const FTransform*)Memory;
+		Transform = Result;
+	}
+	else if(StructProperty->Struct == TBaseStructure<FEulerTransform>::Get())
+	{
+		const FEulerTransform& Result = *(const FEulerTransform*)Memory;
+		Transform = Result.ToFTransform();
+	}
+	else if(StructProperty->Struct == TBaseStructure<FQuat>::Get())
+	{
+		const FQuat& Result = *(const FQuat*)Memory;
+		Transform.SetRotation(Result);
+	}
+	else if(StructProperty->Struct == TBaseStructure<FVector>::Get())
+	{
+		const FVector& Result = *(const FVector*)Memory;
+		Transform.SetTranslation(Result);
+	}
+	else
+	{
+		return false;
+	}
+
+	Hierarchy->SetControlOffsetTransform(InInfo->ControlKey, FTransform::Identity, false);
+	Hierarchy->SetGlobalTransform(InInfo->ControlKey, Transform, false);
+
+	if(!InInfo->bInitialized)
+	{
+		Hierarchy->SetGlobalTransform(InInfo->ControlKey, Transform, true);
+	}
+	return true;
+}
+
+bool FRigUnit::UpdateDirectManipulationFromHierarchy(const URigVMUnitNode* InNode, TSharedPtr<FStructOnScope> InInstance, FControlRigExecuteContext& InContext, TSharedPtr<FRigDirectManipulationInfo> InInfo)
+{
+	const UScriptStruct* Struct = InNode->GetScriptStruct();
+	check(Struct);
+	check(InInstance.IsValid() && InInstance->IsValid());
+	check(InInstance->GetStruct() == Struct);
+
+	if(!InInfo.IsValid())
+	{
+		return false;
+	}
+
+	auto PropertyAndPath = FindPropertyFromPinPath(Struct, InInfo->Target.Name);
+	const FStructProperty* StructProperty = CastField<FStructProperty>(PropertyAndPath.Get<0>());
+	if(StructProperty == nullptr)
+	{
+		return false;
+	}
+
+	uint8* Memory = StructProperty->ContainerPtrToValuePtr<uint8>(InInstance->GetStructMemory());
+	if(PropertyAndPath.Get<1>().IsValid())
+	{
+		Memory = PropertyAndPath.Get<1>().GetData<uint8>(Memory, StructProperty);
+	}
+
+	URigHierarchy* Hierarchy = InContext.Hierarchy;
+	if (Hierarchy == nullptr)
+	{
+		return false;
+	}
+	
+	const FTransform Transform = Hierarchy->GetLocalTransform(InInfo->ControlKey, false);
+
+	if(StructProperty->Struct == TBaseStructure<FTransform>::Get())
+	{
+		FTransform& Result = *(FTransform*)Memory;
+		Result = Transform;
+	}
+	else if(StructProperty->Struct == TBaseStructure<FEulerTransform>::Get())
+	{
+		FEulerTransform& Result = *(FEulerTransform*)Memory;
+		Result = FEulerTransform(Transform);
+	}
+	else if(StructProperty->Struct == TBaseStructure<FQuat>::Get())
+	{
+		FQuat& Result = *(FQuat*)Memory;
+		Result = Transform.GetRotation();
+	}
+	else if(StructProperty->Struct == TBaseStructure<FVector>::Get())
+	{
+		FVector& Result = *(FVector*)Memory;
+		Result = Transform.GetTranslation();
+	}
+	else
+	{
+		return false;
+	}
+	
+	return true;
+}
+
+TArray<const URigVMPin*> FRigUnit::GetPinsForDirectManipulation(const URigVMUnitNode* InNode, const FRigDirectManipulationTarget& InTarget) const
+{
+	// the default is to return the pin matching the target
+	TArray<const URigVMPin*> Pins;
+	if(const URigVMPin* Pin = InNode->FindPin(InTarget.Name))
+	{
+		Pins.Add(Pin);
+	}
+	return Pins;
+}
+
+void FRigUnit::PerformDebugDrawingForDirectManipulation(const URigVMUnitNode* InNode, TSharedPtr<FStructOnScope> InInstance, FControlRigExecuteContext& InContext, TSharedPtr<FRigDirectManipulationInfo> InInfo) const
+{
+	FRigVMDrawInterface* DrawInterface = InContext.GetDrawInterface();
+	if(DrawInterface == nullptr)
+	{
+		return;
+	}
+	
+	if(FRigControlElement* Control = InContext.Hierarchy->Find<FRigControlElement>(InInfo->ControlKey))
+	{
+		const FTransform Transform = InContext.Hierarchy->GetLocalTransform(InInfo->ControlKey, false);
+		const FVector Location = Transform.GetTranslation();
+		if(!Location.IsNearlyZero())
+		{
+			const FTransform OffsetTransform = InContext.Hierarchy->GetControlOffsetTransform(Control, ERigTransformType::CurrentGlobal);
+			DrawInterface->DrawLine(OffsetTransform, FVector::ZeroVector, Location, FLinearColor::Green);
+		}
+	}	
+}
+
+#endif
diff --git a/Engine/Plugins/Animation/ControlRig/Source/ControlRig/Public/ControlRig.h b/Engine/Plugins/Animation/ControlRig/Source/ControlRig/Public/ControlRig.h
--- a/Engine/Plugins/Animation/ControlRig/Source/ControlRig/Public/ControlRig.h
+++ b/Engine/Plugins/Animation/ControlRig/Source/ControlRig/Public/ControlRig.h
@@ -21,6 +21,7 @@
 
 #if WITH_EDITOR
 #include "RigVMModel/RigVMPin.h"
+#include "RigVMModel/Nodes/RigVMUnitNode.h"
 #include "RigVMTypeUtils.h"
 #endif
 
@@ -410,14 +411,14 @@ private:
 	void HandleHierarchyModified(ERigHierarchyNotification InNotification, URigHierarchy* InHierarchy, const FRigBaseElement* InElement);
 
 #if WITH_EDITOR
-	/** Remove a transient / temporary control used to interact with a pin */
-	FName AddTransientControl(URigVMPin* InPin, FRigElementKey SpaceKey = FRigElementKey(), FTransform OffsetTransform = FTransform::Identity);
+	/** Add a transient / temporary control used to interact with a node */
+	FName AddTransientControl(const URigVMUnitNode* InNode, const FRigDirectManipulationTarget& InTarget);
 
-	/** Sets the value of a transient control based on a pin */
-	bool SetTransientControlValue(URigVMPin* InPin);
+	/** Sets the value of a transient control based on a node */
+	bool SetTransientControlValue(const URigVMUnitNode* InNode, TSharedPtr<FRigDirectManipulationInfo> InInfo);
 
-	/** Remove a transient / temporary control used to interact with a pin */
-	FName RemoveTransientControl(URigVMPin* InPin);
+	/** Remove a transient / temporary control used to interact with a node */
+	FName RemoveTransientControl(const URigVMUnitNode* InNode, const FRigDirectManipulationTarget& InTarget);
 
 	FName AddTransientControl(const FRigElementKey& InElement);
 
@@ -428,9 +429,13 @@ private:
 	FName RemoveTransientControl(const FRigElementKey& InElement);
 
 	static FName GetNameForTransientControl(const FRigElementKey& InElement);
-	FName GetNameForTransientControl(URigVMPin* InPin) const;
-	static FString GetPinNameFromTransientControl(const FRigElementKey& InKey);
+	FName GetNameForTransientControl(const URigVMUnitNode* InNode, const FRigDirectManipulationTarget& InTarget) const;
+	static FString GetNodeNameFromTransientControl(const FRigElementKey& InKey);
+	static FString GetTargetFromTransientControl(const FRigElementKey& InKey);
+	TSharedPtr<FRigDirectManipulationInfo> GetRigUnitManipulationInfoForTransientControl(const FRigElementKey& InKey);
+	
 	static FRigElementKey GetElementKeyFromTransientControl(const FRigElementKey& InKey);
+	bool CanAddTransientControl(const URigVMUnitNode* InNode, const FRigDirectManipulationTarget& InTarget, FString* OutFailureReason);
 
 	/** Removes all  transient / temporary control used to interact with pins */
 	void ClearTransientControls();
@@ -489,6 +494,7 @@ public:
 #endif
 
 	float GetDebugBoneRadiusMultiplier() const { return DebugBoneRadiusMultiplier; }
+	static FRigUnit* GetRigUnitInstanceFromScope(TSharedPtr<FStructOnScope> InScope);
 
 public:
 	//~ Begin IInterface_AssetUserData Interface
@@ -520,6 +526,9 @@ protected:
 	int32 ControlUndoBracketIndex;
 	uint8 InteractionType;
 	TArray<FRigElementKey> ElementsBeingInteracted;
+#if WITH_EDITOR
+	TArray<TSharedPtr<FRigDirectManipulationInfo>> RigUnitManipulationInfos;
+#endif
 	bool bInteractionJustBegan;
 
 	TWeakObjectPtr<USceneComponent> OuterSceneComponent;
@@ -732,6 +741,7 @@ private:
 	friend class AControlRigControlActor;
 	friend class AControlRigShapeActor;
 	friend class FRigTransformElementDetails;
+	friend class FControlRigEditorModule;
 };
 
 class CONTROLRIG_API FControlRigBracketScope
diff --git a/Engine/Plugins/Animation/ControlRig/Source/ControlRig/Public/Units/Hierarchy/RigUnit_AddBoneTransform.h b/Engine/Plugins/Animation/ControlRig/Source/ControlRig/Public/Units/Hierarchy/RigUnit_AddBoneTransform.h
--- a/Engine/Plugins/Animation/ControlRig/Source/ControlRig/Public/Units/Hierarchy/RigUnit_AddBoneTransform.h
+++ b/Engine/Plugins/Animation/ControlRig/Source/ControlRig/Public/Units/Hierarchy/RigUnit_AddBoneTransform.h
@@ -21,15 +21,6 @@ struct FRigUnit_AddBoneTransform : public FRigUnitMutable
 		, CachedBone()
 	{}
 
-	virtual FRigElementKey DetermineSpaceForPin(const FString& InPinPath, void* InUserContext) const override
-	{
-		if (InPinPath.StartsWith(TEXT("Transform")))
-		{
-			return FRigElementKey(Bone, ERigElementType::Bone);
-		}
-		return FRigElementKey();
-	}
-
 	RIGVM_METHOD()
 	virtual void Execute() override;
 
diff --git a/Engine/Plugins/Animation/ControlRig/Source/ControlRig/Public/Units/Hierarchy/RigUnit_SetTransform.h b/Engine/Plugins/Animation/ControlRig/Source/ControlRig/Public/Units/Hierarchy/RigUnit_SetTransform.h
--- a/Engine/Plugins/Animation/ControlRig/Source/ControlRig/Public/Units/Hierarchy/RigUnit_SetTransform.h
+++ b/Engine/Plugins/Animation/ControlRig/Source/ControlRig/Public/Units/Hierarchy/RigUnit_SetTransform.h
@@ -29,18 +29,10 @@ struct CONTROLRIG_API FRigUnit_SetTransform : public FRigUnitMutable
 
 	virtual FString GetUnitLabel() const override;
 
-	virtual FRigElementKey DetermineSpaceForPin(const FString& InPinPath, void* InUserContext) const override
-	{
-		if(Space == ERigVMTransformSpace::LocalSpace)
-		{
-			if (const URigHierarchy* Hierarchy = (const URigHierarchy*)InUserContext)
-			{
-				return Hierarchy->GetFirstParent(Item);
-			}
-		}
-		return FRigElementKey();
-	}
-
+#if WITH_EDITOR
+	virtual bool UpdateHierarchyForDirectManipulation(const URigVMUnitNode* InNode, TSharedPtr<FStructOnScope> InInstance, FControlRigExecuteContext& InContext, TSharedPtr<FRigDirectManipulationInfo> InInfo) override;
+	virtual bool UpdateDirectManipulationFromHierarchy(const URigVMUnitNode* InNode, TSharedPtr<FStructOnScope> InInstance, FControlRigExecuteContext& InContext, TSharedPtr<FRigDirectManipulationInfo> InInfo) override;
+#endif
 
 	RIGVM_METHOD()
 	virtual void Execute() override;
@@ -102,18 +94,11 @@ struct CONTROLRIG_API FRigUnit_SetTranslation : public FRigUnitMutable
 
 	virtual FString GetUnitLabel() const override;
 
-	virtual FRigElementKey DetermineSpaceForPin(const FString& InPinPath, void* InUserContext) const override
-	{
-		if(Space == ERigVMTransformSpace::LocalSpace)
-		{
-			if (const URigHierarchy* Hierarchy = (const URigHierarchy*)InUserContext)
-			{
-				return Hierarchy->GetFirstParent(Item);
-			}
-		}
-		return FRigElementKey();
-	}
-
+#if WITH_EDITOR
+	virtual bool UpdateHierarchyForDirectManipulation(const URigVMUnitNode* InNode, TSharedPtr<FStructOnScope> InInstance, FControlRigExecuteContext& InContext, TSharedPtr<FRigDirectManipulationInfo> InInfo) override;
+	virtual bool UpdateDirectManipulationFromHierarchy(const URigVMUnitNode* InNode, TSharedPtr<FStructOnScope> InInstance, FControlRigExecuteContext& InContext, TSharedPtr<FRigDirectManipulationInfo> InInfo) override;
+#endif
+	
 	RIGVM_METHOD()
 	virtual void Execute() override;
 
@@ -174,18 +159,11 @@ struct CONTROLRIG_API FRigUnit_SetRotation : public FRigUnitMutable
 
 	virtual FString GetUnitLabel() const override;
 
-	virtual FRigElementKey DetermineSpaceForPin(const FString& InPinPath, void* InUserContext) const override
-	{
-		if(Space == ERigVMTransformSpace::LocalSpace)
-		{
-			if (const URigHierarchy* Hierarchy = (const URigHierarchy*)InUserContext)
-			{
-				return Hierarchy->GetFirstParent(Item);
-			}
-		}
-		return FRigElementKey();
-	}
-
+#if WITH_EDITOR
+	virtual bool UpdateHierarchyForDirectManipulation(const URigVMUnitNode* InNode, TSharedPtr<FStructOnScope> InInstance, FControlRigExecuteContext& InContext, TSharedPtr<FRigDirectManipulationInfo> InInfo) override;
+	virtual bool UpdateDirectManipulationFromHierarchy(const URigVMUnitNode* InNode, TSharedPtr<FStructOnScope> InInstance, FControlRigExecuteContext& InContext, TSharedPtr<FRigDirectManipulationInfo> InInfo) override;
+#endif
+	
 	RIGVM_METHOD()
 	virtual void Execute() override;
 
@@ -246,6 +224,11 @@ struct FRigUnit_SetScale : public FRigUnitMutable
 
 	virtual FString GetUnitLabel() const override;
 
+#if WITH_EDITOR
+	virtual bool UpdateHierarchyForDirectManipulation(const URigVMUnitNode* InNode, TSharedPtr<FStructOnScope> InInstance, FControlRigExecuteContext& InContext, TSharedPtr<FRigDirectManipulationInfo> InInfo) override;
+	virtual bool UpdateDirectManipulationFromHierarchy(const URigVMUnitNode* InNode, TSharedPtr<FStructOnScope> InInstance, FControlRigExecuteContext& InContext, TSharedPtr<FRigDirectManipulationInfo> InInfo) override;
+#endif
+	
 	RIGVM_METHOD()
 	virtual void Execute() override;
 
diff --git a/Engine/Plugins/Animation/ControlRig/Source/ControlRig/Public/Units/RigUnit.h b/Engine/Plugins/Animation/ControlRig/Source/ControlRig/Public/Units/RigUnit.h
--- a/Engine/Plugins/Animation/ControlRig/Source/ControlRig/Public/Units/RigUnit.h
+++ b/Engine/Plugins/Animation/ControlRig/Source/ControlRig/Public/Units/RigUnit.h
@@ -7,9 +7,52 @@
 #include "RigVMCore/RigVMRegistry.h"
 #include "RigUnitContext.h"
 #include "RigVMFunctions/RigVMFunctionDefines.h"
+
+#if WITH_EDITOR
+#include "RigVMModel/Nodes/RigVMUnitNode.h"
+#endif
+
 #include "RigUnit.generated.h"
 
 struct FRigUnitContext;
+struct FRigDirectManipulationInfo;
+
+#if WITH_EDITOR
+
+struct CONTROLRIG_API FRigDirectManipulationTarget
+{
+	FRigDirectManipulationTarget()
+		: Name()
+		, ControlType(ERigControlType::EulerTransform)
+	{
+	}
+
+	FRigDirectManipulationTarget(const FString& InName, ERigControlType InControlType)
+		: Name(InName)
+		, ControlType(InControlType)
+	{
+	}
+
+	bool operator == (const FRigDirectManipulationTarget& InOther) const
+	{
+		return Name.Equals(InOther.Name, ESearchCase::CaseSensitive);
+	}
+
+	bool operator > (const FRigDirectManipulationTarget& InOther) const
+	{
+		return Name > InOther.Name;
+	}
+
+	bool operator < (const FRigDirectManipulationTarget& InOther) const
+	{
+		return Name < InOther.Name;
+	}
+
+	FString Name;
+	ERigControlType ControlType;
+};
+
+#endif
 
 /** Base class for all rig units */
 USTRUCT(BlueprintType, meta=(Abstract, NodeColor = "0.1 0.1 0.1", ExecuteContext="FControlRigExecuteContext"))
@@ -17,6 +60,8 @@ struct CONTROLRIG_API FRigUnit : public FRigVMStruct
 {
 	GENERATED_BODY()
 
+public:
+
 	FRigUnit()
 	{}
 
@@ -37,6 +82,31 @@ struct CONTROLRIG_API FRigUnit : public FRigVMStruct
 		}
 		return MethodName;
 	}
+
+#if WITH_EDITOR
+	/** Returns the targets for viewport posing */
+	virtual bool GetDirectManipulationTargets(const URigVMUnitNode* InNode, TSharedPtr<FStructOnScope> InInstance, URigHierarchy* InHierarchy, TArray<FRigDirectManipulationTarget>& InOutTargets, FString* OutFailureReason) const;
+
+	/** Optionally configures a control's settings and value for a given target */
+	virtual void ConfigureDirectManipulationControl(const URigVMUnitNode* InNode, TSharedPtr<FRigDirectManipulationInfo> InInfo, FRigControlSettings& InOutSettings, FRigControlValue& InOutValue) const;
+
+	/** Sets a control's pose to represent this viewport pose target */ 
+	virtual bool UpdateHierarchyForDirectManipulation(const URigVMUnitNode* InNode, TSharedPtr<FStructOnScope> InInstance, FControlRigExecuteContext& InContext, TSharedPtr<FRigDirectManipulationInfo> InInfo);
+
+	/** Sets the values on this node based on a viewport pose */
+	virtual bool UpdateDirectManipulationFromHierarchy(const URigVMUnitNode* InNode, TSharedPtr<FStructOnScope> InInstance, FControlRigExecuteContext& InContext, TSharedPtr<FRigDirectManipulationInfo> InInfo);
+
+	/** returns a list of pins affected by the viewport pose */
+	virtual TArray<const URigVMPin*> GetPinsForDirectManipulation(const URigVMUnitNode* InNode, const FRigDirectManipulationTarget& InTarget) const;
+
+	/** Allows the node to draw debug drawing during a manipulation */
+	virtual void PerformDebugDrawingForDirectManipulation(const URigVMUnitNode* InNode, TSharedPtr<FStructOnScope> InInstance, FControlRigExecuteContext& InContext, TSharedPtr<FRigDirectManipulationInfo> InInfo) const;
+
+private:
+	
+	static bool AddDirectManipulationTarget_Internal(TArray<FRigDirectManipulationTarget>& InOutTargets, const URigVMPin* InPin, const UScriptStruct* InScriptStruct);
+	static TTuple<const FProperty*, FRigVMPropertyPath> FindPropertyFromPinPath(const UScriptStruct* InStruct, const FString& InPinPath);
+#endif
 };
 
 /** Base class for all rig units that can change data */
@@ -55,3 +125,34 @@ struct CONTROLRIG_API FRigUnitMutable : public FRigUnit
 	UPROPERTY(DisplayName = "Execute", Transient, meta = (Input, Output))
 	FControlRigExecuteContext ExecuteContext;
 };
+
+#if WITH_EDITOR
+
+struct CONTROLRIG_API FRigDirectManipulationInfo
+{
+	FRigDirectManipulationInfo()
+		: bInitialized(false)
+		, Target()
+		, ControlKey(NAME_None, ERigElementType::Control)
+		, OffsetTransform(FTransform::Identity)
+	{
+		Reset();
+	}
+
+	void Reset()
+	{
+		bInitialized = false;
+		OffsetTransform = FTransform::Identity;
+		Transforms.Reset();
+		Transforms.Add(FTransform::Identity);
+	}
+
+	bool bInitialized;
+	FRigDirectManipulationTarget Target;
+	FRigElementKey ControlKey;
+	FTransform OffsetTransform;
+	TArray<FTransform> Transforms;
+	TWeakObjectPtr<const URigVMUnitNode> Node;
+};
+
+#endif
\ No newline at end of file
diff --git a/Engine/Plugins/Animation/ControlRig/Source/ControlRigDeveloper/Private/ControlRigBlueprint.cpp b/Engine/Plugins/Animation/ControlRig/Source/ControlRigDeveloper/Private/ControlRigBlueprint.cpp
--- a/Engine/Plugins/Animation/ControlRig/Source/ControlRigDeveloper/Private/ControlRigBlueprint.cpp
+++ b/Engine/Plugins/Animation/ControlRig/Source/ControlRigDeveloper/Private/ControlRigBlueprint.cpp
@@ -202,6 +202,12 @@ void UControlRigBlueprint::PreSave(FObjectPreSaveContext ObjectSaveContext)
 {
 	Super::PreSave(ObjectSaveContext);
 
+	// make sure to save the VM with high performance settings
+	// so that during cooking we reach small footprints.
+	// these settings may have changed during the user session.
+	VMCompileSettings.ASTSettings.bFoldAssignments = true;
+	VMCompileSettings.ASTSettings.bFoldLiterals = true;
+
 	bExposesAnimatableControls = false;
 	Hierarchy->ForEach<FRigControlElement>([this](FRigControlElement* ControlElement) -> bool
     {
@@ -501,7 +507,7 @@ const FControlRigShapeDefinition* UControlRigBlueprint::GetControlShapeByName(co
 	return UControlRigShapeLibrary::GetShapeByName(InName, ShapeLibraries, LibraryNameMap);
 }
 
-FName UControlRigBlueprint::AddTransientControl(URigVMPin* InPin)
+FName UControlRigBlueprint::AddTransientControl(const URigVMUnitNode* InNode, const FRigDirectManipulationTarget& InTarget)
 {
 	TUniquePtr<FControlValueScope> ValueScope;
 	if (!UControlRigEditorSettings::Get()->bResetControlsOnPinValueInteraction) // if we need to retain the controls
@@ -515,35 +521,6 @@ FName UControlRigBlueprint::AddTransientControl(URigVMPin* InPin)
 	URigVMBlueprintGeneratedClass* RigClass = GetRigVMBlueprintGeneratedClass();
 	UControlRig* CDO = Cast<UControlRig>(RigClass->GetDefaultObject(true /* create if needed */));
 
-	FRigElementKey SpaceKey;
-	FTransform OffsetTransform = FTransform::Identity;
-	if (URigVMUnitNode* UnitNode = Cast<URigVMUnitNode>(InPin->GetPinForLink()->GetNode()))
-	{
-		if (TSharedPtr<FStructOnScope> DefaultStructScope = UnitNode->ConstructStructInstance())
-		{
-			FRigUnit* DefaultStruct = (FRigUnit*)DefaultStructScope->GetStructMemory();
-
-			FString PinPath = InPin->GetPinForLink()->GetPinPath();
-			FString Left, Right;
-
-			if (URigVMPin::SplitPinPathAtStart(PinPath, Left, Right))
-			{
-				SpaceKey = DefaultStruct->DetermineSpaceForPin(Right, Hierarchy);
-				
-				URigHierarchy* RigHierarchy = Hierarchy;
-
-				// use the active rig instead of the CDO rig because we want to access the evaluation result of the rig graph
-				// to calculate the offset transform, for example take a look at RigUnit_ModifyTransform
-				if (UControlRig* RigBeingDebugged = Cast<UControlRig>(GetObjectBeingDebugged()))
-				{
-					RigHierarchy = RigBeingDebugged->GetHierarchy();
-				}
-				
-				OffsetTransform = DefaultStruct->DetermineOffsetTransformForPin(Right, RigHierarchy);
-			}
-		}
-	}
-
 	FName ReturnName = NAME_None;
 	TArray<UObject*> ArchetypeInstances;
 	CDO->GetArchetypeInstances(ArchetypeInstances);
@@ -552,7 +529,7 @@ FName UControlRigBlueprint::AddTransientControl(URigVMPin* InPin)
 		UControlRig* InstancedControlRig = Cast<UControlRig>(ArchetypeInstance);
 		if (InstancedControlRig)
 		{
-			FName ControlName = InstancedControlRig->AddTransientControl(InPin, SpaceKey, OffsetTransform);
+			FName ControlName = InstancedControlRig->AddTransientControl(InNode, InTarget);
 			if (ReturnName == NAME_None)
 			{
 				ReturnName = ControlName;
@@ -563,7 +540,7 @@ FName UControlRigBlueprint::AddTransientControl(URigVMPin* InPin)
 	return ReturnName;
 }
 
-FName UControlRigBlueprint::RemoveTransientControl(URigVMPin* InPin)
+FName UControlRigBlueprint::RemoveTransientControl(const URigVMUnitNode* InNode, const FRigDirectManipulationTarget& InTarget)
 {
 	TUniquePtr<FControlValueScope> ValueScope;
 	if (!UControlRigEditorSettings::Get()->bResetControlsOnPinValueInteraction) // if we need to retain the controls
@@ -582,7 +559,7 @@ FName UControlRigBlueprint::RemoveTransientControl(URigVMPin* InPin)
 		UControlRig* InstancedControlRig = Cast<UControlRig>(ArchetypeInstance);
 		if (InstancedControlRig)
 		{
-			FName Name = InstancedControlRig->RemoveTransientControl(InPin);
+			FName Name = InstancedControlRig->RemoveTransientControl(InNode, InTarget);
 			if (RemovedName == NAME_None)
 	{
 				RemovedName = Name;
@@ -1353,6 +1330,17 @@ void UControlRigBlueprint::HandleHierarchyModified(ERigHierarchyNotification InN
 			}
 			break;
 		}
+		case ERigHierarchyNotification::ElementDeselected:
+		{
+			if (const FRigControlElement* ControlElement = Cast<FRigControlElement>(InElement))
+			{
+				if (ControlElement->Settings.bIsTransientControl)
+				{
+					ClearTransientControls();
+				}
+			}
+			break;
+		}
 		default:
 		{
 			break;
diff --git a/Engine/Plugins/Animation/ControlRig/Source/ControlRigDeveloper/Public/ControlRigBlueprint.h b/Engine/Plugins/Animation/ControlRig/Source/ControlRigDeveloper/Public/ControlRigBlueprint.h
--- a/Engine/Plugins/Animation/ControlRig/Source/ControlRigDeveloper/Public/ControlRigBlueprint.h
+++ b/Engine/Plugins/Animation/ControlRig/Source/ControlRigDeveloper/Public/ControlRigBlueprint.h
@@ -120,10 +120,10 @@ public:
 
 #if WITH_EDITOR
 	/** Remove a transient / temporary control used to interact with a pin */
-	FName AddTransientControl(URigVMPin* InPin);
+	FName AddTransientControl(const URigVMUnitNode* InNode, const FRigDirectManipulationTarget& InTarget);
 
 	/** Remove a transient / temporary control used to interact with a pin */
-	FName RemoveTransientControl(URigVMPin* InPin);
+	FName RemoveTransientControl(const URigVMUnitNode* InNode, const FRigDirectManipulationTarget& InTarget);
 
 	/** Remove a transient / temporary control used to interact with a bone */
 	FName AddTransientControl(const FRigElementKey& InElement);
diff --git a/Engine/Plugins/Animation/ControlRig/Source/ControlRigEditor/Private/ControlRigEditorModule.cpp b/Engine/Plugins/Animation/ControlRig/Source/ControlRigEditor/Private/ControlRigEditorModule.cpp
--- a/Engine/Plugins/Animation/ControlRig/Source/ControlRigEditor/Private/ControlRigEditorModule.cpp
+++ b/Engine/Plugins/Animation/ControlRig/Source/ControlRigEditor/Private/ControlRigEditorModule.cpp
@@ -92,6 +92,7 @@
 #include "RigVMFunctions/Math/RigVMMathLibrary.h"
 #include "Constraints/ControlRigTransformableHandle.h"
 #include "Constraints/TransformConstraintChannelInterface.h"
+#include "Async/TaskGraphInterfaces.h"
 
 #define LOCTEXT_NAMESPACE "ControlRigEditorModule"
 
@@ -468,6 +469,8 @@ void FControlRigEditorModule::GetNodeContextMenuActions(URigVMBlueprint* RigVMBl
 		}
 	}
 
+	GetDirectManipulationMenuActions(RigVMBlueprint, ModelNode, nullptr, Menu);
+
 	if (RigElementsToSelect.Num() > 0)
 	{
 		FToolMenuSection& Section = Menu->AddSection("RigVMEditorContextMenuHierarchy", LOCTEXT("HierarchyHeader", "Hierarchy"));
@@ -614,27 +617,162 @@ void FControlRigEditorModule::GetNodeContextMenuActions(URigVMBlueprint* RigVMBl
 void FControlRigEditorModule::GetPinContextMenuActions(URigVMBlueprint* RigVMBlueprint, const UEdGraphPin* EdGraphPin, URigVMPin* ModelPin, UToolMenu* Menu) const
 {
 	FRigVMEditorModule::GetPinContextMenuActions(RigVMBlueprint, EdGraphPin, ModelPin, Menu);
+	GetDirectManipulationMenuActions(RigVMBlueprint, ModelPin->GetNode(), ModelPin, Menu);
+}
 
+void FControlRigEditorModule::GetDirectManipulationMenuActions(URigVMBlueprint* RigVMBlueprint, URigVMNode* InNode, URigVMPin* ModelPin, UToolMenu* Menu) const
+{
     // Add direct manipulation context menu entries
 	if(UControlRigBlueprint* ControlRigBlueprint = Cast<UControlRigBlueprint>(RigVMBlueprint))
 	{
-		if ((ModelPin->GetCPPType() == TEXT("FVector") ||
-			 ModelPin->GetCPPType() == TEXT("FQuat") ||
-			 ModelPin->GetCPPType() == TEXT("FTransform")) &&
-			(ModelPin->GetDirection() == ERigVMPinDirection::Input ||
-			 ModelPin->GetDirection() == ERigVMPinDirection::IO) &&
-			 ModelPin->GetPinForLink()->GetRootPin()->GetSourceLinks(true).Num() == 0)
+		UControlRig* DebuggedRig = Cast<UControlRig>(RigVMBlueprint->GetObjectBeingDebugged());
+		if(DebuggedRig == nullptr)
 		{
-			FToolMenuSection& Section = Menu->AddSection("RigVMEditorContextMenuControlPin", LOCTEXT("ControlPin", "Direct Manipulation"));
-			Section.AddMenuEntry(
-				"DirectManipControlPin",
-				LOCTEXT("DirectManipControlPin", "Control Pin Value"),
-				LOCTEXT("DirectManipControlPin_Tooltip", "Configures the pin for direct interaction in the viewport"),
-				FSlateIcon(),
-				FUIAction(FExecuteAction::CreateLambda([ControlRigBlueprint, ModelPin]() {
-					ControlRigBlueprint->AddTransientControl(ModelPin);
-				})
-			));
+			return;
+		}
+		
+		if(const URigVMUnitNode* UnitNode = Cast<URigVMUnitNode>(InNode))
+		{
+			if(!UnitNode->IsPartOfRuntime(DebuggedRig))
+			{
+				return;
+			}
+			
+			const UScriptStruct* ScriptStruct = UnitNode->GetScriptStruct();
+			if(ScriptStruct == nullptr)
+			{
+				return;
+			}
+
+			TSharedPtr<FStructOnScope> NodeInstance = UnitNode->ConstructStructInstance(false);
+			if(!NodeInstance.IsValid() || !NodeInstance->IsValid())
+			{
+				return;
+			}
+
+			const FRigUnit* UnitInstance = UControlRig::GetRigUnitInstanceFromScope(NodeInstance);
+			TArray<FRigDirectManipulationTarget> Targets;
+			if(UnitInstance->GetDirectManipulationTargets(UnitNode, NodeInstance, DebuggedRig->GetHierarchy(), Targets, nullptr))
+			{
+				if(ModelPin)
+				{
+					Targets.RemoveAll([ModelPin, UnitNode, UnitInstance](const FRigDirectManipulationTarget& Target) -> bool
+					{
+						const TArray<const URigVMPin*> AffectedPins = UnitInstance->GetPinsForDirectManipulation(UnitNode, Target);
+						return !AffectedPins.Contains(ModelPin);
+					});
+				}
+
+				if(!Targets.IsEmpty())
+				{
+					FToolMenuSection& Section = Menu->AddSection("RigVMEditorContextMenuControlNode", LOCTEXT("ControlNodeDirectManipulation", "Direct Manipulation"));
+
+					bool bHasPosition = false;
+					bool bHasRotation = false;
+					bool bHasScale = false;
+					
+					for(const FRigDirectManipulationTarget& Target : Targets)
+					{
+						auto IsSliced = [UnitNode]() -> bool
+						{
+							return UnitNode->IsWithinLoop();
+						};
+						
+						auto HasNoUnconstrainedAffectedPin = [NodeInstance, UnitNode, Target]() -> bool
+						{
+							const FRigUnit* UnitInstance = UControlRig::GetRigUnitInstanceFromScope(NodeInstance);
+							const TArray<const URigVMPin*> AffectedPins = UnitInstance->GetPinsForDirectManipulation(UnitNode, Target);
+
+							int32 NumAffectedPinsWithRootLinks = 0;
+							for(const URigVMPin* AffectedPin : AffectedPins)
+							{
+								if(!AffectedPin->GetRootPin()->GetSourceLinks().IsEmpty())
+								{
+									NumAffectedPinsWithRootLinks++;
+								}
+							}
+							return NumAffectedPinsWithRootLinks == AffectedPins.Num();
+						};
+
+						FText Suffix;
+						static const FText SuffixPosition = LOCTEXT("DirectManipulationPosition", " (W)"); 
+						static const FText SuffixRotation = LOCTEXT("DirectManipulationRotation", " (E)"); 
+						static const FText SuffixScale = LOCTEXT("DirectManipulationScale", " (R)"); 
+						TSharedPtr< const FUICommandInfo > CommandInfo;
+						if(!CommandInfo.IsValid() && !bHasPosition)
+						{
+							if(Target.ControlType == ERigControlType::EulerTransform || Target.ControlType == ERigControlType::Position)
+							{
+								CommandInfo = FControlRigEditorCommands::Get().RequestDirectManipulationPosition;
+								Suffix = SuffixPosition;
+								bHasPosition = true;
+							}
+						}
+						if(!CommandInfo.IsValid() && !bHasRotation)
+						{
+							if(Target.ControlType == ERigControlType::EulerTransform || Target.ControlType == ERigControlType::Rotator)
+							{
+								CommandInfo = FControlRigEditorCommands::Get().RequestDirectManipulationRotation;
+								Suffix = SuffixRotation;
+								bHasRotation = true;
+							}
+						}
+						if(!CommandInfo.IsValid() && !bHasScale)
+						{
+							if(Target.ControlType == ERigControlType::EulerTransform)
+							{
+								CommandInfo = FControlRigEditorCommands::Get().RequestDirectManipulationScale;
+								Suffix = SuffixScale;
+								bHasScale = true;
+							}
+						}
+
+						const FText Label = FText::Format(LOCTEXT("ControlNodeLabelFormat", "Manipulate {0}{1}"), FText::FromString(Target.Name), Suffix);
+						TAttribute<FText> ToolTipAttribute = TAttribute<FText>::CreateLambda([HasNoUnconstrainedAffectedPin, IsSliced, Target]() -> FText
+						{
+							if(HasNoUnconstrainedAffectedPin())
+							{
+								return FText::Format(LOCTEXT("ControlNodeLabelFormat_Tooltip_FullyConstrained", "The value of {0} cannot be manipulated, its pins have links fully constraining it."), FText::FromString(Target.Name));
+							}
+							if(IsSliced())
+							{
+								return FText::Format(LOCTEXT("ControlNodeLabelFormat_Tooltip_Sliced", "The value of {0} cannot be manipulated, the node is linked to a loop."), FText::FromString(Target.Name));
+							}
+							return FText::Format(LOCTEXT("ControlNodeLabelFormat_Tooltip", "Manipulate the value of {0} interactively"), FText::FromString(Target.Name));
+						});
+
+						FToolMenuEntry& MenuEntry = Section.AddMenuEntry(
+							*Target.Name,
+							Label,
+							ToolTipAttribute,
+							FSlateIcon(),
+							FUIAction(FExecuteAction::CreateLambda([ControlRigBlueprint, UnitNode, Target]() {
+
+								// disable literal folding for the moment
+								if(ControlRigBlueprint->VMCompileSettings.ASTSettings.bFoldLiterals)
+								{
+									ControlRigBlueprint->VMCompileSettings.ASTSettings.bFoldLiterals = false;
+									ControlRigBlueprint->RecompileVM();
+								}
+
+								// run the task after a bit so that the rig has the opportunity to run first
+								FFunctionGraphTask::CreateAndDispatchWhenReady([ControlRigBlueprint, UnitNode, Target]()
+								{
+									ControlRigBlueprint->AddTransientControl(UnitNode, Target);
+								}, TStatId(), NULL, ENamedThreads::GameThread);
+							}),
+							FCanExecuteAction::CreateLambda([HasNoUnconstrainedAffectedPin, IsSliced]() -> bool
+							{
+								if(HasNoUnconstrainedAffectedPin() || IsSliced())
+								{
+									return false;
+								}
+								return true;
+							})
+						));
+					}
+				}
+			}
 		}
 	}
 }
diff --git a/Engine/Plugins/Animation/ControlRig/Source/ControlRigEditor/Private/ControlRigEditorModule.h b/Engine/Plugins/Animation/ControlRig/Source/ControlRigEditor/Private/ControlRigEditorModule.h
--- a/Engine/Plugins/Animation/ControlRig/Source/ControlRigEditor/Private/ControlRigEditorModule.h
+++ b/Engine/Plugins/Animation/ControlRig/Source/ControlRigEditor/Private/ControlRigEditorModule.h
@@ -48,6 +48,8 @@ public:
 	static void UnLinkLevelSequence(UAnimSequence* AnimSequence);
 	void ExtendAnimSequenceMenu();
 
+	void GetDirectManipulationMenuActions(URigVMBlueprint* RigVMBlueprint, URigVMNode* InNode, URigVMPin* ModelPin, UToolMenu* Menu) const;
+
 private:
 
 	/** Handle for our sequencer control rig parameter track editor */
diff --git a/Engine/Plugins/Animation/ControlRig/Source/ControlRigEditor/Private/EditMode/ControlRigEditMode.cpp b/Engine/Plugins/Animation/ControlRig/Source/ControlRigEditor/Private/EditMode/ControlRigEditMode.cpp
--- a/Engine/Plugins/Animation/ControlRig/Source/ControlRigEditor/Private/EditMode/ControlRigEditMode.cpp
+++ b/Engine/Plugins/Animation/ControlRig/Source/ControlRigEditor/Private/EditMode/ControlRigEditMode.cpp
@@ -585,7 +585,16 @@ void FControlRigEditMode::Tick(FEditorViewportClient* ViewportClient, float Delt
 				if (UControlRig* ControlRig = RuntimeRigPtr.Get())
 				{
 					TArray<FRigElementKey> SelectedRigElements = GetSelectedRigElements(ControlRig);
-					const UE::Widget::EWidgetMode CurrentWidgetMode = ViewportClient->GetWidgetMode();
+					UE::Widget::EWidgetMode CurrentWidgetMode = ViewportClient->GetWidgetMode();
+					if(!RequestedWidgetModes.IsEmpty())
+					{
+						if(RequestedWidgetModes.Last() != CurrentWidgetMode)
+						{
+							CurrentWidgetMode = RequestedWidgetModes.Last();
+							ViewportClient->SetWidgetMode(CurrentWidgetMode);
+						}
+						RequestedWidgetModes.Reset();
+					}
 					for (FRigElementKey SelectedRigElement : SelectedRigElements)
 					{
 						//need to loop through the shape actors and set widget based upon the first one
@@ -2280,6 +2289,11 @@ void FControlRigEditMode::RecalcPivotTransform()
 	}
 }
 
+void FControlRigEditMode::RequestTransformWidgetMode(UE::Widget::EWidgetMode InWidgetMode)
+{
+	RequestedWidgetModes.Add(InWidgetMode);
+}
+
 bool FControlRigEditMode::HasPivotTransformsChanged() const
 {
 	if (PivotTransforms.Num() != LastPivotTransforms.Num())
@@ -2910,7 +2924,10 @@ void FControlRigEditMode::ResetTransforms(bool bSelectionOnly)
 				ControlElement = ControlRig->FindControl(ElementToReset.Name);
 				if (ControlElement->Settings.bIsTransientControl)
 				{
-					ControlElement = nullptr;
+					if(UControlRig::GetNodeNameFromTransientControl(ControlElement->GetKey()).IsEmpty())
+					{
+						ControlElement = nullptr;
+					}
 				}
 			}
 			
diff --git a/Engine/Plugins/Animation/ControlRig/Source/ControlRigEditor/Private/Editor/ControlRigEditor.cpp b/Engine/Plugins/Animation/ControlRig/Source/ControlRigEditor/Private/Editor/ControlRigEditor.cpp
--- a/Engine/Plugins/Animation/ControlRig/Source/ControlRigEditor/Private/Editor/ControlRigEditor.cpp
+++ b/Engine/Plugins/Animation/ControlRig/Source/ControlRigEditor/Private/Editor/ControlRigEditor.cpp
@@ -1028,6 +1028,62 @@ void FControlRigEditor::Compile()
 	}
 }
 
+void FControlRigEditor::HandleModifiedEvent(ERigVMGraphNotifType InNotifType, URigVMGraph* InGraph, UObject* InSubject)
+{
+	IControlRigEditor::HandleModifiedEvent(InNotifType, InGraph, InSubject);
+
+	if(InNotifType == ERigVMGraphNotifType::NodeSelected)
+	{
+		if(const URigVMUnitNode* UnitNode = Cast<URigVMUnitNode>(InSubject))
+		{
+			SetDirectionManipulationSubject(UnitNode);
+		}
+	}
+	else if(InNotifType == ERigVMGraphNotifType::NodeSelectionChanged)
+	{
+		bool bNeedsToClearManipulationSubject = true;
+		const TArray<FName> SelectedNodes = InGraph->GetSelectNodes();
+		if(SelectedNodes.Num() == 1)
+		{
+			if(const URigVMUnitNode* UnitNode = Cast<URigVMUnitNode>(InGraph->FindNodeByName(SelectedNodes[0])))
+			{
+				SetDirectionManipulationSubject(UnitNode);
+				bNeedsToClearManipulationSubject = false;
+			}
+		}
+
+		if(bNeedsToClearManipulationSubject)
+		{
+			ClearDirectManipulationSubject();
+		}
+	}
+	else if(InNotifType == ERigVMGraphNotifType::PinDefaultValueChanged)
+	{
+		if(const URigVMPin* Pin = Cast<URigVMPin>(InSubject))
+		{
+			if(Pin->GetNode() == DirectManipulationSubject.Get())
+			{
+				RefreshDirectManipulationTextList();
+			}
+		}
+	}
+}
+
+void FControlRigEditor::OnCreateGraphEditorCommands(TSharedPtr<FUICommandList> GraphEditorCommandsList)
+{
+	IControlRigEditor::OnCreateGraphEditorCommands(GraphEditorCommandsList);
+
+	GraphEditorCommandsList->MapAction(
+		FControlRigEditorCommands::Get().RequestDirectManipulationPosition,
+		FExecuteAction::CreateSP(this, &FControlRigEditor::HandleRequestDirectManipulationPosition));
+	GraphEditorCommandsList->MapAction(
+		FControlRigEditorCommands::Get().RequestDirectManipulationRotation,
+		FExecuteAction::CreateSP(this, &FControlRigEditor::HandleRequestDirectManipulationRotation));
+	GraphEditorCommandsList->MapAction(
+		FControlRigEditorCommands::Get().RequestDirectManipulationScale,
+		FExecuteAction::CreateSP(this, &FControlRigEditor::HandleRequestDirectManipulationScale));
+}
+
 void FControlRigEditor::SaveAsset_Execute()
 {
 	FRigVMEditor::SaveAsset_Execute();
@@ -1261,6 +1317,37 @@ void FControlRigEditor::HandleVMExecutedEvent(URigVMHost* InHost, const FName& I
 				}
 			}
 		}
+
+		// update transient controls on nodes / pins
+		if(UControlRig* DebuggedControlRig = Cast<UControlRig>(ControlRigBP->GetObjectBeingDebugged()))
+		{
+			if(!DebuggedControlRig->RigUnitManipulationInfos.IsEmpty())
+			{
+				FControlRigExecuteContext& ExecuteContext = DebuggedControlRig->GetExtendedExecuteContext().GetPublicDataSafe<FControlRigExecuteContext>();
+				
+				for(const TSharedPtr<FRigDirectManipulationInfo>& ManipulationInfo : DebuggedControlRig->RigUnitManipulationInfos)
+				{
+					if(const URigVMUnitNode* Node = ManipulationInfo->Node.Get())
+					{
+						const UScriptStruct* ScriptStruct = Node->GetScriptStruct();
+						if(ScriptStruct == nullptr)
+						{
+							continue;
+						}
+
+						TSharedPtr<FStructOnScope> NodeInstance = Node->ConstructLiveStructInstance(DebuggedControlRig);
+						if(!NodeInstance.IsValid() || !NodeInstance->IsValid())
+						{
+							continue;
+						}
+				
+						FRigUnit* UnitInstance = UControlRig::GetRigUnitInstanceFromScope(NodeInstance);
+						UnitInstance->UpdateHierarchyForDirectManipulation(Node, NodeInstance, ExecuteContext, ManipulationInfo);
+						UnitInstance->PerformDebugDrawingForDirectManipulation(Node, NodeInstance, ExecuteContext, ManipulationInfo);
+					}
+				}
+			}
+		}		
 	}
 }
 
@@ -1433,6 +1520,59 @@ void FControlRigEditor::HandleViewportCreated(const TSharedRef<class IPersonaVie
 		return EVisibility::Collapsed;
 	};
 
+	{
+		FPersonaViewportNotificationOptions DirectManipulationNotificationOptions(TAttribute<EVisibility>::CreateSP(this, &FControlRigEditor::GetDirectManipulationVisibility));
+		DirectManipulationNotificationOptions.OnGetBrushOverride = TAttribute<const FSlateBrush*>(FControlRigEditorStyle::Get().GetBrush("ControlRig.Viewport.Notification.DirectManipulation"));
+
+		InViewport->AddNotification(
+			EMessageSeverity::Info,
+			false,
+			SNew(SHorizontalBox)
+			.Visibility(this, &FControlRigEditor::GetDirectManipulationVisibility)
+			.ToolTipText(LOCTEXT("DirectManipulation", "Direct Manipulation"))
+			+SHorizontalBox::Slot()
+			.FillWidth(1.0f)
+			.Padding(4.0f, 4.0f)
+			[
+				SNew(SHorizontalBox)
+				+SHorizontalBox::Slot()
+				.AutoWidth()
+				.VAlign(VAlign_Center)
+				.Padding(0.0f, 0.0f, 4.0f, 0.0f)
+				[
+					SNew(STextBlock)
+					.TextStyle(FAppStyle::Get(), "AnimViewport.MessageText")
+					.Font(FAppStyle::Get().GetFontStyle("FontAwesome.9"))
+					.Text(FEditorFontGlyphs::Crosshairs)
+				]
+				+SHorizontalBox::Slot()
+				.VAlign(VAlign_Center)
+				.AutoWidth()
+				[
+					SNew(SComboBox<TSharedPtr<FString>>)
+                	.ContentPadding(FMargin(4.0f, 2.0f))
+                	.OptionsSource(&DirectManipulationTextList)
+                	.OnGenerateWidget_Lambda([this](TSharedPtr<FString> Item)
+                	{ 
+                		return SNew(SBox)
+                			.MaxDesiredWidth(600.0f)
+                			[
+                				SNew(STextBlock)
+								.TextStyle(FAppStyle::Get(), "AnimViewport.MessageText")
+                				.Text(FText::FromString(*Item))
+                			];
+                	} )	
+                	.OnSelectionChanged(this, &FControlRigEditor::OnDirectManipulationChanged)
+                	[
+			        	SNew(STextBlock)
+						.TextStyle(FAppStyle::Get(), "AnimViewport.MessageText")
+						.Text(this, &FControlRigEditor::GetDirectionManipulationText)
+					]
+				]
+			],
+			DirectManipulationNotificationOptions
+		);
+	}	
 	InViewport->AddNotification(MakeAttributeLambda(GetErrorSeverity),
 		false,
 		SNew(SHorizontalBox)
@@ -1585,43 +1725,6 @@ void FControlRigEditor::HandleViewportCreated(const TSharedRef<class IPersonaVie
 					], 
 					LOCTEXT("ControlRigAxesScale", "Axes Scale")
 				);
-
-				if (UControlRigBlueprint* ControlRigBlueprint = CastChecked<UControlRigBlueprint>(GetBlueprintObj()))
-				{
-					for (UEdGraph* Graph : ControlRigBlueprint->UbergraphPages)
-					{
-						if (UControlRigGraph* RigGraph = Cast<UControlRigGraph>(Graph))
-						{
-							const TArray<TSharedPtr<FString>>* BoneNameList = RigGraph->GetBoneNameList();
-
-							InMenuBuilder.AddWidget(
-								SNew(SBox)
-								.HAlign(HAlign_Right)
-								[
-									SNew(SBox)
-									.Padding(FMargin(4.0f, 0.0f, 0.0f, 0.0f))
-									.WidthOverride(100.0f)
-									.IsEnabled(this, &FControlRigEditor::IsPinControlNameListEnabled)
-									[
-										SAssignNew(PinControlNameList, SRigVMGraphPinNameListValueWidget)
-										.OptionsSource(BoneNameList)
-										.OnGenerateWidget(this, &FControlRigEditor::MakePinControlNameListItemWidget)
-										.OnSelectionChanged(this, &FControlRigEditor::OnPinControlNameListChanged)
-										.OnComboBoxOpening(this, &FControlRigEditor::OnPinControlNameListComboBox, BoneNameList)
-										.InitiallySelectedItem(GetPinControlCurrentlySelectedItem(BoneNameList))
-										.Content()
-										[
-											SNew(STextBlock)
-											.Text(this, &FControlRigEditor::GetPinControlNameListText)
-										]
-									]
-								],
-								LOCTEXT("ControlRigAuthoringSpace", "Pin Control Space")
-							);
-							break;
-						}
-					}
-				}
 			}
 			InMenuBuilder.EndSection();
 		}
@@ -1761,139 +1864,6 @@ void FControlRigEditor::OnToolbarDrawNullsChanged(ECheckBoxState InNewValue)
 	}
 }
 
-bool FControlRigEditor::IsPinControlNameListEnabled() const
-{
-	if (UControlRig* ControlRig = GetControlRig())
-	{
-		TArray<FRigControlElement*>	TransientControls = ControlRig->GetHierarchy()->GetTransientControls();
-		if (TransientControls.Num() > 0)
-		{
-			// if the transient control is not for a rig element, it is for a pin
-			if (UControlRig::GetElementKeyFromTransientControl(TransientControls[0]->GetKey()) == FRigElementKey())
-			{
-				return true;
-			}
-		}
-	}
-	return false;
-}
-
-TSharedRef<SWidget> FControlRigEditor::MakePinControlNameListItemWidget(TSharedPtr<FString> InItem)
-{
-	return 	SNew(STextBlock).Text(FText::FromString(*InItem));
-}
-
-FText FControlRigEditor::GetPinControlNameListText() const
-{
-	if (UControlRig* ControlRig = GetControlRig())
-	{
-		FText Result;
-		ControlRig->GetHierarchy()->ForEach<FRigControlElement>([this, &Result, ControlRig](FRigControlElement* ControlElement) -> bool
-        {
-			if (ControlElement->Settings.bIsTransientControl)
-			{
-				FRigElementKey Parent = ControlRig->GetHierarchy()->GetFirstParent(ControlElement->GetKey());
-				Result = FText::FromName(Parent.Name);
-				
-				return false;
-			}
-			return true;
-		});
-		
-		if(!Result.IsEmpty())
-		{
-			return Result;
-		}
-	}
-	return FText::FromName(NAME_None);
-}
-
-TSharedPtr<FString> FControlRigEditor::GetPinControlCurrentlySelectedItem(const TArray<TSharedPtr<FString>>* InNameList) const
-{
-	FString CurrentItem = GetPinControlNameListText().ToString();
-	for (const TSharedPtr<FString>& Item : *InNameList)
-	{
-		if (Item->Equals(CurrentItem))
-		{
-			return Item;
-		}
-	}
-	return TSharedPtr<FString>();
-}
-
-void FControlRigEditor::SetPinControlNameListText(const FText& NewTypeInValue, ETextCommit::Type /*CommitInfo*/)
-{
-	if (UControlRig* ControlRig = GetControlRig())
-	{
-		ControlRig->GetHierarchy()->ForEach<FRigControlElement>([this, NewTypeInValue, ControlRig](FRigControlElement* ControlElement) -> bool
-        {
-            if (ControlElement->Settings.bIsTransientControl)
-			{
-				FName NewParentName = *NewTypeInValue.ToString();
-				const int32 NewParentIndex = ControlRig->GetHierarchy()->GetIndex(FRigElementKey(NewParentName, ERigElementType::Bone));
-				if (NewParentIndex == INDEX_NONE)
-				{
-					NewParentName = NAME_None;
-					ControlRig->GetHierarchy()->GetController()->RemoveAllParents(ControlElement->GetKey(), true, false);
-				}
-				else
-				{
-					ControlRig->GetHierarchy()->GetController()->SetParent(ControlElement->GetKey(), FRigElementKey(NewParentName, ERigElementType::Bone), true, false);
-				}
-
-				// find out if the controlled pin is part of a visual debug node
-				if (UControlRigBlueprint* ControlRigBlueprint = CastChecked<UControlRigBlueprint>(GetBlueprintObj()))
-				{
-					FString PinName = UControlRig::GetPinNameFromTransientControl(ControlElement->GetKey());
-					if (URigVMPin* ControlledPin = GetFocusedModel()->FindPin(PinName))
-					{
-						URigVMNode* ControlledNode = ControlledPin->GetPinForLink()->GetNode();
-						if (URigVMUnitNode* UnitNode = Cast<URigVMUnitNode>(ControlledNode))
-						{
-							if (const FString* Value = UnitNode->GetScriptStruct()->FindMetaData(FRigVMStruct::TemplateNameMetaName))
-							{
-								if (Value->Equals("VisualDebug"))
-								{
-									if (URigVMPin* SpacePin = ControlledNode->FindPin(TEXT("Space")))
-									{
-										FString DefaultValue;
-										const FRigElementKey NewSpaceKey(NewParentName, ERigElementType::Bone); 
-										FRigElementKey::StaticStruct()->ExportText(DefaultValue, &NewSpaceKey, nullptr, nullptr, PPF_None, nullptr);
-										ensure(GetFocusedController()->SetPinDefaultValue(SpacePin->GetPinPath(), DefaultValue, false, false, false));
-									}
-									else if (URigVMPin* BoneSpacePin = ControlledNode->FindPin(TEXT("BoneSpace")))
-									{
-										if (BoneSpacePin->GetCPPType() == TEXT("FName") && BoneSpacePin->GetCustomWidgetName() == TEXT("BoneName"))
-										{
-											GetFocusedController()->SetPinDefaultValue(BoneSpacePin->GetPinPath(), NewParentName.ToString(), false, false, false);
-										}
-									}
-								}
-							}
-						}
-					}
-				}
-			}
-			return true;
-		});
-	}
-}
-
-void FControlRigEditor::OnPinControlNameListChanged(TSharedPtr<FString> NewSelection, ESelectInfo::Type SelectInfo)
-{
-	if (SelectInfo != ESelectInfo::Direct)
-	{
-		FString NewValue = *NewSelection.Get();
-		SetPinControlNameListText(FText::FromString(NewValue), ETextCommit::OnEnter);
-	}
-}
-
-void FControlRigEditor::OnPinControlNameListComboBox(const TArray<TSharedPtr<FString>>* InNameList)
-{
-	TSharedPtr<FString> CurrentlySelected = GetPinControlCurrentlySelectedItem(InNameList);
-	PinControlNameList->SetSelectedItem(CurrentlySelected);
-}
-
 bool FControlRigEditor::IsConstructionModeEnabled() const
 {
 	return GetEventQueue() == ConstructionEventQueue;
@@ -2421,6 +2391,186 @@ void FControlRigEditor::OnWrappedPropertyChangedChainEvent(URigVMDetailsViewWrap
 	}
 }
 
+bool FControlRigEditor::HandleRequestDirectManipulation(ERigControlType InControlType) const
+{
+	TArray<FRigDirectManipulationTarget> Targets = GetDirectManipulationTargets();
+	for(const FRigDirectManipulationTarget& Target : Targets)
+	{
+		if(Target.ControlType == InControlType || Target.ControlType == ERigControlType::EulerTransform)
+		{
+			if (FControlRigEditorEditMode* EditMode = GetEditMode())
+			{
+				switch(InControlType)
+				{
+					case ERigControlType::Position:
+					{
+						EditMode->RequestTransformWidgetMode(UE::Widget::WM_Translate);
+						break;
+					}
+					case ERigControlType::Rotator:
+					{
+						EditMode->RequestTransformWidgetMode(UE::Widget::WM_Rotate);
+						break;
+					}
+					case ERigControlType::Scale:
+					{
+						EditMode->RequestTransformWidgetMode(UE::Widget::WM_Scale);
+						break;
+					}
+					default:
+					{
+						break;
+					}
+				}
+			}
+
+			if(UControlRigBlueprint* Blueprint = GetControlRigBlueprint())
+			{
+				Blueprint->AddTransientControl(DirectManipulationSubject.Get(), Target);
+			}
+			return true;
+		}
+	}
+	return false;
+}
+
+bool FControlRigEditor::SetDirectionManipulationSubject(const URigVMUnitNode* InNode)
+{
+	if(DirectManipulationSubject.Get() == InNode)
+	{
+		return false;
+	}
+	if(UControlRigBlueprint* Blueprint = GetControlRigBlueprint())
+	{
+		Blueprint->ClearTransientControls();
+	}
+	DirectManipulationSubject = InNode;
+
+	// update the direct manipulation target list
+	RefreshDirectManipulationTextList();
+	return true;
+}
+
+bool FControlRigEditor::IsDirectManipulationEnabled() const
+{
+	return !GetDirectManipulationTargets().IsEmpty();
+}
+
+EVisibility FControlRigEditor::GetDirectManipulationVisibility() const
+{
+	return IsDirectManipulationEnabled() ? EVisibility::Visible : EVisibility::Hidden;
+}
+
+FText FControlRigEditor::GetDirectionManipulationText() const
+{
+	if (UControlRig* DebuggedControlRig = Cast<UControlRig>(GetBlueprintObj()->GetObjectBeingDebugged()))
+	{
+		TArray<FRigControlElement*> TransientControls = DebuggedControlRig->GetHierarchy()->GetTransientControls();
+		for(const FRigControlElement* TransientControl : TransientControls)
+		{
+			const FString Target = UControlRig::GetTargetFromTransientControl(TransientControl->GetKey());
+			if(!Target.IsEmpty())
+			{
+				return FText::FromString(Target);
+			}
+		}
+	}
+	static const FText DefaultText = LOCTEXT("ControlRigDirectManipulation", "Direct Manipulation");
+	return DefaultText;
+}
+
+void FControlRigEditor::OnDirectManipulationChanged(TSharedPtr<FString> NewValue, ESelectInfo::Type SelectInfo)
+{
+	if(!NewValue.IsValid())
+	{
+		return;
+	}
+	
+	const URigVMUnitNode* UnitNode = DirectManipulationSubject.Get();
+	if(UnitNode == nullptr)
+	{
+		return;
+	}
+	
+	UControlRigBlueprint* ControlRigBlueprint = CastChecked<UControlRigBlueprint>(GetBlueprintObj());
+	if(ControlRigBlueprint == nullptr)
+	{
+		return;
+	}
+
+	// disable literal folding for the moment
+	if(ControlRigBlueprint->VMCompileSettings.ASTSettings.bFoldLiterals)
+	{
+		ControlRigBlueprint->VMCompileSettings.ASTSettings.bFoldLiterals = false;
+		ControlRigBlueprint->RecompileVM();
+	}
+
+	const FString& DesiredTarget = *NewValue.Get();
+	const TArray<FRigDirectManipulationTarget> Targets = GetDirectManipulationTargets();
+	for(const FRigDirectManipulationTarget& Target : Targets)
+	{
+		if(Target.Name.Equals(DesiredTarget, ESearchCase::CaseSensitive))
+		{
+			// run the task after a bit so that the rig has the opportunity to run first
+			FFunctionGraphTask::CreateAndDispatchWhenReady([ControlRigBlueprint, UnitNode, Target]()
+			{
+				ControlRigBlueprint->AddTransientControl(UnitNode, Target);
+			}, TStatId(), NULL, ENamedThreads::GameThread);
+			break;
+		}
+	}
+}
+
+const TArray<FRigDirectManipulationTarget> FControlRigEditor::GetDirectManipulationTargets() const
+{
+	if(DirectManipulationSubject.IsValid())
+	{
+		if (UControlRig* DebuggedControlRig = Cast<UControlRig>(GetBlueprintObj()->GetObjectBeingDebugged()))
+		{
+			if(const URigVMUnitNode* Node = DirectManipulationSubject.Get())
+			{
+				if(Node->IsPartOfRuntime(DebuggedControlRig))
+				{
+					const TSharedPtr<FStructOnScope> NodeInstance = Node->ConstructLiveStructInstance(DebuggedControlRig);
+					if(NodeInstance.IsValid() && NodeInstance->IsValid())
+					{
+						if(const FRigUnit* UnitInstance = UControlRig::GetRigUnitInstanceFromScope(NodeInstance))
+						{
+							TArray<FRigDirectManipulationTarget> Targets;
+							if(UnitInstance->GetDirectManipulationTargets(Node, NodeInstance, DebuggedControlRig->GetHierarchy(), Targets, nullptr))
+							{
+								return Targets;
+							}
+						}
+					}
+				}
+			}
+		}
+	}
+
+	static const TArray<FRigDirectManipulationTarget> EmptyTargets;
+	return EmptyTargets;
+}
+
+const TArray<TSharedPtr<FString>>& FControlRigEditor::GetDirectManipulationTargetTextList() const
+{
+	if(DirectManipulationTextList.IsEmpty())
+	{
+		const TArray<FRigDirectManipulationTarget> Targets = GetDirectManipulationTargets();
+		for(const FRigDirectManipulationTarget& Target : Targets)
+		{
+			DirectManipulationTextList.Emplace(new FString(Target.Name));
+		}
+	}
+	return DirectManipulationTextList;
+}
+
+void FControlRigEditor::RefreshDirectManipulationTextList()
+{
+	DirectManipulationTextList.Reset();
+	(void)GetDirectManipulationTargetTextList();
+}
+
 void FControlRigEditor::BindCommands()
 {
 	FRigVMEditor::BindCommands();
@@ -2737,7 +2887,14 @@ void FControlRigEditor::OnHierarchyModified_AnyThread(ERigHierarchyNotification
 						}
 						else
 						{
-							ClearDetailObject();
+							// only clear the details if we are not looking at a transient control
+							if (UControlRig* DebuggedControlRig = Cast<UControlRig>(GetBlueprintObj()->GetObjectBeingDebugged()))
+							{
+								if(DebuggedControlRig->RigUnitManipulationInfos.IsEmpty())
+								{
+									ClearDetailObject();
+								}
+							}
 						}
 					}
 				}						
@@ -3779,88 +3936,89 @@ void FControlRigEditor::HandleOnControlModified(UControlRig* Subject, FRigContro
 
 	URigHierarchy* Hierarchy = Subject->GetHierarchy();
 
-	if (ControlElement->Settings.bIsTransientControl)
+	if (ControlElement->Settings.bIsTransientControl && !GIsTransacting)
 	{
-		FRigControlValue ControlValue = Hierarchy->GetControlValue(ControlElement, ERigControlValueType::Current);
-
-		const FString PinPath = UControlRig::GetPinNameFromTransientControl(ControlElement->GetKey());
-		if (URigVMPin* Pin = Blueprint->GetRigVMClient()->FindPin(PinPath))
+		const URigVMUnitNode* UnitNode = nullptr;
+		UControlRig* DebuggedControlRig = Cast<UControlRig>(GetBlueprintObj()->GetObjectBeingDebugged());
+		const FString NodeName = UControlRig::GetNodeNameFromTransientControl(ControlElement->GetKey());
+		const FString PoseTarget = UControlRig::GetTargetFromTransientControl(ControlElement->GetKey());
+		TSharedPtr<FStructOnScope> NodeInstance;
+		TSharedPtr<FRigDirectManipulationInfo> ManipulationInfo;
+
+		// try to find the direct manipulation info on the rig. if there's no matching information
+		// the manipulation is likely happening on a bone instead.
+		if(DebuggedControlRig && !NodeName.IsEmpty() && !PoseTarget.IsEmpty())
 		{
-			FString NewDefaultValue;
-			switch (ControlElement->Settings.ControlType)
+			UnitNode = Cast<URigVMUnitNode>(GetFocusedModel()->FindNode(NodeName));
+			if(UnitNode)
 			{
-				case ERigControlType::Position:
-				case ERigControlType::Scale:
-				{
-					NewDefaultValue = ControlValue.ToString<FVector>();
-					break;
-				}
-				case ERigControlType::Rotator:
-				{
-					FVector3f RotatorAngles = ControlValue.Get<FVector3f>();
-					FRotator Rotator = FRotator::MakeFromEuler((FVector)RotatorAngles);
-					FRigControlValue RotatorValue = FRigControlValue::Make<FRotator>(Rotator);
-					NewDefaultValue = RotatorValue.ToString<FRotator>();
-					break;
-				}
-				case ERigControlType::Transform:
-				{
-					NewDefaultValue = ControlValue.ToString<FTransform>();
-					break;
-				}
-				case ERigControlType::TransformNoScale:
+				if(UnitNode->GetScriptStruct())
 				{
-					NewDefaultValue = ControlValue.ToString<FTransformNoScale>();
-					break;
+					NodeInstance = UnitNode->ConstructStructInstance(false);
+					ManipulationInfo = DebuggedControlRig->GetRigUnitManipulationInfoForTransientControl(ControlElement->GetKey());
 				}
-				case ERigControlType::EulerTransform:
-				{
-					NewDefaultValue = ControlValue.ToString<FEulerTransform>();
-					break;
-				}
-				default:
+				else
 				{
-					break;
+					UnitNode = nullptr;
 				}
 			}
+		}
+		
+		if (UnitNode && NodeInstance.IsValid() && ManipulationInfo.IsValid())
+		{
+			FRigUnit* UnitInstance = DebuggedControlRig->GetRigUnitInstanceFromScope(NodeInstance);
+			check(UnitInstance);
 
-			if (!NewDefaultValue.IsEmpty())
+			// update the node based on the incoming pose. once that is done we'll need to compare the node instance
+			// with the settings on the node in the graph and update them accordingly.
+			FControlRigExecuteContext& ExecuteContext = DebuggedControlRig->GetExtendedExecuteContext().GetPublicDataSafe<FControlRigExecuteContext>();
+			if(UnitInstance->UpdateDirectManipulationFromHierarchy(UnitNode, NodeInstance, ExecuteContext, ManipulationInfo))
 			{
-				bool bRequiresRecompile = true;
-
-				if(UControlRig* ControlRig = GetControlRig())
+				UnitNode->UpdateHostFromStructInstance(DebuggedControlRig, NodeInstance);
+				URigVMController* Controller = Blueprint->GetOrCreateController(UnitNode->GetGraph());
+				TMap<FString, FString> PinPathToNewDefaultValue;
+				UnitNode->ComputePinValueDifferences(NodeInstance, PinPathToNewDefaultValue);
+				if(!PinPathToNewDefaultValue.IsEmpty())
 				{
-					if(TSharedPtr<FRigVMParserAST> AST = Pin->GetGraph()->GetDiagnosticsAST())
+					// we'll disable compilation since the control rig editor module will have disabled folding of literals
+					// so each register is free to be edited directly.
+					TGuardValue<bool> DisableBlueprintNotifs(Blueprint->bSuspendModelNotificationsForSelf, true);
+
+					if(PinPathToNewDefaultValue.Num() > 1)
 					{
-						FRigVMASTProxy PinProxy = FRigVMASTProxy::MakeFromUObject(Pin);
-						if(const FRigVMExprAST* PinExpr = AST->GetExprForSubject(PinProxy))
+						Controller->OpenUndoBracket(TEXT("Set pin defaults during manipulation"));
+					}
+					bool bChangedSomething = false;
+
+					for(const TPair<FString, FString>& Pair : PinPathToNewDefaultValue)
+					{
+						if(const URigVMPin* Pin = UnitNode->FindPin(Pair.Key))
 						{
-							if(PinExpr->IsA(FRigVMExprAST::Var))
+							if(Controller->SetPinDefaultValue(Pin->GetPinPath(), Pair.Value, true, true, true, false, false))
 							{
-								const FString PinHash = URigVMCompiler::GetPinHash(Pin, PinExpr->To<FRigVMVarExprAST>(), false);
-								if(const FRigVMOperand* OperandForPin = Blueprint->PinToOperandMap.Find(PinHash))
-								{
-									if(URigVM* VM = ControlRig->GetVM())
-									{
-										// only operands which are shared across multiple instructions require recompile
-										if(!VM->GetByteCode().IsOperandShared(*OperandForPin))
-										{
-											VM->SetPropertyValueFromString(*OperandForPin, NewDefaultValue);
-											bRequiresRecompile = false;
-										}
-									}
-								}
+								bChangedSomething = true;
 							}
 						}
 					}
+
+					if(PinPathToNewDefaultValue.Num() > 1)
+					{
+						if(bChangedSomething)
+						{
+							Controller->CloseUndoBracket();
+						}
+						else
+						{
+							Controller->CancelUndoBracket();
+						}
+					}
 				}
-				
-				TGuardValue<bool> DisableBlueprintNotifs(Blueprint->bSuspendModelNotificationsForSelf, !bRequiresRecompile);
-				GetFocusedController()->SetPinDefaultValue(Pin->GetPinPath(), NewDefaultValue, true, true, true);
+
 			}
 		}
 		else
 		{
+			FRigControlValue ControlValue = Hierarchy->GetControlValue(ControlElement, ERigControlValueType::Current);
 			const FRigElementKey ElementKey = UControlRig::GetElementKeyFromTransientControl(ControlElement->GetKey());
 
 			if (ElementKey.Type == ERigElementType::Bone)
diff --git a/Engine/Plugins/Animation/ControlRig/Source/ControlRigEditor/Private/Editor/ControlRigEditor.h b/Engine/Plugins/Animation/ControlRig/Source/ControlRigEditor/Private/Editor/ControlRigEditor.h
--- a/Engine/Plugins/Animation/ControlRig/Source/ControlRigEditor/Private/Editor/ControlRigEditor.h
+++ b/Engine/Plugins/Animation/ControlRig/Source/ControlRigEditor/Private/Editor/ControlRigEditor.h
@@ -22,6 +22,7 @@
 #include "Editor/RigVMDetailsViewWrapperObject.h"
 #include "ControlRigTestData.h"
 #include "RigVMHost.h"
+#include "Units/RigUnit.h"
 
 class UControlRigBlueprint;
 class IPersonaToolkit;
@@ -73,6 +74,8 @@ public:
 	// FRigVMEditor interface
 	virtual UClass* GetDetailWrapperClass() const;
 	virtual void Compile() override;
+	virtual void HandleModifiedEvent(ERigVMGraphNotifType InNotifType, URigVMGraph* InGraph, UObject* InSubject) override;
+	virtual void OnCreateGraphEditorCommands(TSharedPtr<FUICommandList> GraphEditorCommandsList) override;
 
 	// allows the editor to fill an empty graph
 	virtual void CreateEmptyGraphContent(URigVMController* InController) override;
@@ -141,6 +144,21 @@ public:
 	FOnGetContextMenu& OnGetViewportContextMenu() { return OnGetViewportContextMenuDelegate; }
 	FNewMenuCommandsDelegate& OnViewportContextMenuCommands() { return OnViewportContextMenuCommandsDelegate; }
 
+	// DirectManipulation functionality
+	void HandleRequestDirectManipulationPosition() const { (void)HandleRequestDirectManipulation(ERigControlType::Position); }
+	void HandleRequestDirectManipulationRotation() const { (void)HandleRequestDirectManipulation(ERigControlType::Rotator); }
+	void HandleRequestDirectManipulationScale() const { (void)HandleRequestDirectManipulation(ERigControlType::Scale); }
+	bool HandleRequestDirectManipulation(ERigControlType InControlType) const;
+	bool SetDirectionManipulationSubject(const URigVMUnitNode* InNode);
+	bool IsDirectManipulationEnabled() const;
+	EVisibility GetDirectManipulationVisibility() const;
+	FText GetDirectionManipulationText() const;
+	void OnDirectManipulationChanged(TSharedPtr<FString> NewValue, ESelectInfo::Type SelectInfo);
+	const TArray<FRigDirectManipulationTarget> GetDirectManipulationTargets() const;
+	const TArray<TSharedPtr<FString>>& GetDirectManipulationTargetTextList() const;
+	bool ClearDirectManipulationSubject() { return SetDirectionManipulationSubject(nullptr); }
+	void RefreshDirectManipulationTextList();
+
 protected:
 
 	virtual void BindCommands() override;
@@ -266,14 +284,6 @@ protected:
 
 	void OnAnimInitialized();
 
-	TSharedPtr<SRigVMGraphPinNameListValueWidget> PinControlNameList;
-	bool IsPinControlNameListEnabled() const;
-	TSharedRef<SWidget> MakePinControlNameListItemWidget(TSharedPtr<FString> InItem);
-	FText GetPinControlNameListText() const;
-	TSharedPtr<FString> GetPinControlCurrentlySelectedItem(const TArray<TSharedPtr<FString>>* InNameList) const;
-	void SetPinControlNameListText(const FText& NewTypeInValue, ETextCommit::Type /*CommitInfo*/);
-	void OnPinControlNameListChanged(TSharedPtr<FString> NewSelection, ESelectInfo::Type SelectInfo);
-	void OnPinControlNameListComboBox(const TArray<TSharedPtr<FString>>* InNameList);
 	bool IsConstructionModeEnabled() const;
 
 	int32 RigHierarchyTabCount;
@@ -286,6 +296,9 @@ protected:
 	uint32 LastHierarchyHash;
 
 	TStrongObjectPtr<UControlRigTestData> TestDataStrongPtr;
+
+	TWeakObjectPtr<const URigVMUnitNode> DirectManipulationSubject;
+	mutable TArray<TSharedPtr<FString>> DirectManipulationTextList;
 	
 	static const TArray<FName> ForwardsSolveEventQueue;
 	static const TArray<FName> BackwardsSolveEventQueue;
diff --git a/Engine/Plugins/Animation/ControlRig/Source/ControlRigEditor/Private/Editor/ControlRigEditorCommands.cpp b/Engine/Plugins/Animation/ControlRig/Source/ControlRigEditor/Private/Editor/ControlRigEditorCommands.cpp
--- a/Engine/Plugins/Animation/ControlRig/Source/ControlRigEditor/Private/Editor/ControlRigEditorCommands.cpp
+++ b/Engine/Plugins/Animation/ControlRig/Source/ControlRigEditor/Private/Editor/ControlRigEditorCommands.cpp
@@ -10,6 +10,9 @@ void FControlRigEditorCommands::RegisterCommands()
 	UI_COMMAND(ForwardsSolveEvent, "Forwards Solve", "Run the forwards solve graph", EUserInterfaceActionType::Button, FInputChord());
 	UI_COMMAND(BackwardsSolveEvent, "Backwards Solve", "Run the backwards solve graph", EUserInterfaceActionType::Button, FInputChord());
 	UI_COMMAND(BackwardsAndForwardsSolveEvent, "Backwards and Forwards", "Run backwards solve followed by forwards solve", EUserInterfaceActionType::Button, FInputChord());
+	UI_COMMAND(RequestDirectManipulationPosition, "Request Direct Manipulation for Position", "Request per node direct manipulation on a position", EUserInterfaceActionType::Button, FInputChord(EKeys::W));
+	UI_COMMAND(RequestDirectManipulationRotation, "Request Direct Manipulation for Rotation", "Request per node direct manipulation on a rotation", EUserInterfaceActionType::Button, FInputChord(EKeys::E));
+	UI_COMMAND(RequestDirectManipulationScale, "Request Direct Manipulation for Scale", "Request per node direct manipulation on a scale", EUserInterfaceActionType::Button, FInputChord(EKeys::R));
 }
 
 #undef LOCTEXT_NAMESPACE
diff --git a/Engine/Plugins/Animation/ControlRig/Source/ControlRigEditor/Private/Editor/ControlRigEditorCommands.h b/Engine/Plugins/Animation/ControlRig/Source/ControlRigEditor/Private/Editor/ControlRigEditorCommands.h
--- a/Engine/Plugins/Animation/ControlRig/Source/ControlRigEditor/Private/Editor/ControlRigEditorCommands.h
+++ b/Engine/Plugins/Animation/ControlRig/Source/ControlRigEditor/Private/Editor/ControlRigEditorCommands.h
@@ -30,6 +30,15 @@ public:
 	/** Run the backwards solve graph followed by the forwards solve graph */
 	TSharedPtr< FUICommandInfo > BackwardsAndForwardsSolveEvent;
 
+	/** Request per node direct manipulation on a position */
+	TSharedPtr< FUICommandInfo > RequestDirectManipulationPosition;
+
+	/** Request per node direct manipulation on a rotation */
+	TSharedPtr< FUICommandInfo > RequestDirectManipulationRotation;
+
+	/** Request per node direct manipulation on a scale */
+	TSharedPtr< FUICommandInfo > RequestDirectManipulationScale;
+
 	/**
 	 * Initialize commands
 	 */
diff --git a/Engine/Plugins/Animation/ControlRig/Source/ControlRigEditor/Public/EditMode/ControlRigEditMode.h b/Engine/Plugins/Animation/ControlRig/Source/ControlRigEditor/Public/EditMode/ControlRigEditMode.h
--- a/Engine/Plugins/Animation/ControlRig/Source/ControlRigEditor/Public/EditMode/ControlRigEditMode.h
+++ b/Engine/Plugins/Animation/ControlRig/Source/ControlRigEditor/Public/EditMode/ControlRigEditMode.h
@@ -382,6 +382,10 @@ public:
 
 	/** Suspend Rig Hierarchy Notifies*/
 	void SuspendHierarchyNotifs(bool bVal) { bSuspendHierarchyNotifs = bVal; }
+
+	/** Request a certain transform widget for the next update */
+	void RequestTransformWidgetMode(UE::Widget::EWidgetMode InWidgetMode);
+	
 private:
 	/** Whether or not Pivot Transforms have changed, in which case we need to redraw viewport*/
 	bool HasPivotTransformsChanged() const;
@@ -389,6 +393,7 @@ private:
 	void SetRigElementSelectionInternal(UControlRig* ControlRig, ERigElementType Type, const FName& InRigElementName, bool bSelected);
 	
 	FEditorViewportClient* CurrentViewportClient;
+	TArray<UE::Widget::EWidgetMode> RequestedWidgetModes;
 
 /* store coordinate system per widget mode*/
 private:
diff --git a/Engine/Plugins/Runtime/RigVM/Source/RigVMDeveloper/Private/RigVMModel/Nodes/RigVMFunctionEntryNode.cpp b/Engine/Plugins/Runtime/RigVM/Source/RigVMDeveloper/Private/RigVMModel/Nodes/RigVMFunctionEntryNode.cpp
--- a/Engine/Plugins/Runtime/RigVM/Source/RigVMDeveloper/Private/RigVMModel/Nodes/RigVMFunctionEntryNode.cpp
+++ b/Engine/Plugins/Runtime/RigVM/Source/RigVMDeveloper/Private/RigVMModel/Nodes/RigVMFunctionEntryNode.cpp
@@ -2,6 +2,7 @@
 
 #include "RigVMModel/Nodes/RigVMFunctionEntryNode.h"
 #include "RigVMModel/RigVMFunctionLibrary.h"
+#include "RigVMModel/Nodes/RigVMCollapseNode.h"
 
 #include UE_INLINE_GENERATED_CPP_BY_NAME(RigVMFunctionEntryNode)
 
@@ -29,6 +30,15 @@ bool URigVMFunctionEntryNode::IsDefinedAsVarying() const
 	return true; 
 }
 
+bool URigVMFunctionEntryNode::IsWithinLoop() const
+{
+	if(const URigVMCollapseNode* CollapseNode = Cast<URigVMCollapseNode>(GetGraph()->GetOuter()))
+	{
+		return CollapseNode->IsWithinLoop();
+	}
+	return Super::IsWithinLoop();
+}
+
 FString URigVMFunctionEntryNode::GetNodeTitle() const
 {
 	return TEXT("Entry");
diff --git a/Engine/Plugins/Runtime/RigVM/Source/RigVMDeveloper/Private/RigVMModel/Nodes/RigVMUnitNode.cpp b/Engine/Plugins/Runtime/RigVM/Source/RigVMDeveloper/Private/RigVMModel/Nodes/RigVMUnitNode.cpp
--- a/Engine/Plugins/Runtime/RigVM/Source/RigVMDeveloper/Private/RigVMModel/Nodes/RigVMUnitNode.cpp
+++ b/Engine/Plugins/Runtime/RigVM/Source/RigVMDeveloper/Private/RigVMModel/Nodes/RigVMUnitNode.cpp
@@ -5,6 +5,8 @@
 #include "RigVMCore/RigVMStruct.h"
 #include "RigVMModel/RigVMControllerActions.h"
 #include "RigVMUserWorkflowRegistry.h"
+#include "RigVMHost.h"
+#include "RigVMBlueprint.h"
 
 #include UE_INLINE_GENERATED_CPP_BY_NAME(RigVMUnitNode)
 
@@ -353,6 +355,305 @@ TSharedPtr<FStructOnScope> URigVMUnitNode::ConstructStructInstance(bool bUseDefa
 	return nullptr;
 }
 
+void URigVMUnitNode::EnumeratePropertiesOnHostAndStructInstance(URigVMHost* InHost, TSharedPtr<FStructOnScope> InInstance, bool bPreferLiterals, TFunction<void(const URigVMPin*,const FProperty*,uint8*,const FProperty*,uint8*)> InEnumerationFunction, int32 InSliceIndex) const
+{
+	check(InHost);
+	check(InInstance.IsValid() && InInstance->IsValid());
+
+	const URigVMBlueprint* Blueprint = GetTypedOuter<URigVMBlueprint>();
+	if(Blueprint == nullptr)
+	{
+		return;
+	}
+
+	URigVM* VM = InHost->GetVM();
+	if(VM == nullptr)
+	{
+		return;
+	}
+
+	const FRigVMByteCode& ByteCode = VM->GetByteCode();
+	const TArray<int32>& InstructionIndices = ByteCode.GetAllInstructionIndicesForSubject(const_cast<URigVMUnitNode*>(this));
+	if(InstructionIndices.IsEmpty())
+	{
+		return;
+	}
+
+	const FRigVMInstructionArray& Instructions = ByteCode.GetInstructions();
+	for(int32 Index = 0; Index < InstructionIndices.Num(); Index++)
+	{
+		const int32 InstructionIndex = InstructionIndices[Index]; 
+		if(Instructions[InstructionIndices[Index]].OpCode != ERigVMOpCode::Execute)
+		{
+			continue;
+		}
+
+		const FRigVMInstruction& Instruction = Instructions[InstructionIndex];
+		const FRigVMOperandArray Operands = ByteCode.GetOperandsForOp(Instruction);
+		const UScriptStruct* ScriptStruct = Cast<UScriptStruct>(InInstance->GetStruct());
+
+		const FRigVMFunction* Function = FRigVMRegistry::Get().FindFunction(const_cast<UScriptStruct*>(ScriptStruct), *GetMethodName().ToString());
+		if(Function == nullptr)
+		{
+			return;
+		}
+
+		for(int32 ArgumentIndex = 0; ArgumentIndex < Function->GetArguments().Num(); ArgumentIndex++)
+		{
+			const FRigVMFunctionArgument& Argument = Function->GetArguments()[ArgumentIndex];
+			const URigVMPin* Pin = FindPin(Argument.Name);
+			if(Pin == nullptr)
+			{
+				continue;
+			}
+			
+			FRigVMOperand Operand = Operands[ArgumentIndex];
+
+			if(bPreferLiterals && Operand.GetMemoryType() != ERigVMMemoryType::Literal)
+			{
+				// walk backwards on the bytecode to find the copy expression copying
+				// the default value literal into the target operand during initialization
+				for(int32 CopyInstructionIndex = InstructionIndex - 1; CopyInstructionIndex >= 0; CopyInstructionIndex--)
+				{
+					const FRigVMInstruction& CopyInstruction = Instructions[CopyInstructionIndex];
+					if(CopyInstruction.OpCode == ERigVMOpCode::Copy)
+					{
+						const FRigVMCopyOp& CopyOp = ByteCode.GetOpAt<FRigVMCopyOp>(CopyInstruction);
+						if(CopyOp.Target == Operand &&
+							CopyOp.Source.GetMemoryType() == ERigVMMemoryType::Literal &&
+							CopyOp.Source.GetRegisterOffset() == INDEX_NONE)
+						{
+							Operand = CopyOp.Source;
+							break;
+						}
+					}
+				}
+			}
+
+			const FProperty* InstanceProperty = ScriptStruct->FindPropertyByName(Argument.Name);
+			if(InstanceProperty == nullptr)
+			{
+				continue;
+			}
+
+			uint8* InstanceMemory = InstanceProperty->ContainerPtrToValuePtr<uint8>(InInstance->GetStructMemory());
+			uint8* HostMemory = nullptr;
+			const FProperty* HostProperty = nullptr;
+			
+			switch(Operand.GetMemoryType())
+			{
+			case ERigVMMemoryType::External:
+				{
+					TArray<FRigVMExternalVariable> ExternalVariables = InHost->GetExternalVariables();
+					if(ExternalVariables.IsValidIndex(Operand.GetRegisterIndex()))
+					{
+						const FRigVMExternalVariable& ExternalVariable = ExternalVariables[Operand.GetRegisterIndex()];
+						HostMemory = ExternalVariable.Memory; 
+						HostProperty = ExternalVariable.Property;
+					}
+					break;
+				}
+			case ERigVMMemoryType::Literal:
+				{
+					if(URigVMMemoryStorage* LiteralMemory = VM->GetLiteralMemory(false))
+					{
+						if(const FProperty* Property = LiteralMemory->GetProperty(Operand.GetRegisterIndex()))
+						{
+							HostMemory = Property->ContainerPtrToValuePtr<uint8>(LiteralMemory);
+							HostProperty = Property;
+						}
+					}
+					break;
+				}
+			case ERigVMMemoryType::Work:
+				{
+					if(URigVMMemoryStorage* WorkMemory = VM->GetWorkMemory(false))
+					{
+						if(const FProperty* Property = WorkMemory->GetProperty(Operand.GetRegisterIndex()))
+						{
+							HostMemory = Property->ContainerPtrToValuePtr<uint8>(WorkMemory);
+							HostProperty = Property;
+						}
+					}
+					break;
+				}
+			default:
+				{
+					break;
+				}
+			}
+
+			// hidden properties are backed up by sliced memory
+			if(Pin->GetDirection() == ERigVMPinDirection::Hidden)
+			{
+				if(InSliceIndex < 0)
+				{
+					InSliceIndex = 0;
+				}
+				
+				if(const FArrayProperty* HostArrayProperty = CastField<FArrayProperty>(HostProperty))
+				{
+					FScriptArrayHelper ArrayHelper(HostArrayProperty, HostMemory);
+					if(ArrayHelper.Num() <= InSliceIndex)
+					{
+						UE_LOG(LogRigVMDeveloper, Warning,
+							TEXT("Provided slice index %d is invalid when enumerating properties on node '%s' (total number of slices is %d)."),
+							InSliceIndex, *GetPathName(), ArrayHelper.Num());
+						continue;
+					}
+
+					HostProperty = HostArrayProperty->Inner;
+					HostMemory = ArrayHelper.GetRawPtr(InSliceIndex);
+				}
+			}
+
+			if(HostMemory && HostProperty && InstanceProperty && InstanceMemory)
+			{
+				InEnumerationFunction(Pin, HostProperty, HostMemory, InstanceProperty, InstanceMemory);
+			}
+		}
+	}
+}
+
+TSharedPtr<FStructOnScope> URigVMUnitNode::ConstructLiveStructInstance(URigVMHost* InHost, int32 InSliceIndex) const
+{
+	check(InHost);
+	
+	TSharedPtr<FStructOnScope> NodeInstance = ConstructStructInstance(false);
+	if(!NodeInstance.IsValid() || !NodeInstance->IsValid())
+	{
+		return NodeInstance;
+	}
+
+	EnumeratePropertiesOnHostAndStructInstance(InHost, NodeInstance, false,
+	[](const URigVMPin* Pin, const FProperty* HostProperty, uint8* HostMemory, const FProperty* InstanceProperty, uint8* InstanceMemory) {
+
+		// copy the live memory from the host into the node instance
+		URigVMMemoryStorage::CopyProperty(InstanceProperty, InstanceMemory, HostProperty, HostMemory);
+	}, InSliceIndex);
+	
+	return NodeInstance;
+}
+
+bool URigVMUnitNode::UpdateHostFromStructInstance(URigVMHost* InHost, TSharedPtr<FStructOnScope> InInstance, int32 InSliceIndex) const
+{
+	check(InHost);
+	check(InInstance.IsValid() && InInstance->IsValid());
+
+	bool bChangedSomething = false;
+	
+	EnumeratePropertiesOnHostAndStructInstance(InHost, InInstance, true,
+	[&bChangedSomething]
+	(const URigVMPin* Pin, const FProperty* HostProperty, uint8* HostMemory, const FProperty* InstanceProperty, uint8* InstanceMemory) {
+
+		// only update registers for inputs
+		if(Pin->GetDirection() != ERigVMPinDirection::Input &&
+			Pin->GetDirection() != ERigVMPinDirection::IO &&
+			Pin->GetDirection() != ERigVMPinDirection::Visible)
+		{
+			return;
+		}
+
+		// only update registers for pins which don't have a top level link
+		if(!Pin->GetRootPin()->GetLinkedSourcePins().IsEmpty())
+		{
+			return;
+		}
+
+		// ignore properties with matching data
+		if(HostProperty->SameType(InstanceProperty))
+		{
+			if(HostProperty->Identical(HostMemory, InstanceMemory))
+			{
+				return;
+			}
+		}
+
+		// copy the memory back to the host
+		URigVMMemoryStorage::CopyProperty(HostProperty, HostMemory, InstanceProperty, InstanceMemory);
+		bChangedSomething = true;
+
+	}, InSliceIndex);
+
+	return bChangedSomething;
+}
+
+void URigVMUnitNode::ComputePinValueDifferences(TSharedPtr<FStructOnScope> InCurrentInstance, TSharedPtr<FStructOnScope> InDesiredInstance, TMap<FString, FString>& OutNewPinDefaultValues) const
+{
+	check(InCurrentInstance.IsValid() && InCurrentInstance->IsValid());
+	check(InDesiredInstance.IsValid() && InDesiredInstance->IsValid());
+	
+	const UScriptStruct* ScriptStruct = Cast<UScriptStruct>(InCurrentInstance->GetStruct());
+	check(ScriptStruct == Cast<UScriptStruct>(InDesiredInstance->GetStruct()));
+	
+	for(URigVMPin* Pin : GetPins())
+	{
+		if(Pin->GetDirection() != ERigVMPinDirection::Input &&
+			Pin->GetDirection() != ERigVMPinDirection::IO &&
+			Pin->GetDirection() != ERigVMPinDirection::Visible)
+		{
+			continue;
+		}
+
+		const FProperty* Property = ScriptStruct->FindPropertyByName(Pin->GetFName());
+		if(Property == nullptr)
+		{
+			continue;
+		}
+
+		const uint8* CurrentMemory = Property->ContainerPtrToValuePtr<uint8>(InCurrentInstance->GetStructMemory());
+		const uint8* DesiredMemory = Property->ContainerPtrToValuePtr<uint8>(InDesiredInstance->GetStructMemory());
+		if(Property->Identical(CurrentMemory, DesiredMemory))
+		{
+			continue;
+		}
+
+		FString DefaultValue;
+		Property->ExportText_Direct(
+			DefaultValue,
+			DesiredMemory,
+			DesiredMemory,
+			nullptr,
+			PPF_None,
+			nullptr
+		);
+
+		if(!DefaultValue.IsEmpty())
+		{
+			OutNewPinDefaultValues.Add(Property->GetName(), DefaultValue);
+		}
+	}
+}
+
+void URigVMUnitNode::ComputePinValueDifferences(TSharedPtr<FStructOnScope> InDesiredInstance, TMap<FString, FString>& OutNewPinDefaultValues) const
+{
+	const TSharedPtr<FStructOnScope> DefaultInstance = ConstructStructInstance(false);
+	ComputePinValueDifferences(DefaultInstance, InDesiredInstance, OutNewPinDefaultValues);
+}
+
+bool URigVMUnitNode::IsPartOfRuntime() const
+{
+	if(const URigVMBlueprint* Blueprint = GetTypedOuter<URigVMBlueprint>())
+	{
+		if(URigVMHost* DebuggedHost = Cast<URigVMHost>(Blueprint->GetObjectBeingDebugged()))
+		{
+			return IsPartOfRuntime(DebuggedHost);
+		}
+	}
+	return false;
+}
+
+bool URigVMUnitNode::IsPartOfRuntime(URigVMHost* InHost) const
+{
+	check(InHost);
+	if(URigVM* VM = InHost->GetVM())
+	{
+		const FRigVMByteCode& ByteCode = VM->GetByteCode();
+		return ByteCode.GetFirstInstructionIndexForSubject(const_cast<URigVMUnitNode*>(this)) != INDEX_NONE;
+	}
+	return false;
+}
+
 FRigVMStructUpgradeInfo URigVMUnitNode::GetUpgradeInfo() const
 {
 	if(UScriptStruct* Struct = GetScriptStruct())
diff --git a/Engine/Plugins/Runtime/RigVM/Source/RigVMDeveloper/Private/RigVMModel/RigVMController.cpp b/Engine/Plugins/Runtime/RigVM/Source/RigVMDeveloper/Private/RigVMModel/RigVMController.cpp
--- a/Engine/Plugins/Runtime/RigVM/Source/RigVMDeveloper/Private/RigVMModel/RigVMController.cpp
+++ b/Engine/Plugins/Runtime/RigVM/Source/RigVMDeveloper/Private/RigVMModel/RigVMController.cpp
@@ -7762,7 +7762,7 @@ FString URigVMController::GetPinDefaultValue(const FString& InPinPath)
 	return Pin->GetDefaultValue();
 }
 
-bool URigVMController::SetPinDefaultValue(const FString& InPinPath, const FString& InDefaultValue, bool bResizeArrays, bool bSetupUndoRedo, bool bMergeUndoAction, bool bPrintPythonCommand)
+bool URigVMController::SetPinDefaultValue(const FString& InPinPath, const FString& InDefaultValue, bool bResizeArrays, bool bSetupUndoRedo, bool bMergeUndoAction, bool bPrintPythonCommand, bool bSetValueOnLinkedPins)
 {
 	if (!IsValidGraph())
 	{
@@ -7792,7 +7792,7 @@ bool URigVMController::SetPinDefaultValue(const FString& InPinPath, const FStrin
 		}
 	}
 	
-	if (!SetPinDefaultValue(Pin, InDefaultValue, bResizeArrays, bSetupUndoRedo, bMergeUndoAction))
+	if (!SetPinDefaultValue(Pin, InDefaultValue, bResizeArrays, bSetupUndoRedo, bMergeUndoAction, bSetValueOnLinkedPins))
 	{
 		return false;
 	}
@@ -7800,7 +7800,7 @@ bool URigVMController::SetPinDefaultValue(const FString& InPinPath, const FStrin
 	URigVMPin* PinForLink = Pin->GetPinForLink();
 	if (PinForLink != Pin)
 	{
-		if (!SetPinDefaultValue(PinForLink, InDefaultValue, bResizeArrays, false, bMergeUndoAction))
+		if (!SetPinDefaultValue(PinForLink, InDefaultValue, bResizeArrays, false, bMergeUndoAction, bSetValueOnLinkedPins))
 		{
 			return false;
 		}
@@ -7821,12 +7821,20 @@ bool URigVMController::SetPinDefaultValue(const FString& InPinPath, const FStrin
 	return true;
 }
 
-bool URigVMController::SetPinDefaultValue(URigVMPin* InPin, const FString& InDefaultValue, bool bResizeArrays, bool bSetupUndoRedo, bool bMergeUndoAction)
+bool URigVMController::SetPinDefaultValue(URigVMPin* InPin, const FString& InDefaultValue, bool bResizeArrays, bool bSetupUndoRedo, bool bMergeUndoAction, bool bSetValueOnLinkedPins)
 {
 	if (!bIsTransacting && !IsGraphEditable())
 	{
 		return false;
 	}
+
+	if(!bSetValueOnLinkedPins)
+	{
+		if(!InPin->GetSourceLinks(false).IsEmpty())
+		{
+			return false;
+		}
+	}
 	
 	check(InPin);
 
@@ -7897,7 +7905,7 @@ bool URigVMController::SetPinDefaultValue(URigVMPin* InPin, const FString& InDef
 					PostProcessDefaultValue(SubPin, Elements[ElementIndex]);
 					if (!Elements[ElementIndex].IsEmpty())
 					{
-						SetPinDefaultValue(SubPin, Elements[ElementIndex], bResizeArrays, false, false);
+						SetPinDefaultValue(SubPin, Elements[ElementIndex], bResizeArrays, false, false, bSetValueOnLinkedPins);
 						bSetPinDefaultValueSucceeded = true;
 					}
 				}
@@ -7919,7 +7927,7 @@ bool URigVMController::SetPinDefaultValue(URigVMPin* InPin, const FString& InDef
 					PostProcessDefaultValue(SubPin, MemberValue);
 					if (!MemberValue.IsEmpty())
 					{
-						SetPinDefaultValue(SubPin, MemberValue, bResizeArrays, false, false);
+						SetPinDefaultValue(SubPin, MemberValue, bResizeArrays, false, false, bSetValueOnLinkedPins);
 						bSetPinDefaultValueSucceeded = true;
 					}
 				}
diff --git a/Engine/Plugins/Runtime/RigVM/Source/RigVMDeveloper/Private/RigVMModel/RigVMNode.cpp b/Engine/Plugins/Runtime/RigVM/Source/RigVMDeveloper/Private/RigVMModel/RigVMNode.cpp
--- a/Engine/Plugins/Runtime/RigVM/Source/RigVMDeveloper/Private/RigVMModel/RigVMNode.cpp
+++ b/Engine/Plugins/Runtime/RigVM/Source/RigVMDeveloper/Private/RigVMModel/RigVMNode.cpp
@@ -699,6 +699,33 @@ const bool URigVMNode::IsControlFlowBlockSliced(const FName& InBlockName) const
 	return false;
 }
 
+bool URigVMNode::IsWithinLoop() const
+{
+	for(const URigVMPin* Pin : Pins)
+	{
+		const TArray<URigVMPin*> SourcePins = Pin->GetLinkedSourcePins(true);
+		for(const URigVMPin* SourcePin : SourcePins)
+		{
+			if(SourcePin->GetNode()->IsLoopNode())
+			{
+				if(!SourcePin->IsExecuteContext() || SourcePin->GetFName() != FRigVMStruct::ForLoopCompletedPinName)
+				{
+					return true;
+				}
+			}
+		}
+
+		for(const URigVMPin* SourcePin : SourcePins)
+		{
+			if(SourcePin->GetNode()->IsWithinLoop())
+			{
+				return true;
+			}
+		}
+	}
+	return false;
+}
+
 TArray<FRigVMUserWorkflow> URigVMNode::GetSupportedWorkflows(ERigVMUserWorkflowType InType, const UObject* InSubject) const
 {
 	if(InSubject == nullptr)
diff --git a/Engine/Plugins/Runtime/RigVM/Source/RigVMDeveloper/Public/RigVMModel/Nodes/RigVMFunctionEntryNode.h b/Engine/Plugins/Runtime/RigVM/Source/RigVMDeveloper/Public/RigVMModel/Nodes/RigVMFunctionEntryNode.h
--- a/Engine/Plugins/Runtime/RigVM/Source/RigVMDeveloper/Public/RigVMModel/Nodes/RigVMFunctionEntryNode.h
+++ b/Engine/Plugins/Runtime/RigVM/Source/RigVMDeveloper/Public/RigVMModel/Nodes/RigVMFunctionEntryNode.h
@@ -25,6 +25,7 @@ public:
 	// Override node functions
 	virtual FLinearColor GetNodeColor() const override;
 	virtual bool IsDefinedAsVarying() const override;
+	virtual bool IsWithinLoop() const override;
 
 	// URigVMNode interface
 	virtual FString GetNodeTitle() const override;
diff --git a/Engine/Plugins/Runtime/RigVM/Source/RigVMDeveloper/Public/RigVMModel/Nodes/RigVMUnitNode.h b/Engine/Plugins/Runtime/RigVM/Source/RigVMDeveloper/Public/RigVMModel/Nodes/RigVMUnitNode.h
--- a/Engine/Plugins/Runtime/RigVM/Source/RigVMDeveloper/Public/RigVMModel/Nodes/RigVMUnitNode.h
+++ b/Engine/Plugins/Runtime/RigVM/Source/RigVMDeveloper/Public/RigVMModel/Nodes/RigVMUnitNode.h
@@ -7,6 +7,8 @@
 #include "UObject/StructOnScope.h"
 #include "RigVMUnitNode.generated.h"
 
+class URigVMHost;
+
 /**
  * The Struct Node represents a Function Invocation of a RIGVM_METHOD
  * declared on a USTRUCT. Struct Nodes have input / output pins for all
@@ -55,6 +57,24 @@ public:
 	// @param bUseDefault If set to true the default struct will be created - otherwise the struct will contains the values from the node
 	TSharedPtr<FStructOnScope> ConstructStructInstance(bool bUseDefault = false) const;
 
+	// Returns an instance of the struct with values backed by the memory of a currently running host
+	TSharedPtr<FStructOnScope> ConstructLiveStructInstance(URigVMHost* InHost, int32 InSliceIndex = 0) const;
+	
+	// Updates the memory of a host to match a specific struct instance
+	bool UpdateHostFromStructInstance(URigVMHost* InHost, TSharedPtr<FStructOnScope> InInstance, int32 InSliceIndex = 0) const;
+
+	// Compares two struct instances and returns differences in pin values
+	void ComputePinValueDifferences(TSharedPtr<FStructOnScope> InCurrentInstance, TSharedPtr<FStructOnScope> InDesiredInstance, TMap<FString, FString>& OutNewPinDefaultValues) const;
+
+	// Compares a desired struct instance with the node's default and returns differences in pin values
+	void ComputePinValueDifferences(TSharedPtr<FStructOnScope> InDesiredInstance, TMap<FString, FString>& OutNewPinDefaultValues) const;
+
+	// Returns true if the node is part of the debugged runtime rig
+	bool IsPartOfRuntime() const;
+
+	// Returns true if the node is part of the debugged runtime rig
+	bool IsPartOfRuntime(URigVMHost* InHost) const;
+
 	// Returns a copy of the struct with the current values
 	template <
 		typename T,
@@ -80,6 +100,12 @@ protected:
 
 	virtual FText GetToolTipTextForPin(const URigVMPin* InPin) const override;
 	virtual bool ShouldInputPinComputeLazily(const URigVMPin* InPin) const override;
+	void EnumeratePropertiesOnHostAndStructInstance(
+		URigVMHost* InHost,
+		TSharedPtr<FStructOnScope> InInstance, 
+		bool bPreferLiterals,
+		TFunction<void(const URigVMPin*,const FProperty*,uint8*,const FProperty*,uint8*)> InEnumerationFunction,
+		int32 InSliceIndex = 0) const;
 
 private:
 
diff --git a/Engine/Plugins/Runtime/RigVM/Source/RigVMDeveloper/Public/RigVMModel/RigVMController.h b/Engine/Plugins/Runtime/RigVM/Source/RigVMDeveloper/Public/RigVMModel/RigVMController.h
--- a/Engine/Plugins/Runtime/RigVM/Source/RigVMDeveloper/Public/RigVMModel/RigVMController.h
+++ b/Engine/Plugins/Runtime/RigVM/Source/RigVMDeveloper/Public/RigVMModel/RigVMController.h
@@ -739,7 +739,7 @@ public:
 	// Sets the default value of a pin given its pinpath.
 	// This causes a PinDefaultValueChanged modified event.
 	UFUNCTION(BlueprintCallable, Category = RigVMController)
-	bool SetPinDefaultValue(const FString& InPinPath, const FString& InDefaultValue, bool bResizeArrays = true, bool bSetupUndoRedo = true, bool bMergeUndoAction = false, bool bPrintPythonCommand = false);
+	bool SetPinDefaultValue(const FString& InPinPath, const FString& InDefaultValue, bool bResizeArrays = true, bool bSetupUndoRedo = true, bool bMergeUndoAction = false, bool bPrintPythonCommand = false, bool bSetValueOnLinkedPins = true);
 
 	// Resets the default value of a pin given its pinpath.
 	// This causes a PinDefaultValueChanged modified event.
@@ -1112,7 +1112,7 @@ private:
 	void ConfigurePinFromProperty(FProperty* InProperty, URigVMPin* InOutPin, ERigVMPinDirection InPinDirection = ERigVMPinDirection::Invalid);
 	void ConfigurePinFromPin(URigVMPin* InOutPin, URigVMPin* InPin, bool bCopyDisplayName = false);
 	void ConfigurePinFromArgument(URigVMPin* InOutPin, const FRigVMGraphFunctionArgument& InArgument, bool bCopyDisplayName = false);
-	bool SetPinDefaultValue(URigVMPin* InPin, const FString& InDefaultValue, bool bResizeArrays, bool bSetupUndoRedo, bool bMergeUndoAction);
+	bool SetPinDefaultValue(URigVMPin* InPin, const FString& InDefaultValue, bool bResizeArrays, bool bSetupUndoRedo, bool bMergeUndoAction, bool bSetValueOnLinkedPins = true);
 	bool ResetPinDefaultValue(URigVMPin* InPin, bool bSetupUndoRedo);
 	static FString GetPinInitialDefaultValue(const URigVMPin* InPin);
 	static FString GetPinInitialDefaultValueFromStruct(UScriptStruct* ScriptStruct, const URigVMPin* InPin, uint32 InOffset);
diff --git a/Engine/Plugins/Runtime/RigVM/Source/RigVMDeveloper/Public/RigVMModel/RigVMNode.h b/Engine/Plugins/Runtime/RigVM/Source/RigVMDeveloper/Public/RigVMModel/RigVMNode.h
--- a/Engine/Plugins/Runtime/RigVM/Source/RigVMDeveloper/Public/RigVMModel/RigVMNode.h
+++ b/Engine/Plugins/Runtime/RigVM/Source/RigVMDeveloper/Public/RigVMModel/RigVMNode.h
@@ -240,6 +240,9 @@ public:
 	// returns true if a control flow block requires slicing
 	virtual const bool IsControlFlowBlockSliced(const FName& InBlockName) const;
 
+	// returns true if the node is within a loop
+	virtual bool IsWithinLoop() const;
+
 	// returns true if the node can be upgraded
 	UFUNCTION(BlueprintPure, Category = RigVMNode)
 	virtual bool CanBeUpgraded() const { return GetUpgradeInfo().IsValid(); }
