Control Rig: Make sure to set preferred euler angles when adding a control programatically

diff --git a/Engine/Plugins/Animation/ControlRig/Source/ControlRig/Private/Rigs/RigHierarchyController.cpp b/Engine/Plugins/Animation/ControlRig/Source/ControlRig/Private/Rigs/RigHierarchyController.cpp
--- a/Engine/Plugins/Animation/ControlRig/Source/ControlRig/Private/Rigs/RigHierarchyController.cpp
+++ b/Engine/Plugins/Animation/ControlRig/Source/ControlRig/Private/Rigs/RigHierarchyController.cpp
@@ -1,6 +1,7 @@
 // Copyright Epic Games, Inc. All Rights Reserved.
 
 #include "Rigs/RigHierarchyController.h"
+#include "AnimationCoreLibrary.h"
 #include "UObject/Package.h"
 
 #if WITH_EDITOR
@@ -354,6 +355,10 @@ FRigElementKey URigHierarchyController::AddControl(
 		NewElement->Offset.Set(ERigTransformType::InitialLocal, InOffsetTransform);  
 		NewElement->Shape.Set(ERigTransformType::InitialLocal, InShapeTransform);  
 		Hierarchy->SetControlValue(NewElement, InValue, ERigControlValueType::Initial, false);
+		const FTransform LocalTransform = Hierarchy->GetTransform(NewElement, ERigTransformType::InitialLocal);
+
+		const FVector PreferredEulerAngles = AnimationCore::EulerFromQuat(LocalTransform.GetRotation(), NewElement->Settings.PreferredRotationOrder);
+		Hierarchy->SetControlPreferredEulerAngles(NewElement, PreferredEulerAngles, NewElement->Settings.PreferredRotationOrder, true);
 
 		NewElement->Offset.MarkDirty(ERigTransformType::InitialGlobal);
 		NewElement->Pose.MarkDirty(ERigTransformType::InitialGlobal);
