Sequencer: Control Rig: Can select and pose actors with controls and Ctrl selection works again do to issue with Ctrl also selecting a Widget axis. Also fix issue with shift moving camera when in marquee mode

diff --git a/Engine/Plugins/Animation/ControlRig/Source/ControlRigEditor/Private/EditMode/ControlRigEditMode.cpp b/Engine/Plugins/Animation/ControlRig/Source/ControlRigEditor/Private/EditMode/ControlRigEditMode.cpp
--- a/Engine/Plugins/Animation/ControlRig/Source/ControlRigEditor/Private/EditMode/ControlRigEditMode.cpp
+++ b/Engine/Plugins/Animation/ControlRig/Source/ControlRigEditor/Private/EditMode/ControlRigEditMode.cpp
@@ -1262,7 +1262,9 @@ bool FControlRigEditMode::GetCustomInputCoordinateSystem(FMatrix& OutMatrix, voi
 bool FControlRigEditMode::HandleClick(FEditorViewportClient* InViewportClient, HHitProxy *HitProxy, const FViewportClick &Click)
 {
 	const bool bClickSelectThroughGizmo = CVarClickSelectThroughGizmo.GetValueOnGameThread();
-	if (bClickSelectThroughGizmo == false)
+	//if Control is down we act like we are selecting an axis so don't do this check
+	//if doing control else we can't do control selection anymore, see FMouseDeltaTracker::DetermineCurrentAxis(
+	if (Click.IsControlDown() == false && bClickSelectThroughGizmo == false)
 	{
 		const EAxisList::Type CurrentAxis = InViewportClient->GetCurrentWidgetAxis();
 		//if we are hitting a widget, besides arcball then bail saying we are handling it
@@ -1769,6 +1771,22 @@ void FControlRigEditMode::SelectNone()
 	FEdMode::SelectNone();
 }
 
+static bool IsDoingDrag(const TWeakPtr<ISequencer>& WeakSequencer, FViewport* InViewport, EAxisList::Type CurrentAxis)
+{
+	if (WeakSequencer.IsValid())
+	{
+		ISequencer* Sequencer = WeakSequencer.Pin().Get();
+		const USequencerSettings* SequencerSettings = Sequencer->GetSequencerSettings();
+		const bool LeftMouseButtonDown = InViewport->KeyState(EKeys::LeftMouseButton);
+		const bool bIsCtrlKeyDown = InViewport->KeyState(EKeys::LeftControl) || InViewport->KeyState(EKeys::RightControl);
+		const bool bIsAltKeyDown = InViewport->KeyState(EKeys::LeftAlt) || InViewport->KeyState(EKeys::RightAlt);
+
+		//if shfit is down we still want to drag
+
+		return LeftMouseButtonDown && (CurrentAxis == EAxisList::None) && !bIsCtrlKeyDown && !bIsAltKeyDown && (SequencerSettings ? SequencerSettings->GetLeftMouseDragDoesMarquee() : false);
+	}
+	return false;
+}
 bool FControlRigEditMode::InputDelta(FEditorViewportClient* InViewportClient, FViewport* InViewport, FVector& InDrag, FRotator& InRot, FVector& InScale)
 {
 	if (IsDragAnimSliderToolPressed(InViewport)) //this is needed to make sure we get all of the processed mouse events, for some reason the above may not return true
@@ -1796,7 +1814,7 @@ bool FControlRigEditMode::InputDelta(FEditorViewportClient* InViewportClient, FV
 	const bool bDoTranslation = !Drag.IsZero() && (WidgetMode == UE::Widget::WM_Translate || WidgetMode == UE::Widget::WM_TranslateRotateZ);
 	const bool bDoScale = !Scale.IsZero() && WidgetMode == UE::Widget::WM_Scale;
 
-	if (!InDrag.IsZero() && bShiftDown)
+	if (IsDoingDrag(WeakSequencer, InViewport, GetCurrentWidgetAxis()) == false && !InDrag.IsZero() && bShiftDown)
 	{
 		FVector CameraDelta(InDrag);
 
@@ -1938,7 +1956,8 @@ bool FControlRigEditMode::InputDelta(FEditorViewportClient* InViewportClient, FV
 	{
 		TickManipulatableObjects(0.f);
 	}
-	return bManipulatorMadeChange;
+	//if in level editor we want to move other things also
+	return IsInLevelEditor() ? false :bManipulatorMadeChange;
 }
 
 bool FControlRigEditMode::ShouldDrawWidget() const
diff --git a/Engine/Plugins/Animation/ControlRig/Source/ControlRigEditor/Private/Sequencer/ControlRigParameterTrackEditor.cpp b/Engine/Plugins/Animation/ControlRig/Source/ControlRigEditor/Private/Sequencer/ControlRigParameterTrackEditor.cpp
--- a/Engine/Plugins/Animation/ControlRig/Source/ControlRigEditor/Private/Sequencer/ControlRigParameterTrackEditor.cpp
+++ b/Engine/Plugins/Animation/ControlRig/Source/ControlRigEditor/Private/Sequencer/ControlRigParameterTrackEditor.cpp
@@ -2445,8 +2445,9 @@ void FControlRigParameterTrackEditor::OnSelectionChanged(TArray<UMovieSceneTrack
 	const bool UseSelectedKeys = CVarSelectedKeysSelectControls.GetValueOnGameThread();
 	GetSequencer()->GetSelectedKeyAreas(KeyAreas, UseSelectedKeys);
 	if (KeyAreas.Num() <= 0)
-	{
-		if (ControlRigEditMode)
+	{ 
+		if (FSlateApplication::Get().GetModifierKeys().IsShiftDown() == false && 
+			FSlateApplication::Get().GetModifierKeys().IsControlDown() == false && ControlRigEditMode)
 		{
 			TMap<UControlRig*, TArray<FRigElementKey>> AllSelectedControls;
 			ControlRigEditMode->GetAllSelectedControls(AllSelectedControls);
