Control Rig: Follow up to direct manipulation - fix control rig editor style

diff --git a/Engine/Plugins/Animation/ControlRig/Source/ControlRigEditor/Private/ControlRigEditorStyle.h b/Engine/Plugins/Animation/ControlRig/Source/ControlRigEditor/Private/ControlRigEditorStyle.h
--- a/Engine/Plugins/Animation/ControlRig/Source/ControlRigEditor/Private/ControlRigEditorStyle.h
+++ b/Engine/Plugins/Animation/ControlRig/Source/ControlRigEditor/Private/ControlRigEditorStyle.h
@@ -104,6 +104,8 @@ public:
 				Set( "ControlRig.Viewport.Border", new BOX_BRUSH( "Old/Window/ViewportDebugBorder", 0.8f, FLinearColor(1.0f,1.0f,1.0f,1.0f) ) );
 				// similar style to "AnimViewport.Notification.Warning"
 				Set( "ControlRig.Viewport.Notification.ChangeShapeTransform", new BOX_BRUSH("Common/RoundedSelection_16x", 4.0f/16.0f, FLinearColor(FColor(169, 0, 148))));
+				// similar style to "AnimViewport.Notification.Warning"
+				Set( "ControlRig.Viewport.Notification.DirectManipulation", new BOX_BRUSH("Common/RoundedSelection_16x", 4.0f/16.0f, FLinearColor(FColor(0, 112, 224))));
 			}
 		}
 
