Constraints: notify when the parent transform is modified via the editor

diff --git a/Engine/Plugins/Animation/ControlRig/Source/ControlRig/Private/Constraints/ControlRigTransformableHandle.cpp b/Engine/Plugins/Animation/ControlRig/Source/ControlRig/Private/Constraints/ControlRigTransformableHandle.cpp
--- a/Engine/Plugins/Animation/ControlRig/Source/ControlRig/Private/Constraints/ControlRigTransformableHandle.cpp
+++ b/Engine/Plugins/Animation/ControlRig/Source/ControlRig/Private/Constraints/ControlRigTransformableHandle.cpp
@@ -392,14 +392,28 @@ void UTransformableControlHandle::OnControlModified(
 		return;
 	}
 
-	if (ControlRig == InControlRig && InControl->GetName() == ControlName)
+	if (OnHandleModified.IsBound() && (ControlRig == InControlRig))
 	{
-		if(OnHandleModified.IsBound())
-		{
-			const EHandleEvent Event = InContext.bConstraintUpdate ?
-				EHandleEvent::GlobalTransformUpdated : EHandleEvent::LocalTransformUpdated; 
+		const EHandleEvent Event = InContext.bConstraintUpdate ?
+			EHandleEvent::GlobalTransformUpdated : EHandleEvent::LocalTransformUpdated;
+
+		if (InControl->GetName() == ControlName)
+		{	// if that handle is wrapping InControl  
 			OnHandleModified.Broadcast(this, Event);
 		}
+		else if (Event == EHandleEvent::GlobalTransformUpdated)
+		{
+			if (const FRigControlElement* Control = ControlRig->FindControl(ControlName))
+			{	// if that handle control's transform has been dirtied
+				const bool bIsTransformDirty =
+					Control->Pose.IsDirty(ERigTransformType::CurrentLocal) ||
+					Control->Pose.IsDirty(ERigTransformType::CurrentGlobal);
+				if (bIsTransformDirty)
+				{
+					OnHandleModified.Broadcast(this, Event);
+				}
+			}
+		}
 	}
 }
 
diff --git a/Engine/Source/Runtime/Experimental/Animation/Constraints/Private/TransformConstraint.cpp b/Engine/Source/Runtime/Experimental/Animation/Constraints/Private/TransformConstraint.cpp
--- a/Engine/Source/Runtime/Experimental/Animation/Constraints/Private/TransformConstraint.cpp
+++ b/Engine/Source/Runtime/Experimental/Animation/Constraints/Private/TransformConstraint.cpp
@@ -569,6 +569,16 @@ void UTickableTransformConstraint::OnHandleModified(UTransformableHandle* InHand
 			return;
 		}
 	}
+
+	// update the constraint if the parent's transform has been modified 
+	if (InHandle == ParentTRSHandle && InNotification == EHandleEvent::GlobalTransformUpdated)
+	{
+		if (IsFullyActive() && ConstraintTick.IsTickFunctionRegistered() && ConstraintTick.IsTickFunctionEnabled())
+		{
+			Evaluate();
+		}
+		return;
+	}
 }
 
 bool UTickableTransformConstraint::IsFullyActive() const
