RigVM: Follow up fixes for Direct Manipulation workflow

diff --git a/Engine/Plugins/Animation/ControlRig/Source/ControlRigEditor/Private/Editor/ControlRigEditor.cpp b/Engine/Plugins/Animation/ControlRig/Source/ControlRigEditor/Private/Editor/ControlRigEditor.cpp
--- a/Engine/Plugins/Animation/ControlRig/Source/ControlRigEditor/Private/Editor/ControlRigEditor.cpp
+++ b/Engine/Plugins/Animation/ControlRig/Source/ControlRigEditor/Private/Editor/ControlRigEditor.cpp
@@ -126,6 +126,7 @@ FControlRigEditor::FControlRigEditor()
 	, RigHierarchyTabCount(0)
 	, bIsConstructionEventRunning(false)
 	, LastHierarchyHash(INDEX_NONE)
+	, bRefreshDirectionManipulationTargetsRequired(false)
 {
 	LastEventQueue = ConstructionEventQueue;
 }
@@ -1063,7 +1064,19 @@ void FControlRigEditor::HandleModifiedEvent(ERigVMGraphNotifType InNotifType, UR
 		{
 			if(Pin->GetNode() == DirectManipulationSubject.Get())
 			{
-				RefreshDirectManipulationTextList();
+				bRefreshDirectionManipulationTargetsRequired = true;
+			}
+		}
+	}
+	else if(InNotifType == ERigVMGraphNotifType::LinkAdded ||
+		InNotifType == ERigVMGraphNotifType::LinkRemoved)
+	{
+		if(const URigVMLink* Link = Cast<URigVMLink>(InSubject))
+		{
+			if((Link->GetSourceNode() == DirectManipulationSubject.Get()) ||
+				(Link->GetTargetNode() == DirectManipulationSubject.Get()))
+			{
+				bRefreshDirectionManipulationTargetsRequired = true;
 			}
 		}
 	}
@@ -1084,6 +1097,17 @@ void FControlRigEditor::OnCreateGraphEditorCommands(TSharedPtr<FUICommandList> G
 		FExecuteAction::CreateSP(this, &FControlRigEditor::HandleRequestDirectManipulationScale));
 }
 
+void FControlRigEditor::HandleVMCompiledEvent(UObject* InCompiledObject, URigVM* InVM)
+{
+	IControlRigEditor::HandleVMCompiledEvent(InCompiledObject, InVM);
+
+	if(bRefreshDirectionManipulationTargetsRequired)
+	{
+		RefreshDirectManipulationTextList();
+		bRefreshDirectionManipulationTargetsRequired = false;
+	}
+}
+
 void FControlRigEditor::SaveAsset_Execute()
 {
 	FRigVMEditor::SaveAsset_Execute();
@@ -1549,7 +1573,7 @@ void FControlRigEditor::HandleViewportCreated(const TSharedRef<class IPersonaVie
 				.VAlign(VAlign_Center)
 				.AutoWidth()
 				[
-					SNew(SComboBox<TSharedPtr<FString>>)
+					SAssignNew(DirectManipulationCombo, SComboBox<TSharedPtr<FString>>)
                 	.ContentPadding(FMargin(4.0f, 2.0f))
                 	.OptionsSource(&DirectManipulationTextList)
                 	.OnGenerateWidget_Lambda([this](TSharedPtr<FString> Item)
@@ -2569,6 +2593,10 @@ void FControlRigEditor::RefreshDirectManipulationTextList()
 {
 	DirectManipulationTextList.Reset();
 	(void)GetDirectManipulationTargetTextList();
+	if(DirectManipulationCombo.IsValid())
+	{
+		DirectManipulationCombo->RefreshOptions();
+	}
 }
 
 void FControlRigEditor::BindCommands()
diff --git a/Engine/Plugins/Animation/ControlRig/Source/ControlRigEditor/Private/Editor/ControlRigEditor.h b/Engine/Plugins/Animation/ControlRig/Source/ControlRigEditor/Private/Editor/ControlRigEditor.h
--- a/Engine/Plugins/Animation/ControlRig/Source/ControlRigEditor/Private/Editor/ControlRigEditor.h
+++ b/Engine/Plugins/Animation/ControlRig/Source/ControlRigEditor/Private/Editor/ControlRigEditor.h
@@ -14,6 +14,7 @@
 #include "RigVMModel/Nodes/RigVMLibraryNode.h"
 #include "RigVMCore/RigVM.h"
 #include "Widgets/SRigVMGraphPinNameListValueWidget.h"
+#include "Widgets/Input/SComboBox.h"
 #include "Styling/SlateTypes.h"
 #include "AnimPreviewInstance.h"
 #include "ScopedTransaction.h"
@@ -76,6 +77,7 @@ public:
 	virtual void Compile() override;
 	virtual void HandleModifiedEvent(ERigVMGraphNotifType InNotifType, URigVMGraph* InGraph, UObject* InSubject) override;
 	virtual void OnCreateGraphEditorCommands(TSharedPtr<FUICommandList> GraphEditorCommandsList) override;
+	virtual void HandleVMCompiledEvent(UObject* InCompiledObject, URigVM* InVM) override;
 
 	// allows the editor to fill an empty graph
 	virtual void CreateEmptyGraphContent(URigVMController* InController) override;
@@ -299,6 +301,8 @@ protected:
 
 	TWeakObjectPtr<const URigVMUnitNode> DirectManipulationSubject;
 	mutable TArray<TSharedPtr<FString>> DirectManipulationTextList;
+	TSharedPtr<SComboBox<TSharedPtr<FString>>> DirectManipulationCombo;
+	bool bRefreshDirectionManipulationTargetsRequired;
 	
 	static const TArray<FName> ForwardsSolveEventQueue;
 	static const TArray<FName> BackwardsSolveEventQueue;
