Control Rig: Add safe guards to SRigHierarchy context menu

diff --git a/Engine/Plugins/Animation/ControlRig/Source/ControlRigEditor/Private/Editor/SRigHierarchy.cpp b/Engine/Plugins/Animation/ControlRig/Source/ControlRigEditor/Private/Editor/SRigHierarchy.cpp
--- a/Engine/Plugins/Animation/ControlRig/Source/ControlRigEditor/Private/Editor/SRigHierarchy.cpp
+++ b/Engine/Plugins/Animation/ControlRig/Source/ControlRigEditor/Private/Editor/SRigHierarchy.cpp
@@ -1132,9 +1132,15 @@ void SRigHierarchy::CreateContextMenu() const
 	if (!ToolMenus->IsMenuRegistered(MenuName))
 	{
 		UToolMenu* Menu = ToolMenus->RegisterMenu(MenuName);
+		TWeakPtr<const SRigHierarchy> WeakRigHierarchyWidget = SharedThis(this).ToWeakPtr();
 		
-		Menu->AddDynamicSection(NAME_None, FNewToolMenuDelegate::CreateLambda([this](UToolMenu* InMenu)
+		Menu->AddDynamicSection(NAME_None, FNewToolMenuDelegate::CreateLambda([WeakRigHierarchyWidget](UToolMenu* InMenu)
 			{
+				if(!WeakRigHierarchyWidget.IsValid())
+				{
+					return;
+				}
+			
 				UControlRigContextMenuContext* MainContext = InMenu->FindContext<UControlRigContextMenuContext>();
 				
 				if (SRigHierarchy* RigHierarchyPanel = MainContext->GetRigHierarchyPanel())
@@ -1171,30 +1177,39 @@ void SRigHierarchy::CreateContextMenu() const
 					ElementsSection.AddMenuEntry(Commands.RenameItem);
 					ElementsSection.AddMenuEntry(Commands.MirrorItem);
 
-					if(RigHierarchyPanel->IsProceduralElementSelected() && ControlRigBlueprint.IsValid())
+					if(RigHierarchyPanel->IsProceduralElementSelected() && WeakRigHierarchyWidget.Pin()->ControlRigBlueprint.IsValid())
 					{
 						ElementsSection.AddMenuEntry(
 							"SelectSpawnerNode",
 							LOCTEXT("SelectSpawnerNode", "Select Spawner Node"),
 							LOCTEXT("SelectSpawnerNode_Tooltip", "Selects the node that spawn / added this element."),
 							FSlateIcon(),
-							FUIAction(FExecuteAction::CreateLambda([this]() {
-								const TArray<const FRigBaseElement*> Elements = GetHierarchy()->GetSelectedElements();
+							FUIAction(FExecuteAction::CreateLambda([WeakRigHierarchyWidget]() {
+								if(!WeakRigHierarchyWidget.IsValid())
+								{
+									return;
+								}
+								if(!WeakRigHierarchyWidget.Pin()->ControlRigBlueprint.IsValid())
+								{
+									return;
+								}
+								const UControlRigBlueprint* CurrentControlRigBlueprint = WeakRigHierarchyWidget.Pin()->ControlRigBlueprint.Get();
+								const TArray<const FRigBaseElement*> Elements = WeakRigHierarchyWidget.Pin()->GetHierarchy()->GetSelectedElements();
 								for(const FRigBaseElement* Element : Elements)
 								{
 									if(Element->IsProcedural())
 									{
 										const int32 InstructionIndex = Element->GetCreatedAtInstructionIndex();
-										if(UControlRig* ControlRig = Cast<UControlRig>(ControlRigBlueprint->GetObjectBeingDebugged()))
+										if(const UControlRig* ControlRig = Cast<UControlRig>(CurrentControlRigBlueprint->GetObjectBeingDebugged()))
 										{
 											if(ControlRig->VM)
 											{
 												if(URigVMNode* Node = Cast<URigVMNode>(ControlRig->VM->GetByteCode().GetSubjectForInstruction(InstructionIndex)))
 												{
-													if(URigVMController* Controller = ControlRigBlueprint->GetController(Node->GetGraph()))
+													if(URigVMController* Controller = CurrentControlRigBlueprint->GetController(Node->GetGraph()))
 													{
 														Controller->SelectNode(Node);
-														Controller->RequestJumpToHyperlinkDelegate.ExecuteIfBound(Node);
+														(void)Controller->RequestJumpToHyperlinkDelegate.ExecuteIfBound(Node);
 													}
 												}
 											}
