Sequencer: Control Rig: Shift camera move code no longer needed was causing double transform. jus removed some code.
#trival

diff --git a/Engine/Plugins/Animation/ControlRig/Source/ControlRigEditor/Private/EditMode/ControlRigEditMode.cpp b/Engine/Plugins/Animation/ControlRig/Source/ControlRigEditor/Private/EditMode/ControlRigEditMode.cpp
--- a/Engine/Plugins/Animation/ControlRig/Source/ControlRigEditor/Private/EditMode/ControlRigEditMode.cpp
+++ b/Engine/Plugins/Animation/ControlRig/Source/ControlRigEditor/Private/EditMode/ControlRigEditMode.cpp
@@ -1814,19 +1814,6 @@ bool FControlRigEditMode::InputDelta(FEditorViewportClient* InViewportClient, FV
 	const bool bDoTranslation = !Drag.IsZero() && (WidgetMode == UE::Widget::WM_Translate || WidgetMode == UE::Widget::WM_TranslateRotateZ);
 	const bool bDoScale = !Scale.IsZero() && WidgetMode == UE::Widget::WM_Scale;
 
-	if (IsDoingDrag(WeakSequencer, InViewport, GetCurrentWidgetAxis()) == false && !InDrag.IsZero() && bShiftDown)
-	{
-		FVector CameraDelta(InDrag);
-
-		// Only apply camera speed modifiers to the drag if we aren't zooming in an ortho viewport.
-		if (!InViewportClient->IsOrtho() || !(InViewport->KeyState(EKeys::LeftMouseButton) && InViewport->KeyState(EKeys::RightMouseButton)))
-		{
-			const float CameraSpeed = InViewportClient->GetCameraSpeed();
-			CameraDelta *= CameraSpeed;
-		}
-
-		InViewportClient->MoveViewportCamera(CameraDelta, InRot);
-	}
 
 	if (InteractionScopes.Num() > 0 && bMouseButtonDown && CurrentAxis != EAxisList::None
 		&& (bDoRotation || bDoTranslation || bDoScale))
