Control Rig: Public function not available when first created
#rnx

diff --git a/Engine/Plugins/Animation/ControlRig/Source/ControlRigEditor/Private/Editor/ControlRigEditor.cpp b/Engine/Plugins/Animation/ControlRig/Source/ControlRigEditor/Private/Editor/ControlRigEditor.cpp
--- a/Engine/Plugins/Animation/ControlRig/Source/ControlRigEditor/Private/Editor/ControlRigEditor.cpp
+++ b/Engine/Plugins/Animation/ControlRig/Source/ControlRigEditor/Private/Editor/ControlRigEditor.cpp
@@ -1119,6 +1119,10 @@ void FControlRigEditor::SaveAsset_Execute()
 		const UControlRig* CDO = ControlRig->GetClass()->GetDefaultObject<UControlRig>();
 		CDO->DynamicHierarchy->CopyHierarchy(RigBlueprint->Hierarchy);
 	}
+
+	FBlueprintActionDatabase& ActionDatabase = FBlueprintActionDatabase::Get();
+	ActionDatabase.ClearAssetActions(UControlRigBlueprint::StaticClass());
+	ActionDatabase.RefreshClassActions(UControlRigBlueprint::StaticClass());
 }
 
 void FControlRigEditor::SaveAssetAs_Execute()
@@ -1132,6 +1136,10 @@ void FControlRigEditor::SaveAssetAs_Execute()
 		const UControlRig* CDO = ControlRig->GetClass()->GetDefaultObject<UControlRig>();
 		CDO->DynamicHierarchy->CopyHierarchy(RigBlueprint->Hierarchy);
 	}
+
+	FBlueprintActionDatabase& ActionDatabase = FBlueprintActionDatabase::Get();
+	ActionDatabase.ClearAssetActions(UControlRigBlueprint::StaticClass());
+	ActionDatabase.RefreshClassActions(UControlRigBlueprint::StaticClass());
 }
 
 FName FControlRigEditor::GetToolkitFName() const
