Fix deprecated Control Rig GetTransform node no longer returning initial transform on control rig initialization. Causes TransformConstraint to produce wrong results.

Regression from CL 23407131 (removing Init phase from CR)

#swarm https://p4-swarm.epicgames.net/reviews/26867019

diff --git a/Engine/Plugins/Animation/ControlRig/Source/ControlRig/Private/Units/Hierarchy/RigUnit_GetBoneTransform.cpp b/Engine/Plugins/Animation/ControlRig/Source/ControlRig/Private/Units/Hierarchy/RigUnit_GetBoneTransform.cpp
--- a/Engine/Plugins/Animation/ControlRig/Source/ControlRig/Private/Units/Hierarchy/RigUnit_GetBoneTransform.cpp
+++ b/Engine/Plugins/Animation/ControlRig/Source/ControlRig/Private/Units/Hierarchy/RigUnit_GetBoneTransform.cpp
@@ -13,6 +13,16 @@ FRigUnit_GetBoneTransform_Execute()
 	const URigHierarchy* Hierarchy = ExecuteContext.Hierarchy;
 	if (Hierarchy)
 	{
+		// KIARAN AUG - 2023 (HACK)
+		// This restores the original behavior of this node where it returns the InitialTransform during the first
+		// call to Execute() so that downstream nodes (like TransformConstraint) work correctly. This behavior was lost
+		// when 23407131 removed the Context.Init pass which called into FRigUnit_GetInitialBoneTransform
+		if (bFirstUpdate)
+		{
+			FRigUnit_GetInitialBoneTransform::StaticExecute(ExecuteContext, Bone, Space, Transform, CachedBone);
+			bFirstUpdate = false;
+			return;
+		}
 		if (!CachedBone.UpdateCache(FRigElementKey(Bone, ERigElementType::Bone), Hierarchy))
 		{
 			UE_CONTROLRIG_RIGUNIT_REPORT_WARNING(TEXT("Bone '%s' is not valid."), *Bone.ToString());
diff --git a/Engine/Plugins/Animation/ControlRig/Source/ControlRig/Private/Units/Hierarchy/RigUnit_GetBoneTransform.h b/Engine/Plugins/Animation/ControlRig/Source/ControlRig/Private/Units/Hierarchy/RigUnit_GetBoneTransform.h
--- a/Engine/Plugins/Animation/ControlRig/Source/ControlRig/Private/Units/Hierarchy/RigUnit_GetBoneTransform.h
+++ b/Engine/Plugins/Animation/ControlRig/Source/ControlRig/Private/Units/Hierarchy/RigUnit_GetBoneTransform.h
@@ -16,6 +16,7 @@ struct CONTROLRIG_API FRigUnit_GetBoneTransform : public FRigUnit
 	FRigUnit_GetBoneTransform()
 		: Space(ERigVMTransformSpace::GlobalSpace)
 		, CachedBone()
+		, bFirstUpdate(true)
 	{}
 
 	RIGVM_METHOD()
@@ -42,6 +43,10 @@ struct CONTROLRIG_API FRigUnit_GetBoneTransform : public FRigUnit
 	UPROPERTY(transient)
 	FCachedRigElement CachedBone;
 
+	// Used to force first update to return initial pose
+	UPROPERTY(transient)
+	bool bFirstUpdate;
+
 	RIGVM_METHOD()
 	virtual FRigVMStructUpgradeInfo GetUpgradeInfo() const override;
 };
