Control Rig: The rig is no longer in charge of resetting the control pose after a construction event. The control rig host is in charge of that now.

diff --git a/Engine/Plugins/Animation/ControlRig/Source/ControlRig/Private/AnimNode_ControlRig.cpp b/Engine/Plugins/Animation/ControlRig/Source/ControlRig/Private/AnimNode_ControlRig.cpp
--- a/Engine/Plugins/Animation/ControlRig/Source/ControlRig/Private/AnimNode_ControlRig.cpp
+++ b/Engine/Plugins/Animation/ControlRig/Source/ControlRig/Private/AnimNode_ControlRig.cpp
@@ -121,6 +121,19 @@ void FAnimNode_ControlRig::Initialize_AnyThread(const FAnimationInitializeContex
 
 	FAnimNode_ControlRigBase::Initialize_AnyThread(Context);
 
+	if (ControlRig)
+	{
+		//Don't Inititialize the Control Rig here it may have the wrong VM on the CDO
+		SetTargetInstance(ControlRig);
+		ControlRig->RequestInit();
+		bControlRigRequiresInitialization = true;
+		LastBonesSerialNumberForCacheBones = 0;
+	}
+	else
+	{
+		SetTargetInstance(nullptr);
+	}
+
 	AlphaBoolBlend.Reinitialize();
 	AlphaScaleBiasClamp.Reinitialize();
 }
diff --git a/Engine/Plugins/Animation/ControlRig/Source/ControlRig/Private/AnimNode_ControlRigBase.cpp b/Engine/Plugins/Animation/ControlRig/Source/ControlRig/Private/AnimNode_ControlRigBase.cpp
--- a/Engine/Plugins/Animation/ControlRig/Source/ControlRig/Private/AnimNode_ControlRigBase.cpp
+++ b/Engine/Plugins/Animation/ControlRig/Source/ControlRig/Private/AnimNode_ControlRigBase.cpp
@@ -75,19 +75,6 @@ void FAnimNode_ControlRigBase::Initialize_AnyThread(const FAnimationInitializeCo
 
 	FAnimNode_CustomProperty::Initialize_AnyThread(Context);
 	Source.Initialize(Context);
-
-	if (UControlRig* ControlRig = GetControlRig())
-	{
-		//Don't Inititialize the Control Rig here it may have the wrong VM on the CDO
-		SetTargetInstance(ControlRig);
-		ControlRig->RequestInit();
-		bControlRigRequiresInitialization = true;
-		LastBonesSerialNumberForCacheBones = 0;
-	}
-	else
-	{
-		SetTargetInstance(nullptr);
-	}
 }
 
 void FAnimNode_ControlRigBase::GatherDebugData(FNodeDebugData& DebugData)
diff --git a/Engine/Plugins/Animation/ControlRig/Source/ControlRig/Private/AnimNode_ControlRig_ExternalSource.cpp b/Engine/Plugins/Animation/ControlRig/Source/ControlRig/Private/AnimNode_ControlRig_ExternalSource.cpp
--- a/Engine/Plugins/Animation/ControlRig/Source/ControlRig/Private/AnimNode_ControlRig_ExternalSource.cpp
+++ b/Engine/Plugins/Animation/ControlRig/Source/ControlRig/Private/AnimNode_ControlRig_ExternalSource.cpp
@@ -7,6 +7,7 @@
 
 FAnimNode_ControlRig_ExternalSource::FAnimNode_ControlRig_ExternalSource()
 {
+	bControlRigRequiresInitialization = false;
 }
 
 void FAnimNode_ControlRig_ExternalSource::SetControlRig(UControlRig* InControlRig)
@@ -29,4 +30,19 @@ TSubclassOf<UControlRig> FAnimNode_ControlRig_ExternalSource::GetControlRigClass
 	return nullptr;
 }
 
+void FAnimNode_ControlRig_ExternalSource::Initialize_AnyThread(const FAnimationInitializeContext& Context)
+{
+	FAnimNode_ControlRigBase::Initialize_AnyThread(Context);
+
+	if (ControlRig.IsValid())
+	{
+		//Don't Initialize the Control Rig here, the owner of the rig is in charge of initializing
+		SetTargetInstance(ControlRig.Get());
+	}
+	else
+	{
+		SetTargetInstance(nullptr);
+	}
+}
+
 
diff --git a/Engine/Plugins/Animation/ControlRig/Source/ControlRig/Private/ControlRig.cpp b/Engine/Plugins/Animation/ControlRig/Source/ControlRig/Private/ControlRig.cpp
--- a/Engine/Plugins/Animation/ControlRig/Source/ControlRig/Private/ControlRig.cpp
+++ b/Engine/Plugins/Animation/ControlRig/Source/ControlRig/Private/ControlRig.cpp
@@ -751,24 +751,6 @@ bool UControlRig::Execute(const FName& InEventName)
 			// disable selection notifications from the hierarchy
 			TGuardValue<bool> DisableSelectionNotifications(GetHierarchy()->GetController(true)->bSuspendSelectionNotifications, true);
 			{
-				// save the current state of all pose elements to preserve user intention, since construction event can
-				// run in between forward events
-				// the saved pose is reapplied to the rig after construction event as the pose scope goes out of scope
-				TUniquePtr<FPoseScope> PoseScope;
-				if (!bConstructionModeEnabled)
-				{
-					// only do this in non-construction mode because 
-					// when construction mode is enabled, the control values are cleared before reaching here (too late to save them)
-					if (!bJustRanInit)
-					{
-						PoseScope = MakeUnique<FPoseScope>(this, ERigElementType::ToResetAfterConstructionEvent, TArray<FRigElementKey>(), ERigTransformType::CurrentGlobal);
-					}
-					else
-					{
-						PoseScope = MakeUnique<FPoseScope>(this, ERigElementType::ToResetAfterConstructionEvent);
-					}
-				}
-				
 				{
 					// Copy the hierarchy from the default object onto this one
 #if WITH_EDITOR
diff --git a/Engine/Plugins/Animation/ControlRig/Source/ControlRig/Public/AnimNode_ControlRig_ExternalSource.h b/Engine/Plugins/Animation/ControlRig/Source/ControlRig/Public/AnimNode_ControlRig_ExternalSource.h
--- a/Engine/Plugins/Animation/ControlRig/Source/ControlRig/Public/AnimNode_ControlRig_ExternalSource.h
+++ b/Engine/Plugins/Animation/ControlRig/Source/ControlRig/Public/AnimNode_ControlRig_ExternalSource.h
@@ -18,6 +18,7 @@ struct CONTROLRIG_API FAnimNode_ControlRig_ExternalSource : public FAnimNode_Con
 	void SetControlRig(UControlRig* InControlRig);
 	virtual UControlRig* GetControlRig() const override;
 	virtual TSubclassOf<UControlRig> GetControlRigClass() const override;
+	virtual void Initialize_AnyThread(const FAnimationInitializeContext& Context) override;
 
 private:
 	UPROPERTY(transient)
diff --git a/Engine/Plugins/Animation/ControlRig/Source/ControlRigEditor/Private/Editor/ControlRigEditor.cpp b/Engine/Plugins/Animation/ControlRig/Source/ControlRigEditor/Private/Editor/ControlRigEditor.cpp
--- a/Engine/Plugins/Animation/ControlRig/Source/ControlRigEditor/Private/Editor/ControlRigEditor.cpp
+++ b/Engine/Plugins/Animation/ControlRig/Source/ControlRigEditor/Private/Editor/ControlRigEditor.cpp
@@ -4135,6 +4135,8 @@ TSharedPtr<FUICommandList> FControlRigEditor::HandleOnViewportContextMenuCommand
 void FControlRigEditor::OnPreConstruction_AnyThread(UControlRig* InRig, const FName& InEventName)
 {
 	bIsConstructionEventRunning = true;
+	const TArrayView<const FRigElementKey> Elements;
+	PreConstructionPose = InRig->GetHierarchy()->GetPose(false, ERigElementType::ToResetAfterConstructionEvent, Elements);
 }
 
 void FControlRigEditor::OnPostConstruction_AnyThread(UControlRig* InRig, const FName& InEventName)
@@ -4177,6 +4179,10 @@ void FControlRigEditor::OnPostConstruction_AnyThread(UControlRig* InRig, const F
 			}, TStatId(), NULL, ENamedThreads::GameThread);
 		}
 	}
+	else
+	{
+		InRig->GetHierarchy()->SetPose(PreConstructionPose, ERigTransformType::CurrentGlobal);
+	}
 }
 
 #undef LOCTEXT_NAMESPACE
diff --git a/Engine/Plugins/Animation/ControlRig/Source/ControlRigEditor/Private/Editor/ControlRigEditor.h b/Engine/Plugins/Animation/ControlRig/Source/ControlRigEditor/Private/Editor/ControlRigEditor.h
--- a/Engine/Plugins/Animation/ControlRig/Source/ControlRigEditor/Private/Editor/ControlRigEditor.h
+++ b/Engine/Plugins/Animation/ControlRig/Source/ControlRigEditor/Private/Editor/ControlRigEditor.h
@@ -293,6 +293,7 @@ protected:
 
 	void OnPreConstruction_AnyThread(UControlRig* InRig, const FName& InEventName);
 	void OnPostConstruction_AnyThread(UControlRig* InRig, const FName& InEventName);
+	FRigPose PreConstructionPose;
 
 	bool bIsConstructionEventRunning;
 	uint32 LastHierarchyHash;
