ControlRig: Offer ResetCurvesFunctionality

#8588

diff --git a/Engine/Plugins/Animation/ControlRig/Source/ControlRigEditor/Private/Editor/SRigHierarchy.cpp b/Engine/Plugins/Animation/ControlRig/Source/ControlRigEditor/Private/Editor/SRigHierarchy.cpp
--- a/Engine/Plugins/Animation/ControlRig/Source/ControlRigEditor/Private/Editor/SRigHierarchy.cpp
+++ b/Engine/Plugins/Animation/ControlRig/Source/ControlRigEditor/Private/Editor/SRigHierarchy.cpp
@@ -1128,10 +1128,13 @@ void SRigHierarchy::CreateContextMenu() const
 	{
 		return;
 	}
-	
-	if (!ToolMenus->IsMenuRegistered(MenuName))
+
+	static bool bCreatedContextMenu = false;
+	if (!ToolMenus->IsMenuRegistered(MenuName) || !bCreatedContextMenu)
 	{
 		UToolMenu* Menu = ToolMenus->RegisterMenu(MenuName);
+		bCreatedContextMenu = true;
+		
 		TWeakPtr<const SRigHierarchy> WeakRigHierarchyWidget = SharedThis(this).ToWeakPtr();
 		
 		Menu->AddDynamicSection(NAME_None, FNewToolMenuDelegate::CreateLambda([WeakRigHierarchyWidget](UToolMenu* InMenu)
@@ -1290,6 +1293,11 @@ void SRigHierarchy::CreateContextMenu() const
 						LOCTEXT("RefreshSubMenu_ToolTip", "Refresh the existing initial transform from the selected mesh. This only updates if the node is found."),
 						FNewMenuDelegate::CreateSP(RigHierarchyPanel, &SRigHierarchy::CreateRefreshMenu)
 					);	
+
+					AssetsSection.AddSubMenu(TEXT("ResetCurves"), LOCTEXT("ResetCurvesSubMenu", "Reset Curves"),
+						LOCTEXT("ResetCurvesSubMenu_ToolTip", "Reset all curves in this rig asset to the selected mesh, Useful when if you add more morphs to the mesh but control rig does not update."),
+						FNewMenuDelegate::CreateSP(RigHierarchyPanel, &SRigHierarchy::CreateResetCurvesMenu)
+					);				
 				}
 			})
 		);
@@ -1345,7 +1353,34 @@ void SRigHierarchy::CreateRefreshMenu(FMenuBuilder& MenuBuilder)
 		[
 			SNew(SObjectPropertyEntryBox)
 			.AllowedClass(USkeletalMesh::StaticClass())
-			.OnObjectChanged(this, &SRigHierarchy::RefreshHierarchy)
+			.OnObjectChanged(this, &SRigHierarchy::RefreshHierarchy, false)
+		]
+		,
+		FText()
+	);
+}
+
+void SRigHierarchy::CreateResetCurvesMenu(FMenuBuilder& MenuBuilder)
+{
+	MenuBuilder.AddWidget(
+		SNew(SVerticalBox)
+
+		+ SVerticalBox::Slot()
+		.AutoHeight()
+		.Padding(3)
+		[
+			SNew(STextBlock)
+			.Font(FAppStyle::GetFontStyle("ControlRig.Hierarchy.Menu"))
+			.Text(LOCTEXT("ResetMesh_Title", "Select Mesh"))
+			.ToolTipText(LOCTEXT("ResetMesh_Tooltip", "Select mesh to reset curves to."))
+		]
+		+ SVerticalBox::Slot()
+		.AutoHeight()
+		.Padding(3)
+		[
+			SNew(SObjectPropertyEntryBox)
+			.AllowedClass(USkeletalMesh::StaticClass())
+			.OnObjectChanged(this, &SRigHierarchy::RefreshHierarchy, true)
 		]
 		,
 		FText()
@@ -1387,7 +1422,7 @@ void SRigHierarchy::UpdateMesh(USkeletalMesh* InMesh, const bool bImport) const
 	EditorSharedPtr->Compile();
 }
 
-void SRigHierarchy::RefreshHierarchy(const FAssetData& InAssetData)
+void SRigHierarchy::RefreshHierarchy(const FAssetData& InAssetData, bool bOnlyResetCurves)
 {
 	if (bIsChangingRigHierarchy)
 	{
@@ -1424,9 +1459,21 @@ void SRigHierarchy::RefreshHierarchy(const FAssetData& InAssetData)
 
 		URigHierarchyController* Controller = Hierarchy->GetController(true);
 		check(Controller);
-		
-		Controller->ImportBones(Mesh->GetSkeleton(), NAME_None, true, true, bSelectBones, true, true);
-		Controller->ImportCurves(Mesh->GetSkeleton(), NAME_None, false, true, true);
+
+		if(bOnlyResetCurves)
+		{
+			TArray<FRigElementKey> CurveKeys = Hierarchy->GetAllKeys(false, ERigElementType::Curve);
+			for(const FRigElementKey& CurveKey : CurveKeys)
+			{
+				Controller->RemoveElement(CurveKey, true, true);
+			}			
+			Controller->ImportCurves(Mesh->GetSkeleton(), NAME_None, false, true, true);
+		}
+		else
+		{
+			Controller->ImportBones(Mesh->GetSkeleton(), NAME_None, true, true, bSelectBones, true, true);
+			Controller->ImportCurves(Mesh->GetSkeleton(), NAME_None, false, true, true);
+		}
 	}
 
 	ControlRigBlueprint->PropagateHierarchyFromBPToInstances();
diff --git a/Engine/Plugins/Animation/ControlRig/Source/ControlRigEditor/Private/Editor/SRigHierarchy.h b/Engine/Plugins/Animation/ControlRig/Source/ControlRigEditor/Private/Editor/SRigHierarchy.h
--- a/Engine/Plugins/Animation/ControlRig/Source/ControlRigEditor/Private/Editor/SRigHierarchy.h
+++ b/Engine/Plugins/Animation/ControlRig/Source/ControlRigEditor/Private/Editor/SRigHierarchy.h
@@ -199,8 +199,9 @@ private:
 	void ImportHierarchy(const FAssetData& InAssetData);
 	void CreateImportMenu(FMenuBuilder& MenuBuilder);
 	void CreateRefreshMenu(FMenuBuilder& MenuBuilder);
+	void CreateResetCurvesMenu(FMenuBuilder& MenuBuilder);
 	bool ShouldFilterOnImport(const FAssetData& AssetData) const;
-	void RefreshHierarchy(const FAssetData& InAssetData);
+	void RefreshHierarchy(const FAssetData& InAssetData, bool bOnlyResetCurves);
 	void UpdateMesh(USkeletalMesh* InMesh, const bool bImport) const;
 
 	void HandleResetTransform(bool bSelectionOnly);
