Control Rig: Improve rig hierarchy copy performance

diff --git a/Engine/Plugins/Animation/ControlRig/Source/ControlRig/Private/Rigs/RigHierarchy.cpp b/Engine/Plugins/Animation/ControlRig/Source/ControlRig/Private/Rigs/RigHierarchy.cpp
--- a/Engine/Plugins/Animation/ControlRig/Source/ControlRig/Private/Rigs/RigHierarchy.cpp
+++ b/Engine/Plugins/Animation/ControlRig/Source/ControlRig/Private/Rigs/RigHierarchy.cpp
@@ -322,23 +322,31 @@ void URigHierarchy::DeclareConstructClasses(TArray<FTopLevelAssetPath>& OutConst
 void URigHierarchy::Reset()
 {
 	LLM_SCOPE_BYNAME(TEXT("Animation/ControlRig"));
-	
+
+	Reset_Impl(true);
+}
+
+void URigHierarchy::Reset_Impl(bool bResetElements)
+{
 	TopologyVersion = 0;
 	MetadataVersion = 0;
 	bEnableDirtyPropagation = true;
 
-	// walk in reverse since certain elements might not have been allocated themselves
-	for(int32 ElementIndex = Elements.Num() - 1; ElementIndex >= 0; ElementIndex--)
+	if(bResetElements)
 	{
-		DestroyElement(Elements[ElementIndex]);
-	}
-	Elements.Reset();
-	ElementsPerType.Reset();
-	for(int32 TypeIndex=0;TypeIndex<RigElementTypeToFlatIndex(ERigElementType::Last);TypeIndex++)
-	{
-		ElementsPerType.Add(TArray<FRigBaseElement*>());
+		// walk in reverse since certain elements might not have been allocated themselves
+		for(int32 ElementIndex = Elements.Num() - 1; ElementIndex >= 0; ElementIndex--)
+		{
+			DestroyElement(Elements[ElementIndex]);
+		}
+		Elements.Reset();
+		ElementsPerType.Reset();
+		for(int32 TypeIndex=0;TypeIndex<RigElementTypeToFlatIndex(ERigElementType::Last);TypeIndex++)
+		{
+			ElementsPerType.Add(TArray<FRigBaseElement*>());
+		}
+		IndexLookup.Reset();
 	}
-	IndexLookup.Reset();
 
 	ResetPoseHash = INDEX_NONE;
 	ResetPoseIsFilteredOut.Reset();
@@ -379,59 +387,146 @@ void URigHierarchy::CopyHierarchy(URigHierarchy* InHierarchy)
 	LLM_SCOPE_BYNAME(TEXT("Animation/ControlRig"));
 
 	FScopeLock Lock(&ElementsLock);
-	
-	Reset();
-
-	// Allocate the elements in batches to improve performance
-	TArray<uint8*> NewElementsPerType;
-	TArray<int32> StructureSizePerType;
-	for(int32 ElementTypeIndex = 0; ElementTypeIndex < InHierarchy->ElementsPerType.Num(); ElementTypeIndex++)
+	if(Elements.Num() == 0 && InHierarchy->Elements.Num () == 0)
 	{
-		const ERigElementType ElementType = FlatIndexToRigElementType(ElementTypeIndex);
-		int32 StructureSize = 0;
+		return;
+	}
 
-		const int32 Count = InHierarchy->ElementsPerType[ElementTypeIndex].Num();
-		if(Count)
+	bool bReallocateElements = false;
+
+	// check if we really need to do a deep copy all over again.
+	// for rigs which contain more elements (likely procedural elements)
+	// we'll assume we can just remove superfluous elements (from the end of the lists).
+	bReallocateElements = Elements.Num() < InHierarchy->Elements.Num();
+	if(!bReallocateElements)
+	{
+		for(int32 ElementTypeIndex = 0; ElementTypeIndex < RigElementTypeToFlatIndex(ERigElementType::Last); ElementTypeIndex++)
 		{
-			FRigBaseElement* ElementMemory = MakeElement(ElementType, Count, &StructureSize);
-			NewElementsPerType.Add((uint8*)ElementMemory);
+			check(ElementsPerType.IsValidIndex(ElementTypeIndex));
+			check(InHierarchy->ElementsPerType.IsValidIndex(ElementTypeIndex));
+			if(ElementsPerType[ElementTypeIndex].Num() < InHierarchy->ElementsPerType[ElementTypeIndex].Num())
+			{
+				bReallocateElements = true;
+				break;
+			}
 		}
-		else
+
+		// make sure that we have the elements in the right order / type
+		if(!bReallocateElements)
 		{
-			NewElementsPerType.Add(nullptr);
+			for(int32 Index = 0; Index < InHierarchy->Elements.Num(); Index++)
+			{
+				if((Elements[Index]->GetKey().Type != InHierarchy->Elements[Index]->GetKey().Type) ||
+					(Elements[Index]->SubIndex != InHierarchy->Elements[Index]->SubIndex))
+				{
+					bReallocateElements = true;
+					break;
+				}
+			}
 		}
-
-		StructureSizePerType.Add(StructureSize);
-		ElementsPerType[ElementTypeIndex].Reserve(Count);
 	}
 
-	Elements.Reserve(InHierarchy->Elements.Num());
-	IndexLookup.Reserve(InHierarchy->IndexLookup.Num());
-	
-	for(int32 Index = 0; Index < InHierarchy->Num(); Index++)
+	Reset_Impl(bReallocateElements);
+
+	static const TArray<int32> StructureSizePerType = {
+		sizeof(FRigBoneElement),
+		sizeof(FRigNullElement),
+		sizeof(FRigControlElement),
+		sizeof(FRigCurveElement),
+		sizeof(FRigRigidBodyElement),
+		sizeof(FRigReferenceElement),
+	}; 
+
+	if(bReallocateElements)
 	{
-		FRigBaseElement* Source = InHierarchy->Get(Index);
-		const FRigElementKey& Key = Source->Key;
+		// Allocate the elements in batches to improve performance
+		TArray<uint8*> NewElementsPerType;
+		for(int32 ElementTypeIndex = 0; ElementTypeIndex < InHierarchy->ElementsPerType.Num(); ElementTypeIndex++)
+		{
+			const ERigElementType ElementType = FlatIndexToRigElementType(ElementTypeIndex);
+			int32 StructureSize = 0;
 
-		const int32 ElementTypeIndex = RigElementTypeToFlatIndex(Key.Type);
+			const int32 Count = InHierarchy->ElementsPerType[ElementTypeIndex].Num();
+			if(Count)
+			{
+				FRigBaseElement* ElementMemory = MakeElement(ElementType, Count, &StructureSize);
+				verify(StructureSize == StructureSizePerType[ElementTypeIndex]);
+				NewElementsPerType.Add((uint8*)ElementMemory);
+			}
+			else
+			{
+				NewElementsPerType.Add(nullptr);
+			}
+			
+			ElementsPerType[ElementTypeIndex].Reserve(Count);
+		}
+
+		Elements.Reserve(InHierarchy->Elements.Num());
+		IndexLookup.Reserve(InHierarchy->IndexLookup.Num());
+
+		for(int32 Index = 0; Index < InHierarchy->Num(); Index++)
+		{
+			FRigBaseElement* Source = InHierarchy->Get(Index);
+			const FRigElementKey& Key = Source->Key;
+
+			const int32 ElementTypeIndex = RigElementTypeToFlatIndex(Key.Type);
 		
-		const int32 SubIndex = Num(Key.Type);
+			const int32 SubIndex = Num(Key.Type);
 
-		const int32 StructureSize = StructureSizePerType[ElementTypeIndex];
-		check(NewElementsPerType[ElementTypeIndex] != nullptr);
-		FRigBaseElement* Target = (FRigBaseElement*)&NewElementsPerType[ElementTypeIndex][StructureSize * SubIndex];
-		//FRigBaseElement* Target = MakeElement(Key.Type);
+			const int32 StructureSize = StructureSizePerType[ElementTypeIndex];
+			check(NewElementsPerType[ElementTypeIndex] != nullptr);
+			FRigBaseElement* Target = (FRigBaseElement*)&NewElementsPerType[ElementTypeIndex][StructureSize * SubIndex];
 		
-		Target->Key = Key;
-		Target->NameString = Source->NameString;
-		Target->SubIndex = SubIndex;
-		Target->Index = Elements.Add(Target);
+			Target->Key = Key;
+			Target->NameString = Source->NameString;
+			Target->SubIndex = SubIndex;
+			Target->Index = Elements.Add(Target);
+			Target->CreatedAtInstructionIndex = Source->CreatedAtInstructionIndex;
 
-		ElementsPerType[ElementTypeIndex].Add(Target);
-		IndexLookup.Add(Key, Target->Index);
+			ElementsPerType[ElementTypeIndex].Add(Target);
+			IndexLookup.Add(Key, Target->Index);
 
-		check(Source->Index == Index);
-		check(Target->Index == Index);
+			check(Source->Index == Index);
+			check(Target->Index == Index);
+		}
+	}
+	else
+	{
+		// remove the superfluous elements
+		for(int32 ElementIndex = Elements.Num() - 1; ElementIndex >= InHierarchy->Elements.Num(); ElementIndex--)
+		{
+			DestroyElement(Elements[ElementIndex]);
+		}
+
+		// shrink the containers accordingly
+		Elements.SetNum(InHierarchy->Elements.Num());
+		for(int32 ElementTypeIndex = 0; ElementTypeIndex < RigElementTypeToFlatIndex(ERigElementType::Last); ElementTypeIndex++)
+		{
+			ElementsPerType[ElementTypeIndex].SetNum(InHierarchy->ElementsPerType[ElementTypeIndex].Num());
+		}
+
+		for(int32 Index = 0; Index < InHierarchy->Num(); Index++)
+		{
+			const FRigBaseElement* Source = InHierarchy->Get(Index);
+			FRigBaseElement* Target = Elements[Index];
+
+			check(Target->Key.Type == Source->Key.Type);
+            Target->Key = Source->Key;
+            Target->NameString = Source->NameString;
+            Target->SubIndex = Source->SubIndex;
+            Target->Index = Source->Index;
+			Target->CreatedAtInstructionIndex = Source->CreatedAtInstructionIndex;
+			Target->CachedChildren.Reset();
+
+			// fast-pass for update children
+			Target->CachedChildren.SetNumZeroed(Source->CachedChildren.Num());
+			for(int32 ChildIndex = 0; ChildIndex < Target->CachedChildren.Num(); ChildIndex++)
+			{
+				Target->CachedChildren[ChildIndex] = Elements[Source->CachedChildren[ChildIndex]->Index];
+			}
+		}
+
+		IndexLookup = InHierarchy->IndexLookup;
 	}
 
 	for(int32 Index = 0; Index < InHierarchy->Num(); Index++)
@@ -441,15 +536,14 @@ void URigHierarchy::CopyHierarchy(URigHierarchy* InHierarchy)
 		Target->CopyFrom(this, Source, InHierarchy);
 	}
 
-	for (const TPair<FRigElementKey, FRigElementKey>& NameMapPair : InHierarchy->PreviousNameMap)
-	{
-		PreviousNameMap.FindOrAdd(NameMapPair.Key) = NameMapPair.Value;
-	}
-
-
+	PreviousNameMap.Append(InHierarchy->PreviousNameMap);
 	TopologyVersion = InHierarchy->GetTopologyVersion();
 	MetadataVersion = InHierarchy->GetMetadataVersion();
-	UpdateAllCachedChildren();
+
+	if(bReallocateElements)
+	{
+		UpdateAllCachedChildren();
+	}
 	
 	EnsureCacheValidity();
 }
diff --git a/Engine/Plugins/Animation/ControlRig/Source/ControlRig/Public/Rigs/RigHierarchy.h b/Engine/Plugins/Animation/ControlRig/Source/ControlRig/Public/Rigs/RigHierarchy.h
--- a/Engine/Plugins/Animation/ControlRig/Source/ControlRig/Public/Rigs/RigHierarchy.h
+++ b/Engine/Plugins/Animation/ControlRig/Source/ControlRig/Public/Rigs/RigHierarchy.h
@@ -4389,6 +4389,7 @@ private:
 
 	void QueueNotification(ERigHierarchyNotification InNotification, const FRigBaseElement* InElement);
 	void SendQueuedNotifications();
+	void Reset_Impl(bool bResetElements);
 	
 	friend class URigHierarchyController;
 	friend class UControlRig;
