Control Rig: Follow up fix for hierarchy copy performance improvements

diff --git a/Engine/Plugins/Animation/ControlRig/Source/ControlRig/Private/Rigs/RigHierarchy.cpp b/Engine/Plugins/Animation/ControlRig/Source/ControlRig/Private/Rigs/RigHierarchy.cpp
--- a/Engine/Plugins/Animation/ControlRig/Source/ControlRig/Private/Rigs/RigHierarchy.cpp
+++ b/Engine/Plugins/Animation/ControlRig/Source/ControlRig/Private/Rigs/RigHierarchy.cpp
@@ -492,6 +492,8 @@ void URigHierarchy::CopyHierarchy(URigHierarchy* InHierarchy)
 	}
 	else
 	{
+		InHierarchy->UpdateAllCachedChildren();
+		
 		// remove the superfluous elements
 		for(int32 ElementIndex = Elements.Num() - 1; ElementIndex >= InHierarchy->Elements.Num(); ElementIndex--)
 		{
@@ -524,6 +526,8 @@ void URigHierarchy::CopyHierarchy(URigHierarchy* InHierarchy)
 			{
 				Target->CachedChildren[ChildIndex] = Elements[Source->CachedChildren[ChildIndex]->Index];
 			}
+
+			Target->TopologyVersion = InHierarchy->GetTopologyVersion();
 		}
 
 		IndexLookup = InHierarchy->IndexLookup;
