Control Rig: Avoid posting empty transactions to undo stack in sequencer

diff --git a/Engine/Plugins/Animation/ControlRig/Source/ControlRigEditor/Private/Sequencer/ControlRigParameterTrackEditor.cpp b/Engine/Plugins/Animation/ControlRig/Source/ControlRigEditor/Private/Sequencer/ControlRigParameterTrackEditor.cpp
--- a/Engine/Plugins/Animation/ControlRig/Source/ControlRigEditor/Private/Sequencer/ControlRigParameterTrackEditor.cpp
+++ b/Engine/Plugins/Animation/ControlRig/Source/ControlRigEditor/Private/Sequencer/ControlRigParameterTrackEditor.cpp
@@ -219,6 +219,7 @@ FControlRigParameterTrackEditor::FControlRigParameterTrackEditor(TSharedRef<ISeq
 	, bFilterAssetBySkeleton(true)
 	, bFilterAssetByAnimatableControls(false)
 	, ControlUndoBracket(0)
+	, ControlChangedDuringUndoBracket(0)
 {
 	FMovieSceneToolsModule::Get().RegisterAnimationBakeHelper(this);
 
@@ -3491,6 +3492,7 @@ void FControlRigParameterTrackEditor::HandleControlModified(UControlRig* Control
 					}
 					AddControlKeys(Component, ControlRig, Name, ControlElement->GetName(), (EControlRigContextChannelToKey)Context.KeyMask, 
 						KeyMode, Context.LocalTime);
+					ControlChangedDuringUndoBracket++;
 				}
 			}
 		}
@@ -3503,6 +3505,7 @@ void FControlRigParameterTrackEditor::HandleControlUndoBracket(UControlRig* Subj
 	{
 		FScopeLock ScopeLock(&ControlUndoTransactionMutex);
 		ControlUndoTransaction = MakeShareable(new FScopedTransaction(LOCTEXT("KeyMultipleControls", "Auto-Key multiple controls")));
+		ControlChangedDuringUndoBracket = 0;
 	}
 
 	ControlUndoBracket = FMath::Max<int32>(0, ControlUndoBracket + (bOpenUndoBracket ? 1 : -1));
@@ -3510,6 +3513,10 @@ void FControlRigParameterTrackEditor::HandleControlUndoBracket(UControlRig* Subj
 	if(!bOpenUndoBracket && ControlUndoBracket == 0)
 	{
 		FScopeLock ScopeLock(&ControlUndoTransactionMutex);
+		if(ControlChangedDuringUndoBracket == 0 && ControlUndoTransaction.IsValid())
+		{
+			ControlUndoTransaction->Cancel();
+		}
 		ControlUndoTransaction.Reset();
 	}
 }
diff --git a/Engine/Plugins/Animation/ControlRig/Source/ControlRigEditor/Private/Sequencer/ControlRigParameterTrackEditor.h b/Engine/Plugins/Animation/ControlRig/Source/ControlRigEditor/Private/Sequencer/ControlRigParameterTrackEditor.h
--- a/Engine/Plugins/Animation/ControlRig/Source/ControlRigEditor/Private/Sequencer/ControlRigParameterTrackEditor.h
+++ b/Engine/Plugins/Animation/ControlRig/Source/ControlRigEditor/Private/Sequencer/ControlRigParameterTrackEditor.h
@@ -284,6 +284,9 @@ private:
 	/** An index counter for the opened undo brackets */
 	int32 ControlUndoBracket;
 
+	/** A counter for occurred control changes during a control undo bracket */
+	int32 ControlChangedDuringUndoBracket;
+
 	/** Lock to avoid registering multiple transactions from different tracks at the same time */
 	static FCriticalSection ControlUndoTransactionMutex;
 
