Control Rig: Fix fallout in hierarchy selection

diff --git a/Engine/Plugins/Animation/ControlRig/Source/ControlRig/Private/Rigs/RigHierarchy.cpp b/Engine/Plugins/Animation/ControlRig/Source/ControlRig/Private/Rigs/RigHierarchy.cpp
--- a/Engine/Plugins/Animation/ControlRig/Source/ControlRig/Private/Rigs/RigHierarchy.cpp
+++ b/Engine/Plugins/Animation/ControlRig/Source/ControlRig/Private/Rigs/RigHierarchy.cpp
@@ -373,7 +373,21 @@ void URigHierarchy::ResetToDefault()
 	{
 		if(URigHierarchy* DefaultHierarchy = DefaultHierarchyPtr.Get())
 		{
+			const TArray<FRigElementKey> PreviousSelection = OrderedSelection;
 			CopyHierarchy(DefaultHierarchy);
+
+			// reestablish the selection
+			check(OrderedSelection.IsEmpty());
+			for(const FRigElementKey& Key : PreviousSelection)
+			{
+				if(FRigBaseElement* Element = Find(Key))
+				{
+					check(!Element->bSelected);
+					Element->bSelected = true;
+					OrderedSelection.Add(Element->GetKey());
+					Notify(ERigHierarchyNotification::ElementSelected, Element);
+				}
+			}
 			return;
 		}
 	}
@@ -519,6 +533,7 @@ void URigHierarchy::CopyHierarchy(URigHierarchy* InHierarchy)
             Target->Index = Source->Index;
 			Target->CreatedAtInstructionIndex = Source->CreatedAtInstructionIndex;
 			Target->CachedChildren.Reset();
+			Target->bSelected = false;
 
 			// fast-pass for update children
 			Target->CachedChildren.SetNumZeroed(Source->CachedChildren.Num());
@@ -4212,14 +4227,14 @@ bool URigHierarchy::IsSelected(const FRigBaseElement* InElement) const
 	{
 		return false;
 	}
-	if(URigHierarchy* HierarchyForSelection = HierarchyForSelectionPtr.Get())
+	if(const URigHierarchy* HierarchyForSelection = HierarchyForSelectionPtr.Get())
 	{
 		return HierarchyForSelection->IsSelected(InElement->GetKey());
 	}
 
 	const bool bIsSelected = OrderedSelection.Contains(InElement->GetKey());
 	ensure(bIsSelected == InElement->IsSelected());
-	return OrderedSelection.Contains(InElement->GetKey());
+	return bIsSelected;
 }
 
 void URigHierarchy::ResetCachedChildren()
