Control Rig: FIx some spawned controls not appearing in sequencer
#rnx

diff --git a/Engine/Plugins/Animation/ControlRig/Source/ControlRig/Private/Sequencer/MovieSceneControlRigParameterTrack.cpp b/Engine/Plugins/Animation/ControlRig/Source/ControlRig/Private/Sequencer/MovieSceneControlRigParameterTrack.cpp
--- a/Engine/Plugins/Animation/ControlRig/Source/ControlRig/Private/Sequencer/MovieSceneControlRigParameterTrack.cpp
+++ b/Engine/Plugins/Animation/ControlRig/Source/ControlRig/Private/Sequencer/MovieSceneControlRigParameterTrack.cpp
@@ -44,7 +44,7 @@ void UMovieSceneControlRigParameterTrack::BeginDestroy()
 	Super::BeginDestroy();
 	if (IsValid(ControlRig))
 	{
-		ControlRig->OnInitialized_AnyThread().RemoveAll(this);
+		ControlRig->OnPostConstruction_AnyThread().RemoveAll(this);
 	}
 }
 
@@ -215,12 +215,12 @@ UMovieSceneSection* UMovieSceneControlRigParameterTrack::CreateControlRigSection
 
 	if (IsValid(ControlRig))
 	{
-		ControlRig->OnInitialized_AnyThread().RemoveAll(this);
+		ControlRig->OnPostConstruction_AnyThread().RemoveAll(this);
 	}
 
 	ControlRig = InControlRig;
 
-	ControlRig->OnInitialized_AnyThread().AddUObject(this, &UMovieSceneControlRigParameterTrack::HandleOnInitialized);
+	ControlRig->OnPostConstruction_AnyThread().AddUObject(this, &UMovieSceneControlRigParameterTrack::HandleOnPostConstructed);
 
 	UMovieSceneControlRigParameterSection* NewSection = Cast<UMovieSceneControlRigParameterSection>(CreateNewSection());
 
@@ -463,7 +463,7 @@ void UMovieSceneControlRigParameterTrack::PostLoad()
 
 	if (IsValid(ControlRig))
 	{
-		ControlRig->OnInitialized_AnyThread().AddUObject(this, &UMovieSceneControlRigParameterTrack::HandleOnInitialized);
+		ControlRig->OnPostConstruction_AnyThread().AddUObject(this, &UMovieSceneControlRigParameterTrack::HandleOnPostConstructed);
 	}
 }
 
@@ -530,7 +530,7 @@ namespace MovieSceneControlRigTrack
 	}
 }
 
-void UMovieSceneControlRigParameterTrack::HandleOnInitialized_GameThread()
+void UMovieSceneControlRigParameterTrack::HandleOnPostConstructed_GameThread()
 {
 	if (IsValid(ControlRig))
 	{
@@ -554,11 +554,11 @@ void UMovieSceneControlRigParameterTrack::HandleOnInitialized_GameThread()
 	}
 }
 
-void UMovieSceneControlRigParameterTrack::HandleOnInitialized(URigVMHost* Subject, const FName& InEventName)
+void UMovieSceneControlRigParameterTrack::HandleOnPostConstructed(UControlRig* Subject, const FName& InEventName)
 {
 	MovieSceneControlRigTrack::AsyncHelpers::ExecuteOnGameThread<void>([this]()
 	{
-		HandleOnInitialized_GameThread();
+		HandleOnPostConstructed_GameThread();
 
 	}).Wait();
 }
@@ -616,7 +616,7 @@ void UMovieSceneControlRigParameterTrack::PostEditImport()
 	{
 		if (ControlRig->OnInitialized_AnyThread().IsBoundToObject(this) == false)
 		{
-			ControlRig->OnInitialized_AnyThread().AddUObject(this, &UMovieSceneControlRigParameterTrack::HandleOnInitialized);
+			ControlRig->OnPostConstruction_AnyThread().AddUObject(this, &UMovieSceneControlRigParameterTrack::HandleOnPostConstructed);
 		}
 		ControlRig->ClearFlags(RF_Transient); //when copied make sure it's no longer transient, sequencer does this for tracks/sections 
 											  //but not for all objects in them since the control rig itself has transient objects.
@@ -648,14 +648,14 @@ void UMovieSceneControlRigParameterTrack::ReplaceControlRig(UControlRig* NewCont
 {
 	if (IsValid(ControlRig))
 	{
-		ControlRig->OnInitialized_AnyThread().RemoveAll(this);
+		ControlRig->OnPostConstruction_AnyThread().RemoveAll(this);
 	}
 
 	ControlRig = NewControlRig;
 
 	if (IsValid(ControlRig))
 	{
-		ControlRig->OnInitialized_AnyThread().AddUObject(this, &UMovieSceneControlRigParameterTrack::HandleOnInitialized);
+		ControlRig->OnPostConstruction_AnyThread().AddUObject(this, &UMovieSceneControlRigParameterTrack::HandleOnPostConstructed);
 
 		if (ControlRig->GetOuter() != this)
 		{
diff --git a/Engine/Plugins/Animation/ControlRig/Source/ControlRig/Public/Sequencer/MovieSceneControlRigParameterTrack.h b/Engine/Plugins/Animation/ControlRig/Source/ControlRig/Public/Sequencer/MovieSceneControlRigParameterTrack.h
--- a/Engine/Plugins/Animation/ControlRig/Source/ControlRig/Public/Sequencer/MovieSceneControlRigParameterTrack.h
+++ b/Engine/Plugins/Animation/ControlRig/Source/ControlRig/Public/Sequencer/MovieSceneControlRigParameterTrack.h
@@ -76,8 +76,8 @@ public:
 #endif // WITH_EDITOR
 
 	//UControlRig Delegates
-	void HandleOnInitialized_GameThread();
-	void HandleOnInitialized(URigVMHost* Subject, const FName& InEventName);
+	void HandleOnPostConstructed_GameThread();
+	void HandleOnPostConstructed(UControlRig* Subject, const FName& InEventName);
 
 #if WITH_EDITOR
 	void HandlePackageDone(const FEndLoadPackageContext& Context);
diff --git a/Engine/Plugins/Animation/ControlRig/Source/ControlRigEditor/Private/Sequencer/ControlRigParameterTrackEditor.cpp b/Engine/Plugins/Animation/ControlRig/Source/ControlRigEditor/Private/Sequencer/ControlRigParameterTrackEditor.cpp
--- a/Engine/Plugins/Animation/ControlRig/Source/ControlRigEditor/Private/Sequencer/ControlRigParameterTrackEditor.cpp
+++ b/Engine/Plugins/Animation/ControlRig/Source/ControlRigEditor/Private/Sequencer/ControlRigParameterTrackEditor.cpp
@@ -427,7 +427,7 @@ void FControlRigParameterTrackEditor::BindControlRig(UControlRig* ControlRig)
 	if (ControlRig && BoundControlRigs.Contains(ControlRig) == false)
 	{
 		ControlRig->ControlModified().AddRaw(this, &FControlRigParameterTrackEditor::HandleControlModified);
-		ControlRig->OnInitialized_AnyThread().AddRaw(this, &FControlRigParameterTrackEditor::HandleOnInitialized);
+		ControlRig->OnPostConstruction_AnyThread().AddRaw(this, &FControlRigParameterTrackEditor::HandleOnPostConstructed);
 		ControlRig->ControlSelected().AddRaw(this, &FControlRigParameterTrackEditor::HandleControlSelected);
 		ControlRig->ControlUndoBracket().AddRaw(this, &FControlRigParameterTrackEditor::HandleControlUndoBracket);
 		ControlRig->ControlRigBound().AddRaw(this, &FControlRigParameterTrackEditor::HandleOnControlRigBound);
@@ -472,7 +472,7 @@ void FControlRigParameterTrackEditor::UnbindControlRig(UControlRig* ControlRig)
 			Track->ConstraintChannelAdded().RemoveAll(this);
 		}
 		ControlRig->ControlModified().RemoveAll(this);
-		ControlRig->OnInitialized_AnyThread().RemoveAll(this);
+		ControlRig->OnPostConstruction_AnyThread().RemoveAll(this);
 		ControlRig->ControlSelected().RemoveAll(this);
 		if (const TSharedPtr<IControlRigObjectBinding> Binding = ControlRig->GetObjectBinding())
 		{
@@ -3438,7 +3438,7 @@ void FControlRigParameterTrackEditor::HandleControlSelected(UControlRig* Subject
 
 
 
-void FControlRigParameterTrackEditor::HandleOnInitialized(URigVMHost* Subject, const FName& InEventName)
+void FControlRigParameterTrackEditor::HandleOnPostConstructed(UControlRig* Subject, const FName& InEventName)
 {
 	if (IsInGameThread())
 	{
diff --git a/Engine/Plugins/Animation/ControlRig/Source/ControlRigEditor/Private/Sequencer/ControlRigParameterTrackEditor.h b/Engine/Plugins/Animation/ControlRig/Source/ControlRigEditor/Private/Sequencer/ControlRigParameterTrackEditor.h
--- a/Engine/Plugins/Animation/ControlRig/Source/ControlRigEditor/Private/Sequencer/ControlRigParameterTrackEditor.h
+++ b/Engine/Plugins/Animation/ControlRig/Source/ControlRigEditor/Private/Sequencer/ControlRigParameterTrackEditor.h
@@ -127,7 +127,7 @@ private:
 	void HandleControlModified(UControlRig* Subject, FRigControlElement* ControlElement, const FRigControlModifiedContext& Context);
 	void HandleControlSelected(UControlRig* Subject, FRigControlElement* ControlElement, bool bSelected);
 	void HandleControlUndoBracket(UControlRig* Subject, bool bOpenUndoBracket);
-	void HandleOnInitialized(URigVMHost* Subject, const FName& InEventName);
+	void HandleOnPostConstructed(UControlRig* Subject, const FName& InEventName);
 	void HandleOnControlRigBound(UControlRig* InControlRig);
 	void HandleOnObjectBoundToControlRig(UObject* InObject);
 
