Control Rig: Connector Element - first steps

diff --git a/Engine/Plugins/Animation/ControlRig/Source/ControlRig/Private/Rigs/RigHierarchy.cpp b/Engine/Plugins/Animation/ControlRig/Source/ControlRig/Private/Rigs/RigHierarchy.cpp
--- a/Engine/Plugins/Animation/ControlRig/Source/ControlRig/Private/Rigs/RigHierarchy.cpp
+++ b/Engine/Plugins/Animation/ControlRig/Source/ControlRig/Private/Rigs/RigHierarchy.cpp
@@ -45,6 +45,7 @@ static TAutoConsoleVariable<int32> CVarControlRigHierarchyTraceCallstack(TEXT("C
 static TAutoConsoleVariable<int32> CVarControlRigHierarchyTracePrecision(TEXT("ControlRig.Hierarchy.TracePrecision"), 3, TEXT("sets the number digits in a float when tracing hierarchies."));
 static TAutoConsoleVariable<int32> CVarControlRigHierarchyTraceOnSpawn(TEXT("ControlRig.Hierarchy.TraceOnSpawn"), 0, TEXT("sets the number of frames to trace when a new hierarchy is spawned"));
 TAutoConsoleVariable<bool> CVarControlRigHierarchyEnableRotationOrder(TEXT("ControlRig.Hierarchy.EnableRotationOrder"), false, TEXT("enables the rotation order for controls"));
+TAutoConsoleVariable<bool> CVarControlRigHierarchyEnableModules(TEXT("ControlRig.Hierarchy.Modules"), false, TEXT("enables the modular rigging functionality"));
 static int32 sRigHierarchyLastTrace = INDEX_NONE;
 static TCHAR sRigHierarchyTraceFormat[16];
 
@@ -373,21 +374,7 @@ void URigHierarchy::ResetToDefault()
 	{
 		if(URigHierarchy* DefaultHierarchy = DefaultHierarchyPtr.Get())
 		{
-			const TArray<FRigElementKey> PreviousSelection = OrderedSelection;
 			CopyHierarchy(DefaultHierarchy);
-
-			// reestablish the selection
-			check(OrderedSelection.IsEmpty());
-			for(const FRigElementKey& Key : PreviousSelection)
-			{
-				if(FRigBaseElement* Element = Find(Key))
-				{
-					check(!Element->bSelected);
-					Element->bSelected = true;
-					OrderedSelection.Add(Element->GetKey());
-					Notify(ERigHierarchyNotification::ElementSelected, Element);
-				}
-			}
 			return;
 		}
 	}
@@ -449,6 +436,7 @@ void URigHierarchy::CopyHierarchy(URigHierarchy* InHierarchy)
 		sizeof(FRigCurveElement),
 		sizeof(FRigRigidBodyElement),
 		sizeof(FRigReferenceElement),
+		sizeof(FRigConnectorElement),
 	}; 
 
 	if(bReallocateElements)
@@ -724,6 +712,7 @@ void URigHierarchy::CopyPose(URigHierarchy* InHierarchy, bool bCurrent, bool bIn
 								ControlElementA->Offset.MarkDirty(ERigTransformType::CurrentLocal);
 								ControlElementA->Pose.MarkDirty(ERigTransformType::CurrentGlobal);
 								ControlElementA->Shape.MarkDirty(ERigTransformType::CurrentGlobal);
+								ControlElementA->PoseVersion++;
 							}
 							if(bInitial)
 							{
@@ -731,6 +720,7 @@ void URigHierarchy::CopyPose(URigHierarchy* InHierarchy, bool bCurrent, bool bIn
 								ControlElementA->Offset.MarkDirty(ERigTransformType::InitialLocal);
 								ControlElementA->Pose.MarkDirty(ERigTransformType::InitialGlobal);
 								ControlElementA->Shape.MarkDirty(ERigTransformType::InitialGlobal);
+								ControlElementA->PoseVersion++;
 							}
 						}
 						else
@@ -739,11 +729,13 @@ void URigHierarchy::CopyPose(URigHierarchy* InHierarchy, bool bCurrent, bool bIn
 							{
 								MultiParentElementA->Pose.Set(ERigTransformType::CurrentGlobal, InHierarchy->GetTransform(MultiParentElementB, ERigTransformType::CurrentGlobal));
 								MultiParentElementA->Pose.MarkDirty(ERigTransformType::CurrentLocal);
+								MultiParentElementA->PoseVersion++;
 							}
 							if(bInitial)
 							{
 								MultiParentElementA->Pose.Set(ERigTransformType::InitialGlobal, InHierarchy->GetTransform(MultiParentElementB, ERigTransformType::InitialGlobal));
 								MultiParentElementA->Pose.MarkDirty(ERigTransformType::InitialLocal);
+								MultiParentElementA->PoseVersion++;
 							}
 						}
 					}
@@ -1282,6 +1274,15 @@ FName URigHierarchy::GetSafeNewDisplayName(const FRigElementKey& InParentElement
 	return *Name;
 }
 
+int32 URigHierarchy::GetPoseVersion(const FRigElementKey& InKey) const
+{
+	if(const FRigTransformElement* TransformElement = Find<FRigTransformElement>(InKey))
+	{
+		return TransformElement->PoseVersion;
+	}
+	return INDEX_NONE;
+}
+
 FEdGraphPinType URigHierarchy::GetControlPinType(FRigControlElement* InControlElement) const
 {
 	check(InControlElement);
@@ -3051,6 +3052,7 @@ void URigHierarchy::SetTransform(FRigTransformElement* InTransformElement, const
 	const ERigTransformType::Type OpposedType = SwapLocalAndGlobal(InTransformType);
 	InTransformElement->Pose.Set(InTransformType, InTransform);
 	InTransformElement->Pose.MarkDirty(OpposedType);
+	InTransformElement->PoseVersion++;
 
 	if(FRigControlElement* ControlElement = Cast<FRigControlElement>(InTransformElement))
 	{
@@ -3865,6 +3867,83 @@ void URigHierarchy::SetControlVisibility(FRigControlElement* InControlElement, b
 #endif
 }
 
+void URigHierarchy::SetConnectorSettings(FRigConnectorElement* InConnectorElement, FRigConnectorSettings InSettings,
+	bool bSetupUndo, bool bForce, bool bPrintPythonCommands)
+{
+	LLM_SCOPE_BYNAME(TEXT("Animation/ControlRig"));
+	if(InConnectorElement == nullptr)
+	{
+		return;
+	}
+
+	const FRigConnectorSettings PreviousSettings = InConnectorElement->Settings;
+	if(!bForce && PreviousSettings == InSettings)
+	{
+		return;
+	}
+
+	if(bSetupUndo && !HasAnyFlags(RF_Transient))
+	{
+		Modify();
+	}
+
+	InConnectorElement->Settings = InSettings;
+	Notify(ERigHierarchyNotification::ConnectorSettingChanged, InConnectorElement);
+	
+#if WITH_EDITOR
+	if (!bPropagatingChange)
+	{
+		TGuardValue<bool> bPropagatingChangeGuardValue(bPropagatingChange, true);
+			
+		for(FRigHierarchyListener& Listener : ListeningHierarchies)
+		{
+			URigHierarchy* ListeningHierarchy = Listener.Hierarchy.Get();
+			if (ListeningHierarchy)
+			{	
+				if(FRigConnectorElement* ListeningElement = Cast<FRigConnectorElement>(ListeningHierarchy->Find(InConnectorElement->GetKey())))
+				{
+					// bSetupUndo = false such that all listening hierarchies performs undo at the same time the root hierachy undos
+					ListeningHierarchy->SetConnectorSettings(ListeningElement, InSettings, false, bForce);
+				}
+			}
+		}
+	}
+
+	if (bPrintPythonCommands)
+	{
+		FString BlueprintName;
+		if (UBlueprint* Blueprint = GetTypedOuter<UBlueprint>())
+		{
+			BlueprintName = Blueprint->GetFName().ToString();
+		}
+		else if (UControlRig* Rig = Cast<UControlRig>(GetOuter()))
+		{
+			if (UBlueprint* BlueprintCR = Cast<UBlueprint>(Rig->GetClass()->ClassGeneratedBy))
+			{
+				BlueprintName = BlueprintCR->GetFName().ToString();
+			}
+		}
+		if (!BlueprintName.IsEmpty())
+		{
+			FString ControlNamePythonized = RigVMPythonUtils::PythonizeName(InConnectorElement->GetName().ToString());
+			FString SettingsName = FString::Printf(TEXT("connector_settings_%s"),
+				*ControlNamePythonized);
+			TArray<FString> Commands = ConnectorSettingsToPythonCommands(InConnectorElement->Settings, SettingsName);
+
+			for (const FString& Command : Commands)
+			{
+				RigVMPythonUtils::Print(BlueprintName, Command);
+			}
+			
+			RigVMPythonUtils::Print(BlueprintName,
+				FString::Printf(TEXT("hierarchy.set_connector_settings(%s, %s)"),
+				*InConnectorElement->GetKey().ToPythonString(),
+				*SettingsName));
+		}
+	}
+#endif
+}
+
 
 float URigHierarchy::GetCurveValue(FRigCurveElement* InCurveElement) const
 {
@@ -4412,6 +4491,15 @@ FRigBaseElement* URigHierarchy::MakeElement(ERigElementType InElementType, int32
 			Element = NewElement<FRigReferenceElement>(InCount);
 			break;
 		}
+		case ERigElementType::Connector:
+		{
+			if(OutStructureSize)
+			{
+				*OutStructureSize = sizeof(FRigConnectorElement);
+			}
+			Element = NewElement<FRigConnectorElement>(InCount);
+			break;
+		}
 		default:
 		{
 			ensure(false);
@@ -4491,6 +4579,15 @@ void URigHierarchy::DestroyElement(FRigBaseElement*& InElement)
 			}
 			break;
 		}
+		case ERigElementType::Connector:
+		{
+			FRigConnectorElement* ExistingElements = Cast<FRigConnectorElement>(InElement);
+			for(int32 Index=0;Index<Count;Index++)
+			{
+				ExistingElements[Index].~FRigConnectorElement(); 
+			}
+			break;
+		}
 		default:
 		{
 			ensure(false);
@@ -6216,4 +6313,16 @@ TArray<FString> URigHierarchy::ControlSettingsToPythonCommands(const FRigControl
 	return Commands;
 }
 
+TArray<FString> URigHierarchy::ConnectorSettingsToPythonCommands(const FRigConnectorSettings& Settings, const FString& NameSettings)
+{
+	TArray<FString> Commands;
+	Commands.Add(FString::Printf(TEXT("%s = unreal.RigConnectorSettings()"),
+			*NameSettings));
+
+	// no content values just yet - we are skipping the ResolvedItem here since
+	// we don't assume it is going to be resolved initially.
+
+	return Commands;
+}
+
 #endif
diff --git a/Engine/Plugins/Animation/ControlRig/Source/ControlRig/Private/Rigs/RigHierarchyController.cpp b/Engine/Plugins/Animation/ControlRig/Source/ControlRig/Private/Rigs/RigHierarchyController.cpp
--- a/Engine/Plugins/Animation/ControlRig/Source/ControlRig/Private/Rigs/RigHierarchyController.cpp
+++ b/Engine/Plugins/Animation/ControlRig/Source/ControlRig/Private/Rigs/RigHierarchyController.cpp
@@ -550,6 +550,68 @@ FRigElementKey URigHierarchyController::AddReference(FName InName, FRigElementKe
 	return NewElement->Key;
 }
 
+FRigElementKey URigHierarchyController::AddConnector(FName InName, FRigElementKey InParent, FTransform InTransform,
+	bool bTransformInGlobal, FRigConnectorSettings InSettings, bool bSetupUndo, bool bPrintPythonCommand)
+{
+	if(!IsValid())
+	{
+		return FRigElementKey();
+	}
+
+#if WITH_EDITOR
+	TSharedPtr<FScopedTransaction> TransactionPtr;
+	if(bSetupUndo)
+	{
+		TransactionPtr = MakeShared<FScopedTransaction>(NSLOCTEXT("RigHierarchyController", "Add Connector", "Add Connector"));
+		Hierarchy->Modify();
+	}
+#endif
+
+	FRigConnectorElement* NewElement = MakeElement<FRigConnectorElement>();
+	{
+		TGuardValue<bool> DisableCacheValidityChecks(Hierarchy->bEnableCacheValidityCheck, false);
+		NewElement->Key.Type = ERigElementType::Connector;
+		NewElement->Key.Name = Hierarchy->GetSafeNewName(InName.ToString(), NewElement->Key.Type);
+		NewElement->Settings = InSettings;
+		AddElement(NewElement, Hierarchy->Get(Hierarchy->GetIndex(InParent)), true);
+
+		if(bTransformInGlobal)
+		{
+			Hierarchy->SetTransform(NewElement, InTransform, ERigTransformType::InitialGlobal, true, false);
+			Hierarchy->SetTransform(NewElement, InTransform, ERigTransformType::CurrentGlobal, true, false);
+		}
+		else
+		{
+			Hierarchy->SetTransform(NewElement, InTransform, ERigTransformType::InitialLocal, true, false);
+			Hierarchy->SetTransform(NewElement, InTransform, ERigTransformType::CurrentLocal, true, false);
+		}
+
+		NewElement->Pose.Current = NewElement->Pose.Initial;
+	}
+
+#if WITH_EDITOR
+	TransactionPtr.Reset();
+
+	if (bPrintPythonCommand && !bSuspendPythonPrinting)
+	{
+		UBlueprint* Blueprint = GetTypedOuter<UBlueprint>();
+		if (Blueprint)
+		{
+			TArray<FString> Commands = GetAddConnectorPythonCommands(NewElement);
+			for (const FString& Command : Commands)
+			{			
+				RigVMPythonUtils::Print(Blueprint->GetFName().ToString(),
+					FString::Printf(TEXT("%s"), *Command));
+			}
+		}
+	}
+#endif
+
+	Hierarchy->EnsureCacheValidity();
+		
+	return NewElement->Key;
+}
+
 FRigControlSettings URigHierarchyController::GetControlSettings(FRigElementKey InKey) const
 {
 	if(!IsValid())
@@ -1016,6 +1078,12 @@ FString URigHierarchyController::ExportToText(TArray<FRigElementKey> InKeys) con
 				FRigReferenceElement::StaticStruct()->ExportText(PerElementData.Content, Element, &DefaultElement, nullptr, PPF_None, nullptr);
 				break;
 			}
+			case ERigElementType::Connector:
+			{
+				FRigConnectorElement DefaultElement;
+				FRigConnectorElement::StaticStruct()->ExportText(PerElementData.Content, Element, &DefaultElement, nullptr, PPF_None, nullptr);
+				break;
+			}
 			default:
 			{
 				ensure(false);
@@ -1144,6 +1212,12 @@ TArray<FRigElementKey> URigHierarchyController::ImportFromText(FString InContent
 				FRigReferenceElement::StaticStruct()->ImportText(*PerElementData.Content, NewElement, nullptr, EPropertyPortFlags::PPF_None, &ErrorPipe, FRigReferenceElement::StaticStruct()->GetName(), true);
 				break;
 			}
+			case ERigElementType::Connector:
+					{
+				NewElement = MakeElement<FRigConnectorElement>();
+				FRigConnectorElement::StaticStruct()->ImportText(*PerElementData.Content, NewElement, nullptr, EPropertyPortFlags::PPF_None, &ErrorPipe, FRigConnectorElement::StaticStruct()->GetName(), true);
+				break;
+			}
 			default:
 			{
 				ensure(false);
@@ -1450,6 +1524,10 @@ TArray<FString> URigHierarchyController::GetAddElementPythonCommands(FRigBaseEle
 	{
 		ensure(false);
 	}
+	else if(FRigConnectorElement* ConnectorElement = Cast<FRigConnectorElement>(Element))
+	{
+		return GetAddConnectorPythonCommands(ConnectorElement);
+	}
 	return TArray<FString>();
 }
 
@@ -1580,6 +1658,38 @@ TArray<FString> URigHierarchyController::GetAddRigidBodyPythonCommands(FRigRigid
 	return Commands;
 }
 
+TArray<FString> URigHierarchyController::GetAddConnectorPythonCommands(FRigConnectorElement* Connector) const
+{
+	TArray<FString> Commands;
+	FString TransformStr = RigVMPythonUtils::TransformToPythonString(Connector->Pose.Initial.Local.Transform);
+
+	FString ParentKeyStr = "''";
+	if (Connector->ParentElement)
+	{
+		ParentKeyStr = Connector->ParentElement->GetKey().ToPythonString();
+	}
+
+	FRigConnectorSettings& Settings = Connector->Settings;
+	FString SettingsStr;
+	{
+		FString ConnectorNamePythonized = RigVMPythonUtils::PythonizeName(Connector->GetName().ToString());
+		SettingsStr = FString::Printf(TEXT("connector_settings_%s"),
+			*ConnectorNamePythonized);
+			
+		Commands.Append(URigHierarchy::ConnectorSettingsToPythonCommands(Settings, SettingsStr));	
+	}
+
+	// AddConnector(FName InName, FRigElementKey InParent, FTransform InTransform, bool bTransformInGlobal = true, FRigConnectorSettings InSettings = FRigConnectorSettings(), bool bSetupUndo = false);
+	Commands.Add(FString::Printf(TEXT("hierarchy_controller.add_connector('%s', %s, %s, False, %s)"),
+		*Connector->GetName().ToString(),
+		*ParentKeyStr,
+		*TransformStr,
+		*SettingsStr
+	));
+
+	return Commands;
+}
+
 TArray<FString> URigHierarchyController::GetSetControlValuePythonCommands(const FRigControlElement* Control, const FRigControlValue& Value,
                                                                           const ERigControlValueType& Type) const
 {
diff --git a/Engine/Plugins/Animation/ControlRig/Source/ControlRig/Private/Rigs/RigHierarchyElements.cpp b/Engine/Plugins/Animation/ControlRig/Source/ControlRig/Private/Rigs/RigHierarchyElements.cpp
--- a/Engine/Plugins/Animation/ControlRig/Source/ControlRig/Private/Rigs/RigHierarchyElements.cpp
+++ b/Engine/Plugins/Animation/ControlRig/Source/ControlRig/Private/Rigs/RigHierarchyElements.cpp
@@ -74,9 +74,13 @@ UScriptStruct* FRigBaseElement::GetElementStruct() const
 		{
 			return FRigRigidBodyElement::StaticStruct();
 		}
+		case ERigElementType::Connector:
+		{
+			return FRigConnectorElement::StaticStruct();
+		}
 		default:
 		{
-				break;
+			break;
 		}
 	}
 	return FRigBaseElement::StaticStruct();
@@ -494,6 +498,7 @@ void FRigTransformElement::Load(FArchive& Ar, URigHierarchy* Hierarchy, ESeriali
 	if(SerializationPhase == ESerializationPhase::StaticData)
 	{
 		Pose.Load(Ar);
+		PoseVersion = 0;
 	}
 }
 
@@ -511,6 +516,7 @@ void FRigTransformElement::CopyPose(FRigBaseElement* InOther, bool bCurrent, boo
 		{
 			Pose.Initial = Other->Pose.Initial;
 		}
+		PoseVersion++;
 	}
 }
 
@@ -521,6 +527,7 @@ void FRigTransformElement::CopyFrom(URigHierarchy* InHierarchy, FRigBaseElement*
 	
 	const FRigTransformElement* SourceTransform = CastChecked<FRigTransformElement>(InOther);
 	Pose = SourceTransform->Pose;
+	PoseVersion++;
 
 	ElementsToDirty.Reset();
 	ElementsToDirty.Reserve(SourceTransform->ElementsToDirty.Num());
@@ -1035,10 +1042,6 @@ uint32 GetTypeHash(const FRigControlSettings& Settings)
 	return Hash;
 }
 
-////////////////////////////////////////////////////////////////////////////////
-// FRigControlElement
-////////////////////////////////////////////////////////////////////////////////
-
 bool FRigControlSettings::operator==(const FRigControlSettings& InOther) const
 {
 	if(AnimationType != InOther.AnimationType)
@@ -1200,6 +1203,10 @@ void FRigControlSettings::SetupLimitArrayForType(bool bLimitTranslation, bool bL
 	}
 }
 
+////////////////////////////////////////////////////////////////////////////////
+// FRigControlElement
+////////////////////////////////////////////////////////////////////////////////
+
 void FRigControlElement::Save(FArchive& Ar, URigHierarchy* Hierarchy, ESerializationPhase SerializationPhase)
 {
 	Super::Save(Ar, Hierarchy, SerializationPhase);
@@ -1417,3 +1424,72 @@ void FRigReferenceElement::CopyPose(FRigBaseElement* InOther, bool bCurrent, boo
 		}
 	}
 }
+
+////////////////////////////////////////////////////////////////////////////////
+// FRigConnectorSettings
+////////////////////////////////////////////////////////////////////////////////
+
+FRigConnectorSettings::FRigConnectorSettings()
+{
+}
+
+void FRigConnectorSettings::Save(FArchive& Ar)
+{
+	Ar.UsingCustomVersion(FControlRigObjectVersion::GUID);
+
+	Ar << ResolvedItem;
+}
+
+void FRigConnectorSettings::Load(FArchive& Ar)
+{
+	Ar.UsingCustomVersion(FControlRigObjectVersion::GUID);
+
+	Ar << ResolvedItem;
+}
+
+uint32 GetTypeHash(const FRigConnectorSettings& Settings)
+{
+	uint32 Hash = GetTypeHash(Settings.ResolvedItem);
+	return Hash;
+}
+
+bool FRigConnectorSettings::operator==(const FRigConnectorSettings& InOther) const
+{
+	if(ResolvedItem != InOther.ResolvedItem)
+	{
+		return false;
+	}
+	return true;
+}
+
+////////////////////////////////////////////////////////////////////////////////
+// FRigConnectorElement
+////////////////////////////////////////////////////////////////////////////////
+
+void FRigConnectorElement::Save(FArchive& Ar, URigHierarchy* Hierarchy, ESerializationPhase SerializationPhase)
+{
+	Super::Save(Ar, Hierarchy, SerializationPhase);
+
+	if(SerializationPhase == ESerializationPhase::StaticData)
+	{
+		Settings.Save(Ar);
+	}
+}
+
+void FRigConnectorElement::Load(FArchive& Ar, URigHierarchy* Hierarchy, ESerializationPhase SerializationPhase)
+{
+	Super::Load(Ar, Hierarchy, SerializationPhase);
+
+	if(SerializationPhase == ESerializationPhase::StaticData)
+	{
+		Settings.Load(Ar);
+	}
+}
+
+void FRigConnectorElement::CopyFrom(URigHierarchy* InHierarchy, FRigBaseElement* InOther, URigHierarchy* InOtherHierarchy)
+{
+	Super::CopyFrom(InHierarchy, InOther, InOtherHierarchy);
+	
+	const FRigConnectorElement* Source = CastChecked<FRigConnectorElement>(InOther);
+	Settings = Source->Settings;
+}
diff --git a/Engine/Plugins/Animation/ControlRig/Source/ControlRig/Public/Rigs/RigHierarchy.h b/Engine/Plugins/Animation/ControlRig/Source/ControlRig/Public/Rigs/RigHierarchy.h
--- a/Engine/Plugins/Animation/ControlRig/Source/ControlRig/Public/Rigs/RigHierarchy.h
+++ b/Engine/Plugins/Animation/ControlRig/Source/ControlRig/Public/Rigs/RigHierarchy.h
@@ -26,6 +26,7 @@ DECLARE_MULTICAST_DELEGATE_TwoParams(FRigHierarchyMetadataChangedDelegate, const
 DECLARE_MULTICAST_DELEGATE_ThreeParams(FRigHierarchyMetadataTagChangedDelegate, const FRigElementKey& /* Key */, const FName& /* Tag */, bool /* AddedOrRemoved */);
 
 extern CONTROLRIG_API TAutoConsoleVariable<bool> CVarControlRigHierarchyEnableRotationOrder;
+extern CONTROLRIG_API TAutoConsoleVariable<bool> CVarControlRigHierarchyEnableModules;
 
 UENUM()
 enum ERigTransformStackEntryType : int
@@ -1776,6 +1777,14 @@ public:
 		SetGlobalTransform(InKey, InTransform, true, bAffectChildren, bSetupUndo);
 	}
 
+	/**
+	 * Returns the version of the transform / pose on the element given its key.
+	 * Versions are incremented with every change occured to the transform. You
+	 * can use this to compare your previous "knowledge" of the pose - and see
+	 * if anybody has changed it during your last access.
+	 */
+	int32 GetPoseVersion(const FRigElementKey& InKey) const;
+
 	/**
 	 * Returns the global offset transform for a given control element.
 	 * @param InKey The key of the control to retrieve the transform for
@@ -2705,6 +2714,36 @@ public:
 		}
 	}
 
+	/**
+	 * Sets the connector settings for a given connector element by key
+	 * @param InKey The key of the connector element to set the settings for
+	 * @param InSettings The new connector settings value to set
+	 * @param bSetupUndo If true the transform stack will be setup for undo / redo
+	 */
+	UFUNCTION(BlueprintCallable, Category = URigHierarchy)
+	void SetConnectorSettings(FRigElementKey InKey, FRigConnectorSettings InSettings, bool bSetupUndo = false, bool bForce = false, bool bPrintPythonCommands = false)
+	{
+		return SetConnectorSettingsByIndex(GetIndex(InKey), InSettings, bSetupUndo, bForce, bPrintPythonCommands);
+	}
+
+	/**
+	 * Sets the connector settings for a given connector element by index
+	 * @param InElementIndex The index of the connector element to set the settings for
+	 * @param InSettings The new connector settings value to set
+	 * @param bSetupUndo If true the transform stack will be setup for undo / redo
+	 */
+	UFUNCTION(BlueprintCallable, Category = URigHierarchy)
+	void SetConnectorSettingsByIndex(int32 InElementIndex, FRigConnectorSettings InSettings, bool bSetupUndo = false, bool bForce = false, bool bPrintPythonCommands = false)
+	{
+		if(Elements.IsValidIndex(InElementIndex))
+		{
+			if(FRigConnectorElement* ConnectorElement = Cast<FRigConnectorElement>(Elements[InElementIndex]))
+			{
+				SetConnectorSettings(ConnectorElement, InSettings, bSetupUndo, bForce, bPrintPythonCommands);
+			}
+		}
+	}
+
 	/**
 	 * Returns the global current or initial value for a given key.
 	 * If the element does not have a parent FTransform::Identity will be returned.
@@ -3659,6 +3698,14 @@ public:
 	 */
 	void SetControlVisibility(FRigControlElement* InControlElement, bool bVisibility);
 
+	/**
+	 * Sets the connector settings for a given connector element
+	 * @param InConnectorElement The element to set the settings for
+	 * @param InSettings The new connector settings value to set
+	 * @param bSetupUndo If true the transform stack will be setup for undo / redo
+	 */
+	void SetConnectorSettings(FRigConnectorElement* InConnectorElement, FRigConnectorSettings InSettings, bool bSetupUndo = false, bool bForce = false, bool bPrintPythonCommands = false);
+
 	/**
 	 * Returns a curve's value. If the curve value is not set, returns 
 	 * @param InCurveElement The element to retrieve the value for
@@ -4105,10 +4152,14 @@ protected:
 			{
 				return 5;
 			}
-			case ERigElementType::Last:
+			case ERigElementType::Connector:
 			{
 				return 6;
 			}
+			case ERigElementType::Last:
+			{
+				return 7;
+			}
 			case ERigElementType::All:
 			default:
 			{
@@ -4149,6 +4200,10 @@ protected:
 				return ERigElementType::Reference;
 			}
 			case 6:
+			{
+				return ERigElementType::Connector;
+			}
+			case 7:
 			{
 				return ERigElementType::Last;
 			}
@@ -4258,6 +4313,7 @@ private:
 
 #if WITH_EDITOR
 	static TArray<FString> ControlSettingsToPythonCommands(const FRigControlSettings& Settings, const FString& NameSettings);
+	static TArray<FString> ConnectorSettingsToPythonCommands(const FRigConnectorSettings& Settings, const FString& NameSettings);
 #endif
 
 	template<typename T>
diff --git a/Engine/Plugins/Animation/ControlRig/Source/ControlRig/Public/Rigs/RigHierarchyController.h b/Engine/Plugins/Animation/ControlRig/Source/ControlRig/Public/Rigs/RigHierarchyController.h
--- a/Engine/Plugins/Animation/ControlRig/Source/ControlRig/Public/Rigs/RigHierarchyController.h
+++ b/Engine/Plugins/Animation/ControlRig/Source/ControlRig/Public/Rigs/RigHierarchyController.h
@@ -237,6 +237,20 @@ public:
         FRigReferenceGetWorldTransformDelegate InDelegate,
         bool bSetupUndo = false);
 
+	/**
+	 * Adds a connector to the hierarchy
+	 * @param InName The suggested name of the new connector - will eventually be corrected by the namespace
+	 * @param InParent The (optional) parent of the new connector. If you don't need a parent, pass FRigElementKey()
+	 * @param InTransform The transform for the new connector - either in local or global space, based on bTransformInGlobal
+	 * @param bTransformInGlobal Set this to true if the Transform passed is expressed in global space, false for local space.
+	 * @param InSettings All of the connector's settings
+	 * @param bSetupUndo If set to true the stack will record the change for undo / redo
+	 * @param bPrintPythonCommand If set to true a python command equivalent to this call will be printed out
+	 * @return The key for the newly created bone.
+	 */
+	UFUNCTION(BlueprintCallable, Category = URigHierarchyController)
+	FRigElementKey AddConnector(FName InName, FRigElementKey InParent, FTransform InTransform, bool bTransformInGlobal = true, FRigConnectorSettings InSettings = FRigConnectorSettings(), bool bSetupUndo = false, bool bPrintPythonCommand = false);
+
 	/**
 	 * Returns the control settings of a given control
 	 * @param InKey The key of the control to receive the settings for
@@ -578,6 +592,8 @@ public:
 
 	TArray<FString> GetAddRigidBodyPythonCommands(FRigRigidBodyElement* RigidBody) const;
 
+	TArray<FString> GetAddConnectorPythonCommands(FRigConnectorElement* Control) const;
+
 	TArray<FString> GetSetControlValuePythonCommands(const FRigControlElement* Control, const FRigControlValue& Value, const ERigControlValueType& Type) const;
 	
 	TArray<FString> GetSetControlOffsetTransformPythonCommands(const FRigControlElement* Control, const FTransform& Offset, bool bInitial = false, bool bAffectChildren = true) const;
diff --git a/Engine/Plugins/Animation/ControlRig/Source/ControlRig/Public/Rigs/RigHierarchyDefines.h b/Engine/Plugins/Animation/ControlRig/Source/ControlRig/Public/Rigs/RigHierarchyDefines.h
--- a/Engine/Plugins/Animation/ControlRig/Source/ControlRig/Public/Rigs/RigHierarchyDefines.h
+++ b/Engine/Plugins/Animation/ControlRig/Source/ControlRig/Public/Rigs/RigHierarchyDefines.h
@@ -28,8 +28,9 @@ enum class ERigElementType : uint8
 	Curve = 0x008,
 	RigidBody = 0x010 UMETA(Hidden), 
 	Reference = 0x020,
-	Last = 0x040 UMETA(Hidden),
-	All = Bone | Null | Control | Curve | RigidBody | Reference,
+	Connector = 0x040,
+	Last = 0x080 UMETA(Hidden),
+	All = Bone | Null | Control | Curve | RigidBody | Reference | Connector,
 	ToResetAfterConstructionEvent = Bone | Control | Curve UMETA(Hidden),
 };
 
@@ -89,6 +90,7 @@ enum class ERigHierarchyNotification : uint8
 	InteractionBracketOpened,
 	InteractionBracketClosed,
 	ElementReordered,
+	ConnectorSettingChanged,
 
 	/** MAX - invalid */
 	Max UMETA(Hidden),
diff --git a/Engine/Plugins/Animation/ControlRig/Source/ControlRig/Public/Rigs/RigHierarchyElements.h b/Engine/Plugins/Animation/ControlRig/Source/ControlRig/Public/Rigs/RigHierarchyElements.h
--- a/Engine/Plugins/Animation/ControlRig/Source/ControlRig/Public/Rigs/RigHierarchyElements.h
+++ b/Engine/Plugins/Animation/ControlRig/Source/ControlRig/Public/Rigs/RigHierarchyElements.h
@@ -746,6 +746,11 @@ public:
 	GENERATED_BODY()
 	DECLARE_RIG_ELEMENT_METHODS(FRigTransformElement)
 
+	FRigTransformElement()
+		: FRigBaseElement()
+		, PoseVersion(0)
+	{}
+
 	virtual ~FRigTransformElement(){}
 
 	virtual void Save(FArchive& A, URigHierarchy* Hierarchy, ESerializationPhase SerializationPhase) override;
@@ -754,6 +759,9 @@ public:
 	UPROPERTY(BlueprintReadOnly, EditAnywhere, Category = RigElement, meta = (DisplayAfter = "Index"))
 	FRigCurrentAndInitialTransform Pose;
 
+	UPROPERTY(BlueprintReadOnly, VisibleAnywhere, Category = RigElement)
+	int32 PoseVersion;
+	
 protected:
 
 	struct FElementToDirty
@@ -792,7 +800,8 @@ protected:
 			InElement->GetType() == ERigElementType::Null ||
 			InElement->GetType() == ERigElementType::Control ||
 			InElement->GetType() == ERigElementType::RigidBody ||
-			InElement->GetType() == ERigElementType::Reference;
+			InElement->GetType() == ERigElementType::Reference ||
+			InElement->GetType() == ERigElementType::Connector;
 	}
 
 public:
@@ -838,7 +847,8 @@ protected:
 	{
 		return InElement->GetType() == ERigElementType::Bone ||
 			InElement->GetType() == ERigElementType::RigidBody ||
-			InElement->GetType() == ERigElementType::Reference;
+			InElement->GetType() == ERigElementType::Reference ||
+			InElement->GetType() == ERigElementType::Connector;
 	}
 
 	friend struct FRigBaseElement;
@@ -1552,6 +1562,71 @@ protected:
 	friend class URigHierarchyController;
 };
 
+USTRUCT(BlueprintType)
+struct CONTROLRIG_API FRigConnectorSettings
+{
+	GENERATED_BODY()
+
+	FRigConnectorSettings();
+
+	void Save(FArchive& Ar);
+	void Load(FArchive& Ar);
+
+	friend uint32 GetTypeHash(const FRigConnectorSettings& Settings);
+
+	UPROPERTY(BlueprintReadWrite, EditAnywhere, Category = Control)
+	FRigElementKey ResolvedItem;
+
+	bool operator == (const FRigConnectorSettings& InOther) const;
+
+	bool operator != (const FRigConnectorSettings& InOther) const
+	{
+		return !(*this == InOther);
+	}
+};
+
+USTRUCT(BlueprintType)
+struct CONTROLRIG_API FRigConnectorElement : public FRigSingleParentElement
+{
+public:
+	
+	GENERATED_BODY()
+	DECLARE_RIG_ELEMENT_METHODS(FRigConnectorElement)
+
+	FRigConnectorElement()
+		: FRigSingleParentElement()
+	{
+		Key.Type = ERigElementType::Connector; 
+	}
+
+	virtual ~FRigConnectorElement(){}
+	
+	virtual void Save(FArchive& A, URigHierarchy* Hierarchy, ESerializationPhase SerializationPhase) override;
+	virtual void Load(FArchive& Ar, URigHierarchy* Hierarchy, ESerializationPhase SerializationPhase) override;
+
+private:
+
+	virtual void CopyFrom(URigHierarchy* InHierarchy, FRigBaseElement* InOther, URigHierarchy* InOtherHierarchy) override;
+
+public:
+
+	UPROPERTY(BlueprintReadOnly, EditAnywhere, Category = Control)
+	FRigConnectorSettings Settings;
+
+protected:
+
+	static bool IsClassOf(const FRigBaseElement* InElement)
+	{
+		return InElement->GetType() == ERigElementType::Connector;
+	}
+
+protected:
+
+	friend struct FRigBaseElement;
+	friend class URigHierarchy;
+	friend class URigHierarchyController;
+};
+
 USTRUCT()
 struct CONTROLRIG_API FRigHierarchyCopyPasteContentPerElement
 {
diff --git a/Engine/Plugins/Animation/ControlRig/Source/ControlRigDeveloper/Private/Graph/ControlRigGraph.cpp b/Engine/Plugins/Animation/ControlRig/Source/ControlRigDeveloper/Private/Graph/ControlRigGraph.cpp
--- a/Engine/Plugins/Animation/ControlRig/Source/ControlRigDeveloper/Private/Graph/ControlRigGraph.cpp
+++ b/Engine/Plugins/Animation/ControlRig/Source/ControlRigDeveloper/Private/Graph/ControlRigGraph.cpp
@@ -82,6 +82,7 @@ void UControlRigGraph::CacheNameLists(URigHierarchy* InHierarchy, const FRigVMDr
 		ElementNameLists.FindOrAdd(ERigElementType::Curve);
 		ElementNameLists.FindOrAdd(ERigElementType::RigidBody);
 		ElementNameLists.FindOrAdd(ERigElementType::Reference);
+		ElementNameLists.FindOrAdd(ERigElementType::Connector);
 
 		TArray<TSharedPtr<FString>>& AllNameList = ElementNameLists.FindChecked(ERigElementType::All);
 		TArray<TSharedPtr<FString>>& BoneNameList = ElementNameLists.FindChecked(ERigElementType::Bone);
@@ -90,6 +91,7 @@ void UControlRigGraph::CacheNameLists(URigHierarchy* InHierarchy, const FRigVMDr
 		TArray<TSharedPtr<FString>>& CurveNameList = ElementNameLists.FindChecked(ERigElementType::Curve);
 		TArray<TSharedPtr<FString>>& RigidBodyNameList = ElementNameLists.FindChecked(ERigElementType::RigidBody);
 		TArray<TSharedPtr<FString>>& ReferenceNameList = ElementNameLists.FindChecked(ERigElementType::Reference);
+		TArray<TSharedPtr<FString>>& ConnectorNameList = ElementNameLists.FindChecked(ERigElementType::Connector);
 		
 		CacheNameListForHierarchy<FRigBaseElement>(InHierarchy, AllNameList, false);
 		CacheNameListForHierarchy<FRigBoneElement>(InHierarchy, BoneNameList, false);
@@ -99,6 +101,7 @@ void UControlRigGraph::CacheNameLists(URigHierarchy* InHierarchy, const FRigVMDr
 		CacheNameListForHierarchy<FRigCurveElement>(InHierarchy, CurveNameList, false);
 		CacheNameListForHierarchy<FRigRigidBodyElement>(InHierarchy, RigidBodyNameList, false);
 		CacheNameListForHierarchy<FRigReferenceElement>(InHierarchy, ReferenceNameList, false);
+		CacheNameListForHierarchy<FRigConnectorElement>(InHierarchy, ConnectorNameList, false);
 
 		LastHierarchyTopologyVersion = InHierarchy->GetTopologyVersion();
 	}
diff --git a/Engine/Plugins/Animation/ControlRig/Source/ControlRigDeveloper/Public/Graph/ControlRigGraph.h b/Engine/Plugins/Animation/ControlRig/Source/ControlRigDeveloper/Public/Graph/ControlRigGraph.h
--- a/Engine/Plugins/Animation/ControlRig/Source/ControlRigDeveloper/Public/Graph/ControlRigGraph.h
+++ b/Engine/Plugins/Animation/ControlRig/Source/ControlRigDeveloper/Public/Graph/ControlRigGraph.h
@@ -53,6 +53,10 @@ public:
 	{
 		return GetElementNameList(ERigElementType::Curve);
 	}
+	const TArray<TSharedPtr<FString>>* GetConnectorNameList(URigVMPin* InPin = nullptr) const
+	{
+		return GetElementNameList(ERigElementType::Connector);
+	}
 
 	virtual const TArray<TSharedPtr<FString>>* GetNameListForWidget(const FString& InWidgetName) const override;
 
diff --git a/Engine/Plugins/Animation/ControlRig/Source/ControlRigEditor/Private/ControlRigEditorStyle.h b/Engine/Plugins/Animation/ControlRig/Source/ControlRigEditor/Private/ControlRigEditorStyle.h
--- a/Engine/Plugins/Animation/ControlRig/Source/ControlRigEditor/Private/ControlRigEditorStyle.h
+++ b/Engine/Plugins/Animation/ControlRig/Source/ControlRigEditor/Private/ControlRigEditorStyle.h
@@ -118,6 +118,7 @@ public:
 			Set("ControlRig.Tree.Null", new IMAGE_BRUSH("Slate/Null_16x", Icon16x16));
 			Set("ControlRig.Tree.RigidBody", new IMAGE_BRUSH("Slate/RigidBody_16x", Icon16x16));
 			Set("ControlRig.Tree.Socket", new IMAGE_BRUSH("Slate/Socket_16x", Icon16x16));
+			Set("ControlRig.Tree.Connector", new IMAGE_BRUSH("Slate/Socket_16x", Icon16x16));
 		}
 
 		// Font?
diff --git a/Engine/Plugins/Animation/ControlRig/Source/ControlRigEditor/Private/ControlRigElementDetails.cpp b/Engine/Plugins/Animation/ControlRig/Source/ControlRigEditor/Private/ControlRigElementDetails.cpp
--- a/Engine/Plugins/Animation/ControlRig/Source/ControlRigEditor/Private/ControlRigElementDetails.cpp
+++ b/Engine/Plugins/Animation/ControlRig/Source/ControlRigEditor/Private/ControlRigElementDetails.cpp
@@ -780,7 +780,8 @@ void FRigBaseElementDetails::CustomizeDetails(IDetailLayoutBuilder& DetailBuilde
 	// if we are not a bone, control or null
 	if(!IsAnyElementOfType(ERigElementType::Bone) &&
 		!IsAnyElementOfType(ERigElementType::Control) &&
-		!IsAnyElementOfType(ERigElementType::Null))
+		!IsAnyElementOfType(ERigElementType::Null) &&
+		!IsAnyElementOfType(ERigElementType::Connector))
 	{
 		CustomizeMetadata(DetailBuilder);
 	}
@@ -5039,4 +5040,32 @@ void FRigNullElementDetails::CustomizeDetails(IDetailLayoutBuilder& DetailBuilde
 	CustomizeMetadata(DetailBuilder);
 }
 
+void FRigConnectorElementDetails::CustomizeDetails(IDetailLayoutBuilder& DetailBuilder)
+{
+	FRigTransformElementDetails::CustomizeDetails(DetailBuilder);
+	CustomizeTransform(DetailBuilder);
+	CustomizeSettings(DetailBuilder);
+	CustomizeMetadata(DetailBuilder);
+}
+
+void FRigConnectorElementDetails::CustomizeSettings(IDetailLayoutBuilder& DetailBuilder)
+{
+	if(PerElementInfos.IsEmpty())
+	{
+		return;
+	}
+
+	if(IsAnyElementNotOfType(ERigElementType::Connector))
+	{
+		return;
+	}
+
+	const TSharedPtr<IPropertyHandle> SettingsHandle = DetailBuilder.GetProperty(GET_MEMBER_NAME_CHECKED(FRigConnectorElement, Settings));
+	DetailBuilder.HideProperty(SettingsHandle);
+
+	IDetailCategoryBuilder& SettingsCategory = DetailBuilder.EditCategory(TEXT("Settings"), LOCTEXT("Settings", "Settings"));
+
+	SettingsCategory.AddProperty(SettingsHandle->GetChildHandle(GET_MEMBER_NAME_CHECKED(FRigConnectorSettings, ResolvedItem))).IsEnabled(false);
+}
+
 #undef LOCTEXT_NAMESPACE
diff --git a/Engine/Plugins/Animation/ControlRig/Source/ControlRigEditor/Private/ControlRigElementDetails.h b/Engine/Plugins/Animation/ControlRig/Source/ControlRigEditor/Private/ControlRigElementDetails.h
--- a/Engine/Plugins/Animation/ControlRig/Source/ControlRigEditor/Private/ControlRigElementDetails.h
+++ b/Engine/Plugins/Animation/ControlRig/Source/ControlRigEditor/Private/ControlRigElementDetails.h
@@ -739,3 +739,19 @@ public:
 	/** IDetailCustomization interface */
 	virtual void CustomizeDetails(IDetailLayoutBuilder& DetailBuilder) override;
 };
+
+class FRigConnectorElementDetails : public FRigTransformElementDetails
+{
+public:
+
+	// Makes a new instance of this detail layout class for a specific detail view requesting it
+	static TSharedRef<IDetailCustomization> MakeInstance()
+	{
+		return MakeShareable(new FRigConnectorElementDetails);
+	}
+
+	/** IDetailCustomization interface */
+	virtual void CustomizeDetails(IDetailLayoutBuilder& DetailBuilder) override;
+
+	void CustomizeSettings(IDetailLayoutBuilder& DetailBuilder);
+};
\ No newline at end of file
diff --git a/Engine/Plugins/Animation/ControlRig/Source/ControlRigEditor/Private/ControlRigHierarchyCommands.cpp b/Engine/Plugins/Animation/ControlRig/Source/ControlRigEditor/Private/ControlRigHierarchyCommands.cpp
--- a/Engine/Plugins/Animation/ControlRig/Source/ControlRigEditor/Private/ControlRigHierarchyCommands.cpp
+++ b/Engine/Plugins/Animation/ControlRig/Source/ControlRigEditor/Private/ControlRigHierarchyCommands.cpp
@@ -10,6 +10,7 @@ void FControlRigHierarchyCommands::RegisterCommands()
 	UI_COMMAND(AddControlItem, "New Control", "Add new control at the origin (0, 0, 0) to the hierarchy.", EUserInterfaceActionType::Button, FInputChord(EKeys::N, EModifierKey::Control));
 	UI_COMMAND(AddAnimationChannelItem, "New Animation Channel", "Add new animation channel.", EUserInterfaceActionType::Button, FInputChord());
 	UI_COMMAND(AddNullItem, "New Null", "Add new null at the origin (0, 0, 0) to the hierarchy.", EUserInterfaceActionType::Button, FInputChord());
+	UI_COMMAND(AddConnectorItem, "New Connector", "Add new connector at the origin (0, 0, 0) to the hierarchy.", EUserInterfaceActionType::Button, FInputChord());
 	UI_COMMAND(DuplicateItem, "Duplicate", "Duplicate the selected items in the hierarchy.", EUserInterfaceActionType::Button, FInputChord(EKeys::D, EModifierKey::Control));
 	UI_COMMAND(MirrorItem, "Mirror", "Mirror the selected items in the hierarchy.", EUserInterfaceActionType::Button, FInputChord());
 	UI_COMMAND(DeleteItem, "Delete", "Deletes the selected items from the hierarchy.", EUserInterfaceActionType::Button, FInputChord(EKeys::Delete));
diff --git a/Engine/Plugins/Animation/ControlRig/Source/ControlRigEditor/Private/ControlRigHierarchyCommands.h b/Engine/Plugins/Animation/ControlRig/Source/ControlRigEditor/Private/ControlRigHierarchyCommands.h
--- a/Engine/Plugins/Animation/ControlRig/Source/ControlRigEditor/Private/ControlRigHierarchyCommands.h
+++ b/Engine/Plugins/Animation/ControlRig/Source/ControlRigEditor/Private/ControlRigHierarchyCommands.h
@@ -30,6 +30,9 @@ public:
 	/** Add Item at origin */
 	TSharedPtr< FUICommandInfo > AddNullItem;
 
+	/** Add Item at origin */
+	TSharedPtr< FUICommandInfo > AddConnectorItem;
+
 	/** Duplicate currently selected items */
 	TSharedPtr< FUICommandInfo > DuplicateItem;
 
diff --git a/Engine/Plugins/Animation/ControlRig/Source/ControlRigEditor/Private/EditMode/ControlRigEditMode.cpp b/Engine/Plugins/Animation/ControlRig/Source/ControlRigEditor/Private/EditMode/ControlRigEditMode.cpp
--- a/Engine/Plugins/Animation/ControlRig/Source/ControlRigEditor/Private/EditMode/ControlRigEditMode.cpp
+++ b/Engine/Plugins/Animation/ControlRig/Source/ControlRigEditor/Private/EditMode/ControlRigEditMode.cpp
@@ -3670,6 +3670,7 @@ void FControlRigEditMode::OnHierarchyModified(ERigHierarchyNotification InNotif,
             	case ERigElementType::Control:
             	case ERigElementType::RigidBody:
             	case ERigElementType::Reference:
+            	case ERigElementType::Connector:
 				{
 					const bool bSelected = InNotif == ERigHierarchyNotification::ElementSelected;
 					// users may select gizmo and control rig units, so we have to let them go through both of them if they do
diff --git a/Engine/Plugins/Animation/ControlRig/Source/ControlRigEditor/Private/Editor/ControlRigEditor.cpp b/Engine/Plugins/Animation/ControlRig/Source/ControlRigEditor/Private/Editor/ControlRigEditor.cpp
--- a/Engine/Plugins/Animation/ControlRig/Source/ControlRigEditor/Private/Editor/ControlRigEditor.cpp
+++ b/Engine/Plugins/Animation/ControlRig/Source/ControlRigEditor/Private/Editor/ControlRigEditor.cpp
@@ -2171,6 +2171,7 @@ void FControlRigEditor::SetRigElementTransform(const FRigElementKey& InElement,
 	switch (InElement.Type)
 	{
 		case ERigElementType::Bone:
+		case ERigElementType::Connector:
 		{
 			FTransform Transform = InTransform;
 			if (bLocal)
@@ -2402,15 +2403,30 @@ void FControlRigEditor::OnWrappedPropertyChangedChainEvent(URigVMDetailsViewWrap
 			}
 			else if(PropertyPath.RemoveFromStart(SettingsString))
 			{
-				const FRigControlSettings Settings  = InWrapperObject->GetContent<FRigControlElement>().Settings;
-
-				FRigControlElement* ControlElement = ControlRigBP->Hierarchy->Find<FRigControlElement>(WrappedElement.GetKey());
-				if(ControlElement == nullptr)
+				if(Key.Type == ERigElementType::Control)
 				{
-					return;
+					const FRigControlSettings Settings  = InWrapperObject->GetContent<FRigControlElement>().Settings;
+
+					FRigControlElement* ControlElement = ControlRigBP->Hierarchy->Find<FRigControlElement>(WrappedElement.GetKey());
+					if(ControlElement == nullptr)
+					{
+						return;
+					}
+
+					ControlRigBP->Hierarchy->SetControlSettings(ControlElement, Settings, true, false, true);
 				}
+				else if(Key.Type == ERigElementType::Connector)
+				{
+					const FRigConnectorSettings Settings  = InWrapperObject->GetContent<FRigConnectorElement>().Settings;
+
+					FRigConnectorElement* ConnectorElement = ControlRigBP->Hierarchy->Find<FRigConnectorElement>(WrappedElement.GetKey());
+					if(ConnectorElement == nullptr)
+					{
+						return;
+					}
 
-				ControlRigBP->Hierarchy->SetControlSettings(ControlElement, Settings, true, false, true);
+					ControlRigBP->Hierarchy->SetConnectorSettings(ConnectorElement, Settings, true, false, true);
+				}
 			}
 
 			if(IsConstructionModeEnabled() || bIsInitial)
@@ -2735,7 +2751,8 @@ void FControlRigEditor::OnHierarchyModified(ERigHierarchyNotification InNotif, U
 								if ((ModelPin->GetCPPType() == TEXT("FName") && ModelPin->GetCustomWidgetName() == TEXT("BoneName") && RemovedElementType == ERigElementType::Bone) ||
 									(ModelPin->GetCPPType() == TEXT("FName") && ModelPin->GetCustomWidgetName() == TEXT("ControlName") && RemovedElementType == ERigElementType::Control) ||
 									(ModelPin->GetCPPType() == TEXT("FName") && ModelPin->GetCustomWidgetName() == TEXT("SpaceName") && RemovedElementType == ERigElementType::Null) ||
-									(ModelPin->GetCPPType() == TEXT("FName") && ModelPin->GetCustomWidgetName() == TEXT("CurveName") && RemovedElementType == ERigElementType::Curve))
+									(ModelPin->GetCPPType() == TEXT("FName") && ModelPin->GetCustomWidgetName() == TEXT("CurveName") && RemovedElementType == ERigElementType::Curve) ||
+									(ModelPin->GetCPPType() == TEXT("FName") && ModelPin->GetCustomWidgetName() == TEXT("ConnectorName") && RemovedElementType == ERigElementType::Connector))
 								{
 									if (ModelPin->GetDefaultValue() == RemovedElementName)
 									{
@@ -2815,7 +2832,8 @@ void FControlRigEditor::OnHierarchyModified(ERigHierarchyNotification InNotif, U
 									if ((ModelPin->GetCPPType() == TEXT("FName") && ModelPin->GetCustomWidgetName() == TEXT("BoneName") && ElementType == ERigElementType::Bone) ||
 										(ModelPin->GetCPPType() == TEXT("FName") && ModelPin->GetCustomWidgetName() == TEXT("ControlName") && ElementType == ERigElementType::Control) ||
 										(ModelPin->GetCPPType() == TEXT("FName") && ModelPin->GetCustomWidgetName() == TEXT("SpaceName") && ElementType == ERigElementType::Null) ||
-										(ModelPin->GetCPPType() == TEXT("FName") && ModelPin->GetCustomWidgetName() == TEXT("CurveName") && ElementType == ERigElementType::Curve))
+										(ModelPin->GetCPPType() == TEXT("FName") && ModelPin->GetCustomWidgetName() == TEXT("CurveName") && ElementType == ERigElementType::Curve) ||
+										(ModelPin->GetCPPType() == TEXT("FName") && ModelPin->GetCustomWidgetName() == TEXT("ConnectorName") && ElementType == ERigElementType::Connector))
 									{
 										if (ModelPin->GetDefaultValue() == OldNameStr)
 										{
@@ -3228,7 +3246,7 @@ void FControlRigEditor::CreateRigHierarchyToGraphDragAndDropMenu() const
 							GetterLabel = LOCTEXT("GetNull","Get Null");
 							GetterTooltip = LOCTEXT("GetNull_ToolTip", "Getter For Null");
 							SetterLabel = LOCTEXT("SetNull","Set Null");
-							SetterTooltip = LOCTEXT("SetNull_eoolTip", "Setter For Null");
+							SetterTooltip = LOCTEXT("SetNull_ToolTip", "Setter For Null");
 						}
 						else if ((DraggedTypes & (uint8)ERigElementType::Control) != 0)
 						{
@@ -3247,6 +3265,13 @@ void FControlRigEditor::CreateRigHierarchyToGraphDragAndDropMenu() const
 								SetterTooltip = LOCTEXT("SetAnimationChannel_ToolTip", "Setter For Animation Channel");
 							}
 						}
+						else if ((DraggedTypes & (uint8)ERigElementType::Connector) != 0)
+						{
+							GetterLabel = LOCTEXT("GetConnector","Get Connector");
+							GetterTooltip = LOCTEXT("GetConnector_ToolTip", "Getter For Connector");
+							SetterLabel = LOCTEXT("SetConnector","Set Connector");
+							SetterTooltip = LOCTEXT("SetConnector_ToolTip", "Setter For Connector");
+						}
 					}
 
 					FToolMenuEntry GetElementsEntry = FToolMenuEntry::InitMenuEntry(
@@ -3282,7 +3307,8 @@ void FControlRigEditor::CreateRigHierarchyToGraphDragAndDropMenu() const
 
 					if (((DraggedTypes & (uint8)ERigElementType::Bone) != 0) ||
 						((DraggedTypes & (uint8)ERigElementType::Control) != 0) ||
-						((DraggedTypes & (uint8)ERigElementType::Null) != 0))
+						((DraggedTypes & (uint8)ERigElementType::Null) != 0) ||
+						((DraggedTypes & (uint8)ERigElementType::Connector) != 0))
 					{
 						FToolMenuEntry& RotationTranslationSeparator = Section.AddSeparator(TEXT("RotationTranslationSeparator"));
 						RotationTranslationSeparator.InsertPosition.Name = SetElementsEntry.Name;
diff --git a/Engine/Plugins/Animation/ControlRig/Source/ControlRigEditor/Private/Editor/ControlRigWrapperObject.cpp b/Engine/Plugins/Animation/ControlRig/Source/ControlRigEditor/Private/Editor/ControlRigWrapperObject.cpp
--- a/Engine/Plugins/Animation/ControlRig/Source/ControlRigEditor/Private/Editor/ControlRigWrapperObject.cpp
+++ b/Engine/Plugins/Animation/ControlRig/Source/ControlRigEditor/Private/Editor/ControlRigWrapperObject.cpp
@@ -36,6 +36,10 @@ UClass* UControlRigWrapperObject::GetClassForStruct(UScriptStruct* InStruct, boo
 			{
 				PropertyEditorModule.RegisterCustomClassLayout(WrapperClassName, FOnGetDetailCustomizationInstance::CreateStatic(&FRigControlElementDetails::MakeInstance));
 			}
+			else if(InStruct == FRigConnectorElement::StaticStruct())
+			{
+				PropertyEditorModule.RegisterCustomClassLayout(WrapperClassName, FOnGetDetailCustomizationInstance::CreateStatic(&FRigConnectorElementDetails::MakeInstance));
+			}
 		}
 	}
 #endif
diff --git a/Engine/Plugins/Animation/ControlRig/Source/ControlRigEditor/Private/Editor/SRigHierarchy.cpp b/Engine/Plugins/Animation/ControlRig/Source/ControlRigEditor/Private/Editor/SRigHierarchy.cpp
--- a/Engine/Plugins/Animation/ControlRig/Source/ControlRigEditor/Private/Editor/SRigHierarchy.cpp
+++ b/Engine/Plugins/Animation/ControlRig/Source/ControlRigEditor/Private/Editor/SRigHierarchy.cpp
@@ -313,6 +313,10 @@ void SRigHierarchy::BindCommands()
 		FExecuteAction::CreateSP(this, &SRigHierarchy::HandleNewItem, ERigElementType::Null, false),
 		FCanExecuteAction::CreateSP(this, &SRigHierarchy::IsNonProceduralElementSelected));
 
+	CommandList->MapAction(Commands.AddConnectorItem,
+		FExecuteAction::CreateSP(this, &SRigHierarchy::HandleNewItem, ERigElementType::Connector, false),
+		FCanExecuteAction::CreateSP(this, &SRigHierarchy::IsNonProceduralElementSelected));
+
 	CommandList->MapAction(Commands.DuplicateItem,
 		FExecuteAction::CreateSP(this, &SRigHierarchy::HandleDuplicateItem),
 		FCanExecuteAction::CreateSP(this, &SRigHierarchy::CanDuplicateItem));
@@ -481,7 +485,8 @@ EVisibility SRigHierarchy::IsSearchbarVisible() const
 	{
 		if ((Hierarchy->Num(ERigElementType::Bone) +
 			Hierarchy->Num(ERigElementType::Null) +
-			Hierarchy->Num(ERigElementType::Control)) > 0)
+			Hierarchy->Num(ERigElementType::Control) +
+			Hierarchy->Num(ERigElementType::Connector)) > 0)
 		{
 			return EVisibility::Visible;
 		}
@@ -745,9 +750,10 @@ void SRigHierarchy::OnHierarchyModified(ERigHierarchyNotification InNotif, URigH
 			break;
 		}
 		case ERigHierarchyNotification::ControlSettingChanged:
+		case ERigHierarchyNotification::ConnectorSettingChanged:
 		{
 			// update color and other settings of the item
-			if(InElement && InElement->GetType() == ERigElementType::Control)
+			if(InElement && ((InElement->GetType() == ERigElementType::Control) || (InElement->GetType() == ERigElementType::Connector)))
 			{
 				for (int32 RootIndex = 0; RootIndex < TreeView->RootElements.Num(); ++RootIndex)
 				{
@@ -1172,6 +1178,11 @@ void SRigHierarchy::CreateContextMenu() const
 								DefaultSection.AddMenuEntry(Commands.AddAnimationChannelItem);
 							}
 							DefaultSection.AddMenuEntry(Commands.AddNullItem);
+
+							if(CVarControlRigHierarchyEnableModules.GetValueOnAnyThread())
+							{
+								DefaultSection.AddMenuEntry(Commands.AddConnectorItem);
+							}
 						})
 					);
 					
@@ -1880,6 +1891,11 @@ void SRigHierarchy::HandleNewItem(ERigElementType InElementType, bool bIsAnimati
 					NewItemKey = Controller->AddNull(NewElementName, ParentKey, ParentTransform, true, true, true);
 					break;
 				}
+				case ERigElementType::Connector:
+				{
+					NewItemKey = Controller->AddConnector(NewElementName, ParentKey, ParentTransform, true, FRigConnectorSettings(), true);
+					break;
+				}
 				default:
 				{
 					return;
@@ -2339,6 +2355,7 @@ TOptional<EItemDropZone> SRigHierarchy::OnCanAcceptDrop(const FDragDropEvent& Dr
 			case ERigElementType::Null:
 			case ERigElementType::RigidBody:
 			case ERigElementType::Reference:
+			case ERigElementType::Connector:
 			{
 				for (const FRigElementKey& DraggedKey : RigDragDropOp->GetElements())
 				{
@@ -2348,6 +2365,7 @@ TOptional<EItemDropZone> SRigHierarchy::OnCanAcceptDrop(const FDragDropEvent& Dr
 						case ERigElementType::Null:
 						case ERigElementType::RigidBody:
 						case ERigElementType::Reference:
+						case ERigElementType::Connector:
 						{
 							break;
 						}
@@ -2641,7 +2659,8 @@ void SRigHierarchy::HandleSetInitialTransformFromCurrentTransform()
 						}
 					}
 					else if (SelectedKey.Type == ERigElementType::Null ||
-						SelectedKey.Type == ERigElementType::Bone)
+						SelectedKey.Type == ERigElementType::Bone ||
+						SelectedKey.Type == ERigElementType::Connector)
 					{
 						FTransform InitialTransform = LocalTransform;
 						if (ControlRigEditor.Pin()->PreviewInstance)
@@ -2707,7 +2726,8 @@ void SRigHierarchy::HandleControlBoneOrSpaceTransform()
 	if (SelectedKeys.Num() == 1)
 	{
 		if (SelectedKeys[0].Type == ERigElementType::Bone ||
-			SelectedKeys[0].Type == ERigElementType::Null)
+			SelectedKeys[0].Type == ERigElementType::Null ||
+			SelectedKeys[0].Type == ERigElementType::Connector)
 		{
 			if(!DebuggedControlRig->GetHierarchy()->IsProcedural(SelectedKeys[0]))
 			{
@@ -2787,6 +2807,7 @@ void SRigHierarchy::HandleUnparent()
 			}
 			case ERigElementType::Null:
 			case ERigElementType::Control:
+			case ERigElementType::Connector:
 			{
 				Controller->RemoveAllParents(SelectedKey, true, true, true);
 				break;
diff --git a/Engine/Plugins/Animation/ControlRig/Source/ControlRigEditor/Private/Editor/SRigHierarchyTreeView.cpp b/Engine/Plugins/Animation/ControlRig/Source/ControlRigEditor/Private/Editor/SRigHierarchyTreeView.cpp
--- a/Engine/Plugins/Animation/ControlRig/Source/ControlRigEditor/Private/Editor/SRigHierarchyTreeView.cpp
+++ b/Engine/Plugins/Animation/ControlRig/Source/ControlRigEditor/Private/Editor/SRigHierarchyTreeView.cpp
@@ -869,6 +869,11 @@ TPair<const FSlateBrush*, FSlateColor> SRigHierarchyItem::GetBrushForElementType
 			Brush = FControlRigEditorStyle::Get().GetBrush("ControlRig.Tree.Socket");
 			break;
 		}
+		case ERigElementType::Connector:
+		{
+			Brush = FControlRigEditorStyle::Get().GetBrush("ControlRig.Tree.Connector");
+			break;
+		}
 		default:
 		{
 			break;
diff --git a/Engine/Plugins/Animation/ControlRig/Source/ControlRigEditor/Private/Graph/ControlRigGraphPanelPinFactory.cpp b/Engine/Plugins/Animation/ControlRig/Source/ControlRigEditor/Private/Graph/ControlRigGraphPanelPinFactory.cpp
--- a/Engine/Plugins/Animation/ControlRig/Source/ControlRigEditor/Private/Graph/ControlRigGraphPanelPinFactory.cpp
+++ b/Engine/Plugins/Animation/ControlRig/Source/ControlRigEditor/Private/Graph/ControlRigGraphPanelPinFactory.cpp
@@ -182,6 +182,15 @@ TSharedPtr<SGraphPin> FControlRigGraphPanelPinFactory::CreatePin_Internal(UEdGra
 						.OnGetSelectedClicked_UObject(RigGraph, &UControlRigGraph::HandleGetSelectedClicked)
 						.OnBrowseClicked_UObject(RigGraph, &UControlRigGraph::HandleBrowseClicked);
 				}
+				else if (CustomWidgetName == TEXT("ConnectorName"))
+				{
+					return SNew(SRigVMGraphPinNameList, InPin)
+						.ModelPin(ModelPin)
+						.OnGetNameFromSelection_UObject(RigGraph, &UControlRigGraph::GetSelectedElementsNameList)
+						.OnGetNameListContent_UObject(RigGraph, &UControlRigGraph::GetConnectorNameList)
+						.OnGetSelectedClicked_UObject(RigGraph, &UControlRigGraph::HandleGetSelectedClicked)
+						.OnBrowseClicked_UObject(RigGraph, &UControlRigGraph::HandleBrowseClicked);
+				}
 				else if (CustomWidgetName == TEXT("DrawingName"))
 				{
 					return SNew(SRigVMGraphPinNameList, InPin)
