Control Rig: Fix static assert on size of FRigBoneElement

diff --git a/Engine/Plugins/Animation/ControlRig/Source/ControlRig/Private/Rigs/RigHierarchy.cpp b/Engine/Plugins/Animation/ControlRig/Source/ControlRig/Private/Rigs/RigHierarchy.cpp
--- a/Engine/Plugins/Animation/ControlRig/Source/ControlRig/Private/Rigs/RigHierarchy.cpp
+++ b/Engine/Plugins/Animation/ControlRig/Source/ControlRig/Private/Rigs/RigHierarchy.cpp
@@ -712,7 +712,7 @@ void URigHierarchy::CopyPose(URigHierarchy* InHierarchy, bool bCurrent, bool bIn
 								ControlElementA->Offset.MarkDirty(ERigTransformType::CurrentLocal);
 								ControlElementA->Pose.MarkDirty(ERigTransformType::CurrentGlobal);
 								ControlElementA->Shape.MarkDirty(ERigTransformType::CurrentGlobal);
-								ControlElementA->PoseVersion++;
+								//ControlElementA->PoseVersion++;
 							}
 							if(bInitial)
 							{
@@ -720,7 +720,7 @@ void URigHierarchy::CopyPose(URigHierarchy* InHierarchy, bool bCurrent, bool bIn
 								ControlElementA->Offset.MarkDirty(ERigTransformType::InitialLocal);
 								ControlElementA->Pose.MarkDirty(ERigTransformType::InitialGlobal);
 								ControlElementA->Shape.MarkDirty(ERigTransformType::InitialGlobal);
-								ControlElementA->PoseVersion++;
+								//ControlElementA->PoseVersion++;
 							}
 						}
 						else
@@ -729,13 +729,13 @@ void URigHierarchy::CopyPose(URigHierarchy* InHierarchy, bool bCurrent, bool bIn
 							{
 								MultiParentElementA->Pose.Set(ERigTransformType::CurrentGlobal, InHierarchy->GetTransform(MultiParentElementB, ERigTransformType::CurrentGlobal));
 								MultiParentElementA->Pose.MarkDirty(ERigTransformType::CurrentLocal);
-								MultiParentElementA->PoseVersion++;
+								//MultiParentElementA->PoseVersion++;
 							}
 							if(bInitial)
 							{
 								MultiParentElementA->Pose.Set(ERigTransformType::InitialGlobal, InHierarchy->GetTransform(MultiParentElementB, ERigTransformType::InitialGlobal));
 								MultiParentElementA->Pose.MarkDirty(ERigTransformType::InitialLocal);
-								MultiParentElementA->PoseVersion++;
+								//MultiParentElementA->PoseVersion++;
 							}
 						}
 					}
@@ -1276,10 +1276,12 @@ FName URigHierarchy::GetSafeNewDisplayName(const FRigElementKey& InParentElement
 
 int32 URigHierarchy::GetPoseVersion(const FRigElementKey& InKey) const
 {
+	/*
 	if(const FRigTransformElement* TransformElement = Find<FRigTransformElement>(InKey))
 	{
 		return TransformElement->PoseVersion;
 	}
+	*/
 	return INDEX_NONE;
 }
 
@@ -3052,7 +3054,7 @@ void URigHierarchy::SetTransform(FRigTransformElement* InTransformElement, const
 	const ERigTransformType::Type OpposedType = SwapLocalAndGlobal(InTransformType);
 	InTransformElement->Pose.Set(InTransformType, InTransform);
 	InTransformElement->Pose.MarkDirty(OpposedType);
-	InTransformElement->PoseVersion++;
+	//InTransformElement->PoseVersion++;
 
 	if(FRigControlElement* ControlElement = Cast<FRigControlElement>(InTransformElement))
 	{
diff --git a/Engine/Plugins/Animation/ControlRig/Source/ControlRig/Private/Rigs/RigHierarchyElements.cpp b/Engine/Plugins/Animation/ControlRig/Source/ControlRig/Private/Rigs/RigHierarchyElements.cpp
--- a/Engine/Plugins/Animation/ControlRig/Source/ControlRig/Private/Rigs/RigHierarchyElements.cpp
+++ b/Engine/Plugins/Animation/ControlRig/Source/ControlRig/Private/Rigs/RigHierarchyElements.cpp
@@ -498,7 +498,7 @@ void FRigTransformElement::Load(FArchive& Ar, URigHierarchy* Hierarchy, ESeriali
 	if(SerializationPhase == ESerializationPhase::StaticData)
 	{
 		Pose.Load(Ar);
-		PoseVersion = 0;
+		//PoseVersion = 0;
 	}
 }
 
@@ -516,7 +516,7 @@ void FRigTransformElement::CopyPose(FRigBaseElement* InOther, bool bCurrent, boo
 		{
 			Pose.Initial = Other->Pose.Initial;
 		}
-		PoseVersion++;
+		//PoseVersion++;
 	}
 }
 
@@ -527,7 +527,7 @@ void FRigTransformElement::CopyFrom(URigHierarchy* InHierarchy, FRigBaseElement*
 	
 	const FRigTransformElement* SourceTransform = CastChecked<FRigTransformElement>(InOther);
 	Pose = SourceTransform->Pose;
-	PoseVersion++;
+	//PoseVersion++;
 
 	ElementsToDirty.Reset();
 	ElementsToDirty.Reserve(SourceTransform->ElementsToDirty.Num());
diff --git a/Engine/Plugins/Animation/ControlRig/Source/ControlRig/Public/Rigs/RigHierarchyElements.h b/Engine/Plugins/Animation/ControlRig/Source/ControlRig/Public/Rigs/RigHierarchyElements.h
--- a/Engine/Plugins/Animation/ControlRig/Source/ControlRig/Public/Rigs/RigHierarchyElements.h
+++ b/Engine/Plugins/Animation/ControlRig/Source/ControlRig/Public/Rigs/RigHierarchyElements.h
@@ -748,7 +748,6 @@ public:
 
 	FRigTransformElement()
 		: FRigBaseElement()
-		, PoseVersion(0)
 	{}
 
 	virtual ~FRigTransformElement(){}
@@ -758,9 +757,6 @@ public:
 
 	UPROPERTY(BlueprintReadOnly, EditAnywhere, Category = RigElement, meta = (DisplayAfter = "Index"))
 	FRigCurrentAndInitialTransform Pose;
-
-	UPROPERTY(BlueprintReadOnly, VisibleAnywhere, Category = RigElement)
-	int32 PoseVersion;
 	
 protected:
 
