Control Rig: Direct Manipulation support for array elements

diff --git a/Engine/Plugins/Animation/ControlRig/Source/ControlRig/Private/Units/RigUnit.cpp b/Engine/Plugins/Animation/ControlRig/Source/ControlRig/Private/Units/RigUnit.cpp
--- a/Engine/Plugins/Animation/ControlRig/Source/ControlRig/Private/Units/RigUnit.cpp
+++ b/Engine/Plugins/Animation/ControlRig/Source/ControlRig/Private/Units/RigUnit.cpp
@@ -36,10 +36,12 @@ bool FRigUnit::AddDirectManipulationTarget_Internal(TArray<FRigDirectManipulatio
 		InPin->GetDirection() == ERigVMPinDirection::IO ||
 		InPin->GetDirection() == ERigVMPinDirection::Visible)
 	{
-		if(InPin->GetCPPTypeObject() == TBaseStructure<FTransform>::Get() ||
+		if(!InPin->IsArray() && (
+			InPin->GetCPPTypeObject() == TBaseStructure<FTransform>::Get() ||
 			InPin->GetCPPTypeObject() == TBaseStructure<FEulerTransform>::Get() ||
 			InPin->GetCPPTypeObject() == TBaseStructure<FVector>::Get() ||
-			InPin->GetCPPTypeObject() == TBaseStructure<FQuat>::Get())
+			InPin->GetCPPTypeObject() == TBaseStructure<FQuat>::Get()
+		))
 		{
 			if(const URigVMPin* ParentPin = InPin->GetParentPin())
 			{
@@ -71,7 +73,8 @@ bool FRigUnit::AddDirectManipulationTarget_Internal(TArray<FRigDirectManipulatio
 	return false;
 }
 
-TTuple<const FProperty*, FRigVMPropertyPath> FRigUnit::FindPropertyFromPinPath(const UScriptStruct* InStruct, const FString& InPinPath)
+TTuple<const FStructProperty*, uint8*> FRigUnit::FindStructPropertyAndTargetMemory(
+	TSharedPtr<FStructOnScope> InInstance, const UScriptStruct* InStruct, const FString& InPinPath)
 {
 	FString PinPath = InPinPath;
 	PinPath = PinPath.Replace(TEXT("["), TEXT("."));
@@ -89,15 +92,33 @@ TTuple<const FProperty*, FRigVMPropertyPath> FRigUnit::FindPropertyFromPinPath(c
 	const FProperty* Property = InStruct->FindPropertyByName(*Left);
 	if(Property == nullptr)
 	{
-		return {nullptr, FRigVMPropertyPath() };
+		return {nullptr, nullptr};
 	}
 
-	if(Right.IsEmpty())
+	const FStructProperty* StructProperty = CastField<FStructProperty>(Property);
+	if(StructProperty == nullptr)
+	{
+		if(const FArrayProperty* ArrayProperty = CastField<FArrayProperty>(Property))
+		{
+			StructProperty = CastField<FStructProperty>(ArrayProperty->Inner);
+			// if we are in an array we expect the path to exist since we
+			// need to exist a sub element.
+			check(!Right.IsEmpty());
+		}
+	}
+	if(StructProperty == nullptr)
+	{
+		return {nullptr, nullptr};
+	}
+
+	uint8* Memory = Property->ContainerPtrToValuePtr<uint8>(InInstance->GetStructMemory());
+	if(!Right.IsEmpty())
 	{
-		return {Property, FRigVMPropertyPath() };
+		const FRigVMPropertyPath PropertyPath(Property, Right);
+		Memory = PropertyPath.GetData<uint8>(Memory, Property);
 	}
 
-	return {Property, FRigVMPropertyPath(Property, Right) };
+	return {StructProperty, Memory};
 }
 
 void FRigUnit::ConfigureDirectManipulationControl(const URigVMUnitNode* InNode, TSharedPtr<FRigDirectManipulationInfo> InInfo, FRigControlSettings& InOutSettings, FRigControlValue& InOutValue) const
@@ -131,18 +152,16 @@ bool FRigUnit::UpdateHierarchyForDirectManipulation(const URigVMUnitNode* InNode
 		return false;
 	}
 
-	auto PropertyAndPath = FindPropertyFromPinPath(Struct, InInfo->Target.Name);
-	const FStructProperty* StructProperty = CastField<FStructProperty>(PropertyAndPath.Get<0>());
-	if(StructProperty == nullptr)
+	TTuple<const FStructProperty*, uint8*> StructPropertyAndTargetMemory =
+		FindStructPropertyAndTargetMemory(InInstance, Struct, InInfo->Target.Name);
+
+	const FStructProperty* StructProperty = StructPropertyAndTargetMemory.Get<0>();
+	const uint8* Memory = StructPropertyAndTargetMemory.Get<1>();
+	if(StructProperty == nullptr || Memory == nullptr)
 	{
 		return false;
 	}
 
-	const uint8* Memory = StructProperty->ContainerPtrToValuePtr<uint8>(InInstance->GetStructMemory());
-	if(PropertyAndPath.Get<1>().IsValid())
-	{
-		Memory = PropertyAndPath.Get<1>().GetData<uint8>((uint8*)Memory, StructProperty);
-	}
 
 	URigHierarchy* Hierarchy = InContext.Hierarchy;
 	if (Hierarchy == nullptr)
@@ -199,17 +218,14 @@ bool FRigUnit::UpdateDirectManipulationFromHierarchy(const URigVMUnitNode* InNod
 		return false;
 	}
 
-	auto PropertyAndPath = FindPropertyFromPinPath(Struct, InInfo->Target.Name);
-	const FStructProperty* StructProperty = CastField<FStructProperty>(PropertyAndPath.Get<0>());
-	if(StructProperty == nullptr)
-	{
-		return false;
-	}
+	TTuple<const FStructProperty*, uint8*> StructPropertyAndTargetMemory =
+		FindStructPropertyAndTargetMemory(InInstance, Struct, InInfo->Target.Name);
 
-	uint8* Memory = StructProperty->ContainerPtrToValuePtr<uint8>(InInstance->GetStructMemory());
-	if(PropertyAndPath.Get<1>().IsValid())
+	const FStructProperty* StructProperty = StructPropertyAndTargetMemory.Get<0>();
+	uint8* Memory = StructPropertyAndTargetMemory.Get<1>();
+	if(StructProperty == nullptr || Memory == nullptr)
 	{
-		Memory = PropertyAndPath.Get<1>().GetData<uint8>(Memory, StructProperty);
+		return false;
 	}
 
 	URigHierarchy* Hierarchy = InContext.Hierarchy;
diff --git a/Engine/Plugins/Animation/ControlRig/Source/ControlRig/Public/Units/RigUnit.h b/Engine/Plugins/Animation/ControlRig/Source/ControlRig/Public/Units/RigUnit.h
--- a/Engine/Plugins/Animation/ControlRig/Source/ControlRig/Public/Units/RigUnit.h
+++ b/Engine/Plugins/Animation/ControlRig/Source/ControlRig/Public/Units/RigUnit.h
@@ -105,7 +105,7 @@ public:
 private:
 	
 	static bool AddDirectManipulationTarget_Internal(TArray<FRigDirectManipulationTarget>& InOutTargets, const URigVMPin* InPin, const UScriptStruct* InScriptStruct);
-	static TTuple<const FProperty*, FRigVMPropertyPath> FindPropertyFromPinPath(const UScriptStruct* InStruct, const FString& InPinPath);
+	static TTuple<const FStructProperty*, uint8*> FindStructPropertyAndTargetMemory(TSharedPtr<FStructOnScope> InInstance, const UScriptStruct* InStruct, const FString& InPinPath); 
 #endif
 };
 
