Control Rig: Improve pose version access.

Only track pose versions if anybody ever asked for it.

Add unit test

diff --git a/Engine/Plugins/Animation/ControlRig/Source/ControlRig/Private/Rigs/RigHierarchy.cpp b/Engine/Plugins/Animation/ControlRig/Source/ControlRig/Private/Rigs/RigHierarchy.cpp
--- a/Engine/Plugins/Animation/ControlRig/Source/ControlRig/Private/Rigs/RigHierarchy.cpp
+++ b/Engine/Plugins/Animation/ControlRig/Source/ControlRig/Private/Rigs/RigHierarchy.cpp
@@ -234,6 +234,8 @@ void URigHierarchy::Load(FArchive& Ar)
 	int32 ElementCount = 0;
 	Ar << ElementCount;
 
+	PoseVersionPerElement.Reset();
+
 	for(int32 ElementIndex = 0; ElementIndex < ElementCount; ElementIndex++)
 	{
 		FRigElementKey Key;
@@ -354,6 +356,7 @@ void URigHierarchy::Reset_Impl(bool bResetElements)
 	ElementsToRetainLocalTransform.Reset();
 	DefaultParentPerElement.Reset();
 	OrderedSelection.Reset();
+	PoseVersionPerElement.Reset();
 
 	if(!IsGarbageCollecting())
 	{
@@ -541,6 +544,7 @@ void URigHierarchy::CopyHierarchy(URigHierarchy* InHierarchy)
 		FRigBaseElement* Source = InHierarchy->Get(Index);
 		FRigBaseElement* Target = Elements[Index];
 		Target->CopyFrom(this, Source, InHierarchy);
+		IncrementPoseVersion(Index);
 	}
 
 	PreviousNameMap.Append(InHierarchy->PreviousNameMap);
@@ -691,6 +695,7 @@ void URigHierarchy::CopyPose(URigHierarchy* InHierarchy, bool bCurrent, bool bIn
 		if(FRigBaseElement* OtherElement = InHierarchy->Find(Element->GetKey()))
 		{
 			Element->CopyPose(OtherElement, bCurrent, bInitial, bWeights);
+			IncrementPoseVersion(Element->Index);
 
 			// if the topologies don't match and we are supposed to match
 			// elements in global space...
@@ -712,7 +717,7 @@ void URigHierarchy::CopyPose(URigHierarchy* InHierarchy, bool bCurrent, bool bIn
 								ControlElementA->Offset.MarkDirty(ERigTransformType::CurrentLocal);
 								ControlElementA->Pose.MarkDirty(ERigTransformType::CurrentGlobal);
 								ControlElementA->Shape.MarkDirty(ERigTransformType::CurrentGlobal);
-								//ControlElementA->PoseVersion++;
+								IncrementPoseVersion(ControlElementA->Index);
 							}
 							if(bInitial)
 							{
@@ -720,7 +725,7 @@ void URigHierarchy::CopyPose(URigHierarchy* InHierarchy, bool bCurrent, bool bIn
 								ControlElementA->Offset.MarkDirty(ERigTransformType::InitialLocal);
 								ControlElementA->Pose.MarkDirty(ERigTransformType::InitialGlobal);
 								ControlElementA->Shape.MarkDirty(ERigTransformType::InitialGlobal);
-								//ControlElementA->PoseVersion++;
+								IncrementPoseVersion(ControlElementA->Index);
 							}
 						}
 						else
@@ -729,13 +734,13 @@ void URigHierarchy::CopyPose(URigHierarchy* InHierarchy, bool bCurrent, bool bIn
 							{
 								MultiParentElementA->Pose.Set(ERigTransformType::CurrentGlobal, InHierarchy->GetTransform(MultiParentElementB, ERigTransformType::CurrentGlobal));
 								MultiParentElementA->Pose.MarkDirty(ERigTransformType::CurrentLocal);
-								//MultiParentElementA->PoseVersion++;
+								IncrementPoseVersion(MultiParentElementA->Index);
 							}
 							if(bInitial)
 							{
 								MultiParentElementA->Pose.Set(ERigTransformType::InitialGlobal, InHierarchy->GetTransform(MultiParentElementB, ERigTransformType::InitialGlobal));
 								MultiParentElementA->Pose.MarkDirty(ERigTransformType::InitialLocal);
-								//MultiParentElementA->PoseVersion++;
+								IncrementPoseVersion(MultiParentElementA->Index);
 							}
 						}
 					}
@@ -1276,12 +1281,10 @@ FName URigHierarchy::GetSafeNewDisplayName(const FRigElementKey& InParentElement
 
 int32 URigHierarchy::GetPoseVersion(const FRigElementKey& InKey) const
 {
-	/*
 	if(const FRigTransformElement* TransformElement = Find<FRigTransformElement>(InKey))
 	{
-		return TransformElement->PoseVersion;
+		return GetPoseVersion(TransformElement->Index);
 	}
-	*/
 	return INDEX_NONE;
 }
 
@@ -3054,7 +3057,7 @@ void URigHierarchy::SetTransform(FRigTransformElement* InTransformElement, const
 	const ERigTransformType::Type OpposedType = SwapLocalAndGlobal(InTransformType);
 	InTransformElement->Pose.Set(InTransformType, InTransform);
 	InTransformElement->Pose.MarkDirty(OpposedType);
-	//InTransformElement->PoseVersion++;
+	IncrementPoseVersion(InTransformElement->Index);
 
 	if(FRigControlElement* ControlElement = Cast<FRigControlElement>(InTransformElement))
 	{
diff --git a/Engine/Plugins/Animation/ControlRig/Source/ControlRig/Private/Rigs/RigHierarchyElements.cpp b/Engine/Plugins/Animation/ControlRig/Source/ControlRig/Private/Rigs/RigHierarchyElements.cpp
--- a/Engine/Plugins/Animation/ControlRig/Source/ControlRig/Private/Rigs/RigHierarchyElements.cpp
+++ b/Engine/Plugins/Animation/ControlRig/Source/ControlRig/Private/Rigs/RigHierarchyElements.cpp
@@ -498,7 +498,6 @@ void FRigTransformElement::Load(FArchive& Ar, URigHierarchy* Hierarchy, ESeriali
 	if(SerializationPhase == ESerializationPhase::StaticData)
 	{
 		Pose.Load(Ar);
-		//PoseVersion = 0;
 	}
 }
 
@@ -516,7 +515,6 @@ void FRigTransformElement::CopyPose(FRigBaseElement* InOther, bool bCurrent, boo
 		{
 			Pose.Initial = Other->Pose.Initial;
 		}
-		//PoseVersion++;
 	}
 }
 
@@ -527,7 +525,6 @@ void FRigTransformElement::CopyFrom(URigHierarchy* InHierarchy, FRigBaseElement*
 	
 	const FRigTransformElement* SourceTransform = CastChecked<FRigTransformElement>(InOther);
 	Pose = SourceTransform->Pose;
-	//PoseVersion++;
 
 	ElementsToDirty.Reset();
 	ElementsToDirty.Reserve(SourceTransform->ElementsToDirty.Num());
diff --git a/Engine/Plugins/Animation/ControlRig/Source/ControlRig/Public/Rigs/RigHierarchy.h b/Engine/Plugins/Animation/ControlRig/Source/ControlRig/Public/Rigs/RigHierarchy.h
--- a/Engine/Plugins/Animation/ControlRig/Source/ControlRig/Public/Rigs/RigHierarchy.h
+++ b/Engine/Plugins/Animation/ControlRig/Source/ControlRig/Public/Rigs/RigHierarchy.h
@@ -4424,9 +4424,28 @@ public:
 	TElementDependencyMap GetDependenciesForVM(const URigVM* InVM, FName InEventName = NAME_None) const;
 
 private:
-	
+
 #endif
 
+	mutable TArray<int32> PoseVersionPerElement;
+	FORCEINLINE int32& GetPoseVersion(int32 InIndex) const
+	{
+		if(!PoseVersionPerElement.IsValidIndex(InIndex))
+		{
+			PoseVersionPerElement.SetNumZeroed(InIndex + 1);
+		}
+		return PoseVersionPerElement[InIndex];
+	}
+	FORCEINLINE void IncrementPoseVersion(int32 InIndex) const
+	{
+		// don't do anything if the pose version array is empty
+		// or the element has not been requested yet.
+		if(PoseVersionPerElement.IsValidIndex(InIndex))
+		{
+			PoseVersionPerElement[InIndex]++;
+		}
+	}
+
 	void UpdateVisibilityOnProxyControls();
 
 	static const TArray<FString>& GetTransformTypeStrings();
