Control Rig: Fix additive rigs ignoring curves during backwards solve

diff --git a/Engine/Plugins/Animation/ControlRig/Source/ControlRig/Private/ControlRig.cpp b/Engine/Plugins/Animation/ControlRig/Source/ControlRig/Private/ControlRig.cpp
--- a/Engine/Plugins/Animation/ControlRig/Source/ControlRig/Private/ControlRig.cpp
+++ b/Engine/Plugins/Animation/ControlRig/Source/ControlRig/Private/ControlRig.cpp
@@ -344,7 +344,8 @@ void UControlRig::Evaluate_AnyThread()
 		if (PoseBeforeBackwardsSolve.Num() == 0)
 		{
 			// If the pose is empty, this is an indication that a new pose is coming in
-			PoseBeforeBackwardsSolve = Hierarchy->GetPose(false, ERigElementType::Bone, TArrayView<const FRigElementKey>());
+			uint8 Types = (uint8)ERigElementType::Bone | (uint8)ERigElementType::Curve;
+			PoseBeforeBackwardsSolve = Hierarchy->GetPose(false, (ERigElementType)Types, TArrayView<const FRigElementKey>());
 		}
 		else
 		{
