Control Rig: Fix assert when moving control while undo window is open
#rnx

diff --git a/Engine/Plugins/Animation/ControlRig/Source/ControlRigEditor/Private/Sequencer/ControlRigParameterTrackEditor.cpp b/Engine/Plugins/Animation/ControlRig/Source/ControlRigEditor/Private/Sequencer/ControlRigParameterTrackEditor.cpp
--- a/Engine/Plugins/Animation/ControlRig/Source/ControlRigEditor/Private/Sequencer/ControlRigParameterTrackEditor.cpp
+++ b/Engine/Plugins/Animation/ControlRig/Source/ControlRigEditor/Private/Sequencer/ControlRigParameterTrackEditor.cpp
@@ -3501,7 +3501,7 @@ void FControlRigParameterTrackEditor::HandleControlModified(UControlRig* Control
 
 void FControlRigParameterTrackEditor::HandleControlUndoBracket(UControlRig* Subject, bool bOpenUndoBracket)
 {
-	if(bOpenUndoBracket && ControlUndoBracket == 0)
+	if(IsInGameThread() && bOpenUndoBracket && ControlUndoBracket == 0)
 	{
 		FScopeLock ScopeLock(&ControlUndoTransactionMutex);
 		ControlUndoTransaction = MakeShareable(new FScopedTransaction(LOCTEXT("KeyMultipleControls", "Auto-Key multiple controls")));
