Control Rig: Fix race condition when copying hierarchy from default, and updating hierarchy cache

diff --git a/Engine/Plugins/Animation/ControlRig/Source/ControlRig/Private/Rigs/RigHierarchy.cpp b/Engine/Plugins/Animation/ControlRig/Source/ControlRig/Private/Rigs/RigHierarchy.cpp
--- a/Engine/Plugins/Animation/ControlRig/Source/ControlRig/Private/Rigs/RigHierarchy.cpp
+++ b/Engine/Plugins/Animation/ControlRig/Source/ControlRig/Private/Rigs/RigHierarchy.cpp
@@ -497,8 +497,6 @@ void URigHierarchy::CopyHierarchy(URigHierarchy* InHierarchy)
 	}
 	else
 	{
-		InHierarchy->UpdateAllCachedChildren();
-		
 		// remove the superfluous elements
 		for(int32 ElementIndex = Elements.Num() - 1; ElementIndex >= InHierarchy->Elements.Num(); ElementIndex--)
 		{
@@ -526,13 +524,6 @@ void URigHierarchy::CopyHierarchy(URigHierarchy* InHierarchy)
 			Target->CachedChildren.Reset();
 			Target->bSelected = false;
 
-			// fast-pass for update children
-			Target->CachedChildren.SetNumZeroed(Source->CachedChildren.Num());
-			for(int32 ChildIndex = 0; ChildIndex < Target->CachedChildren.Num(); ChildIndex++)
-			{
-				Target->CachedChildren[ChildIndex] = Elements[Source->CachedChildren[ChildIndex]->Index];
-			}
-
 			Target->TopologyVersion = InHierarchy->GetTopologyVersion();
 		}
 
@@ -551,11 +542,7 @@ void URigHierarchy::CopyHierarchy(URigHierarchy* InHierarchy)
 	TopologyVersion = InHierarchy->GetTopologyVersion();
 	MetadataVersion = InHierarchy->GetMetadataVersion();
 
-	if(bReallocateElements)
-	{
-		UpdateAllCachedChildren();
-	}
-	
+	UpdateAllCachedChildren();
 	EnsureCacheValidity();
 }
 
