Control Rig: Ensure to build up menus outside of Control Rig Editor

diff --git a/Engine/Plugins/Animation/ControlRig/Source/ControlRigEditor/Private/ControlRigBlueprintEditorLibrary.cpp b/Engine/Plugins/Animation/ControlRig/Source/ControlRigEditor/Private/ControlRigBlueprintEditorLibrary.cpp
--- a/Engine/Plugins/Animation/ControlRig/Source/ControlRigEditor/Private/ControlRigBlueprintEditorLibrary.cpp
+++ b/Engine/Plugins/Animation/ControlRig/Source/ControlRigEditor/Private/ControlRigBlueprintEditorLibrary.cpp
@@ -1,6 +1,7 @@
 // Copyright Epic Games, Inc. All Rights Reserved.
 
 #include "ControlRigBlueprintEditorLibrary.h"
+#include "Editor/SRigHierarchy.h"
 
 #include UE_INLINE_GENERATED_CPP_BY_NAME(ControlRigBlueprintEditorLibrary)
 
@@ -116,3 +117,9 @@ URigHierarchyController* UControlRigBlueprintEditorLibrary::GetHierarchyControll
 	return InRigBlueprint->GetHierarchyController();
 }
 
+void UControlRigBlueprintEditorLibrary::SetupAllEditorMenus()
+{
+	SRigHierarchy::CreateContextMenu();
+	SRigHierarchy::CreateDragDropMenu();
+}
+
diff --git a/Engine/Plugins/Animation/ControlRig/Source/ControlRigEditor/Private/Editor/SRigHierarchy.cpp b/Engine/Plugins/Animation/ControlRig/Source/ControlRigEditor/Private/Editor/SRigHierarchy.cpp
--- a/Engine/Plugins/Animation/ControlRig/Source/ControlRigEditor/Private/Editor/SRigHierarchy.cpp
+++ b/Engine/Plugins/Animation/ControlRig/Source/ControlRigEditor/Private/Editor/SRigHierarchy.cpp
@@ -90,6 +90,9 @@ FString FRigElementHierarchyDragDropOp::GetJoinedElementNames() const
 
 ///////////////////////////////////////////////////////////
 
+const FName SRigHierarchy::ContextMenuName = TEXT("ControlRigEditor.RigHierarchy.ContextMenu");
+const FName SRigHierarchy::DragDropMenuName = TEXT("ControlRigEditor.RigHierarchy.DragDropMenu");
+
 SRigHierarchy::~SRigHierarchy()
 {
 	const FControlRigEditor* Editor = ControlRigEditor.IsValid() ? ControlRigEditor.Pin().Get() : nullptr;
@@ -1039,8 +1042,15 @@ void SRigHierarchy::OnSetExpansionRecursive(TSharedPtr<FRigTreeElement> InItem,
 	TreeView->SetExpansionRecursive(InItem, false, bShouldBeExpanded);
 }
 
-void SRigHierarchy::CreateDragDropMenu() const
+void SRigHierarchy::CreateDragDropMenu()
 {
+	static bool bCreatedMenu = false;
+	if(bCreatedMenu)
+	{
+		return;
+	}
+	bCreatedMenu = true;
+
 	const FName MenuName = DragDropMenuName;
 	UToolMenus* ToolMenus = UToolMenus::Get();
 
@@ -1049,10 +1059,8 @@ void SRigHierarchy::CreateDragDropMenu() const
 		return;
 	}
 
-	if (!ToolMenus->IsMenuRegistered(MenuName))
+	if (UToolMenu* Menu = ToolMenus->ExtendMenu(MenuName))
 	{
-		UToolMenu* Menu = ToolMenus->RegisterMenu(MenuName);
-
 		Menu->AddDynamicSection(NAME_None, FNewToolMenuDelegate::CreateLambda([](UToolMenu* InMenu)
 			{
 				UToolMenus* ToolMenus = UToolMenus::Get();
@@ -1124,8 +1132,15 @@ UToolMenu* SRigHierarchy::GetDragDropMenu(const TArray<FRigElementKey>& DraggedK
 	return Menu;
 }
 
-void SRigHierarchy::CreateContextMenu() const
+void SRigHierarchy::CreateContextMenu()
 {
+	static bool bCreatedMenu = false;
+	if(bCreatedMenu)
+	{
+		return;
+	}
+	bCreatedMenu = true;
+	
 	const FName MenuName = ContextMenuName;
 
 	UToolMenus* ToolMenus = UToolMenus::Get();
@@ -1135,21 +1150,10 @@ void SRigHierarchy::CreateContextMenu() const
 		return;
 	}
 
-	static bool bCreatedContextMenu = false;
-	if (!ToolMenus->IsMenuRegistered(MenuName) || !bCreatedContextMenu)
+	if (UToolMenu* Menu = ToolMenus->ExtendMenu(MenuName))
 	{
-		UToolMenu* Menu = ToolMenus->RegisterMenu(MenuName);
-		bCreatedContextMenu = true;
-		
-		TWeakPtr<const SRigHierarchy> WeakRigHierarchyWidget = SharedThis(this).ToWeakPtr();
-		
-		Menu->AddDynamicSection(NAME_None, FNewToolMenuDelegate::CreateLambda([WeakRigHierarchyWidget](UToolMenu* InMenu)
+		Menu->AddDynamicSection(NAME_None, FNewToolMenuDelegate::CreateLambda([](UToolMenu* InMenu)
 			{
-				if(!WeakRigHierarchyWidget.IsValid())
-				{
-					return;
-				}
-			
 				UControlRigContextMenuContext* MainContext = InMenu->FindContext<UControlRigContextMenuContext>();
 				
 				if (SRigHierarchy* RigHierarchyPanel = MainContext->GetRigHierarchyPanel())
@@ -1191,24 +1195,20 @@ void SRigHierarchy::CreateContextMenu() const
 					ElementsSection.AddMenuEntry(Commands.RenameItem);
 					ElementsSection.AddMenuEntry(Commands.MirrorItem);
 
-					if(RigHierarchyPanel->IsProceduralElementSelected() && WeakRigHierarchyWidget.Pin()->ControlRigBlueprint.IsValid())
+					if(RigHierarchyPanel->IsProceduralElementSelected() && RigHierarchyPanel->ControlRigBlueprint.IsValid())
 					{
 						ElementsSection.AddMenuEntry(
 							"SelectSpawnerNode",
 							LOCTEXT("SelectSpawnerNode", "Select Spawner Node"),
 							LOCTEXT("SelectSpawnerNode_Tooltip", "Selects the node that spawn / added this element."),
 							FSlateIcon(),
-							FUIAction(FExecuteAction::CreateLambda([WeakRigHierarchyWidget]() {
-								if(!WeakRigHierarchyWidget.IsValid())
-								{
-									return;
-								}
-								if(!WeakRigHierarchyWidget.Pin()->ControlRigBlueprint.IsValid())
+							FUIAction(FExecuteAction::CreateLambda([RigHierarchyPanel]() {
+								if(!RigHierarchyPanel->ControlRigBlueprint.IsValid())
 								{
 									return;
 								}
-								const UControlRigBlueprint* CurrentControlRigBlueprint = WeakRigHierarchyWidget.Pin()->ControlRigBlueprint.Get();
-								const TArray<const FRigBaseElement*> Elements = WeakRigHierarchyWidget.Pin()->GetHierarchy()->GetSelectedElements();
+								const UControlRigBlueprint* CurrentControlRigBlueprint = RigHierarchyPanel->ControlRigBlueprint.Get();
+								const TArray<const FRigBaseElement*> Elements = RigHierarchyPanel->GetHierarchy()->GetSelectedElements();
 								for(const FRigBaseElement* Element : Elements)
 								{
 									if(Element->IsProcedural())
diff --git a/Engine/Plugins/Animation/ControlRig/Source/ControlRigEditor/Private/Editor/SRigHierarchy.h b/Engine/Plugins/Animation/ControlRig/Source/ControlRigEditor/Private/Editor/SRigHierarchy.h
--- a/Engine/Plugins/Animation/ControlRig/Source/ControlRigEditor/Private/Editor/SRigHierarchy.h
+++ b/Engine/Plugins/Animation/ControlRig/Source/ControlRigEditor/Private/Editor/SRigHierarchy.h
@@ -151,13 +151,13 @@ private:
 	TOptional<EItemDropZone> OnCanAcceptDrop(const FDragDropEvent& DragDropEvent, EItemDropZone DropZone, TSharedPtr<FRigTreeElement> TargetItem);
 	FReply OnAcceptDrop(const FDragDropEvent& DragDropEvent, EItemDropZone DropZone, TSharedPtr<FRigTreeElement> TargetItem);
 
-	const FName ContextMenuName = TEXT("ControlRigEditor.RigHierarchy.ContextMenu");
-	void CreateContextMenu() const;
+	static const FName ContextMenuName;
+	static void CreateContextMenu();
 	UToolMenu* GetContextMenu();
 	TSharedPtr<FUICommandList> GetContextMenuCommands() const;
 	
-	const FName DragDropMenuName = TEXT("ControlRigEditor.RigHierarchy.DragDropMenu");
-	void CreateDragDropMenu() const;
+	static const FName DragDropMenuName;
+	static void CreateDragDropMenu();
 	UToolMenu* GetDragDropMenu(const TArray<FRigElementKey>& DraggedKeys, FRigElementKey TargetKey);
 
 	/** Our owning control rig editor */
@@ -244,4 +244,5 @@ public:
 
 	friend class FRigTreeElement;
 	friend class SRigHierarchyItem;
+	friend class UControlRigBlueprintEditorLibrary;
 };
\ No newline at end of file
diff --git a/Engine/Plugins/Animation/ControlRig/Source/ControlRigEditor/Public/ControlRigBlueprintEditorLibrary.h b/Engine/Plugins/Animation/ControlRig/Source/ControlRigEditor/Public/ControlRigBlueprintEditorLibrary.h
--- a/Engine/Plugins/Animation/ControlRig/Source/ControlRigEditor/Public/ControlRigBlueprintEditorLibrary.h
+++ b/Engine/Plugins/Animation/ControlRig/Source/ControlRigEditor/Public/ControlRigBlueprintEditorLibrary.h
@@ -20,7 +20,7 @@ enum class ECastToControlRigBlueprintCases : uint8
 	CastFailed
 };
 
-UCLASS(meta=(ScriptName="ControlRigBlueprintLibrary"))
+UCLASS(BlueprintType, meta=(ScriptName="ControlRigBlueprintLibrary"))
 class CONTROLRIGEDITOR_API UControlRigBlueprintEditorLibrary : public UBlueprintFunctionLibrary
 {
 	GENERATED_BODY()
@@ -68,5 +68,8 @@ public:
 
 	UFUNCTION(BlueprintCallable, BlueprintPure, Category = "Control Rig Blueprint")
 	static URigHierarchyController* GetHierarchyController(UControlRigBlueprint* InRigBlueprint);
+
+	UFUNCTION(BlueprintCallable, BlueprintCallable, Category = "Control Rig Blueprint")
+	static void SetupAllEditorMenus();
 };
 
