Control Rig: Fix crash when compiling CR_Mannequin_Body in ControlRigExample
#rnx

diff --git a/Engine/Plugins/Animation/ControlRig/Source/ControlRig/Private/Rigs/RigHierarchyElements.cpp b/Engine/Plugins/Animation/ControlRig/Source/ControlRig/Private/Rigs/RigHierarchyElements.cpp
--- a/Engine/Plugins/Animation/ControlRig/Source/ControlRig/Private/Rigs/RigHierarchyElements.cpp
+++ b/Engine/Plugins/Animation/ControlRig/Source/ControlRig/Private/Rigs/RigHierarchyElements.cpp
@@ -701,11 +701,38 @@ void FRigMultiParentElement::CopyPose(FRigBaseElement* InOther, bool bCurrent, b
 		FRigMultiParentElement* Source = Cast<FRigMultiParentElement>(InOther);
 		if(ensure(Source))
 		{
-			if(ensure(ParentConstraints.Num() == Source->ParentConstraints.Num()))
+			// Find the map between constraint indices
+			TMap<int32, int32> ConstraintIndexToSourceConstraintIndex;
+			for(int32 ConstraintIndex = 0; ConstraintIndex < ParentConstraints.Num(); ConstraintIndex++)
 			{
-				for(int32 ParentIndex = 0; ParentIndex < ParentConstraints.Num(); ParentIndex++)
+				const FRigElementParentConstraint& ParentConstraint = ParentConstraints[ConstraintIndex];
+				int32 SourceConstraintIndex = Source->ParentConstraints.IndexOfByPredicate([ParentConstraint](const FRigElementParentConstraint& Constraint)
 				{
-					ParentConstraints[ParentIndex].CopyPose(Source->ParentConstraints[ParentIndex], bCurrent, bInitial);
+					return Constraint.ParentElement->GetKey() == ParentConstraint.ParentElement->GetKey();
+				});
+				if (SourceConstraintIndex != INDEX_NONE)
+				{
+					ConstraintIndexToSourceConstraintIndex.Add(ConstraintIndex, SourceConstraintIndex);
+				}
+			}
+			
+			for(int32 ParentIndex = 0; ParentIndex < ParentConstraints.Num(); ParentIndex++)
+			{
+				if (int32* SourceConstraintIndex = ConstraintIndexToSourceConstraintIndex.Find(ParentIndex))
+				{
+					ParentConstraints[ParentIndex].CopyPose(Source->ParentConstraints[*SourceConstraintIndex], bCurrent, bInitial);
+				}
+				else
+				{
+					// Otherwise, reset the weights to 0
+					if (bCurrent)
+					{
+						ParentConstraints[ParentIndex].Weight = 0.f;
+					}
+					if (bInitial)
+					{
+						ParentConstraints[ParentIndex].InitialWeight = 0.f;
+					}
 				}
 			}
 		}
