Sequencer: Control Rig:  Remove auto compensation when we modify a channel, can lead to inaccuricies. Will need to add a button eventually to do this as needed

diff --git a/Engine/Plugins/Animation/ControlRig/Source/ControlRigEditor/Private/Sequencer/ControlRigParameterTrackEditor.cpp b/Engine/Plugins/Animation/ControlRig/Source/ControlRigEditor/Private/Sequencer/ControlRigParameterTrackEditor.cpp
--- a/Engine/Plugins/Animation/ControlRig/Source/ControlRigEditor/Private/Sequencer/ControlRigParameterTrackEditor.cpp
+++ b/Engine/Plugins/Animation/ControlRig/Source/ControlRigEditor/Private/Sequencer/ControlRigParameterTrackEditor.cpp
@@ -2226,7 +2226,7 @@ void FControlRigParameterTrackEditor::OnChannelChanged(const FMovieSceneChannelM
 							GEditor->GetTimerManager()->ClearTimer(ChannelChanged.TimerHandle);
 							ChannelChanged.TimerHandle.Invalidate();
 							ChannelChanged.bWasSetAlready = false;
-							TOptional<FFrameNumber> Optional;
+							TOptional<FFrameNumber> Optional;	
 							ISequencer* SequencerRaw = SequencerPtr.Get();
 							for (TPair <UMovieSceneControlRigParameterSection*, TSet<FName>>& Pair : ChannelChanged.SectionControlNames)
 							{
@@ -2243,8 +2243,6 @@ void FControlRigParameterTrackEditor::OnChannelChanged(const FMovieSceneChannelM
 											EvaluateThisControl(Pair.Key, ControlName, Time);
 										}
 									}
-									FControlRigSpaceChannelHelpers::CompensateIfNeeded(Pair.Key->GetControlRig(), SequencerRaw, Pair.Key,
-										ControlName, Optional);
 									Pair.Key->ControlsToSet.Empty();
 								}
 							}
