Control Rig Modules: Blueprint settings

diff --git a/Engine/Plugins/Animation/ControlRig/Source/ControlRig/Private/Rigs/RigModuleDefines.cpp b/Engine/Plugins/Animation/ControlRig/Source/ControlRig/Private/Rigs/RigModuleDefines.cpp
new file mode 100644
--- /dev/null
+++ b/Engine/Plugins/Animation/ControlRig/Source/ControlRig/Private/Rigs/RigModuleDefines.cpp
@@ -0,0 +1,6 @@
+// Copyright Epic Games, Inc. All Rights Reserved.
+
+#include "Rigs/RigModuleDefines.h"
+
+#include UE_INLINE_GENERATED_CPP_BY_NAME(RigModuleDefines)
+
diff --git a/Engine/Plugins/Animation/ControlRig/Source/ControlRig/Public/Rigs/RigModuleDefines.h b/Engine/Plugins/Animation/ControlRig/Source/ControlRig/Public/Rigs/RigModuleDefines.h
new file mode 100644
--- /dev/null
+++ b/Engine/Plugins/Animation/ControlRig/Source/ControlRig/Public/Rigs/RigModuleDefines.h
@@ -0,0 +1,56 @@
+// Copyright Epic Games, Inc. All Rights Reserved.
+
+#pragma once
+
+#include "CoreMinimal.h"
+#include "RigHierarchyElements.h"
+#include "RigModuleDefines.generated.h"
+
+USTRUCT(BlueprintType)
+struct CONTROLRIG_API FRigModuleConnector
+{
+	GENERATED_BODY()
+	
+	FRigModuleConnector()
+		: bIsRoot(false)
+	{}
+
+	UPROPERTY(BlueprintReadWrite, EditAnywhere, Category = Connector)
+	FString Name;
+
+	UPROPERTY(BlueprintReadWrite, EditAnywhere, Category = Connector)
+	bool bIsRoot;
+
+	UPROPERTY(BlueprintReadWrite, EditAnywhere, Category = Connector)
+	FRigConnectorSettings Settings;
+};
+
+USTRUCT(BlueprintType)
+struct CONTROLRIG_API FRigModuleSettings
+{
+	GENERATED_BODY()
+	
+	FRigModuleSettings()
+	{}
+
+	bool IsValidModule() const { return !ExposedConnectors.IsEmpty(); }
+
+	UPROPERTY(BlueprintReadWrite, EditAnywhere, Category = Module,  meta = (AllowedClasses = "/Script/Engine.Texture2D"))
+	FSoftObjectPath Icon;
+
+	UPROPERTY(BlueprintReadOnly, VisibleAnywhere, Category = Module)
+	TArray<FRigModuleConnector> ExposedConnectors;
+};
+
+
+USTRUCT(BlueprintType)
+struct CONTROLRIG_API FRigModuleDescription
+{
+	GENERATED_BODY()
+	
+	UPROPERTY(BlueprintReadOnly, VisibleAnywhere, Category = Module)
+	FSoftObjectPath Path;
+
+	UPROPERTY(BlueprintReadOnly, VisibleAnywhere, Category = Module)
+	FRigModuleSettings Settings;
+};
\ No newline at end of file
diff --git a/Engine/Plugins/Animation/ControlRig/Source/ControlRigDeveloper/Private/ControlRigBlueprint.cpp b/Engine/Plugins/Animation/ControlRig/Source/ControlRigDeveloper/Private/ControlRigBlueprint.cpp
--- a/Engine/Plugins/Animation/ControlRig/Source/ControlRigDeveloper/Private/ControlRigBlueprint.cpp
+++ b/Engine/Plugins/Animation/ControlRig/Source/ControlRigDeveloper/Private/ControlRigBlueprint.cpp
@@ -136,6 +136,15 @@ USkeletalMesh* UControlRigBlueprint::GetPreviewMesh() const
 #endif
 }
 
+bool UControlRigBlueprint::IsControlRigModule() const
+{
+	if(Hierarchy)
+	{
+		return Hierarchy->Num(ERigElementType::Connector) > 0;
+	}
+	return false;
+}
+
 void UControlRigBlueprint::SetPreviewMesh(USkeletalMesh* PreviewMesh, bool bMarkAsDirty/*=true*/)
 {
 #if WITH_EDITORONLY_DATA
@@ -221,6 +230,22 @@ void UControlRigBlueprint::PreSave(FObjectPreSaveContext ObjectSaveContext)
 		}
 		return true;
 	});
+
+	UpdateExposedModuleConnectors();
+}
+
+void UControlRigBlueprint::UpdateExposedModuleConnectors()
+{
+	ModuleSettings.ExposedConnectors.Reset();
+	Hierarchy->ForEach<FRigConnectorElement>([this](const FRigConnectorElement* ConnectorElement) -> bool
+	{
+		FRigModuleConnector ExposedConnector;
+		ExposedConnector.Name = ConnectorElement->GetNameString();
+		ExposedConnector.bIsRoot = Hierarchy->GetNumberOfParents(ConnectorElement) == 0;
+		ExposedConnector.Settings = ConnectorElement->Settings;
+		ModuleSettings.ExposedConnectors.Add(ExposedConnector);
+		return true;
+	});
 }
 
 void UControlRigBlueprint::PostLoad()
diff --git a/Engine/Plugins/Animation/ControlRig/Source/ControlRigDeveloper/Public/ControlRigBlueprint.h b/Engine/Plugins/Animation/ControlRig/Source/ControlRigDeveloper/Public/ControlRigBlueprint.h
--- a/Engine/Plugins/Animation/ControlRig/Source/ControlRigDeveloper/Public/ControlRigBlueprint.h
+++ b/Engine/Plugins/Animation/ControlRig/Source/ControlRigDeveloper/Public/ControlRigBlueprint.h
@@ -16,6 +16,7 @@
 #include "RigVMModel/RigVMClient.h"
 #include "ControlRigValidationPass.h"
 #include "RigVMBlueprint.h"
+#include "Rigs/RigModuleDefines.h"
 
 #if WITH_EDITOR
 #include "Kismet2/CompilerResultsLog.h"
@@ -86,11 +87,19 @@ public:
 	UFUNCTION(BlueprintCallable, Category = "Control Rig Blueprint")
 	virtual USkeletalMesh* GetPreviewMesh() const override;
 
+	UFUNCTION(BlueprintPure, Category = "Control Rig Blueprint")
+	bool IsControlRigModule() const;
+
 	UPROPERTY(EditAnywhere, Category = "Hierarchy")
 	FRigHierarchySettings HierarchySettings;
 
+	UPROPERTY(EditAnywhere, Category = "Hierarchy", AssetRegistrySearchable)
+	FRigModuleSettings ModuleSettings;
+
 protected:
 
+	void UpdateExposedModuleConnectors();
+
 	/** Asset searchable information about exposed public functions on this rig */
 	UPROPERTY(AssetRegistrySearchable)
 	TArray<FRigVMOldPublicFunctionData> PublicFunctions_DEPRECATED;
@@ -269,4 +278,5 @@ private:
 	friend class UControlRigComponent;
 	friend struct FControlRigGraphSchemaAction_PromoteToVariable;
 	friend class UControlRigGraphSchema;
+	friend class FControlRigBlueprintDetails;
 };
diff --git a/Engine/Plugins/Animation/ControlRig/Source/ControlRigEditor/Private/ControlRigBlueprintDetails.cpp b/Engine/Plugins/Animation/ControlRig/Source/ControlRigEditor/Private/ControlRigBlueprintDetails.cpp
new file mode 100644
--- /dev/null
+++ b/Engine/Plugins/Animation/ControlRig/Source/ControlRigEditor/Private/ControlRigBlueprintDetails.cpp
@@ -0,0 +1,39 @@
+// Copyright Epic Games, Inc. All Rights Reserved.
+
+#include "ControlRigBlueprintDetails.h"
+#include "ControlRigBlueprint.h"
+#include "DetailLayoutBuilder.h"
+#include "Rigs/RigHierarchy.h"
+
+TSharedRef<IDetailCustomization> FControlRigBlueprintDetails::MakeInstance()
+{
+	return MakeShareable(new FControlRigBlueprintDetails);
+}
+
+void FControlRigBlueprintDetails::CustomizeDetails(IDetailLayoutBuilder& DetailLayout)
+{
+	bool bIsValidRigModule = CVarControlRigHierarchyEnableModules.GetValueOnAnyThread();
+	if(bIsValidRigModule)
+	{
+		bIsValidRigModule = false;
+		
+		TArray<TWeakObjectPtr<UObject>> DetailObjects;
+		DetailLayout.GetObjectsBeingCustomized(DetailObjects);
+		for(TWeakObjectPtr<UObject> DetailObject : DetailObjects)
+		{
+			if(UControlRigBlueprint* Blueprint = Cast<UControlRigBlueprint>(DetailObject.Get()))
+			{
+				if(Blueprint->IsControlRigModule())
+				{
+					Blueprint->UpdateExposedModuleConnectors();
+					bIsValidRigModule = true;
+				}
+			}
+		}
+	}
+	
+	if(!bIsValidRigModule)
+	{
+		DetailLayout.HideProperty(GET_MEMBER_NAME_CHECKED(UControlRigBlueprint, ModuleSettings));
+	}
+}
\ No newline at end of file
diff --git a/Engine/Plugins/Animation/ControlRig/Source/ControlRigEditor/Private/ControlRigBlueprintDetails.h b/Engine/Plugins/Animation/ControlRig/Source/ControlRigEditor/Private/ControlRigBlueprintDetails.h
new file mode 100644
--- /dev/null
+++ b/Engine/Plugins/Animation/ControlRig/Source/ControlRigEditor/Private/ControlRigBlueprintDetails.h
@@ -0,0 +1,19 @@
+// Copyright Epic Games, Inc. All Rights Reserved.
+
+#pragma once
+
+#include "CoreMinimal.h"
+#include "IDetailCustomization.h"
+#include "Styling/SlateTypes.h"
+
+class IDetailLayoutBuilder;
+
+class FControlRigBlueprintDetails : public IDetailCustomization
+{
+public:
+	/** Makes a new instance of this detail layout class for a specific detail view requesting it */
+	static TSharedRef<IDetailCustomization> MakeInstance();
+
+	// IDetailCustomization interface
+	virtual void CustomizeDetails(IDetailLayoutBuilder& DetailLayout) override;
+};
diff --git a/Engine/Plugins/Animation/ControlRig/Source/ControlRigEditor/Private/ControlRigBlueprintEditorLibrary.cpp b/Engine/Plugins/Animation/ControlRig/Source/ControlRigEditor/Private/ControlRigBlueprintEditorLibrary.cpp
--- a/Engine/Plugins/Animation/ControlRig/Source/ControlRigEditor/Private/ControlRigBlueprintEditorLibrary.cpp
+++ b/Engine/Plugins/Animation/ControlRig/Source/ControlRigEditor/Private/ControlRigBlueprintEditorLibrary.cpp
@@ -2,6 +2,9 @@
 
 #include "ControlRigBlueprintEditorLibrary.h"
 #include "Editor/SRigHierarchy.h"
+#include "AssetRegistry/AssetData.h"
+#include "AssetRegistry/AssetRegistryModule.h"
+#include "RigVMModel/RigVMPin.h"
 
 #include UE_INLINE_GENERATED_CPP_BY_NAME(ControlRigBlueprintEditorLibrary)
 
@@ -123,3 +126,41 @@ void UControlRigBlueprintEditorLibrary::SetupAllEditorMenus()
 	SRigHierarchy::CreateDragDropMenu();
 }
 
+TArray<FRigModuleDescription> UControlRigBlueprintEditorLibrary::GetAvailableRigModules()
+{
+	TArray<FRigModuleDescription> ModuleDescriptions;
+	
+	// Load the asset registry module
+	const FAssetRegistryModule& AssetRegistryModule = FModuleManager::LoadModuleChecked<FAssetRegistryModule>(TEXT("AssetRegistry"));
+
+	// Collect a full list of assets with the specified class
+	TArray<FAssetData> AssetDataList;
+	AssetRegistryModule.Get().GetAssetsByClass(UControlRigBlueprint::StaticClass()->GetClassPathName(), AssetDataList, true);
+
+	for(const FAssetData& AssetData : AssetDataList)
+	{
+		static const FName ModuleSettingsName = GET_MEMBER_NAME_CHECKED(UControlRigBlueprint, ModuleSettings);
+		if(AssetData.FindTag(ModuleSettingsName))
+		{
+			const FString TagValue = AssetData.GetTagValueRef<FString>(ModuleSettingsName);
+			if(TagValue.IsEmpty())
+			{
+				FRigModuleDescription ModuleDescription;
+				ModuleDescription.Path = AssetData.ToSoftObjectPath();
+				
+				FRigVMPinDefaultValueImportErrorContext ErrorPipe;
+				FRigModuleSettings::StaticStruct()->ImportText(*TagValue, &ModuleDescription.Settings, nullptr, PPF_None, &ErrorPipe, FString());
+				if(ErrorPipe.NumErrors == 0)
+				{
+					if(ModuleDescription.Settings.IsValidModule())
+					{
+						ModuleDescriptions.Add(ModuleDescription);
+					}
+				}
+			}
+		}
+	}
+	
+	return ModuleDescriptions;
+}
+
diff --git a/Engine/Plugins/Animation/ControlRig/Source/ControlRigEditor/Private/ControlRigEditorModule.cpp b/Engine/Plugins/Animation/ControlRig/Source/ControlRigEditor/Private/ControlRigEditorModule.cpp
--- a/Engine/Plugins/Animation/ControlRig/Source/ControlRigEditor/Private/ControlRigEditorModule.cpp
+++ b/Engine/Plugins/Animation/ControlRig/Source/ControlRigEditor/Private/ControlRigEditorModule.cpp
@@ -44,6 +44,7 @@
 #include "ControlRigCompilerDetails.h"
 #include "ControlRigDrawingDetails.h"
 #include "ControlRigAnimGraphDetails.h"
+#include "ControlRigBlueprintDetails.h"
 #include "Subsystems/AssetEditorSubsystem.h"
 #include "Sequencer/ControlRigParameterTrackEditor.h"
 #include "ActorFactories/ActorFactorySkeletalMesh.h"
@@ -117,6 +118,9 @@ void FControlRigEditorModule::StartupModule()
 	ClassesToUnregisterOnShutdown.Add(UMovieSceneControlRigParameterSection::StaticClass()->GetFName());
 	PropertyEditorModule.RegisterCustomClassLayout(ClassesToUnregisterOnShutdown.Last(), FOnGetDetailCustomizationInstance::CreateStatic(&FMovieSceneControlRigSectionDetailsCustomization::MakeInstance));
 
+	ClassesToUnregisterOnShutdown.Add(UControlRigBlueprint::StaticClass()->GetFName());
+	PropertyEditorModule.RegisterCustomClassLayout(ClassesToUnregisterOnShutdown.Last(), FOnGetDetailCustomizationInstance::CreateStatic(&FControlRigBlueprintDetails::MakeInstance));
+
 	ClassesToUnregisterOnShutdown.Add(UControlRig::StaticClass()->GetFName());
 
 	// same as ClassesToUnregisterOnShutdown but for properties, there is none right now
diff --git a/Engine/Plugins/Animation/ControlRig/Source/ControlRigEditor/Public/ControlRigBlueprintEditorLibrary.h b/Engine/Plugins/Animation/ControlRig/Source/ControlRigEditor/Public/ControlRigBlueprintEditorLibrary.h
--- a/Engine/Plugins/Animation/ControlRig/Source/ControlRigEditor/Public/ControlRigBlueprintEditorLibrary.h
+++ b/Engine/Plugins/Animation/ControlRig/Source/ControlRigEditor/Public/ControlRigBlueprintEditorLibrary.h
@@ -71,5 +71,8 @@ public:
 
 	UFUNCTION(BlueprintCallable, BlueprintCallable, Category = "Control Rig Blueprint")
 	static void SetupAllEditorMenus();
+
+	UFUNCTION(BlueprintCallable, BlueprintPure, Category = "Control Rig Blueprint")
+	static TArray<FRigModuleDescription> GetAvailableRigModules();
 };
 
