Control RIg: Do not unbind ControlRigLayerInstance if any control rigs remain active
#rnx

diff --git a/Engine/Plugins/Animation/ControlRig/Source/ControlRig/Private/Sequencer/MovieSceneControlRigParameterTemplate.cpp b/Engine/Plugins/Animation/ControlRig/Source/ControlRig/Private/Sequencer/MovieSceneControlRigParameterTemplate.cpp
--- a/Engine/Plugins/Animation/ControlRig/Source/ControlRig/Private/Sequencer/MovieSceneControlRigParameterTemplate.cpp
+++ b/Engine/Plugins/Animation/ControlRig/Source/ControlRig/Private/Sequencer/MovieSceneControlRigParameterTemplate.cpp
@@ -668,6 +668,7 @@ void FControlRigBindingHelper::UnBindFromSequencerInstance(UControlRig* ControlR
 		}
 
 		UControlRigLayerInstance* AnimInstance = Cast<UControlRigLayerInstance>(SkeletalMeshComponent->GetAnimInstance());
+		bool bShouldUnbind = true;
 		if (AnimInstance)
 		{
 			if (!AnimInstance->IsValidLowLevel() ||
@@ -680,9 +681,14 @@ void FControlRigBindingHelper::UnBindFromSequencerInstance(UControlRig* ControlR
 			AnimInstance->ResetNodes();
 			AnimInstance->RecalcRequiredBones();
 			AnimInstance->RemoveControlRigTrack(ControlRig->GetUniqueID());
+
+			bShouldUnbind = AnimInstance->GetFirstAvailableControlRig() == nullptr;
 		}
 
-		FAnimCustomInstanceHelper::UnbindFromSkeletalMeshComponent< UControlRigLayerInstance>(SkeletalMeshComponent);
+		if (bShouldUnbind)
+		{
+			FAnimCustomInstanceHelper::UnbindFromSkeletalMeshComponent< UControlRigLayerInstance>(SkeletalMeshComponent);
+		}
 	}
 }
 
