Control Rig: avoid spamming undo stack when interacting with rig

diff --git a/Engine/Plugins/Animation/ControlRig/Source/ControlRigEditor/Private/Sequencer/ControlRigParameterTrackEditor.cpp b/Engine/Plugins/Animation/ControlRig/Source/ControlRigEditor/Private/Sequencer/ControlRigParameterTrackEditor.cpp
--- a/Engine/Plugins/Animation/ControlRig/Source/ControlRigEditor/Private/Sequencer/ControlRigParameterTrackEditor.cpp
+++ b/Engine/Plugins/Animation/ControlRig/Source/ControlRigEditor/Private/Sequencer/ControlRigParameterTrackEditor.cpp
@@ -3511,10 +3511,14 @@ void FControlRigParameterTrackEditor::HandleControlUndoBracket(UControlRig* Subj
 	if(!bOpenUndoBracket && ControlUndoBracket == 0)
 	{
 		FScopeLock ScopeLock(&ControlUndoTransactionMutex);
+
+		/*
+		// canceling a sub transaction cancels everything to the top. we need to find a better mechanism for this.
 		if(ControlChangedDuringUndoBracket == 0 && ControlUndoTransaction.IsValid())
 		{
 			ControlUndoTransaction->Cancel();
 		}
+		*/
 		ControlUndoTransaction.Reset();
 	}
 }
