Control RIg: Some tracks not autokeying when multiple control rigs are present
#rnx

diff --git a/Engine/Plugins/Animation/ControlRig/Source/ControlRigEditor/Private/Sequencer/ControlRigParameterTrackEditor.cpp b/Engine/Plugins/Animation/ControlRig/Source/ControlRigEditor/Private/Sequencer/ControlRigParameterTrackEditor.cpp
--- a/Engine/Plugins/Animation/ControlRig/Source/ControlRigEditor/Private/Sequencer/ControlRigParameterTrackEditor.cpp
+++ b/Engine/Plugins/Animation/ControlRig/Source/ControlRigEditor/Private/Sequencer/ControlRigParameterTrackEditor.cpp
@@ -3474,24 +3474,29 @@ void FControlRigParameterTrackEditor::HandleControlModified(UControlRig* Control
 	const TArray<FMovieSceneBinding>& Bindings = MovieScene->GetBindings();
 	for (const FMovieSceneBinding& Binding : Bindings)
 	{
-		UMovieSceneControlRigParameterTrack* Track = Cast<UMovieSceneControlRigParameterTrack>(MovieScene->FindTrack(UMovieSceneControlRigParameterTrack::StaticClass(), Binding.GetObjectGuid(), NAME_None));
-		if (Track && Track->GetControlRig() == ControlRig)
+		TArray<UMovieSceneTrack*> Tracks = MovieScene->FindTracks(UMovieSceneControlRigParameterTrack::StaticClass(), Binding.GetObjectGuid(), NAME_None);
+		for (UMovieSceneTrack* BaseTrack : Tracks)
 		{
-			FName Name(*ControlRig->GetName());
-			if (TSharedPtr<IControlRigObjectBinding> ObjectBinding = ControlRig->GetObjectBinding())
+			UMovieSceneControlRigParameterTrack* Track = Cast<UMovieSceneControlRigParameterTrack>(BaseTrack);
+			if (Track && Track->GetControlRig() == ControlRig)
 			{
-				USceneComponent* Component = Cast<USceneComponent>(ObjectBinding->GetBoundObject());
-				if (Component)
+				FName Name(*ControlRig->GetName());
+				if (TSharedPtr<IControlRigObjectBinding> ObjectBinding = ControlRig->GetObjectBinding())
 				{
-					ESequencerKeyMode KeyMode = ESequencerKeyMode::AutoKey;
-					if (Context.SetKey == EControlRigSetKey::Always)
+					USceneComponent* Component = Cast<USceneComponent>(ObjectBinding->GetBoundObject());
+					if (Component)
 					{
-						KeyMode = ESequencerKeyMode::ManualKeyForced;
+						ESequencerKeyMode KeyMode = ESequencerKeyMode::AutoKey;
+						if (Context.SetKey == EControlRigSetKey::Always)
+						{
+							KeyMode = ESequencerKeyMode::ManualKeyForced;
+						}
+						AddControlKeys(Component, ControlRig, Name, ControlElement->GetName(), (EControlRigContextChannelToKey)Context.KeyMask, 
+							KeyMode, Context.LocalTime);
+						ControlChangedDuringUndoBracket++;
 					}
-					AddControlKeys(Component, ControlRig, Name, ControlElement->GetName(), (EControlRigContextChannelToKey)Context.KeyMask, 
-						KeyMode, Context.LocalTime);
-					ControlChangedDuringUndoBracket++;
 				}
+				break;
 			}
 		}
 	}
