Control Rig: Avoid destroying control actors on editor close

diff --git a/Engine/Plugins/Animation/ControlRig/Source/ControlRigEditor/Private/EditMode/ControlRigEditMode.cpp b/Engine/Plugins/Animation/ControlRig/Source/ControlRigEditor/Private/EditMode/ControlRigEditMode.cpp
--- a/Engine/Plugins/Animation/ControlRig/Source/ControlRigEditor/Private/EditMode/ControlRigEditMode.cpp
+++ b/Engine/Plugins/Animation/ControlRig/Source/ControlRigEditor/Private/EditMode/ControlRigEditMode.cpp
@@ -4551,6 +4551,12 @@ void FControlRigEditMode::OnWorldCleanup(UWorld* World, bool bSessionEnded, bool
 	}
 }
 
+void FControlRigEditMode::OnEditorClosed()
+{
+	ControlRigShapeActors.Reset();
+	ControlRigsToRecreate.Reset();
+}
+
 void FControlRigEditMode::DestroyShapesActors(UControlRig* ControlRig)
 {
 	if (ControlRig == nullptr)
diff --git a/Engine/Plugins/Animation/ControlRig/Source/ControlRigEditor/Private/Editor/ControlRigEditor.cpp b/Engine/Plugins/Animation/ControlRig/Source/ControlRigEditor/Private/Editor/ControlRigEditor.cpp
--- a/Engine/Plugins/Animation/ControlRig/Source/ControlRigEditor/Private/Editor/ControlRigEditor.cpp
+++ b/Engine/Plugins/Animation/ControlRig/Source/ControlRigEditor/Private/Editor/ControlRigEditor.cpp
@@ -143,6 +143,7 @@ FControlRigEditor::~FControlRigEditor()
 		if (FControlRigEditMode* EditMode = GetEditMode())
 		{
 			RigBlueprint->OnHierarchyModified().RemoveAll(EditMode);
+			EditMode->OnEditorClosed();
 		}
 	}
 
diff --git a/Engine/Plugins/Animation/ControlRig/Source/ControlRigEditor/Public/EditMode/ControlRigEditMode.h b/Engine/Plugins/Animation/ControlRig/Source/ControlRigEditor/Public/EditMode/ControlRigEditMode.h
--- a/Engine/Plugins/Animation/ControlRig/Source/ControlRigEditor/Public/EditMode/ControlRigEditMode.h
+++ b/Engine/Plugins/Animation/ControlRig/Source/ControlRigEditor/Public/EditMode/ControlRigEditMode.h
@@ -18,6 +18,7 @@
 
 
 class FEditorViewportClient;
+class FRigVMEditor;
 class FViewport;
 class UActorFactory;
 struct FViewportClick;
@@ -459,6 +460,7 @@ private:
 	void OnWorldCleanup(UWorld* World, bool bSessionEnded, bool bCleanupResources);
 	UWorld* WorldPtr = nullptr;
 
+	void OnEditorClosed();
 	
 	TArray<TWeakObjectPtr<UControlRig>> RuntimeControlRigs;
 	TMap<UControlRig*,TStrongObjectPtr<UControlRigEditModeDelegateHelper>> DelegateHelpers;
