Control Rig: Hierarchy can connect functionality + tests

diff --git a/Engine/Plugins/Animation/ControlRig/Source/ControlRig/Private/Rigs/RigConnectionRules.cpp b/Engine/Plugins/Animation/ControlRig/Source/ControlRig/Private/Rigs/RigConnectionRules.cpp
new file mode 100644
--- /dev/null
+++ b/Engine/Plugins/Animation/ControlRig/Source/ControlRig/Private/Rigs/RigConnectionRules.cpp
@@ -0,0 +1,377 @@
+// Copyright Epic Games, Inc. All Rights Reserved.
+
+#include "Rigs/RigConnectionRules.h"
+#include "Rigs/RigHierarchyElements.h"
+#include "Rigs/RigHierarchy.h"
+
+#include UE_INLINE_GENERATED_CPP_BY_NAME(RigConnectionRules)
+
+////////////////////////////////////////////////////////////////////////////////
+// FRigConnectionRuleStash
+////////////////////////////////////////////////////////////////////////////////
+
+FRigConnectionRuleStash::FRigConnectionRuleStash()
+{
+}
+
+FRigConnectionRuleStash::FRigConnectionRuleStash(const FRigConnectionRule* InRule)
+{
+	check(InRule);
+	ScriptStructPath = InRule->GetScriptStruct()->GetPathName();
+	InRule->GetScriptStruct()->ExportText(ExportedText, InRule, InRule, nullptr, PPF_None, nullptr);
+}
+
+void FRigConnectionRuleStash::Save(FArchive& Ar)
+{
+	Ar << ScriptStructPath;
+	Ar << ExportedText;
+}
+
+void FRigConnectionRuleStash::Load(FArchive& Ar)
+{
+	Ar << ScriptStructPath;
+	Ar << ExportedText;
+}
+
+bool FRigConnectionRuleStash::IsValid() const
+{
+	return !ScriptStructPath.IsEmpty() && !ExportedText.IsEmpty();
+}
+
+UScriptStruct* FRigConnectionRuleStash::GetScriptStruct() const
+{
+	if(!ScriptStructPath.IsEmpty())
+	{
+		return FindObject<UScriptStruct>(nullptr, *ScriptStructPath);
+	}
+	return nullptr;
+}
+
+TSharedPtr<FStructOnScope> FRigConnectionRuleStash::Get() const
+{
+	class FErrorPipe : public FOutputDevice
+	{
+	public:
+
+		int32 NumErrors;
+
+		FErrorPipe()
+			: FOutputDevice()
+			, NumErrors(0)
+		{
+		}
+
+		virtual void Serialize(const TCHAR* V, ELogVerbosity::Type Verbosity, const class FName& Category) override
+		{
+			NumErrors++;
+		}
+	};
+
+	UScriptStruct* ScriptStruct = GetScriptStruct();
+	check(ScriptStruct);
+	TSharedPtr<FStructOnScope> StructOnScope = MakeShareable(new FStructOnScope(ScriptStruct));
+	FErrorPipe ErrorPipe;
+	ScriptStruct->ImportText(*ExportedText, StructOnScope->GetStructMemory(), nullptr, PPF_None, &ErrorPipe, ScriptStruct->GetName());
+	return StructOnScope;
+}
+
+const FRigConnectionRule* FRigConnectionRuleStash::Get(TSharedPtr<FStructOnScope>& InOutStorage) const
+{
+	InOutStorage = Get();
+	if(InOutStorage.IsValid() && InOutStorage->IsValid())
+	{
+		check(InOutStorage->GetStruct()->IsChildOf(FRigConnectionRule::StaticStruct()));
+		return reinterpret_cast<const FRigConnectionRule*>(InOutStorage->GetStructMemory());
+	}
+	return nullptr;
+}
+
+bool FRigConnectionRuleStash::operator==(const FRigConnectionRuleStash& InOther) const
+{
+	return ScriptStructPath.Equals(InOther.ScriptStructPath, ESearchCase::CaseSensitive) &&
+		ExportedText.Equals(InOther.ExportedText, ESearchCase::CaseSensitive);
+}
+
+void FRigConnectionRuleStash::AppendAdditionalConnectors(TArray<FRigElementKey>& OutConnectors) const
+{
+	TSharedPtr<FStructOnScope> Storage;
+	if(const FRigConnectionRule* Rule = Get(Storage))
+	{
+		TArray<FRigElementKey> AdditionalConnectors = Rule->GetAdditionalConnectors();
+		for(const FRigElementKey& AdditionalConnector : AdditionalConnectors)
+		{
+			OutConnectors.AddUnique(AdditionalConnector);
+		}
+	}
+}
+
+uint32 GetTypeHash(const FRigConnectionRuleStash& InRuleStash)
+{
+	return HashCombine(GetTypeHash(InRuleStash.ScriptStructPath), GetTypeHash(InRuleStash.ExportedText));
+}
+
+////////////////////////////////////////////////////////////////////////////////
+// FRigConnectionRule
+////////////////////////////////////////////////////////////////////////////////
+
+TArray<FRigElementKey> FRigConnectionRule::GetAdditionalConnectors() const
+{
+	return TArray<FRigElementKey>();
+}
+
+bool FRigConnectionRule::CanConnect(const FRigConnectionInfo* InConnectionInfo, FString* OutFailureReason) const
+{
+	for(const TPair<FRigElementKey, FRigElementKey>& Pair : InConnectionInfo->ConnectionMap)
+	{
+		const FRigElementKey& ConnectorKey = Pair.Key;
+		if(!ConnectorKey.IsValid())
+		{
+			if(OutFailureReason)
+			{
+				static constexpr TCHAR Format[] = TEXT("Connector element '%s' is not valid.");
+				OutFailureReason->Appendf(Format, *ConnectorKey.ToString());
+			}
+			return false;
+		}
+		if(ConnectorKey.Type != ERigElementType::Connector)
+		{
+			if(OutFailureReason)
+			{
+				static constexpr TCHAR Format[] = TEXT("Connector element '%s' is not a Connector.");
+				OutFailureReason->Appendf(Format, *ConnectorKey.ToString());
+			}
+			return false;
+		}
+		if(InConnectionInfo->SourceHierarchy->Find<FRigConnectorElement>(ConnectorKey) == nullptr)
+		{
+			if(OutFailureReason)
+			{
+				static constexpr TCHAR Format[] = TEXT("Connector element '%s' does not exist.");
+				OutFailureReason->Appendf(Format, *ConnectorKey.ToString());
+			}
+			return false;
+		}
+		const FRigElementKey& TargetKey = Pair.Value;
+		if(!TargetKey.IsValid())
+		{
+			if(OutFailureReason)
+			{
+				static constexpr TCHAR Format[] = TEXT("Target element '%s' is not valid.");
+				OutFailureReason->Appendf(Format, *TargetKey.ToString());
+			}
+			return false;
+		}
+		if(InConnectionInfo->TargetHierarchy->Find(TargetKey) == nullptr)
+		{
+			if(OutFailureReason)
+			{
+				static constexpr TCHAR Format[] = TEXT("Target element '%s' does not exist.");
+				OutFailureReason->Appendf(Format, *TargetKey.ToString());
+			}
+			return false;
+		}
+	}
+	return true;
+}
+
+////////////////////////////////////////////////////////////////////////////////
+// FRigAndConnectionRule
+////////////////////////////////////////////////////////////////////////////////
+
+TArray<FRigElementKey> FRigAndConnectionRule::GetAdditionalConnectors() const
+{
+	TArray<FRigElementKey> AdditionalConnectors;
+	for(const FRigConnectionRuleStash& ChildRule : ChildRules)
+	{
+		ChildRule.AppendAdditionalConnectors(AdditionalConnectors);
+	}
+	return AdditionalConnectors;
+}
+
+bool FRigAndConnectionRule::CanConnect(const FRigConnectionInfo* InConnectionInfo, FString* OutFailureReason) const
+{
+	if(!Super::CanConnect(InConnectionInfo, OutFailureReason))
+	{
+		return false;
+	}
+
+	TSharedPtr<FStructOnScope> Storage;
+	for(const FRigConnectionRuleStash& ChildRule : ChildRules)
+	{
+		if(const FRigConnectionRule* Rule = ChildRule.Get(Storage))
+		{
+			if(!Rule->CanConnect(InConnectionInfo, OutFailureReason))
+			{
+				return false;
+			}
+		}
+	}
+
+	return true;
+}
+
+////////////////////////////////////////////////////////////////////////////////
+// FRigOrConnectionRule
+////////////////////////////////////////////////////////////////////////////////
+
+TArray<FRigElementKey> FRigOrConnectionRule::GetAdditionalConnectors() const
+{
+	TArray<FRigElementKey> AdditionalConnectors;
+	for(const FRigConnectionRuleStash& ChildRule : ChildRules)
+	{
+		ChildRule.AppendAdditionalConnectors(AdditionalConnectors);
+	}
+	return AdditionalConnectors;
+}
+
+bool FRigOrConnectionRule::CanConnect(const FRigConnectionInfo* InConnectionInfo, FString* OutFailureReason) const
+{
+	if(!Super::CanConnect(InConnectionInfo, OutFailureReason))
+	{
+		return false;
+	}
+
+	TSharedPtr<FStructOnScope> Storage;
+	for(const FRigConnectionRuleStash& ChildRule : ChildRules)
+	{
+		if(const FRigConnectionRule* Rule = ChildRule.Get(Storage))
+		{
+			if(Rule->CanConnect(InConnectionInfo, OutFailureReason))
+			{
+				return true;
+			}
+		}
+	}
+
+	return false;
+}
+
+////////////////////////////////////////////////////////////////////////////////
+// FRigTypeConnectionRule
+////////////////////////////////////////////////////////////////////////////////
+
+bool FRigTypeConnectionRule::CanConnect(const FRigConnectionInfo* InConnectionInfo, FString* OutFailureReason) const
+{
+	if(!Super::CanConnect(InConnectionInfo, OutFailureReason))
+	{
+		return false;
+	}
+	
+	for(const TPair<FRigElementKey, FRigElementKey>& Pair : InConnectionInfo->ConnectionMap)
+	{
+		const FRigElementKey& TargetKey = Pair.Value;
+		if(!TargetKey.IsTypeOf(ElementType))
+		{
+			if(OutFailureReason)
+			{
+				const FString ExpectedType = StaticEnum<ERigElementType>()->GetDisplayNameTextByValue((int64)ElementType).ToString();
+				static constexpr TCHAR Format[] = TEXT("Element '%s' is not of the expected type (%s).");
+				OutFailureReason->Appendf(Format, *TargetKey.ToString(), *ExpectedType);
+			}
+			return false;
+		}
+	}
+	return true;
+}
+
+////////////////////////////////////////////////////////////////////////////////
+// FRigTagConnectionRule
+////////////////////////////////////////////////////////////////////////////////
+
+bool FRigTagConnectionRule::CanConnect(const FRigConnectionInfo* InConnectionInfo, FString* OutFailureReason) const
+{
+	if(!Super::CanConnect(InConnectionInfo, OutFailureReason))
+	{
+		return false;
+	}
+	const FRigElementKey TargetKey = InConnectionInfo->ConnectionMap.begin().Value();
+	if(!InConnectionInfo->TargetHierarchy->HasTag(TargetKey, Tag))
+	{
+		if(OutFailureReason)
+		{
+			static constexpr TCHAR Format[] = TEXT("Element '%s' does not contain tag '%s'.");
+			OutFailureReason->Appendf(Format, *TargetKey.ToString(), *Tag.ToString());
+		}
+		return false;
+	}
+	return true;
+}
+
+////////////////////////////////////////////////////////////////////////////////
+// FRigChainConnectionRule
+////////////////////////////////////////////////////////////////////////////////
+
+bool FRigChainConnectionRule::CanConnect(const FRigConnectionInfo* InConnectionInfo, FString* OutFailureReason) const
+{
+	if(!Super::CanConnect(InConnectionInfo, OutFailureReason))
+	{
+		return false;
+	}
+
+	const FRigElementKey StartKey = InConnectionInfo->ConnectionMap.begin().Value();
+	const FRigElementKey EndKey = InConnectionInfo->ConnectionMap.end().Value();
+
+	FRigElementKey NextKey = EndKey;
+	TArray<FRigElementKey> ReversedChain;
+
+	while(NextKey.IsValid())
+	{
+		ReversedChain.Add(NextKey);
+		if(NextKey == StartKey)
+		{
+			break;
+		}
+		NextKey = InConnectionInfo->TargetHierarchy->GetFirstParent(NextKey);
+	}
+
+	if(ReversedChain.Last() != StartKey)
+	{
+		if(OutFailureReason)
+		{
+			static constexpr TCHAR Format[] = TEXT("Elements '%s' and '%s' are not on a chain.");
+			OutFailureReason->Appendf(Format, *StartKey.ToString(), *EndKey.ToString());
+		}
+		return false;
+	}
+
+	if(ReversedChain.Num() < MinNumBones)
+	{
+		if(OutFailureReason)
+		{
+			static constexpr TCHAR Format[] = TEXT("Chain from '%s' to '%s' only has %d elements (minimum of %d required).");
+			OutFailureReason->Appendf(Format, *StartKey.ToString(), *EndKey.ToString(), ReversedChain.Num(), MinNumBones);
+		}
+		return false;
+	}
+	
+	if(ReversedChain.Num() > MaxNumBones && MaxNumBones > 2)
+	{
+		if(OutFailureReason)
+		{
+			static constexpr TCHAR Format[] = TEXT("Chain from '%s' to '%s' has %d elements (maximum of %d allowed).");
+			OutFailureReason->Appendf(Format, *StartKey.ToString(), *EndKey.ToString(), ReversedChain.Num(), MaxNumBones);
+		}
+		return false;
+	}
+
+	if(!bAllowBranches)
+	{
+		for(int32 Index = 1; Index < ReversedChain.Num() - 1; Index++)
+		{
+			NextKey = ReversedChain[Index];
+			if(InConnectionInfo->TargetHierarchy->GetChildren(NextKey).Num() > 1)
+			{
+				if(OutFailureReason)
+				{
+					static constexpr TCHAR Format[] = TEXT("Chain element '%s' contains a branch.");
+					OutFailureReason->Appendf(Format, *NextKey.ToString());
+				}
+				return false;
+			}
+		}
+	}
+
+	return true;
+}
+
+
diff --git a/Engine/Plugins/Animation/ControlRig/Source/ControlRig/Private/Rigs/RigHierarchy.cpp b/Engine/Plugins/Animation/ControlRig/Source/ControlRig/Private/Rigs/RigHierarchy.cpp
--- a/Engine/Plugins/Animation/ControlRig/Source/ControlRig/Private/Rigs/RigHierarchy.cpp
+++ b/Engine/Plugins/Animation/ControlRig/Source/ControlRig/Private/Rigs/RigHierarchy.cpp
@@ -124,6 +124,7 @@ URigHierarchy::URigHierarchy()
 #if WITH_EDITOR
 , bRecordTransformsAtRuntime(true)
 #endif
+, ElementRedirector(nullptr)
 {
 	Reset();
 #if WITH_EDITOR
@@ -1970,7 +1971,7 @@ bool URigHierarchy::CanSwitchToParent(FRigElementKey InChild, FRigElementKey InP
 	{
 		if(OutFailureReason)
 		{
-			*OutFailureReason = FString::Printf(TEXT("Child Element %s cannot be found."), *InChild.ToString());
+			OutFailureReason->Appendf(TEXT("Child Element %s cannot be found."), *InChild.ToString());
 		}
 		return false;
 	}
@@ -1987,7 +1988,7 @@ bool URigHierarchy::CanSwitchToParent(FRigElementKey InChild, FRigElementKey InP
 		
 		if(OutFailureReason)
 		{
-			*OutFailureReason = FString::Printf(TEXT("Parent Element %s cannot be found."), *InParent.ToString());
+			OutFailureReason->Appendf(TEXT("Parent Element %s cannot be found."), *InParent.ToString());
 		}
 		return false;
 	}
@@ -2003,7 +2004,7 @@ bool URigHierarchy::CanSwitchToParent(FRigElementKey InChild, FRigElementKey InP
 	{
 		if(OutFailureReason)
 		{
-			*OutFailureReason = FString::Printf(TEXT("Child Element %s does not allow space switching (it's not a multi parent element)."), *InChild.ToString());
+			OutFailureReason->Appendf(TEXT("Child Element %s does not allow space switching (it's not a multi parent element)."), *InChild.ToString());
 		}
 	}
 
@@ -2012,7 +2013,7 @@ bool URigHierarchy::CanSwitchToParent(FRigElementKey InChild, FRigElementKey InP
 	{
 		if(OutFailureReason)
 		{
-			*OutFailureReason = FString::Printf(TEXT("Parent Element %s is not a transform element"), *InParent.ToString());
+			OutFailureReason->Appendf(TEXT("Parent Element %s is not a transform element"), *InParent.ToString());
 		}
 	}
 
@@ -2020,7 +2021,7 @@ bool URigHierarchy::CanSwitchToParent(FRigElementKey InChild, FRigElementKey InP
 	{
 		if(OutFailureReason)
 		{
-			*OutFailureReason = FString::Printf(TEXT("Cannot switch '%s' to '%s' - would cause a cycle."), *InChild.ToString(), *InParent.ToString());
+			OutFailureReason->Appendf(TEXT("Cannot switch '%s' to '%s' - would cause a cycle."), *InChild.ToString(), *InParent.ToString());
 		}
 		return false;
 	}
@@ -2294,6 +2295,98 @@ void URigHierarchy::Traverse(TFunction<void(FRigBaseElement*, bool& /* continue
 	}
 }
 
+bool URigHierarchy::CanConnect(const FRigConnectionInfo* InConnectionInfo, FString* OutFailureReason) const
+{
+	check(InConnectionInfo);
+	check(InConnectionInfo->IsValid());
+
+	// make sure all connectors are provided
+	const TArray<FRigConnectorElement*> ExpectedConnectors = InConnectionInfo->SourceHierarchy->GetConnectors(false);
+	for(const FRigConnectorElement* ExpectedConnector : ExpectedConnectors)
+	{
+		if(!InConnectionInfo->ConnectionMap.Contains(ExpectedConnector->GetKey()))
+		{
+			if(OutFailureReason)
+			{
+				static constexpr TCHAR Format[] = TEXT("Connector key '%s' not provided as part of the connection map.");
+				OutFailureReason->Appendf(Format, *ExpectedConnector->GetKey().ToString());
+			}
+			return false;
+		}
+	}
+	
+	for(const TPair<FRigElementKey, FRigElementKey>& Pair : InConnectionInfo->ConnectionMap)
+	{
+		const FRigElementKey& ConnectorKey = Pair.Key;
+		const FRigElementKey& TargetKey = Pair.Value;
+
+		if(!ConnectorKey.IsValid())
+		{
+			if(OutFailureReason)
+			{
+				static constexpr TCHAR Format[] = TEXT("Connector key '%s' is not valid.");
+				OutFailureReason->Appendf(Format, *ConnectorKey.ToString());
+			}
+			return false;
+		}
+		if(ConnectorKey.Type != ERigElementType::Connector)
+		{
+			if(OutFailureReason)
+			{
+				static constexpr TCHAR Format[] = TEXT("Connector key '%s' is not a connector.");
+				OutFailureReason->Appendf(Format, *ConnectorKey.ToString());
+			}
+			return false;
+		}
+		if(!TargetKey.IsValid())
+		{
+			if(OutFailureReason)
+			{
+				static constexpr TCHAR Format[] = TEXT("Target key '%s' is not valid.");
+				OutFailureReason->Appendf(Format, *TargetKey.ToString());
+			}
+			return false;
+		}
+
+		const FRigConnectorElement* Connector = InConnectionInfo->SourceHierarchy->Find<FRigConnectorElement>(ConnectorKey); 
+		if(Connector == nullptr)
+		{
+			if(OutFailureReason)
+			{
+				static constexpr TCHAR Format[] = TEXT("Connector element '%s' does not exist.");
+				OutFailureReason->Appendf(Format, *ConnectorKey.ToString());
+			}
+			return false;
+		}
+		const FRigBaseElement* Target = Find<FRigBaseElement>(TargetKey);
+		if(Target == nullptr)
+		{
+			if(OutFailureReason)
+			{
+				static constexpr TCHAR Format[] = TEXT("Target element '%s' does not exist.");
+				OutFailureReason->Appendf(Format, *TargetKey.ToString());
+			}
+			return false;
+		}
+		const FRigTransformElement* TargetTransform = Cast<FRigTransformElement>(Target);
+		if(TargetTransform == nullptr)
+		{
+			if(OutFailureReason)
+			{
+				static constexpr TCHAR Format[] = TEXT("Target element '%s' is not a transform.");
+				OutFailureReason->Appendf(Format, *TargetKey.ToString());
+			}
+			return false;
+		}
+
+		if(!Connector->CanConnect(InConnectionInfo, OutFailureReason))
+		{
+			return false;
+		}
+	}
+	return true;
+}
+
 bool URigHierarchy::Undo()
 {
 #if WITH_EDITOR
diff --git a/Engine/Plugins/Animation/ControlRig/Source/ControlRig/Private/Rigs/RigHierarchyCache.cpp b/Engine/Plugins/Animation/ControlRig/Source/ControlRig/Private/Rigs/RigHierarchyCache.cpp
--- a/Engine/Plugins/Animation/ControlRig/Source/ControlRig/Private/Rigs/RigHierarchyCache.cpp
+++ b/Engine/Plugins/Animation/ControlRig/Source/ControlRig/Private/Rigs/RigHierarchyCache.cpp
@@ -67,3 +67,16 @@ bool FCachedRigElement::IsIdentical(const FRigElementKey& InKey, const URigHiera
 	return InKey == Key && InHierarchy->GetTopologyVersion() == ContainerVersion;
 }
 
+FRigElementKeyRedirector::FRigElementKeyRedirector(const TMap<FRigElementKey, FRigElementKey>& InMap, const URigHierarchy* InHierarchy)
+{
+	check(InHierarchy);
+	InternalKeyToExternalKey.Reserve(InMap.Num());
+	ExternalKeys.Reserve(InMap.Num());
+	for(const TPair<FRigElementKey, FRigElementKey>& Pair : InMap)
+	{
+		check(Pair.Key.IsValid());
+		InternalKeyToExternalKey.Add(Pair.Key, FCachedRigElement(Pair.Value, InHierarchy));
+		ExternalKeys.Add(Pair.Value);
+	}
+}
+
diff --git a/Engine/Plugins/Animation/ControlRig/Source/ControlRig/Private/Rigs/RigHierarchyElements.cpp b/Engine/Plugins/Animation/ControlRig/Source/ControlRig/Private/Rigs/RigHierarchyElements.cpp
--- a/Engine/Plugins/Animation/ControlRig/Source/ControlRig/Private/Rigs/RigHierarchyElements.cpp
+++ b/Engine/Plugins/Animation/ControlRig/Source/ControlRig/Private/Rigs/RigHierarchyElements.cpp
@@ -1461,31 +1461,68 @@ void FRigConnectorSettings::Save(FArchive& Ar)
 {
 	Ar.UsingCustomVersion(FControlRigObjectVersion::GUID);
 
-	Ar << ResolvedItem;
+	Ar << Description;
+
+	int32 NumRules = Rules.Num();
+	Ar << NumRules;
+	for(int32 Index = 0; Index < NumRules; Index++)
+	{
+
+		Rules[Index].Save(Ar);
+	}
 }
 
 void FRigConnectorSettings::Load(FArchive& Ar)
 {
 	Ar.UsingCustomVersion(FControlRigObjectVersion::GUID);
 
-	Ar << ResolvedItem;
-}
+	Ar << Description;
 
-uint32 GetTypeHash(const FRigConnectorSettings& Settings)
-{
-	uint32 Hash = GetTypeHash(Settings.ResolvedItem);
-	return Hash;
+	int32 NumRules = 0;
+	Ar << NumRules;
+	Rules.SetNumZeroed(NumRules);
+	for(int32 Index = 0; Index < NumRules; Index++)
+	{
+
+		Rules[Index].Load(Ar);
+	}
 }
 
 bool FRigConnectorSettings::operator==(const FRigConnectorSettings& InOther) const
 {
-	if(ResolvedItem != InOther.ResolvedItem)
+	if(!Description.Equals(InOther.Description, ESearchCase::CaseSensitive))
+	{
+		return false;
+	}
+	if(Rules.Num() != InOther.Rules.Num())
 	{
 		return false;
 	}
+	for(int32 Index = 0; Index < Rules.Num(); Index++)
+	{
+		if(Rules[Index] != InOther.Rules[Index])
+		{
+			return false;
+		}
+	}
 	return true;
 }
 
+uint32 FRigConnectorSettings::GetRulesHash() const
+{
+	uint32 Hash = GetTypeHash(Rules.Num());
+	for(const FRigConnectionRuleStash& Rule : Rules)
+	{
+		Hash = HashCombine(Hash, GetTypeHash(Rule));
+	}
+	return Hash;
+}
+
+uint32 GetTypeHash(const FRigConnectorSettings& Settings)
+{
+	return Settings.GetRulesHash();
+}
+
 ////////////////////////////////////////////////////////////////////////////////
 // FRigConnectorElement
 ////////////////////////////////////////////////////////////////////////////////
@@ -1510,6 +1547,92 @@ void FRigConnectorElement::Load(FArchive& Ar, URigHierarchy* Hierarchy, ESeriali
 	}
 }
 
+bool FRigConnectorElement::CanConnect(const FRigConnectionInfo* InConnectionInfo, FString* OutFailureReason) const
+{
+	check(InConnectionInfo);
+	check(InConnectionInfo->IsValid());
+
+	if(!InConnectionInfo->ConnectionMap.Contains(GetKey()))
+	{
+		if(OutFailureReason)
+		{
+			static constexpr TCHAR Format[] = TEXT("Connector '%s' is not mapped in ConnectionMap.");
+			OutFailureReason->Appendf(Format, *GetKey().ToString());
+		}
+		return false;
+	}
+
+	for(const FRigConnectionRuleStash& Stash : Settings.Rules)
+	{
+		TSharedPtr<FStructOnScope> RuleScope;
+		const FRigConnectionRule* Rule = Stash.Get(RuleScope);
+		if(Rule == nullptr)
+		{
+			if(OutFailureReason)
+			{
+				static constexpr TCHAR Format[] = TEXT("Rule '%s' could not be loaded from '%s'");
+				OutFailureReason->Appendf(Format, *Stash.ScriptStructPath, *Stash.ExportedText);
+			}
+			return false;
+		}
+
+		FRigConnectionInfo InfoPerRule = *InConnectionInfo;
+		InfoPerRule.ConnectionMap.Reset();
+		InfoPerRule.ConnectionMap.Add(GetKey(), InConnectionInfo->ConnectionMap.FindChecked(GetKey()));
+
+		const TArray<FRigElementKey> AdditionalConnectorKeys = Rule->GetAdditionalConnectors();
+		for(const FRigElementKey& AdditionalConnectorKey : AdditionalConnectorKeys)
+		{
+			if(AdditionalConnectorKey.Type != ERigElementType::Connector)
+			{
+				if(OutFailureReason)
+				{
+					static constexpr TCHAR Format[] = TEXT("Additionally required connector '%s' is not a connector.");
+					OutFailureReason->Appendf(Format, *AdditionalConnectorKey.ToString());
+				}
+				return false;
+			}
+			if(!InConnectionInfo->SourceHierarchy->Contains(AdditionalConnectorKey))
+			{
+				if(OutFailureReason)
+				{
+					static constexpr TCHAR Format[] = TEXT("Additionally required connector '%s' does not exist.");
+					OutFailureReason->Appendf(Format, *AdditionalConnectorKey.ToString());
+				}
+				return false;
+			}
+			if(!InConnectionInfo->ConnectionMap.Contains(AdditionalConnectorKey))
+			{
+				if(OutFailureReason)
+				{
+					static constexpr TCHAR Format[] = TEXT("Additionally required connector '%s' not provided as part of ConnectionMap.");
+					OutFailureReason->Appendf(Format, *AdditionalConnectorKey.ToString());
+				}
+				return false;
+			}
+			const FRigElementKey& AdditionalTargetKey = InConnectionInfo->ConnectionMap.FindChecked(AdditionalConnectorKey);
+			if(!InConnectionInfo->TargetHierarchy->Contains(AdditionalTargetKey))
+			{
+				if(OutFailureReason)
+				{
+					static constexpr TCHAR Format[] = TEXT("Target element '%s' does not exist.");
+					OutFailureReason->Appendf(Format, *AdditionalConnectorKey.ToString());
+				}
+				return false;
+			}
+
+			InfoPerRule.ConnectionMap.Add(AdditionalConnectorKey, AdditionalTargetKey);
+		}
+
+		if(!Rule->CanConnect(&InfoPerRule, OutFailureReason))
+		{
+			return false;
+		}
+	}
+
+	return true;
+}
+
 void FRigConnectorElement::CopyFrom(URigHierarchy* InHierarchy, FRigBaseElement* InOther, URigHierarchy* InOtherHierarchy)
 {
 	Super::CopyFrom(InHierarchy, InOther, InOtherHierarchy);
diff --git a/Engine/Plugins/Animation/ControlRig/Source/ControlRig/Public/Rigs/RigConnectionRules.h b/Engine/Plugins/Animation/ControlRig/Source/ControlRig/Public/Rigs/RigConnectionRules.h
new file mode 100644
--- /dev/null
+++ b/Engine/Plugins/Animation/ControlRig/Source/ControlRig/Public/Rigs/RigConnectionRules.h
@@ -0,0 +1,215 @@
+// Copyright Epic Games, Inc. All Rights Reserved.
+
+#pragma once
+
+#include "CoreMinimal.h"
+#include "RigHierarchyDefines.h"
+#include "UObject/StructOnScope.h"
+#include "RigConnectionRules.generated.h"
+
+struct FRigConnectionRule;
+struct FRigConnectionInfo;
+
+USTRUCT(BlueprintType)
+struct CONTROLRIG_API FRigConnectionRuleStash
+{
+	GENERATED_BODY()
+
+	FRigConnectionRuleStash();
+	FRigConnectionRuleStash(const FRigConnectionRule* InRule);
+
+	// todo: copy content from FRigVMActionKey
+
+	void Save(FArchive& Ar);
+	void Load(FArchive& Ar);
+	
+	friend uint32 GetTypeHash(const FRigConnectionRuleStash& InRuleStash);
+
+	bool IsValid() const;
+	UScriptStruct* GetScriptStruct() const;
+	TSharedPtr<FStructOnScope> Get() const;
+	const FRigConnectionRule* Get(TSharedPtr<FStructOnScope>& InOutStorage) const;
+
+	bool operator == (const FRigConnectionRuleStash& InOther) const;
+
+	bool operator != (const FRigConnectionRuleStash& InOther) const
+	{
+		return !(*this == InOther);
+	}
+
+	void AppendAdditionalConnectors(TArray<FRigElementKey>& OutConnectors) const;
+
+	UPROPERTY(BlueprintReadWrite, EditAnywhere, Category=Rule)
+	FString ScriptStructPath;
+
+	UPROPERTY(BlueprintReadWrite, EditAnywhere, Category=Rule)
+	FString ExportedText;
+};
+
+USTRUCT(meta=(Hidden))
+struct CONTROLRIG_API FRigConnectionRule
+{
+	GENERATED_BODY()
+
+public:
+
+	FRigConnectionRule() {}
+	virtual ~FRigConnectionRule() {}
+
+	virtual UScriptStruct* GetScriptStruct() const { return FRigConnectionRule::StaticStruct(); }
+
+	/** If the rule requires additional connectors for validation return them here */
+	virtual TArray<FRigElementKey> GetAdditionalConnectors() const;
+
+	/**
+	 * Validate if the rule permits the link.
+	 * Provided are the connectors keys and their hierarchy,
+	 * as well as the desired target keys and their hierarchy.
+	 */
+	virtual bool CanConnect(const FRigConnectionInfo* InConnectionInfo, FString* OutFailureReason) const;
+};
+
+USTRUCT(BlueprintType, DisplayName="And Rule")
+struct CONTROLRIG_API FRigAndConnectionRule : public FRigConnectionRule
+{
+	GENERATED_BODY()
+
+public:
+
+	FRigAndConnectionRule()
+	{}
+
+	template<typename TypeA, typename TypeB>
+	FRigAndConnectionRule(const TypeA& InA, const TypeB& InB)
+	{
+		ChildRules.Emplace(&InA);
+		ChildRules.Emplace(&InB);
+	}
+
+	virtual ~FRigAndConnectionRule() override {}
+
+	virtual UScriptStruct* GetScriptStruct() const override { return FRigAndConnectionRule::StaticStruct(); }
+	virtual TArray<FRigElementKey> GetAdditionalConnectors() const override;
+
+	virtual bool CanConnect(const FRigConnectionInfo* InConnectionInfo, FString* OutFailureReason) const override;
+
+	UPROPERTY(BlueprintReadWrite, EditAnywhere, Category=Settings)
+	TArray<FRigConnectionRuleStash> ChildRules;
+};
+
+USTRUCT(BlueprintType, DisplayName="Or Rule")
+struct CONTROLRIG_API FRigOrConnectionRule : public FRigConnectionRule
+{
+	GENERATED_BODY()
+
+public:
+
+	FRigOrConnectionRule()
+	{}
+
+	template<typename TypeA, typename TypeB>
+	FRigOrConnectionRule(const TypeA& InA, const TypeB& InB)
+	{
+		ChildRules.Emplace(&InA);
+		ChildRules.Emplace(&InB);
+	}
+
+	virtual ~FRigOrConnectionRule() override {}
+
+	virtual UScriptStruct* GetScriptStruct() const override { return FRigOrConnectionRule::StaticStruct(); }
+	virtual TArray<FRigElementKey> GetAdditionalConnectors() const override;
+
+	virtual bool CanConnect(const FRigConnectionInfo* InConnectionInfo, FString* OutFailureReason) const override;
+
+	UPROPERTY(BlueprintReadWrite, EditAnywhere, Category=Settings)
+	TArray<FRigConnectionRuleStash> ChildRules;
+};
+
+USTRUCT(BlueprintType, DisplayName="Type Rule")
+struct CONTROLRIG_API FRigTypeConnectionRule : public FRigConnectionRule
+{
+	GENERATED_BODY()
+
+public:
+
+	FRigTypeConnectionRule()
+		: ElementType(ERigElementType::Connector)
+	{}
+
+	FRigTypeConnectionRule(ERigElementType InElementType)
+	: ElementType(InElementType)
+	{}
+
+	virtual ~FRigTypeConnectionRule() override {}
+
+	virtual UScriptStruct* GetScriptStruct() const override { return FRigTypeConnectionRule::StaticStruct(); }
+
+	virtual bool CanConnect(const FRigConnectionInfo* InConnectionInfo, FString* OutFailureReason) const override;
+
+	UPROPERTY(BlueprintReadWrite, EditAnywhere, Category=Settings)
+	ERigElementType ElementType;
+};
+
+USTRUCT(BlueprintType, DisplayName="Tag Rule")
+struct CONTROLRIG_API FRigTagConnectionRule : public FRigConnectionRule
+{
+	GENERATED_BODY()
+
+public:
+
+	FRigTagConnectionRule()
+		: Tag(NAME_None)
+	{}
+
+	FRigTagConnectionRule(const FName& InTag)
+	: Tag(InTag)
+	{}
+
+	virtual ~FRigTagConnectionRule() override {}
+
+	virtual UScriptStruct* GetScriptStruct() const override { return FRigTagConnectionRule::StaticStruct(); }
+
+	virtual bool CanConnect(const FRigConnectionInfo* InConnectionInfo, FString* OutFailureReason) const override;
+
+	UPROPERTY(BlueprintReadWrite, EditAnywhere, Category=Settings)
+	FName Tag;
+};
+
+USTRUCT(BlueprintType, DisplayName="Chain Rule")
+struct CONTROLRIG_API FRigChainConnectionRule : public FRigConnectionRule
+{
+	GENERATED_BODY()
+
+public:
+
+	FRigChainConnectionRule()
+	: MinNumBones(2)
+	, MaxNumBones(0)
+	, bAllowBranches(false)
+	{}
+
+	FRigChainConnectionRule(const FRigElementKey& InEndConnector, int32 InMinNumBones = 2, int32 InMaxNumBones = 0, bool InAllowBranches = false)
+	: EndConnector(InEndConnector)
+	, MinNumBones(InMinNumBones)
+	, MaxNumBones(InMaxNumBones)
+	, bAllowBranches(InAllowBranches)
+	{}
+
+	virtual ~FRigChainConnectionRule() override {}
+
+	virtual UScriptStruct* GetScriptStruct() const override { return FRigChainConnectionRule::StaticStruct(); }
+
+	virtual bool CanConnect(const FRigConnectionInfo* InConnectionInfo, FString* OutFailureReason) const override;
+
+	UPROPERTY(BlueprintReadWrite, EditAnywhere, Category=Settings)
+	FRigElementKey EndConnector;
+
+	UPROPERTY(BlueprintReadWrite, EditAnywhere, Category=Settings)
+	int32 MinNumBones;
+
+	UPROPERTY(BlueprintReadWrite, EditAnywhere, Category=Settings)
+	int32 MaxNumBones;
+
+	UPROPERTY(BlueprintReadWrite, EditAnywhere, Category=Settings)
+	bool bAllowBranches;
+};
\ No newline at end of file
diff --git a/Engine/Plugins/Animation/ControlRig/Source/ControlRig/Public/Rigs/RigHierarchy.h b/Engine/Plugins/Animation/ControlRig/Source/ControlRig/Public/Rigs/RigHierarchy.h
--- a/Engine/Plugins/Animation/ControlRig/Source/ControlRig/Public/Rigs/RigHierarchy.h
+++ b/Engine/Plugins/Animation/ControlRig/Source/ControlRig/Public/Rigs/RigHierarchy.h
@@ -6,6 +6,7 @@
 #include "RigVMCore/RigVMExecuteContext.h"
 #include "RigVMCore/RigVM.h"
 #include "RigHierarchyElements.h"
+#include "RigHierarchyCache.h"
 #include "RigHierarchyPose.h"
 #include "UObject/WeakObjectPtrTemplates.h"
 #include "EdGraph/EdGraphPin.h"
@@ -336,6 +337,22 @@ public:
 	 */
 	int32 GetIndex(const FRigElementKey& InKey) const
 	{
+		if(ElementRedirector)
+		{
+			if(FCachedRigElement* CachedRigElement = ElementRedirector->Find(InKey))
+			{
+				if(CachedRigElement->UpdateCache(this))
+				{
+					return CachedRigElement->GetIndex();
+				}
+				return INDEX_NONE;
+			}
+			if(ElementRedirector->ContainsExternalKey(InKey))
+			{
+				return INDEX_NONE;
+			}
+		}
+		
 		if(const int32* Index = IndexLookup.Find(InKey))
 		{
 			return *Index;
@@ -872,6 +889,24 @@ public:
 		return GetKeysOfType<FRigReferenceElement>(bTraverse);
 	}
 
+	/**
+	 * Returns all Connector elements
+	 * @param bTraverse Returns the elements in order of a depth first traversal
+	 */
+	TArray<FRigConnectorElement*> GetConnectors(bool bTraverse = false) const
+	{
+		return GetElementsOfType<FRigConnectorElement>(bTraverse);
+	}
+
+	/**
+	 * Returns all Connector elements
+	 * @param bTraverse Returns the elements in order of a depth first traversal
+	 */
+	UFUNCTION(BlueprintCallable, Category = URigHierarchy, meta = (DisplayName = "Get Connectors", ScriptName = "GetConnectors"))
+	TArray<FRigElementKey> GetConnectorKeys(bool bTraverse = true) const
+	{
+		return GetKeysOfType<FRigConnectorElement>(bTraverse);
+	}
 	/**
 	 * Returns all root elements
 	 */
@@ -3174,6 +3209,17 @@ public:
 	 */
 	void Traverse(TFunction<void(FRigBaseElement*, bool& /* continue */)> PerElementFunction, bool bTowardsChildren = true) const;
 
+	/**
+	 * Returns true if a hierarchy can be linked into another hierarchy. This relies on Connector elements
+	 * as the relevant landmarks providing rules for validation.
+	 */
+	UFUNCTION(BlueprintCallable, Category = URigHierarchy, meta = (DisplayName = "Can Link", ScriptName = "CanConnect"))
+	bool CanConnect_ForBlueprint(FRigConnectionInfo InConnectionInfo) const
+	{
+		return CanConnect(&InConnectionInfo);
+	}
+	bool CanConnect(const FRigConnectionInfo* InConnectionInfo, FString* OutFailureReason = nullptr) const;
+
 	/**
 	 * Performs undo for one transform change
 	 */
@@ -4446,6 +4492,8 @@ private:
 		}
 	}
 
+	FRigElementKeyRedirector* ElementRedirector;
+
 	void UpdateVisibilityOnProxyControls();
 
 	static const TArray<FString>& GetTransformTypeStrings();
@@ -4474,6 +4522,7 @@ private:
 	friend struct FControlRigVisualGraphUtils;
 	friend struct FRigHierarchyEnableControllerBracket;
 	friend struct FRigHierarchyExecuteContextBracket;
+	friend struct FRigHierarchyRedirectorGuard;
 };
 
 struct CONTROLRIG_API FRigHierarchyInteractionBracket
@@ -4591,6 +4640,18 @@ private:
 	bool bPreviousValue;
 };
 
+struct CONTROLRIG_API FRigHierarchyRedirectorGuard
+{
+public:
+	FRigHierarchyRedirectorGuard(URigHierarchy* InHierarchy, FRigElementKeyRedirector& InRedirector)
+		: Guard(InHierarchy->ElementRedirector, &InRedirector)
+	{
+	}
+
+private:
+	TGuardValue<FRigElementKeyRedirector*> Guard;
+};
+
 template<>
 inline FVector2D URigHierarchy::GetControlValue(FRigControlElement* InControlElement, ERigControlValueType InValueType) const
 {
diff --git a/Engine/Plugins/Animation/ControlRig/Source/ControlRig/Public/Rigs/RigHierarchyCache.h b/Engine/Plugins/Animation/ControlRig/Source/ControlRig/Public/Rigs/RigHierarchyCache.h
--- a/Engine/Plugins/Animation/ControlRig/Source/ControlRig/Public/Rigs/RigHierarchyCache.h
+++ b/Engine/Plugins/Animation/ControlRig/Source/ControlRig/Public/Rigs/RigHierarchyCache.h
@@ -150,3 +150,24 @@ private:
 
 	const FRigBaseElement* Element;
 };
+
+class CONTROLRIG_API FRigElementKeyRedirector
+{
+public:
+
+	FRigElementKeyRedirector()
+		: InternalKeyToExternalKey()
+	{}
+
+	FRigElementKeyRedirector(const TMap<FRigElementKey, FRigElementKey>& InMap, const URigHierarchy* InHierarchy);
+	
+	bool Contains(const FRigElementKey& InKey) const { return InternalKeyToExternalKey.Contains(InKey); }
+	bool ContainsExternalKey(const FRigElementKey& InKey) const { return ExternalKeys.Contains(InKey); }
+	const FCachedRigElement* Find(const FRigElementKey& InKey) const { return InternalKeyToExternalKey.Find(InKey); }
+	FCachedRigElement* Find(const FRigElementKey& InKey) { return InternalKeyToExternalKey.Find(InKey); }
+
+private:
+
+	TMap<FRigElementKey, FCachedRigElement> InternalKeyToExternalKey;
+	TSet<FRigElementKey> ExternalKeys;
+};
\ No newline at end of file
diff --git a/Engine/Plugins/Animation/ControlRig/Source/ControlRig/Public/Rigs/RigHierarchyDefines.h b/Engine/Plugins/Animation/ControlRig/Source/ControlRig/Public/Rigs/RigHierarchyDefines.h
--- a/Engine/Plugins/Animation/ControlRig/Source/ControlRig/Public/Rigs/RigHierarchyDefines.h
+++ b/Engine/Plugins/Animation/ControlRig/Source/ControlRig/Public/Rigs/RigHierarchyDefines.h
@@ -1534,7 +1534,7 @@ public:
 			}
 			case ERigElementType::Null:
 			{
-				return FString::Printf(TEXT("Space(%s)"), *Name.ToString());
+				return FString::Printf(TEXT("Null(%s)"), *Name.ToString());
 			}
 			case ERigElementType::Control:
 			{
@@ -1552,7 +1552,12 @@ public:
 			{
 				return FString::Printf(TEXT("Reference(%s)"), *Name.ToString());
 			}
+			case ERigElementType::Connector:
+			{
+				return FString::Printf(TEXT("Connector(%s)"), *Name.ToString());
+			}
 		}
+		
 		return FString();
 	}
 
diff --git a/Engine/Plugins/Animation/ControlRig/Source/ControlRig/Public/Rigs/RigHierarchyElements.h b/Engine/Plugins/Animation/ControlRig/Source/ControlRig/Public/Rigs/RigHierarchyElements.h
--- a/Engine/Plugins/Animation/ControlRig/Source/ControlRig/Public/Rigs/RigHierarchyElements.h
+++ b/Engine/Plugins/Animation/ControlRig/Source/ControlRig/Public/Rigs/RigHierarchyElements.h
@@ -5,6 +5,7 @@
 #include "CoreMinimal.h"
 #include "RigHierarchyDefines.h"
 #include "RigHierarchyMetadata.h"
+#include "RigConnectionRules.h"
 #include "RigHierarchyElements.generated.h"
 
 struct FRigVMExecuteContext;
@@ -1558,6 +1559,43 @@ protected:
 	friend class URigHierarchyController;
 };
 
+USTRUCT(BlueprintType)
+struct CONTROLRIG_API FRigConnectionInfo
+{
+	GENERATED_BODY()
+
+public:
+	
+	FRigConnectionInfo()
+	: SourceHierarchy(nullptr)
+	, TargetHierarchy(nullptr)
+	{}
+
+	FRigConnectionInfo(const TMap<FRigElementKey, FRigElementKey>& InMap, URigHierarchy* InSourceHierarchy, URigHierarchy* InTargetHierarchy)
+	: ConnectionMap(InMap)
+	, SourceHierarchy(InSourceHierarchy)
+	, TargetHierarchy(InTargetHierarchy)
+	{}
+
+	bool IsValid() const
+	{
+		return (!ConnectionMap.IsEmpty()) && (SourceHierarchy != nullptr) && (TargetHierarchy != nullptr);
+	}
+
+	// The keys of the connectors in the source hierarchy
+	// mapping to the to-be-linked keys of the elements in the target hierarchy
+	UPROPERTY(BlueprintReadWrite, EditAnywhere, Category = Limit)
+	TMap<FRigElementKey, FRigElementKey> ConnectionMap;
+
+	// The hierarchy owning the connectors
+	UPROPERTY(BlueprintReadWrite, EditAnywhere, Category = Limit)
+	URigHierarchy* SourceHierarchy; 
+
+	// The hierarchy to be linked into
+	UPROPERTY(BlueprintReadWrite, EditAnywhere, Category = Limit)
+	URigHierarchy* TargetHierarchy; 
+};
+
 USTRUCT(BlueprintType)
 struct CONTROLRIG_API FRigConnectorSettings
 {
@@ -1571,7 +1609,10 @@ struct CONTROLRIG_API FRigConnectorSettings
 	friend uint32 GetTypeHash(const FRigConnectorSettings& Settings);
 
 	UPROPERTY(BlueprintReadWrite, EditAnywhere, Category = Control)
-	FRigElementKey ResolvedItem;
+	FString Description;
+
+	UPROPERTY(BlueprintReadWrite, EditAnywhere, Category = Control)
+	TArray<FRigConnectionRuleStash> Rules;
 
 	bool operator == (const FRigConnectorSettings& InOther) const;
 
@@ -1579,6 +1620,14 @@ struct CONTROLRIG_API FRigConnectorSettings
 	{
 		return !(*this == InOther);
 	}
+
+	template<typename T>
+	int32 AddRule(const T& InRule)
+	{
+		return Rules.Emplace(&InRule);
+	}
+
+	uint32 GetRulesHash() const;
 };
 
 USTRUCT(BlueprintType)
@@ -1600,6 +1649,8 @@ public:
 	virtual void Save(FArchive& A, URigHierarchy* Hierarchy, ESerializationPhase SerializationPhase) override;
 	virtual void Load(FArchive& Ar, URigHierarchy* Hierarchy, ESerializationPhase SerializationPhase) override;
 
+	bool CanConnect(const FRigConnectionInfo* InConnectionInfo, FString* OutFailureReason) const;
+	
 private:
 
 	virtual void CopyFrom(URigHierarchy* InHierarchy, FRigBaseElement* InOther, URigHierarchy* InOtherHierarchy) override;
diff --git a/Engine/Plugins/Animation/ControlRig/Source/ControlRigEditor/Private/ControlRigElementDetails.cpp b/Engine/Plugins/Animation/ControlRig/Source/ControlRigEditor/Private/ControlRigElementDetails.cpp
--- a/Engine/Plugins/Animation/ControlRig/Source/ControlRigEditor/Private/ControlRigElementDetails.cpp
+++ b/Engine/Plugins/Animation/ControlRig/Source/ControlRigEditor/Private/ControlRigElementDetails.cpp
@@ -5064,8 +5064,6 @@ void FRigConnectorElementDetails::CustomizeSettings(IDetailLayoutBuilder& Detail
 	DetailBuilder.HideProperty(SettingsHandle);
 
 	IDetailCategoryBuilder& SettingsCategory = DetailBuilder.EditCategory(TEXT("Settings"), LOCTEXT("Settings", "Settings"));
-
-	SettingsCategory.AddProperty(SettingsHandle->GetChildHandle(GET_MEMBER_NAME_CHECKED(FRigConnectorSettings, ResolvedItem))).IsEnabled(false);
 }
 
 #undef LOCTEXT_NAMESPACE
