Control Rig: Fix icons for event queue dropdown / event button

diff --git a/Engine/Plugins/Animation/ControlRig/Source/ControlRigEditor/Private/Editor/ControlRigEditor.cpp b/Engine/Plugins/Animation/ControlRig/Source/ControlRigEditor/Private/Editor/ControlRigEditor.cpp
--- a/Engine/Plugins/Animation/ControlRig/Source/ControlRigEditor/Private/Editor/ControlRigEditor.cpp
+++ b/Engine/Plugins/Animation/ControlRig/Source/ControlRigEditor/Private/Editor/ControlRigEditor.cpp
@@ -774,19 +774,19 @@ FSlateIcon FControlRigEditor::GetEventQueueIcon(const TArray<FName>& InEventQueu
 {
 	if(InEventQueue == ConstructionEventQueue)
 	{
-		return FSlateIcon(FRigVMEditorStyle::Get().GetStyleSetName(), "ControlRig.ConstructionMode");
+		return FSlateIcon(FControlRigEditorStyle::Get().GetStyleSetName(), "ControlRig.ConstructionMode");
 	}
 	if(InEventQueue == ForwardsSolveEventQueue)
 	{
-		return FSlateIcon(FRigVMEditorStyle::Get().GetStyleSetName(), "ControlRig.ForwardsSolveEvent");
+		return FSlateIcon(FControlRigEditorStyle::Get().GetStyleSetName(), "ControlRig.ForwardsSolveEvent");
 	}
 	if(InEventQueue == BackwardsSolveEventQueue)
 	{
-		return FSlateIcon(FRigVMEditorStyle::Get().GetStyleSetName(), "ControlRig.BackwardsSolveEvent");
+		return FSlateIcon(FControlRigEditorStyle::Get().GetStyleSetName(), "ControlRig.BackwardsSolveEvent");
 	}
 	if(InEventQueue == BackwardsAndForwardsSolveEventQueue)
 	{
-		return FSlateIcon(FRigVMEditorStyle::Get().GetStyleSetName(), "ControlRig.BackwardsAndForwardsSolveEvent");
+		return FSlateIcon(FControlRigEditorStyle::Get().GetStyleSetName(), "ControlRig.BackwardsAndForwardsSolveEvent");
 	}
 
 	return FSlateIcon(FAppStyle::GetAppStyleSetName(), "GraphEditor.Event_16x");
