Control Rig Modules: Details panel rule editing

diff --git a/Engine/Plugins/Animation/ControlRig/Source/ControlRigEditor/Private/ControlRigEditorModule.cpp b/Engine/Plugins/Animation/ControlRig/Source/ControlRigEditor/Private/ControlRigEditorModule.cpp
--- a/Engine/Plugins/Animation/ControlRig/Source/ControlRigEditor/Private/ControlRigEditorModule.cpp
+++ b/Engine/Plugins/Animation/ControlRig/Source/ControlRigEditor/Private/ControlRigEditorModule.cpp
@@ -150,6 +150,9 @@ void FControlRigEditorModule::StartupModule()
 	PropertiesToUnregisterOnShutdown.Add(StaticEnum<ERigControlTransformChannel>()->GetFName());
 	PropertyEditorModule.RegisterCustomPropertyTypeLayout(PropertiesToUnregisterOnShutdown.Last(), FOnGetPropertyTypeCustomizationInstance::CreateStatic(&FRigControlTransformChannelDetails::MakeInstance));
 
+	PropertiesToUnregisterOnShutdown.Add(FRigConnectionRuleStash::StaticStruct()->GetFName());
+	PropertyEditorModule.RegisterCustomPropertyTypeLayout(PropertiesToUnregisterOnShutdown.Last(), FOnGetPropertyTypeCustomizationInstance::CreateStatic(&FRigConnectionRuleDetails::MakeInstance));
+
 	FRigBaseElementDetails::RegisterSectionMappings(PropertyEditorModule);
 
 	// Register asset tools
diff --git a/Engine/Plugins/Animation/ControlRig/Source/ControlRigEditor/Private/ControlRigEditorStyle.h b/Engine/Plugins/Animation/ControlRig/Source/ControlRigEditor/Private/ControlRigEditorStyle.h
--- a/Engine/Plugins/Animation/ControlRig/Source/ControlRigEditor/Private/ControlRigEditorStyle.h
+++ b/Engine/Plugins/Animation/ControlRig/Source/ControlRigEditor/Private/ControlRigEditorStyle.h
@@ -118,7 +118,10 @@ public:
 			Set("ControlRig.Tree.Null", new IMAGE_BRUSH("Slate/Null_16x", Icon16x16));
 			Set("ControlRig.Tree.RigidBody", new IMAGE_BRUSH("Slate/RigidBody_16x", Icon16x16));
 			Set("ControlRig.Tree.Socket", new IMAGE_BRUSH("Slate/Socket_16x", Icon16x16));
-			Set("ControlRig.Tree.Connector", new IMAGE_BRUSH("Slate/Socket_16x", Icon16x16));
+			{
+				FContentRootBracket Bracket(this, EngineEditorSlateDir);
+				Set("ControlRig.Tree.Connector", new IMAGE_BRUSH_SVG("Starship/Common/SetShowSockets", Icon16x16));
+			}
 		}
 
 		// Font?
diff --git a/Engine/Plugins/Animation/ControlRig/Source/ControlRigEditor/Private/ControlRigElementDetails.cpp b/Engine/Plugins/Animation/ControlRig/Source/ControlRigEditor/Private/ControlRigElementDetails.cpp
--- a/Engine/Plugins/Animation/ControlRig/Source/ControlRigEditor/Private/ControlRigElementDetails.cpp
+++ b/Engine/Plugins/Animation/ControlRig/Source/ControlRigEditor/Private/ControlRigElementDetails.cpp
@@ -12,6 +12,7 @@
 #include "ControlRigBlueprint.h"
 #include "Graph/ControlRigGraph.h"
 #include "PropertyCustomizationHelpers.h"
+#include "PropertyEditorModule.h"
 #include "SEnumCombo.h"
 #include "Units/Execution/RigUnit_BeginExecution.h"
 #include "Units/Execution/RigUnit_DynamicHierarchy.h"
@@ -19,6 +20,8 @@
 #include "HAL/PlatformApplicationMisc.h"
 #include "Styling/AppStyle.h"
 #include "Editor/SRigHierarchyTreeView.h"
+#include "StructViewerFilter.h"
+#include "StructViewerModule.h"
 
 #define LOCTEXT_NAMESPACE "ControlRigElementDetails"
 
@@ -5064,6 +5067,229 @@ void FRigConnectorElementDetails::CustomizeSettings(IDetailLayoutBuilder& Detail
 	DetailBuilder.HideProperty(SettingsHandle);
 
 	IDetailCategoryBuilder& SettingsCategory = DetailBuilder.EditCategory(TEXT("Settings"), LOCTEXT("Settings", "Settings"));
+
+	bool bHideRules = false;
+	uint32 FirstHash = UINT32_MAX;
+	for (const FPerElementInfo& Info : PerElementInfos)
+	{
+		if (URigHierarchy* Hierarchy = Info.IsValid() ? Info.GetHierarchy() : nullptr)
+		{
+			if(const FRigConnectorElement* Connector = Info.GetElement<FRigConnectorElement>())
+			{
+				const uint32 Hash = Connector->Settings.GetRulesHash();
+				if(FirstHash == UINT32_MAX)
+				{
+					FirstHash = Hash;
+				}
+				else if(FirstHash != Hash)
+				{
+					bHideRules = true;
+					break;
+				}
+			}
+			else
+			{
+				bHideRules = true;
+			}
+		}
+	}
+
+	if(!bHideRules)
+	{
+		SettingsCategory
+			.AddProperty(SettingsHandle->GetChildHandle(GET_MEMBER_NAME_CHECKED(FRigConnectorSettings, Rules)))
+			.IsEnabled(!IsAnyElementProcedural());
+	}
+}
+
+
+void FRigConnectionRuleDetails::CustomizeHeader(TSharedRef<IPropertyHandle> InStructPropertyHandle, FDetailWidgetRow& HeaderRow, IPropertyTypeCustomizationUtils& StructCustomizationUtils)
+{
+	StructPropertyHandle = InStructPropertyHandle.ToSharedPtr();
+	PropertyUtilities = StructCustomizationUtils.GetPropertyUtilities();
+	BlueprintBeingCustomized = nullptr;
+	RigElementKeyDetails_GetCustomizedInfo(InStructPropertyHandle, BlueprintBeingCustomized);
+
+	TArray<UObject*> Objects;
+	StructPropertyHandle->GetOuterObjects(Objects);
+	FString FirstObjectValue;
+	for (int32 Index = 0; Index < Objects.Num(); Index++)
+	{
+		FString ObjectValue;
+		if(InStructPropertyHandle->GetPerObjectValue(Index, ObjectValue) == FPropertyAccess::Result::Success)
+		{
+			if(FirstObjectValue.IsEmpty())
+			{
+				FirstObjectValue = ObjectValue;
+			}
+			else
+			{
+				if(!FirstObjectValue.Equals(ObjectValue, ESearchCase::CaseSensitive))
+				{
+					FirstObjectValue.Reset();
+					break;
+				}
+			}
+		}
+	}
+
+	if(!FirstObjectValue.IsEmpty())
+	{
+		FRigConnectionRuleStash::StaticStruct()->ImportText(*FirstObjectValue, &RuleStash, nullptr, EPropertyPortFlags::PPF_None, nullptr, FRigConnectionRuleStash::StaticStruct()->GetName(), true);
+	}
+
+	if (BlueprintBeingCustomized == nullptr || FirstObjectValue.IsEmpty())
+	{
+		HeaderRow
+		.NameContent()
+		[
+			StructPropertyHandle->CreatePropertyNameWidget()
+		]
+		.ValueContent()
+		[
+			StructPropertyHandle->CreatePropertyValueWidget()
+		];
+	}
+	else
+	{
+		HeaderRow
+		.NameContent()
+		[
+			StructPropertyHandle->CreatePropertyNameWidget()
+		]
+		.ValueContent()
+		[
+			SNew(SComboButton)
+			.ContentPadding(FMargin(2,2,2,1))
+			.ButtonContent()
+			[
+				SNew(STextBlock)
+				.Text(this, &FRigConnectionRuleDetails::OnGetStructTextValue)
+			]
+			.OnGetMenuContent(this, &FRigConnectionRuleDetails::GenerateStructPicker)
+		];
+	}
+}
+
+void FRigConnectionRuleDetails::CustomizeChildren(TSharedRef<IPropertyHandle> InStructPropertyHandle, IDetailChildrenBuilder& StructBuilder, IPropertyTypeCustomizationUtils& StructCustomizationUtils)
+{
+	const UScriptStruct* ScriptStruct = RuleStash.GetScriptStruct();
+	if(ScriptStruct == nullptr)
+	{
+		return;
+	}
+
+	(void)RuleStash.Get(Storage);
+	const TSharedRef<FStructOnScope> StorageRef = Storage.ToSharedRef();
+	const FSimpleDelegate OnPropertyChanged = FSimpleDelegate::CreateSP(this, &FRigConnectionRuleDetails::OnRuleContentChanged);
+
+	TArray<TSharedPtr<IPropertyHandle>> ChildProperties = InStructPropertyHandle->AddChildStructure(StorageRef);
+	for (TSharedPtr<IPropertyHandle> ChildHandle : ChildProperties)
+	{
+		ChildHandle->SetOnPropertyValueChanged(OnPropertyChanged);
+		(void)StructBuilder.AddProperty(ChildHandle.ToSharedRef());
+	}
+}
+
+TSharedRef<SWidget> FRigConnectionRuleDetails::GenerateStructPicker()
+{
+	FStructViewerModule& StructViewerModule = FModuleManager::LoadModuleChecked<FStructViewerModule>("StructViewer");
+
+	class FRigConnectionRuleFilter : public IStructViewerFilter
+	{
+	public:
+		FRigConnectionRuleFilter()
+		{
+		}
+
+		virtual bool IsStructAllowed(const FStructViewerInitializationOptions& InInitOptions, const UScriptStruct* InStruct, TSharedRef<FStructViewerFilterFuncs> InFilterFuncs) override
+		{
+			static const UScriptStruct* BaseStruct = FRigConnectionRule::StaticStruct();
+			return InStruct != BaseStruct && InStruct->IsChildOf(BaseStruct);
+		}
+
+		virtual bool IsUnloadedStructAllowed(const FStructViewerInitializationOptions& InInitOptions, const FSoftObjectPath& InStructPath, TSharedRef<FStructViewerFilterFuncs> InFilterFuncs) override
+		{
+			return false;
+		}
+	};
+		
+	static TSharedPtr<FRigConnectionRuleFilter> Filter = MakeShared<FRigConnectionRuleFilter>();
+	FStructViewerInitializationOptions Options;
+	{
+		Options.StructFilter = Filter;
+		Options.Mode = EStructViewerMode::StructPicker;
+		Options.DisplayMode = EStructViewerDisplayMode::ListView;
+		Options.NameTypeToDisplay = EStructViewerNameTypeToDisplay::DisplayName;
+		Options.bShowNoneOption = false;
+		Options.bShowUnloadedStructs = false;
+		Options.bAllowViewOptions = false;
+	}
+
+	return
+		SNew(SBox)
+		.WidthOverride(330.0f)
+		[
+			SNew(SVerticalBox)
+
+			+SVerticalBox::Slot()
+			.FillHeight(1.0f)
+			.MaxHeight(500)
+			[
+				SNew(SBorder)
+				.Padding(4)
+				.BorderImage(FAppStyle::GetBrush("ToolPanel.GroupBorder"))
+				[
+					StructViewerModule.CreateStructViewer(Options, FOnStructPicked::CreateSP(this, &FRigConnectionRuleDetails::OnPickedStruct))
+				]
+			]
+		];
+}
+
+void FRigConnectionRuleDetails::OnPickedStruct(const UScriptStruct* ChosenStruct)
+{
+	if(ChosenStruct == nullptr)
+	{
+		RuleStash = FRigConnectionRuleStash();
+	}
+	else
+	{
+		RuleStash.ScriptStructPath = ChosenStruct->GetPathName();
+		RuleStash.ExportedText = TEXT("()");
+		Storage.Reset();
+	}
+	OnRuleContentChanged();
+	PropertyUtilities->ForceRefresh();
+}
+
+FText FRigConnectionRuleDetails::OnGetStructTextValue() const
+{
+	const UScriptStruct* ScriptStruct = RuleStash.GetScriptStruct();
+	return ScriptStruct
+		? FText::AsCultureInvariant(ScriptStruct->GetName())
+		: LOCTEXT("None", "None");
+}
+
+void FRigConnectionRuleDetails::OnRuleContentChanged()
+{
+	const UScriptStruct* ScriptStruct = RuleStash.GetScriptStruct();
+	if(Storage && Storage->GetStruct() == ScriptStruct)
+	{
+		RuleStash.ExportedText.Reset();
+		const uint8* StructMemory = Storage->GetStructMemory();
+		ScriptStruct->ExportText(RuleStash.ExportedText, StructMemory, StructMemory, nullptr, PPF_None, nullptr);
+	}
+	
+	FString Content;
+	FRigConnectionRuleStash::StaticStruct()->ExportText(Content, &RuleStash, &RuleStash, nullptr, PPF_None, nullptr);
+
+	TArray<UObject*> Objects;
+	StructPropertyHandle->GetOuterObjects(Objects);
+	FString FirstObjectValue;
+	for (int32 Index = 0; Index < Objects.Num(); Index++)
+	{
+		(void)StructPropertyHandle->SetPerObjectValue(Index, Content, EPropertyValueSetFlags::DefaultFlags);
+	}
 }
 
 #undef LOCTEXT_NAMESPACE
diff --git a/Engine/Plugins/Animation/ControlRig/Source/ControlRigEditor/Private/ControlRigElementDetails.h b/Engine/Plugins/Animation/ControlRig/Source/ControlRigEditor/Private/ControlRigElementDetails.h
--- a/Engine/Plugins/Animation/ControlRig/Source/ControlRigEditor/Private/ControlRigElementDetails.h
+++ b/Engine/Plugins/Animation/ControlRig/Source/ControlRigEditor/Private/ControlRigElementDetails.h
@@ -754,4 +754,31 @@ public:
 	virtual void CustomizeDetails(IDetailLayoutBuilder& DetailBuilder) override;
 
 	void CustomizeSettings(IDetailLayoutBuilder& DetailBuilder);
-};
\ No newline at end of file
+};
+
+class FRigConnectionRuleDetails : public IPropertyTypeCustomization
+{
+public:
+
+	static TSharedRef<IPropertyTypeCustomization> MakeInstance()
+	{
+		return MakeShareable(new FRigConnectionRuleDetails);
+	}
+
+	/** IPropertyTypeCustomization interface */
+	virtual void CustomizeHeader(TSharedRef<class IPropertyHandle> InStructPropertyHandle, class FDetailWidgetRow& HeaderRow, IPropertyTypeCustomizationUtils& StructCustomizationUtils) override;
+	virtual void CustomizeChildren(TSharedRef<class IPropertyHandle> InStructPropertyHandle, class IDetailChildrenBuilder& StructBuilder, IPropertyTypeCustomizationUtils& StructCustomizationUtils) override;
+
+protected:
+	
+	TSharedRef<SWidget> GenerateStructPicker();
+	void OnPickedStruct(const UScriptStruct* ChosenStruct);
+	FText OnGetStructTextValue() const;
+	void OnRuleContentChanged();
+	
+	FRigConnectionRuleStash RuleStash;
+	TSharedPtr<FStructOnScope> Storage;
+	UControlRigBlueprint* BlueprintBeingCustomized;
+	TSharedPtr<IPropertyHandle> StructPropertyHandle;
+	TSharedPtr<IPropertyUtilities> PropertyUtilities;
+};
diff --git a/Engine/Plugins/Animation/ControlRig/Source/ControlRigEditor/Private/Editor/SRigHierarchy.cpp b/Engine/Plugins/Animation/ControlRig/Source/ControlRigEditor/Private/Editor/SRigHierarchy.cpp
--- a/Engine/Plugins/Animation/ControlRig/Source/ControlRigEditor/Private/Editor/SRigHierarchy.cpp
+++ b/Engine/Plugins/Animation/ControlRig/Source/ControlRigEditor/Private/Editor/SRigHierarchy.cpp
@@ -1172,7 +1172,7 @@ void SRigHierarchy::CreateContextMenu()
 								SelectedKey = SelectedItems[0]->Key;
 							}
 							
-							if (!SelectedKey || SelectedKey.Type == ERigElementType::Bone)
+							if (!SelectedKey || SelectedKey.Type == ERigElementType::Bone || SelectedKey.Type == ERigElementType::Connector)
 							{
 								DefaultSection.AddMenuEntry(Commands.AddBoneItem);
 							}
@@ -2355,7 +2355,6 @@ TOptional<EItemDropZone> SRigHierarchy::OnCanAcceptDrop(const FDragDropEvent& Dr
 			case ERigElementType::Null:
 			case ERigElementType::RigidBody:
 			case ERigElementType::Reference:
-			case ERigElementType::Connector:
 			{
 				for (const FRigElementKey& DraggedKey : RigDragDropOp->GetElements())
 				{
@@ -2378,6 +2377,12 @@ TOptional<EItemDropZone> SRigHierarchy::OnCanAcceptDrop(const FDragDropEvent& Dr
 				ReturnDropZone = DropZone;
 				break;
 			}
+			case ERigElementType::Connector:
+			{
+				// anything can be parented under a connector
+				ReturnDropZone = DropZone;
+				break;
+			}
 			default:
 			{
 				ReturnDropZone = DropZone;
