Control Rig: Fix drag and drop of animation channels to node graph

diff --git a/Engine/Plugins/Animation/ControlRig/Source/ControlRig/Private/Rigs/RigHierarchyController.cpp b/Engine/Plugins/Animation/ControlRig/Source/ControlRig/Private/Rigs/RigHierarchyController.cpp
--- a/Engine/Plugins/Animation/ControlRig/Source/ControlRig/Private/Rigs/RigHierarchyController.cpp
+++ b/Engine/Plugins/Animation/ControlRig/Source/ControlRig/Private/Rigs/RigHierarchyController.cpp
@@ -2289,7 +2289,23 @@ FName URigHierarchyController::SetDisplayName(FRigControlElement* InControlEleme
 	{
 		RenameElement(InControlElement, InControlElement->Settings.DisplayName, false);
 	}
-
+#if WITH_EDITOR
+	else
+	{
+		// if we are merely setting the display name - we want to update all listening hierarchies
+		for(URigHierarchy::FRigHierarchyListener& Listener : Hierarchy->ListeningHierarchies)
+		{
+			if(URigHierarchy* ListeningHierarchy = Listener.Hierarchy.Get())
+			{
+				if(URigHierarchyController* ListeningController = ListeningHierarchy->GetController())
+				{
+					const TGuardValue<bool> Guard(ListeningController->bSuspendAllNotifications, true);
+					ListeningController->SetDisplayName(InControlElement->GetKey(), InDisplayName, bRenameElement, false, false);
+				}
+			}
+		}
+	}
+#endif
 	return InControlElement->Settings.DisplayName;
 }
 
diff --git a/Engine/Plugins/Animation/ControlRig/Source/ControlRigEditor/Private/Editor/ControlRigEditor.cpp b/Engine/Plugins/Animation/ControlRig/Source/ControlRigEditor/Private/Editor/ControlRigEditor.cpp
--- a/Engine/Plugins/Animation/ControlRig/Source/ControlRigEditor/Private/Editor/ControlRigEditor.cpp
+++ b/Engine/Plugins/Animation/ControlRig/Source/ControlRigEditor/Private/Editor/ControlRigEditor.cpp
@@ -3541,6 +3541,8 @@ void FControlRigEditor::HandleMakeElementGetterSetter(ERigElementGetterSetterTyp
 		{
 			if(ControlElement->IsAnimationChannel())
 			{
+				ChannelValue = ControlElement->GetDisplayName();
+				
 				if(const FRigControlElement* ParentControlElement = Cast<FRigControlElement>(Hierarchy->GetFirstParent(ControlElement)))
 				{
 					NameValue = ParentControlElement->GetName();
