Control Rig: Fix access of connection info map

diff --git a/Engine/Plugins/Animation/ControlRig/Source/ControlRig/Private/Rigs/RigConnectionRules.cpp b/Engine/Plugins/Animation/ControlRig/Source/ControlRig/Private/Rigs/RigConnectionRules.cpp
--- a/Engine/Plugins/Animation/ControlRig/Source/ControlRig/Private/Rigs/RigConnectionRules.cpp
+++ b/Engine/Plugins/Animation/ControlRig/Source/ControlRig/Private/Rigs/RigConnectionRules.cpp
@@ -308,8 +308,19 @@ bool FRigChainConnectionRule::CanConnect(const FRigConnectionInfo* InConnectionI
 		return false;
 	}
 
-	const FRigElementKey StartKey = InConnectionInfo->ConnectionMap.begin().Value();
-	const FRigElementKey EndKey = InConnectionInfo->ConnectionMap.end().Value();
+	TArray<FRigElementKey> Targets;
+	Targets.Reserve(InConnectionInfo->ConnectionMap.Num());
+	for(const TPair<FRigElementKey, FRigElementKey>& Pair : InConnectionInfo->ConnectionMap)
+	{
+		Targets.Add(Pair.Value);
+	}
+	if(Targets.IsEmpty())
+	{
+		return false;
+	}
+
+	const FRigElementKey StartKey = Targets[0];
+	const FRigElementKey EndKey = Targets.Last();
 
 	FRigElementKey NextKey = EndKey;
 	TArray<FRigElementKey> ReversedChain;
