Control Rig: Fix chain rule

diff --git a/Engine/Plugins/Animation/ControlRig/Source/ControlRig/Private/Rigs/RigConnectionRules.cpp b/Engine/Plugins/Animation/ControlRig/Source/ControlRig/Private/Rigs/RigConnectionRules.cpp
--- a/Engine/Plugins/Animation/ControlRig/Source/ControlRig/Private/Rigs/RigConnectionRules.cpp
+++ b/Engine/Plugins/Animation/ControlRig/Source/ControlRig/Private/Rigs/RigConnectionRules.cpp
@@ -301,6 +301,11 @@ bool FRigTagConnectionRule::CanConnect(const FRigConnectionInfo* InConnectionInf
 // FRigChainConnectionRule
 ////////////////////////////////////////////////////////////////////////////////
 
+TArray<FRigElementKey> FRigChainConnectionRule::GetAdditionalConnectors() const
+{
+	return {EndConnector};
+}
+
 bool FRigChainConnectionRule::CanConnect(const FRigConnectionInfo* InConnectionInfo, FString* OutFailureReason) const
 {
 	if(!Super::CanConnect(InConnectionInfo, OutFailureReason))
diff --git a/Engine/Plugins/Animation/ControlRig/Source/ControlRig/Public/Rigs/RigConnectionRules.h b/Engine/Plugins/Animation/ControlRig/Source/ControlRig/Public/Rigs/RigConnectionRules.h
--- a/Engine/Plugins/Animation/ControlRig/Source/ControlRig/Public/Rigs/RigConnectionRules.h
+++ b/Engine/Plugins/Animation/ControlRig/Source/ControlRig/Public/Rigs/RigConnectionRules.h
@@ -198,7 +198,7 @@ public:
 	virtual ~FRigChainConnectionRule() override {}
 
 	virtual UScriptStruct* GetScriptStruct() const override { return FRigChainConnectionRule::StaticStruct(); }
-
+	virtual TArray<FRigElementKey> GetAdditionalConnectors() const override;
 	virtual bool CanConnect(const FRigConnectionInfo* InConnectionInfo, FString* OutFailureReason) const override;
 
 	UPROPERTY(BlueprintReadWrite, EditAnywhere, Category=Settings)
