Control Rig: Connection Rule Details Panel Enabled fix

diff --git a/Engine/Plugins/Animation/ControlRig/Source/ControlRigEditor/Private/ControlRigElementDetails.cpp b/Engine/Plugins/Animation/ControlRig/Source/ControlRigEditor/Private/ControlRigElementDetails.cpp
--- a/Engine/Plugins/Animation/ControlRig/Source/ControlRigEditor/Private/ControlRigElementDetails.cpp
+++ b/Engine/Plugins/Animation/ControlRig/Source/ControlRigEditor/Private/ControlRigElementDetails.cpp
@@ -5108,6 +5108,7 @@ void FRigConnectionRuleDetails::CustomizeHeader(TSharedRef<IPropertyHandle> InSt
 	StructPropertyHandle = InStructPropertyHandle.ToSharedPtr();
 	PropertyUtilities = StructCustomizationUtils.GetPropertyUtilities();
 	BlueprintBeingCustomized = nullptr;
+	EnabledAttribute = false;
 	RigElementKeyDetails_GetCustomizedInfo(InStructPropertyHandle, BlueprintBeingCustomized);
 
 	TArray<UObject*> Objects;
@@ -5131,6 +5132,12 @@ void FRigConnectionRuleDetails::CustomizeHeader(TSharedRef<IPropertyHandle> InSt
 				}
 			}
 		}
+
+		// only enable editing of the rule if the widget is nested under a wrapper object (a rig element)
+		if(Objects[Index]->IsA<URigVMDetailsViewWrapperObject>())
+		{
+			EnabledAttribute = true;
+		}
 	}
 
 	if(!FirstObjectValue.IsEmpty())
@@ -5167,6 +5174,7 @@ void FRigConnectionRuleDetails::CustomizeHeader(TSharedRef<IPropertyHandle> InSt
 				.Text(this, &FRigConnectionRuleDetails::OnGetStructTextValue)
 			]
 			.OnGetMenuContent(this, &FRigConnectionRuleDetails::GenerateStructPicker)
+			.IsEnabled(EnabledAttribute)
 		];
 	}
 }
@@ -5187,7 +5195,8 @@ void FRigConnectionRuleDetails::CustomizeChildren(TSharedRef<IPropertyHandle> In
 	for (TSharedPtr<IPropertyHandle> ChildHandle : ChildProperties)
 	{
 		ChildHandle->SetOnPropertyValueChanged(OnPropertyChanged);
-		(void)StructBuilder.AddProperty(ChildHandle.ToSharedRef());
+		IDetailPropertyRow& ChildRow = StructBuilder.AddProperty(ChildHandle.ToSharedRef());
+		ChildRow.IsEnabled(EnabledAttribute);
 	}
 }
 
diff --git a/Engine/Plugins/Animation/ControlRig/Source/ControlRigEditor/Private/ControlRigElementDetails.h b/Engine/Plugins/Animation/ControlRig/Source/ControlRigEditor/Private/ControlRigElementDetails.h
--- a/Engine/Plugins/Animation/ControlRig/Source/ControlRigEditor/Private/ControlRigElementDetails.h
+++ b/Engine/Plugins/Animation/ControlRig/Source/ControlRigEditor/Private/ControlRigElementDetails.h
@@ -781,4 +781,5 @@ protected:
 	UControlRigBlueprint* BlueprintBeingCustomized;
 	TSharedPtr<IPropertyHandle> StructPropertyHandle;
 	TSharedPtr<IPropertyUtilities> PropertyUtilities;
+	TAttribute<bool> EnabledAttribute;
 };
