Control Rig: show user facing notification when adding a connector

diff --git a/Engine/Plugins/Animation/ControlRig/Source/ControlRigEditor/Private/Editor/ControlRigEditor.cpp b/Engine/Plugins/Animation/ControlRig/Source/ControlRigEditor/Private/Editor/ControlRigEditor.cpp
--- a/Engine/Plugins/Animation/ControlRig/Source/ControlRigEditor/Private/Editor/ControlRigEditor.cpp
+++ b/Engine/Plugins/Animation/ControlRig/Source/ControlRigEditor/Private/Editor/ControlRigEditor.cpp
@@ -2710,6 +2710,22 @@ void FControlRigEditor::OnHierarchyModified(ERigHierarchyNotification InNotif, U
 	switch(InNotif)
 	{
 		case ERigHierarchyNotification::ElementAdded:
+		{
+			if(InElement->GetType() == ERigElementType::Connector)
+			{
+				if(InHierarchy->GetConnectors().Num() == 1)
+				{
+					FNotificationInfo Info(LOCTEXT("FirstConnectorEncountered", "Looks like you have added the first connector. This rig can now be configured as a module in the class settings Hierarchy -> Module Settings."));
+					Info.bFireAndForget = true;
+					Info.FadeOutDuration = 5.0f;
+					Info.ExpireDuration = 5.0f;
+
+					TSharedPtr<SNotificationItem> NotificationPtr = FSlateNotificationManager::Get().AddNotification(Info);
+					NotificationPtr->SetCompletionState(SNotificationItem::CS_Success);
+				}
+			}
+			// no break - fall through
+		}
 		case ERigHierarchyNotification::ParentChanged:
 		case ERigHierarchyNotification::HierarchyReset:
 		{
