Control Rig: Command to load all control rigs

diff --git a/Engine/Plugins/Animation/ControlRig/Source/ControlRigEditor/Private/ControlRigBlueprintEditorLibrary.cpp b/Engine/Plugins/Animation/ControlRig/Source/ControlRigEditor/Private/ControlRigBlueprintEditorLibrary.cpp
--- a/Engine/Plugins/Animation/ControlRig/Source/ControlRigEditor/Private/ControlRigBlueprintEditorLibrary.cpp
+++ b/Engine/Plugins/Animation/ControlRig/Source/ControlRigEditor/Private/ControlRigBlueprintEditorLibrary.cpp
@@ -5,9 +5,20 @@
 #include "AssetRegistry/AssetData.h"
 #include "AssetRegistry/AssetRegistryModule.h"
 #include "RigVMModel/RigVMPin.h"
+#include "Misc/ScopedSlowTask.h"
 
 #include UE_INLINE_GENERATED_CPP_BY_NAME(ControlRigBlueprintEditorLibrary)
 
+FAutoConsoleCommand FCmdControlRigLoadAllAssets
+(
+	TEXT("ControlRig.LoadAllAssets"),
+	TEXT("Loads all control rig assets."),
+	FConsoleCommandDelegate::CreateLambda([]()
+	{
+		UControlRigBlueprintEditorLibrary::LoadAllControlRigs();
+	})
+);
+
 void UControlRigBlueprintEditorLibrary::CastToControlRigBlueprint(
 	UObject* Object,
 	ECastToControlRigBlueprintCases& Branches,
@@ -164,3 +175,34 @@ TArray<FRigModuleDescription> UControlRigBlueprintEditorLibrary::GetAvailableRig
 	return ModuleDescriptions;
 }
 
+void UControlRigBlueprintEditorLibrary::LoadAllControlRigs()
+{
+	const FAssetRegistryModule& AssetRegistryModule = FModuleManager::LoadModuleChecked<FAssetRegistryModule>(TEXT("AssetRegistry"));
+
+	// Collect a full list of assets with the specified class
+	TArray<FAssetData> AssetDataList;
+	AssetRegistryModule.Get().GetAssetsByClass(UControlRigBlueprint::StaticClass()->GetClassPathName(), AssetDataList, true);
+
+	const int32 NumAssets = AssetDataList.Num();
+
+	FScopedSlowTask LoadAllRigsTask(NumAssets, FText::FromString("Load all control rigs..."));
+	LoadAllRigsTask.MakeDialog(true);
+
+	for(int32 Index = 0; Index < NumAssets; Index++)
+	{
+		if (LoadAllRigsTask.ShouldCancel())
+		{
+			break;
+		}
+
+		const FAssetData& AssetData = AssetDataList[Index];
+		LoadAllRigsTask.EnterProgressFrame(1, FText::FromName(AssetData.PackageName));		
+
+		static constexpr TCHAR Format[] = TEXT("[%d/%d]: %s -> %s");
+		UE_LOG(LogControlRig, Display, Format, Index, NumAssets, *AssetData.AssetName.ToString(), *AssetData.PackageName.ToString())
+		(void)AssetData.GetAsset();
+
+		LoadAllRigsTask.ForceRefresh();
+	}
+}
+
diff --git a/Engine/Plugins/Animation/ControlRig/Source/ControlRigEditor/Public/ControlRigBlueprintEditorLibrary.h b/Engine/Plugins/Animation/ControlRig/Source/ControlRigEditor/Public/ControlRigBlueprintEditorLibrary.h
--- a/Engine/Plugins/Animation/ControlRig/Source/ControlRigEditor/Public/ControlRigBlueprintEditorLibrary.h
+++ b/Engine/Plugins/Animation/ControlRig/Source/ControlRigEditor/Public/ControlRigBlueprintEditorLibrary.h
@@ -74,5 +74,8 @@ public:
 
 	UFUNCTION(BlueprintCallable, BlueprintPure, Category = "Control Rig Blueprint")
 	static TArray<FRigModuleDescription> GetAvailableRigModules();
+	
+	UFUNCTION(BlueprintCallable, Category = "Control Rig Blueprint")
+	static void LoadAllControlRigs();
 };
 
