Control Rig: UX improvements on AdditiveControlRig tracks (name and color)
#rnx

diff --git a/Engine/Plugins/Animation/ControlRig/Source/ControlRigEditor/Private/Sequencer/ControlRigParameterTrackEditor.cpp b/Engine/Plugins/Animation/ControlRig/Source/ControlRigEditor/Private/Sequencer/ControlRigParameterTrackEditor.cpp
--- a/Engine/Plugins/Animation/ControlRig/Source/ControlRigEditor/Private/Sequencer/ControlRigParameterTrackEditor.cpp
+++ b/Engine/Plugins/Animation/ControlRig/Source/ControlRigEditor/Private/Sequencer/ControlRigParameterTrackEditor.cpp
@@ -1127,7 +1127,6 @@ void FControlRigParameterTrackEditor::AddAdditiveControlRig(FGuid ObjectBinding,
 
 				FString ObjectName = ControlRigClass->GetName();
 				ObjectName.RemoveFromEnd(TEXT("_C"));
-				const FString AdditiveObjectName = ObjectName + "_Additive";
 				UControlRig* ControlRig = NewObject<UControlRig>(Track, ControlRigClass, FName(*ObjectName), RF_Transactional);
 				ControlRig->SetIsAdditive(true);
 				if (!ControlRig->SupportsEvent(FRigUnit_InverseExecution::EventName))
@@ -1155,8 +1154,11 @@ void FControlRigParameterTrackEditor::AddAdditiveControlRig(FGuid ObjectBinding,
 				const bool bSequencerOwnsControlRig = true;
 				UMovieSceneSection* NewSection = Track->CreateControlRigSection(0, ControlRig, bSequencerOwnsControlRig);
 				
+				const FString AdditiveObjectName = ObjectName + " (Additive)";
 				Track->SetTrackName(FName(*ObjectName));
-				Track->SetDisplayName(FText::FromString(ObjectName));
+				Track->SetDisplayName(FText::FromString(AdditiveObjectName));
+				Track->SetColorTint(FColor(173, 151, 114));
+				
 
 				ParentSequencer->EmptySelection();
 				ParentSequencer->SelectSection(NewSection);
