Control Rig: Fix additive control rigs not picking up blueprint animations
#rnx

diff --git a/Engine/Plugins/Animation/ControlRig/Source/ControlRig/Private/Sequencer/ControlRigLayerInstance.cpp b/Engine/Plugins/Animation/ControlRig/Source/ControlRig/Private/Sequencer/ControlRigLayerInstance.cpp
--- a/Engine/Plugins/Animation/ControlRig/Source/ControlRig/Private/Sequencer/ControlRigLayerInstance.cpp
+++ b/Engine/Plugins/Animation/ControlRig/Source/ControlRig/Private/Sequencer/ControlRigLayerInstance.cpp
@@ -7,6 +7,7 @@
 
 #include "Sequencer/ControlRigLayerInstance.h"
 #include "Sequencer/ControlRigLayerInstanceProxy.h"
+#include "ControlRig.h"
 
 #include UE_INLINE_GENERATED_CPP_BY_NAME(ControlRigLayerInstance)
 
@@ -27,6 +28,22 @@ FAnimInstanceProxy* UControlRigLayerInstance::CreateAnimInstanceProxy()
 	return new FControlRigLayerInstanceProxy(this);
 }
 
+void UControlRigLayerInstance::NativeUpdateAnimation(float DeltaSeconds)
+{
+	Super::NativeUpdateAnimation(DeltaSeconds);
+	
+	if (UAnimInstance* SourceAnimInstance = GetSourceAnimInstance())
+	{
+		if (UControlRig* ControlRig = GetFirstAvailableControlRig())
+		{
+			if (ControlRig->IsAdditive())
+			{
+				SourceAnimInstance->UpdateAnimation(DeltaSeconds, true, EUpdateAnimationFlag::Default);
+			}
+		}
+	}
+}
+
 void UControlRigLayerInstance::UpdateAnimTrack(UAnimSequenceBase* InAnimSequence, int32 SequenceId, float InPosition, float Weight, bool bFireNotifies)
 {
 	GetProxyOnGameThread<FControlRigLayerInstanceProxy>().UpdateAnimTrack(InAnimSequence, SequenceId, InPosition, Weight, bFireNotifies);
@@ -85,6 +102,13 @@ void UControlRigLayerInstance::SetSourceAnimInstance(UAnimInstance* SourceAnimIn
 /** ControlRig related support */
 void UControlRigLayerInstance::AddControlRigTrack(int32 ControlRigID, UControlRig* InControlRig)
 {
+	if (InControlRig->IsAdditive())
+	{
+		if (UAnimInstance* SourceAnimInstance = GetSourceAnimInstance())
+		{
+			SourceAnimInstance->bUseMultiThreadedAnimationUpdate = false;
+		}
+	}
 	GetProxyOnGameThread<FControlRigLayerInstanceProxy>().AddControlRigTrack(ControlRigID, InControlRig);
 }
 
@@ -101,6 +125,17 @@ void UControlRigLayerInstance::UpdateControlRigTrack(int32 ControlRigID, float W
 void UControlRigLayerInstance::RemoveControlRigTrack(int32 ControlRigID)
 {
 	GetProxyOnGameThread<FControlRigLayerInstanceProxy>().RemoveControlRigTrack(ControlRigID);
+	if (UAnimInstance* SourceAnimInstance = GetSourceAnimInstance())
+	{
+		if (UControlRig* ControlRig = GetFirstAvailableControlRig())
+		{
+			SourceAnimInstance->bUseMultiThreadedAnimationUpdate = !ControlRig->IsAdditive();
+		}
+		else
+		{
+			SourceAnimInstance->bUseMultiThreadedAnimationUpdate = true;
+		}
+	}
 }
 
 void UControlRigLayerInstance::ResetControlRigTracks()
diff --git a/Engine/Plugins/Animation/ControlRig/Source/ControlRig/Public/Sequencer/ControlRigLayerInstance.h b/Engine/Plugins/Animation/ControlRig/Source/ControlRig/Public/Sequencer/ControlRigLayerInstance.h
--- a/Engine/Plugins/Animation/ControlRig/Source/ControlRig/Public/Sequencer/ControlRigLayerInstance.h
+++ b/Engine/Plugins/Animation/ControlRig/Source/ControlRig/Public/Sequencer/ControlRigLayerInstance.h
@@ -52,6 +52,8 @@ public:
 protected:
 	// UAnimInstance interface
 	virtual FAnimInstanceProxy* CreateAnimInstanceProxy() override;
+public:
+	virtual void NativeUpdateAnimation(float DeltaSeconds) override;
 
 public:
 	static const FName SequencerPoseName;
