Control Rig: Improve FName FString conversion

diff --git a/Engine/Plugins/Animation/AnimationData/Source/AnimationData/Private/AnimSequencerController.cpp b/Engine/Plugins/Animation/AnimationData/Source/AnimationData/Private/AnimSequencerController.cpp
--- a/Engine/Plugins/Animation/AnimationData/Source/AnimationData/Private/AnimSequencerController.cpp
+++ b/Engine/Plugins/Animation/AnimationData/Source/AnimationData/Private/AnimSequencerController.cpp
@@ -1312,7 +1312,7 @@ bool UAnimSequencerController::SetCurveAttributes(const FAnimationCurveIdentifie
 				const URigHierarchy* Hierarchy = FKRig->GetHierarchy();
 				if (Hierarchy ||  IgnoreSkeletonValidation())
 				{
-					const FRigElementKey CurveKey(URigHierarchy::GetSanitizedName(CurveId.CurveName.ToString()), ERigElementType::Curve);
+					const FRigElementKey CurveKey(URigHierarchy::GetSanitizedName(CurveId.CurveName), ERigElementType::Curve);
 					const FRigElementKey CurveControlKey(UFKControlRig::GetControlName(CurveKey.Name, ERigElementType::Curve), ERigElementType::Control);
 
 					const bool bContainsCurve = Hierarchy ? Hierarchy->Contains(CurveKey) : IgnoreSkeletonValidation();
@@ -1547,7 +1547,7 @@ bool UAnimSequencerController::SetBoneTrackKeys(FName BoneName, const TArray<FVe
 				{
 					if (URigHierarchy* Hierarchy = ControlRig->GetHierarchy())
 					{
-						const FRigElementKey BoneKey(URigHierarchy::GetSanitizedName(BoneName.ToString()), ERigElementType::Bone);
+						const FRigElementKey BoneKey(URigHierarchy::GetSanitizedName(BoneName), ERigElementType::Bone);
 						const FRigElementKey BoneNameControlKey(UFKControlRig::GetControlName(BoneKey.Name, ERigElementType::Bone), ERigElementType::Control);
 						
 						const bool bContainsBone = Hierarchy->Contains(BoneKey);
@@ -2161,12 +2161,12 @@ void UAnimSequencerController::RemoveUnusedControlsAndCurves() const
 				{
 					if(!(TransformCurves.ContainsByPredicate([ControlElement](const FTransformParameterNameAndCurves& TransformParameter)
 							{
-								return TransformParameter.ParameterName == ControlElement->GetName();
+								return TransformParameter.ParameterName == ControlElement->GetFName();
 							})
 							||
 							ScalarCurves.ContainsByPredicate([ControlElement](const FScalarParameterNameAndCurve& ScalarParameter)
 							{
-								return ScalarParameter.ParameterName == ControlElement->GetName();
+								return ScalarParameter.ParameterName == ControlElement->GetFName();
 							}))
 					)
 					{
@@ -2182,7 +2182,7 @@ void UAnimSequencerController::RemoveUnusedControlsAndCurves() const
 				{
 					if(!FloatCurves.ContainsByPredicate([CurveElement](const FFloatCurve& FloatCurve)
 					{
-						return FloatCurve.GetName() == CurveElement->GetName();
+						return FloatCurve.GetName() == CurveElement->GetFName();
 					}))
 					{						
 						ElementsToRemove.Add(CurveElement->GetKey());
@@ -2452,7 +2452,7 @@ bool UAnimSequencerController::AddBoneControl(const FName& BoneName) const
 		{
 			if(URigHierarchy* Hierarchy = FKRig->GetHierarchy())
 			{
-				const FRigElementKey BoneKey(URigHierarchy::GetSanitizedName(BoneName.ToString()), ERigElementType::Bone);
+				const FRigElementKey BoneKey(URigHierarchy::GetSanitizedName(BoneName), ERigElementType::Bone);
 				const FRigElementKey BoneNameControlKey(UFKControlRig::GetControlName(BoneKey.Name, ERigElementType::Bone), ERigElementType::Control);
 
 				const bool bContainsBone = Hierarchy->Contains(BoneKey);
@@ -2557,7 +2557,7 @@ bool UAnimSequencerController::RemoveBoneControl(const FName& BoneName) const
 		{
 			if(URigHierarchy* Hierarchy = FKRig->GetHierarchy())
 			{
-				const FRigElementKey BoneNameControlKey(UFKControlRig::GetControlName(URigHierarchy::GetSanitizedName(BoneName.ToString()), ERigElementType::Bone), ERigElementType::Control);
+				const FRigElementKey BoneNameControlKey(UFKControlRig::GetControlName(URigHierarchy::GetSanitizedName(BoneName), ERigElementType::Bone), ERigElementType::Control);
 				
 				const bool bContainsTransformCurve = Section->GetTransformParameterNamesAndCurves().ContainsByPredicate([BoneNameControlKey](const FTransformParameterNameAndCurves& Parameter)
 				{
@@ -2615,7 +2615,7 @@ bool UAnimSequencerController::SetBoneCurveKeys(const FName& BoneName, const TAr
 			const URigHierarchy* Hierarchy = FKRig->GetHierarchy();
 			if(Hierarchy || IgnoreSkeletonValidation())
 			{
-				const FRigElementKey BoneKey(URigHierarchy::GetSanitizedName(BoneName.ToString()), ERigElementType::Bone);
+				const FRigElementKey BoneKey(URigHierarchy::GetSanitizedName(BoneName), ERigElementType::Bone);
 				const FRigElementKey BoneNameControlKey(UFKControlRig::GetControlName(BoneKey.Name, ERigElementType::Bone), ERigElementType::Control);
 
 				const bool bContainsBone = Hierarchy ? Hierarchy->Contains(BoneKey) : IgnoreSkeletonValidation();
@@ -2806,7 +2806,7 @@ bool UAnimSequencerController::UpdateBoneCurveKeys(const FName& BoneName, const
 		{
 			if(URigHierarchy* Hierarchy = FKRig->GetHierarchy())
 			{
-				const FRigElementKey BoneKey(URigHierarchy::GetSanitizedName(BoneName.ToString()), ERigElementType::Bone);
+				const FRigElementKey BoneKey(URigHierarchy::GetSanitizedName(BoneName), ERigElementType::Bone);
 				const FRigElementKey BoneNameControlKey(UFKControlRig::GetControlName(BoneKey.Name, ERigElementType::Bone), ERigElementType::Control);
 
 				const bool bContainsBone = Hierarchy->Contains(BoneKey);
@@ -2902,7 +2902,7 @@ bool UAnimSequencerController::RemoveBoneCurveKey(const FName& BoneName, float T
 			{
 				if(const URigHierarchy* Hierarchy = FKRig->GetHierarchy())
 				{
-					const FRigElementKey BoneKey(URigHierarchy::GetSanitizedName(BoneName.ToString()), ERigElementType::Bone);
+					const FRigElementKey BoneKey(URigHierarchy::GetSanitizedName(BoneName), ERigElementType::Bone);
 					const FRigElementKey BoneNameControlKey(UFKControlRig::GetControlName(BoneKey.Name, ERigElementType::Bone), ERigElementType::Control);
 	
 					const bool bContainsBone = Hierarchy->Contains(BoneKey);
@@ -2975,7 +2975,7 @@ bool UAnimSequencerController::AddCurveControl(const FName& CurveName) const
 		{
 			if(URigHierarchy* Hierarchy = FKRig->GetHierarchy())
 			{				
-				const FRigElementKey CurveKey(URigHierarchy::GetSanitizedName(CurveName.ToString()), ERigElementType::Curve);
+				const FRigElementKey CurveKey(URigHierarchy::GetSanitizedName(CurveName), ERigElementType::Curve);
 				const FRigElementKey CurveControlKey(UFKControlRig::GetControlName(CurveKey.Name, ERigElementType::Curve), ERigElementType::Control);
 
 				const bool bContainsCurve = Hierarchy->Contains(CurveKey);
@@ -3059,10 +3059,10 @@ bool UAnimSequencerController::RenameCurveControl(const FName& CurveName, const
 			{
 				if (URigHierarchyController* Controller = RigHierarchy->GetController())
 				{
-					const FRigElementKey CurveKey(URigHierarchy::GetSanitizedName(CurveName.ToString()), ERigElementType::Curve);
+					const FRigElementKey CurveKey(URigHierarchy::GetSanitizedName(CurveName), ERigElementType::Curve);
 					const FRigElementKey CurveControlKey(UFKControlRig::GetControlName(CurveKey.Name, ERigElementType::Curve), ERigElementType::Control);
 
-					const FName SanitizedNewCurveName = URigHierarchy::GetSanitizedName(NewCurveName.ToString());
+					const FName SanitizedNewCurveName = URigHierarchy::GetSanitizedName(NewCurveName);
 					// Rename the curve element itself
 					if (RigHierarchy->Contains(CurveKey))
 					{											
@@ -3146,7 +3146,7 @@ bool UAnimSequencerController::RemoveCurveControl(const FName& CurveName) const
 		{
 			if(URigHierarchy* Hierarchy = FKRig->GetHierarchy())
 			{									
-				const FRigElementKey CurveKey(URigHierarchy::GetSanitizedName(CurveName.ToString()), ERigElementType::Curve);
+				const FRigElementKey CurveKey(URigHierarchy::GetSanitizedName(CurveName), ERigElementType::Curve);
 				const FRigElementKey CurveControlKey(UFKControlRig::GetControlName(CurveKey.Name, ERigElementType::Curve), ERigElementType::Control);
 
 				const bool bContainsCurve = Hierarchy->Contains(CurveKey);
@@ -3212,7 +3212,7 @@ bool UAnimSequencerController::SetCurveControlKeys(const FName& CurveName, const
 
 			if (Hierarchy ||  IgnoreSkeletonValidation())
 			{
-				const FRigElementKey CurveKey(URigHierarchy::GetSanitizedName(CurveName.ToString()), ERigElementType::Curve);
+				const FRigElementKey CurveKey(URigHierarchy::GetSanitizedName(CurveName), ERigElementType::Curve);
 				const FRigElementKey CurveControlKey(UFKControlRig::GetControlName(CurveKey.Name, ERigElementType::Curve), ERigElementType::Control);
 
 				const bool bContainsCurve = Hierarchy ? Hierarchy->Contains(CurveKey) : IgnoreSkeletonValidation();
@@ -3298,7 +3298,7 @@ bool UAnimSequencerController::SetCurveControlKey(const FName& CurveName, const
 		{
 			if(const URigHierarchy* Hierarchy = FKRig->GetHierarchy())
 			{
-				const FRigElementKey CurveKey(URigHierarchy::GetSanitizedName(CurveName.ToString()), ERigElementType::Curve);
+				const FRigElementKey CurveKey(URigHierarchy::GetSanitizedName(CurveName), ERigElementType::Curve);
 				const FRigElementKey CurveControlKey(UFKControlRig::GetControlName(CurveKey.Name, ERigElementType::Curve), ERigElementType::Control);
 				const bool bContainsCurve = Hierarchy->Contains(CurveKey);
 				const bool bContainsCurveControl = Hierarchy->Contains(CurveControlKey);
@@ -3388,7 +3388,7 @@ bool UAnimSequencerController::RemoveCurveControlKey(const FName& CurveName, flo
 			if(const URigHierarchy* Hierarchy = FKRig->GetHierarchy())
 			{
 				const FFrameNumber FrameNumber = Model->MovieScene->GetTickResolution().AsFrameNumber(Time);
-				const FRigElementKey CurveKey(URigHierarchy::GetSanitizedName(CurveName.ToString()), ERigElementType::Curve);
+				const FRigElementKey CurveKey(URigHierarchy::GetSanitizedName(CurveName), ERigElementType::Curve);
 				const FRigElementKey CurveControlKey(UFKControlRig::GetControlName(CurveKey.Name, ERigElementType::Curve), ERigElementType::Control);
 				const bool bContainsCurve = Hierarchy->Contains(CurveKey);
 				const bool bContainsCurveControl = Hierarchy->Contains(CurveControlKey);
@@ -3453,7 +3453,7 @@ bool UAnimSequencerController::DuplicateCurveControl(const FName& CurveName, con
 			{
 				if (URigHierarchyController* HierarchyController = Hierarchy->GetController())
 				{
-					const FRigElementKey NewCurveKey(URigHierarchy::GetSanitizedName(DuplicateCurveName.ToString()), ERigElementType::Curve);
+					const FRigElementKey NewCurveKey(URigHierarchy::GetSanitizedName(DuplicateCurveName), ERigElementType::Curve);
 					const FRigElementKey NewCurveControlKey(UFKControlRig::GetControlName(NewCurveKey.Name, ERigElementType::Curve), ERigElementType::Control);
 
 					const bool bAlreadyContainsCurve = Hierarchy->Contains(NewCurveKey);
@@ -3473,7 +3473,7 @@ bool UAnimSequencerController::DuplicateCurveControl(const FName& CurveName, con
 						// Rename the curve driving the control value
 						const FScalarParameterNameAndCurve* ParameterCurvePair = Section->GetScalarParameterNamesAndCurves().FindByPredicate([CurveName](const FScalarParameterNameAndCurve& Parameter)
 						{
-							return Parameter.ParameterName == UFKControlRig::GetControlName(URigHierarchy::GetSanitizedName(CurveName.ToString()), ERigElementType::Curve);
+							return Parameter.ParameterName == UFKControlRig::GetControlName(URigHierarchy::GetSanitizedName(CurveName), ERigElementType::Curve);
 						});
 				
 						if(ParameterCurvePair)
diff --git a/Engine/Plugins/Animation/AnimationData/Source/AnimationData/Private/AnimSequencerDataModel.cpp b/Engine/Plugins/Animation/AnimationData/Source/AnimationData/Private/AnimSequencerDataModel.cpp
--- a/Engine/Plugins/Animation/AnimationData/Source/AnimationData/Private/AnimSequencerDataModel.cpp
+++ b/Engine/Plugins/Animation/AnimationData/Source/AnimationData/Private/AnimSequencerDataModel.cpp
@@ -60,8 +60,8 @@ void UAnimationSequencerDataModel::RemoveOutOfDateControls() const
 					TArray<FRigElementKey> ElementKeysToRemove;
 					Hierarchy->ForEach<FRigControlElement>([this, Section, &ElementKeysToRemove](const FRigControlElement* ControlElement) -> bool
 					{
-						const bool bContainsBone = Section->HasTransformParameter(ControlElement->GetName());
-						const bool bContainsCurve = Section->HasScalarParameter(ControlElement->GetName());
+						const bool bContainsBone = Section->HasTransformParameter(ControlElement->GetFName());
+						const bool bContainsCurve = Section->HasScalarParameter(ControlElement->GetFName());
 						
 						if (!bContainsBone && !bContainsCurve)
 						{
@@ -73,7 +73,7 @@ void UAnimationSequencerDataModel::RemoveOutOfDateControls() const
 						
 					Hierarchy->ForEach<FRigCurveElement>([this, &ElementKeysToRemove](const FRigCurveElement* CurveElement) -> bool
 					{
-						const FName TargetCurveName = CurveElement->GetName();
+						const FName TargetCurveName = CurveElement->GetFName();
 						if(!LegacyCurveData.FloatCurves.ContainsByPredicate([TargetCurveName](const FFloatCurve& Curve) { return Curve.GetName() == TargetCurveName; }))
 						{
 							ElementKeysToRemove.Add(CurveElement->GetKey());	
@@ -884,17 +884,17 @@ void UAnimationSequencerDataModel::GenerateLegacyCurveData()
 
 						Hierarchy->ForEach<FRigCurveElement>([Hierarchy, this, ScalarCurves, FrameRate = GetFrameRate(), SequencerSuffix](const FRigCurveElement* CurveElement) -> bool
 							{
-								const FRigElementKey ControlKey(UFKControlRig::GetControlName(CurveElement->GetName(), ERigElementType::Curve), ERigElementType::Control);
+								const FRigElementKey ControlKey(UFKControlRig::GetControlName(CurveElement->GetFName(), ERigElementType::Curve), ERigElementType::Control);
 								if (const FRigControlElement* Element = Hierarchy->Find<FRigControlElement>(ControlKey))
 								{
 									FFloatCurve& FloatCurve = LegacyCurveData.FloatCurves.AddDefaulted_GetRef();
 									if (RetainFloatCurves)
 									{
-									FloatCurve.SetName(FName(*(CurveElement->GetName().ToString() + TEXT("_Sequencer"))));
+									FloatCurve.SetName(FName(*(CurveElement->GetFName().ToString() + TEXT("_Sequencer"))));
 									}
 									else
 									{
-									FloatCurve.SetName(CurveElement->GetName());
+									FloatCurve.SetName(CurveElement->GetFName());
 									}						
 								
 									FloatCurve.Color = Element->Settings.ShapeColor;
@@ -912,7 +912,7 @@ void UAnimationSequencerDataModel::GenerateLegacyCurveData()
 
 									if (const FScalarParameterNameAndCurve* ScalarCurve = ScalarCurves.FindByPredicate([Element](FScalarParameterNameAndCurve Curve)
 									{
-										return Curve.ParameterName == Element->GetName();
+										return Curve.ParameterName == Element->GetFName();
 									}))
 									{							
 										AnimSequencerHelpers::ConvertFloatChannelToRichCurve(ScalarCurve->ParameterCurve, FloatCurve.FloatCurve, FrameRate);
@@ -1183,7 +1183,7 @@ void UAnimationSequencerDataModel::GeneratePoseData(UControlRig* ControlRig, FAn
 				QUICK_SCOPE_CYCLE_COUNTER(STAT_GetMappings);
 				RigHierarchy->ForEach<FRigBoneElement>([&MeshRefSkeleton, &RequiredBones, &SkeletonRefSkeleton, &RigHierarchy, &RetargetingScope, &RigPose](const FRigBoneElement* BoneElement) -> bool
 				{
-					const FName& BoneName = BoneElement->GetName();
+					const FName& BoneName = BoneElement->GetFName();
 					const int32 BoneIndex = MeshRefSkeleton.FindBoneIndex(BoneName);
 					if (BoneIndex != INDEX_NONE)
 					{
@@ -1205,7 +1205,7 @@ void UAnimationSequencerDataModel::GeneratePoseData(UControlRig* ControlRig, FAn
 
 				RigHierarchy->ForEach<FRigCurveElement>([this, &RigHierarchy, &Curve](const FRigCurveElement* CurveElement) -> bool
 					{
-						const FName& CurveName = CurveElement->GetName();
+						const FName& CurveName = CurveElement->GetFName();
 					Curve.Add(CurveName, RigHierarchy->GetCurveValue(CurveElement->GetKey()));
 						return true;
 					});
diff --git a/Engine/Plugins/Animation/ControlRig/Source/ControlRig/Private/AnimNode_ControlRigBase.cpp b/Engine/Plugins/Animation/ControlRig/Source/ControlRig/Private/AnimNode_ControlRigBase.cpp
--- a/Engine/Plugins/Animation/ControlRig/Source/ControlRig/Private/AnimNode_ControlRigBase.cpp
+++ b/Engine/Plugins/Animation/ControlRig/Source/ControlRig/Private/AnimNode_ControlRigBase.cpp
@@ -358,7 +358,7 @@ void FAnimNode_ControlRigBase::UpdateOutput(UControlRig* ControlRig, FPoseContex
 		const TArray<FRigCurveElement*> Curves = Hierarchy->GetCurves();
 		auto GetNameFromIndex = [&Curves](int32 InCurveIndex)
 		{
-			return Curves[InCurveIndex]->GetName();
+			return Curves[InCurveIndex]->GetFName();
 		};
 
 		auto GetValueFromIndex = [&Curves](int32 InCurveIndex)
diff --git a/Engine/Plugins/Animation/ControlRig/Source/ControlRig/Private/Constraints/ControlRigTransformableHandle.cpp b/Engine/Plugins/Animation/ControlRig/Source/ControlRig/Private/Constraints/ControlRigTransformableHandle.cpp
--- a/Engine/Plugins/Animation/ControlRig/Source/ControlRig/Private/Constraints/ControlRigTransformableHandle.cpp
+++ b/Engine/Plugins/Animation/ControlRig/Source/ControlRig/Private/Constraints/ControlRigTransformableHandle.cpp
@@ -226,7 +226,7 @@ bool UTransformableControlHandle::HasDirectDependencyWith(const UTransformableHa
 	const FRigBaseElementParentArray AllParents = Hierarchy->GetParents(ControlElement, bRecursive);
 	const bool bIsParent = AllParents.ContainsByPredicate([this, OtherHash](const FRigBaseElement* Parent)
 	{
-		const uint32 ParentHash = ComputeHash(ControlRig.Get(), Parent->GetName());
+		const uint32 ParentHash = ComputeHash(ControlRig.Get(), Parent->GetFName());
 		return ParentHash == OtherHash;		
 	});
 
@@ -239,7 +239,7 @@ bool UTransformableControlHandle::HasDirectDependencyWith(const UTransformableHa
 	const TArray<FRigControlElement*> AllControls = Hierarchy->GetControls();
 	const int32 IndexOfPossibleParent = AllControls.IndexOfByPredicate([this, OtherHash](const FRigBaseElement* Parent)
 	{
-		const uint32 ChildHash = ComputeHash(ControlRig.Get(), Parent->GetName());
+		const uint32 ChildHash = ComputeHash(ControlRig.Get(), Parent->GetFName());
 		return ChildHash == OtherHash;
 	});
 
@@ -368,7 +368,7 @@ void UTransformableControlHandle::OnHierarchyModified(
 			const FName OldName = Hierarchy->GetPreviousName(InElement->GetKey());
 			if (OldName == ControlName)
 			{
-				ControlName = InElement->GetName();
+				ControlName = InElement->GetFName();
 			}
 			break;
 		}
@@ -397,7 +397,7 @@ void UTransformableControlHandle::OnControlModified(
 		const EHandleEvent Event = InContext.bConstraintUpdate ?
 			EHandleEvent::GlobalTransformUpdated : EHandleEvent::LocalTransformUpdated;
 
-		if (InControl->GetName() == ControlName)
+		if (InControl->GetFName() == ControlName)
 		{	// if that handle is wrapping InControl  
 			OnHandleModified.Broadcast(this, Event);
 		}
diff --git a/Engine/Plugins/Animation/ControlRig/Source/ControlRig/Private/ControlRig.cpp b/Engine/Plugins/Animation/ControlRig/Source/ControlRig/Private/ControlRig.cpp
--- a/Engine/Plugins/Animation/ControlRig/Source/ControlRig/Private/ControlRig.cpp
+++ b/Engine/Plugins/Animation/ControlRig/Source/ControlRig/Private/ControlRig.cpp
@@ -1358,7 +1358,7 @@ void UControlRig::GetMappableNodeData(TArray<FName>& OutNames, TArray<FNodeItem>
 	// now add all nodes
 	DynamicHierarchy->ForEach<FRigBoneElement>([&OutNames, &OutNodeItems, this](FRigBoneElement* BoneElement) -> bool
     {
-		OutNames.Add(BoneElement->GetName());
+		OutNames.Add(BoneElement->GetFName());
 		FRigElementKey ParentKey = DynamicHierarchy->GetFirstParent(BoneElement->GetKey());
 		if(ParentKey.Type != ERigElementType::Bone)
 		{
@@ -1558,7 +1558,7 @@ void UControlRig::CreateRigControlsForCurveContainer()
 
 		DynamicHierarchy->ForEach<FRigCurveElement>([this, Controller](FRigCurveElement* CurveElement) -> bool
         {
-			const FString Name = CurveElement->GetName().ToString();
+			const FString Name = CurveElement->GetFName().ToString();
 			
 			if (Name.Contains(CtrlPrefix) && !DynamicHierarchy->Contains(FRigElementKey(*Name, ERigElementType::Curve))) //-V1051
 			{
@@ -1571,7 +1571,7 @@ void UControlRig::CreateRigControlsForCurveContainer()
 				FRigControlValue Value;
 				Value.Set<float>(CurveElement->Value);
 
-				Controller->AddControl(CurveElement->GetName(), FRigElementKey(), Settings, Value, FTransform::Identity, FTransform::Identity); 
+				Controller->AddControl(CurveElement->GetFName(), FRigElementKey(), Settings, Value, FTransform::Identity, FTransform::Identity); 
 			}
 
 			return true;
@@ -1586,7 +1586,7 @@ void UControlRig::HandleOnControlModified(UControlRig* Subject, FRigControlEleme
 	if (Control->Settings.bIsCurve && DynamicHierarchy)
 	{
 		const FRigControlValue Value = DynamicHierarchy->GetControlValue(Control, IsConstructionModeEnabled() ? ERigControlValueType::Initial : ERigControlValueType::Current);
-		DynamicHierarchy->SetCurveValue(FRigElementKey(Control->GetName(), ERigElementType::Curve), Value.Get<float>());
+		DynamicHierarchy->SetCurveValue(FRigElementKey(Control->GetFName(), ERigElementType::Curve), Value.Get<float>());
 	}	
 }
 
@@ -1816,7 +1816,7 @@ TArray<FName> UControlRig::CurrentControlSelection() const
 		TArray<const FRigBaseElement*> SelectedControls = DynamicHierarchy->GetSelectedElements(ERigElementType::Control);
 		for (const FRigBaseElement* SelectedControl : SelectedControls)
 		{
-			SelectedControlNames.Add(SelectedControl->GetName());
+			SelectedControlNames.Add(SelectedControl->GetFName());
 		}
 #if WITH_EDITOR
 		for(const TSharedPtr<FRigDirectManipulationInfo>& ManipulationInfo : RigUnitManipulationInfos)
@@ -2276,7 +2276,7 @@ FName UControlRig::GetNameForTransientControl(const URigVMUnitNode* InNode, cons
 	check(DynamicHierarchy);
 	
 	const FString NodeName = InNode->GetName();
-	return DynamicHierarchy->GetSanitizedName(FString::Printf(TEXT("ControlForNode|%s|%s"), *NodeName, *InTarget.Name));
+	return DynamicHierarchy->GetSanitizedName(FRigName(FString::Printf(TEXT("ControlForNode|%s|%s"), *NodeName, *InTarget.Name)));
 }
 
 FString UControlRig::GetNodeNameFromTransientControl(const FRigElementKey& InKey)
@@ -2892,7 +2892,7 @@ void UControlRig::HandleInteractionRigControlSelected(UControlRig* Subject, FRig
 
 		if (bInverted)
 		{
-			const FString ControlName = Control->GetName().ToString();
+			const FString ControlName = Control->GetFName().ToString();
 			if (ControlName.EndsWith(TEXT("_CONTROL")))
 			{
 				const FString BaseName = ControlName.Left(ControlName.Len() - 8);
@@ -2951,7 +2951,7 @@ void UControlRig::SetBoneInitialTransformsFromRefSkeleton(const FReferenceSkelet
 	{
 		if(BoneElement->BoneType == ERigBoneType::Imported)
 		{
-			const int32 BoneIndex = InReferenceSkeleton.FindBoneIndex(BoneElement->GetName());
+			const int32 BoneIndex = InReferenceSkeleton.FindBoneIndex(BoneElement->GetFName());
 			if (BoneIndex != INDEX_NONE)
 			{
 				const FTransform LocalInitialTransform = InReferenceSkeleton.GetRefBonePose()[BoneIndex];
@@ -2984,7 +2984,7 @@ void UControlRig::SetBoneInitialTransformsFromCompactPose(FCompactPose* InCompac
 		{
 			if (BoneElement->BoneType == ERigBoneType::Imported)
 			{
-				int32 MeshIndex = InCompactPose->GetBoneContainer().GetPoseBoneIndexForBoneName(BoneElement->GetName());
+				int32 MeshIndex = InCompactPose->GetBoneContainer().GetPoseBoneIndexForBoneName(BoneElement->GetFName());
 				if (MeshIndex != INDEX_NONE)
 				{
 					FCompactPoseBoneIndex CPIndex = InCompactPose->GetBoneContainer().MakeCompactPoseIndex(FMeshPoseBoneIndex(MeshIndex));
@@ -3222,7 +3222,7 @@ UControlRig::FTransientControlScope::FTransientControlScope(TObjectPtr<URigHiera
 	for (FRigControlElement* Control : Hierarchy->GetTransientControls())
 	{
 		FTransientControlInfo Info;
-		Info.Name = Control->GetName();
+		Info.Name = Control->GetFName();
 		Info.Parent = Hierarchy->GetFirstParent(Control->GetKey());
 		Info.Settings = Control->Settings;
 		// preserve whatever value that was produced by this transient control at the moment
diff --git a/Engine/Plugins/Animation/ControlRig/Source/ControlRig/Private/ControlRigComponent.cpp b/Engine/Plugins/Animation/ControlRig/Source/ControlRig/Private/ControlRigComponent.cpp
--- a/Engine/Plugins/Animation/ControlRig/Source/ControlRig/Private/ControlRigComponent.cpp
+++ b/Engine/Plugins/Animation/ControlRig/Source/ControlRig/Private/ControlRigComponent.cpp
@@ -537,7 +537,7 @@ TArray<FName> UControlRigComponent::GetElementNames(ERigElementType ElementType)
 		{
 			if(Element->IsTypeOf(ElementType))
 			{
-				Names.Add(Element->GetName());
+				Names.Add(Element->GetFName());
 			}
 		}
 	}
@@ -662,11 +662,11 @@ void UControlRigComponent::AddMappedSkeletalMesh(USkeletalMeshComponent* Skeleta
 			{
 				CR->GetHierarchy()->ForEach<FRigBoneElement>([SkeletalMeshComponent, &BonesToMap](FRigBoneElement* BoneElement) -> bool
 				{
-					if (SkeletalMeshComponent->GetBoneIndex(BoneElement->GetName()) != INDEX_NONE)
+					if (SkeletalMeshComponent->GetBoneIndex(BoneElement->GetFName()) != INDEX_NONE)
 					{
 						FControlRigComponentMappedBone BoneToMap;
-						BoneToMap.Source = BoneElement->GetName();
-						BoneToMap.Target = BoneElement->GetName();
+						BoneToMap.Source = BoneElement->GetFName();
+						BoneToMap.Target = BoneElement->GetFName();
 						BonesToMap.Add(BoneToMap);
 					}
 					return true;
@@ -689,8 +689,8 @@ void UControlRigComponent::AddMappedSkeletalMesh(USkeletalMeshComponent* Skeleta
 				CR->GetHierarchy()->ForEach<FRigCurveElement>([Skeleton, &CurvesToMap](FRigCurveElement* CurveElement) -> bool
                 {
 					FControlRigComponentMappedCurve CurveToMap;
-					CurveToMap.Source = CurveElement->GetName();
-					CurveToMap.Target = CurveElement->GetName();
+					CurveToMap.Source = CurveElement->GetFName();
+					CurveToMap.Target = CurveElement->GetFName();
 					CurvesToMap.Add(CurveToMap);
 					return true;
 				});
diff --git a/Engine/Plugins/Animation/ControlRig/Source/ControlRig/Private/ControlRigControlActor.cpp b/Engine/Plugins/Animation/ControlRig/Source/ControlRig/Private/ControlRigControlActor.cpp
--- a/Engine/Plugins/Animation/ControlRig/Source/ControlRig/Private/ControlRigControlActor.cpp
+++ b/Engine/Plugins/Animation/ControlRig/Source/ControlRig/Private/ControlRigControlActor.cpp
@@ -247,7 +247,7 @@ void AControlRigControlActor::Refresh()
 							Component->SetMaterial(i, MaterialInstance);
 						}
 
-						ControlNames.Add(ControlElement->GetName());
+						ControlNames.Add(ControlElement->GetFName());
 						ShapeTransforms.Add(ShapeDef->Transform * Hierarchy->GetControlShapeTransform(ControlElement, ERigTransformType::CurrentLocal));
 						Components.Add(Component);
 						Materials.Add(MaterialInstance);
diff --git a/Engine/Plugins/Animation/ControlRig/Source/ControlRig/Private/Rigs/AdditiveControlRig.cpp b/Engine/Plugins/Animation/ControlRig/Source/ControlRig/Private/Rigs/AdditiveControlRig.cpp
--- a/Engine/Plugins/Animation/ControlRig/Source/ControlRig/Private/Rigs/AdditiveControlRig.cpp
+++ b/Engine/Plugins/Animation/ControlRig/Source/ControlRig/Private/Rigs/AdditiveControlRig.cpp
@@ -75,7 +75,7 @@ bool UAdditiveControlRig::Execute_Internal(const FName& InEventName)
 		GetHierarchy()->ForEach<FRigBoneElement>([&](FRigBoneElement* BoneElement) -> bool
 		{
 			FRigUnit_AddBoneTransform NewUnit;
-			NewUnit.Bone = BoneElement->GetName();
+			NewUnit.Bone = BoneElement->GetFName();
 			NewUnit.bPropagateToChildren = true;
 			AddBoneRigUnits.Add(NewUnit);
 			return true;
@@ -119,7 +119,7 @@ void UAdditiveControlRig::CreateRigElements(const FReferenceSkeleton& InReferenc
 		// add control for all bone hierarchy
 		GetHierarchy()->ForEach<FRigBoneElement>([&](FRigBoneElement* BoneElement) -> bool
         {
-            const FName BoneName = BoneElement->GetName();
+            const FName BoneName = BoneElement->GetFName();
             const int32 ParentIndex = GetHierarchy()->GetFirstParent(BoneElement->GetIndex());
             const FName NullName = GetNullName(BoneName);// name conflict?
             const FName ControlName = GetControlName(BoneName); // name conflict?
diff --git a/Engine/Plugins/Animation/ControlRig/Source/ControlRig/Private/Rigs/FKControlRig.cpp b/Engine/Plugins/Animation/ControlRig/Source/ControlRig/Private/Rigs/FKControlRig.cpp
--- a/Engine/Plugins/Animation/ControlRig/Source/ControlRig/Private/Rigs/FKControlRig.cpp
+++ b/Engine/Plugins/Animation/ControlRig/Source/ControlRig/Private/Rigs/FKControlRig.cpp
@@ -120,7 +120,7 @@ bool UFKControlRig::Execute_Internal(const FName& InEventName)
 	{
 		GetHierarchy()->ForEach<FRigBoneElement>([&](FRigBoneElement* BoneElement) -> bool
         {
-			const FName ControlName = GetControlName(BoneElement->GetName(), BoneElement->GetType());
+			const FName ControlName = GetControlName(BoneElement->GetFName(), BoneElement->GetType());
 			const FRigElementKey ControlKey(ControlName, ERigElementType::Control);
 			const int32 ControlIndex = GetHierarchy()->GetIndex(ControlKey);
 			if (ControlIndex != INDEX_NONE && IsControlActive[ControlIndex])
@@ -157,7 +157,7 @@ bool UFKControlRig::Execute_Internal(const FName& InEventName)
 
 		GetHierarchy()->ForEach<FRigCurveElement>([&](FRigCurveElement* CurveElement) -> bool
         {
-			const FName ControlName = GetControlName(CurveElement->GetName(), CurveElement->GetType());
+			const FName ControlName = GetControlName(CurveElement->GetFName(), CurveElement->GetType());
 			const FRigElementKey ControlKey(ControlName, ERigElementType::Control);
 			const int32 ControlIndex = GetHierarchy()->GetIndex(ControlKey);
 			
@@ -200,7 +200,7 @@ bool UFKControlRig::Execute_Internal(const FName& InEventName)
 
 			FRigBoneElement* BoneElement = CastChecked<FRigBoneElement>(InElement);
 
-			const FName ControlName = GetControlName(BoneElement->GetName(), BoneElement->GetType());
+			const FName ControlName = GetControlName(BoneElement->GetFName(), BoneElement->GetType());
 			const FRigElementKey ControlKey(ControlName, ERigElementType::Control);
 			const int32 ControlIndex = GetHierarchy()->GetIndex(ControlKey);
 			
@@ -235,7 +235,7 @@ bool UFKControlRig::Execute_Internal(const FName& InEventName)
 
 		GetHierarchy()->ForEach<FRigCurveElement>([&](FRigCurveElement* CurveElement) -> bool
         {
-			const FName ControlName = GetControlName(CurveElement->GetName(), CurveElement->GetType());
+			const FName ControlName = GetControlName(CurveElement->GetFName(), CurveElement->GetType());
 			const FRigElementKey ControlKey(ControlName, ERigElementType::Control);
 			const int32 ControlIndex = GetHierarchy()->GetIndex(ControlKey);
 
@@ -308,7 +308,7 @@ TArray<FName> UFKControlRig::GetControlNames()
 	TArray<FName> Names;
 	for (FRigControlElement* ControlElement : Controls) 
 	{
-		Names.Add(ControlElement->GetName());
+		Names.Add(ControlElement->GetFName());
 	}
 	return Names;
 }
@@ -372,7 +372,7 @@ void UFKControlRig::CreateRigElements(const FReferenceSkeleton& InReferenceSkele
 		{
 			auto CreateControlForBoneElement = [this, Controller](const FRigBoneElement* BoneElement)
 			{
-				const FName BoneName = BoneElement->GetName();
+				const FName BoneName = BoneElement->GetFName();
 				const FName ControlName = GetControlName(BoneName, BoneElement->GetType());
 				const FRigElementKey ParentKey = GetHierarchy()->GetFirstParent(BoneElement->GetKey());
 
@@ -410,13 +410,13 @@ void UFKControlRig::CreateRigElements(const FReferenceSkeleton& InReferenceSkele
 		{
 			GetHierarchy()->ForEach<FRigCurveElement>([&](const FRigCurveElement* CurveElement) -> bool
 			{
-				const FName ControlName = GetControlName(CurveElement->GetName(), CurveElement->GetType());
+				const FName ControlName = GetControlName(CurveElement->GetFName(), CurveElement->GetType());
 
 				FRigControlSettings Settings;
 				Settings.ControlType = ERigControlType::Float;
-				Settings.DisplayName = CurveElement->GetName();
+				Settings.DisplayName = CurveElement->GetFName();
 
-				const FName DisplayCurveControlName(*(CurveElement->GetName().ToString() + TEXT(" Curve")));
+				const FName DisplayCurveControlName(*(CurveElement->GetFName().ToString() + TEXT(" Curve")));
 				Settings.DisplayName = DisplayCurveControlName;
 				Controller->AddControl(ControlName, FRigElementKey(), Settings, FRigControlValue::Make(CurveElement->Value), FTransform::Identity, FTransform::Identity, false);
 				
@@ -439,7 +439,7 @@ void UFKControlRig::SetControlOffsetsFromBoneInitials()
 		}
 
 		const FRigBoneElement* BoneElement = CastChecked<FRigBoneElement>(InElement);
-		const FName BoneName = BoneElement->GetName();
+		const FName BoneName = BoneElement->GetFName();
 		const FName ControlName = GetControlName(BoneName, BoneElement->GetType());
 
 		FRigControlElement* ControlElement = GetHierarchy()->Find<FRigControlElement>(FRigElementKey(ControlName, ERigElementType::Control));
@@ -521,11 +521,11 @@ void UFKControlRig::ToggleApplyMode()
         {
             if (ControlElement->Settings.ControlType == ERigControlType::EulerTransform)
 			{
-				SetControlValue<FEulerTransform>(ControlElement->GetName(), FEulerTransform::Identity, true, Context, bSetupUndo);
+				SetControlValue<FEulerTransform>(ControlElement->GetFName(), FEulerTransform::Identity, true, Context, bSetupUndo);
 			}
 			else if (ControlElement->Settings.ControlType == ERigControlType::Float)
 			{
-				SetControlValue<float>(ControlElement->GetName(), 0.f, true, Context, bSetupUndo);
+				SetControlValue<float>(ControlElement->GetFName(), 0.f, true, Context, bSetupUndo);
 			}
 
 			return true;
@@ -543,12 +543,12 @@ void UFKControlRig::ToggleApplyMode()
 			{
 				const FRigControlValue::FEulerTransform_Float InitValue =
 					GetHierarchy()->GetControlValue(ControlElement, ERigControlValueType::Initial).Get<FRigControlValue::FEulerTransform_Float>();
-				SetControlValue<FRigControlValue::FEulerTransform_Float>(ControlElement->GetName(), InitValue, true, Context, bSetupUndo);
+				SetControlValue<FRigControlValue::FEulerTransform_Float>(ControlElement->GetFName(), InitValue, true, Context, bSetupUndo);
 			}
 			else if (ControlElement->Settings.ControlType == ERigControlType::Float)
 			{
 				const float InitValue = GetHierarchy()->GetControlValue(ControlElement, ERigControlValueType::Initial).Get<float>();
-				SetControlValue<float>(ControlElement->GetName(), InitValue, true, Context, bSetupUndo);
+				SetControlValue<float>(ControlElement->GetFName(), InitValue, true, Context, bSetupUndo);
 			}
 
 			return true;
diff --git a/Engine/Plugins/Animation/ControlRig/Source/ControlRig/Private/Rigs/RigHierarchy.cpp b/Engine/Plugins/Animation/ControlRig/Source/ControlRig/Private/Rigs/RigHierarchy.cpp
--- a/Engine/Plugins/Animation/ControlRig/Source/ControlRig/Private/Rigs/RigHierarchy.cpp
+++ b/Engine/Plugins/Animation/ControlRig/Source/ControlRig/Private/Rigs/RigHierarchy.cpp
@@ -484,7 +484,7 @@ void URigHierarchy::CopyHierarchy(URigHierarchy* InHierarchy)
 			FRigBaseElement* Target = (FRigBaseElement*)&NewElementsPerType[ElementTypeIndex][StructureSize * SubIndex];
 		
 			Target->Key = Key;
-			Target->NameString = Source->NameString;
+			Target->NameString.Reset();
 			Target->SubIndex = SubIndex;
 			Target->Index = Elements.Add(Target);
 			Target->CreatedAtInstructionIndex = Source->CreatedAtInstructionIndex;
@@ -518,7 +518,7 @@ void URigHierarchy::CopyHierarchy(URigHierarchy* InHierarchy)
 
 			check(Target->Key.Type == Source->Key.Type);
             Target->Key = Source->Key;
-            Target->NameString = Source->NameString;
+            Target->NameString.Reset();
             Target->SubIndex = Source->SubIndex;
             Target->Index = Source->Index;
 			Target->CreatedAtInstructionIndex = Source->CreatedAtInstructionIndex;
@@ -555,7 +555,7 @@ uint32 URigHierarchy::GetNameHash() const
 	for (int32 ElementIndex = 0; ElementIndex < Elements.Num(); ElementIndex++)
 	{
 		const FRigBaseElement* Element = Elements[ElementIndex];
-		Hash = HashCombine(Hash, GetTypeHash(Element->GetName()));
+		Hash = HashCombine(Hash, GetTypeHash(Element->GetFName()));
 	}
 	return Hash;	
 }
@@ -1042,12 +1042,14 @@ TArray<FRigElementKey> URigHierarchy::GetSelectedKeys(ERigElementType InTypeFilt
 	return Selection;
 }
 
-void URigHierarchy::SanitizeName(FString& InOutName)
+void URigHierarchy::SanitizeName(FRigName& InOutName)
 {
 	// Sanitize the name
-	for (int32 i = 0; i < InOutName.Len(); ++i)
+	FString SanitizedNameString = InOutName.GetName();
+	bool bChangedSomething = false;
+	for (int32 i = 0; i < SanitizedNameString.Len(); ++i)
 	{
-		TCHAR& C = InOutName[i];
+		TCHAR& C = SanitizedNameString[i];
 
 		const bool bGoodChar = FChar::IsAlpha(C) ||					 // Any letter
 			(C == '_') || (C == '-') || (C == '.') || (C == '|') ||	 // _  - .  | anytime
@@ -1057,63 +1059,33 @@ void URigHierarchy::SanitizeName(FString& InOutName)
 		if (!bGoodChar)
 		{
 			C = '_';
+			bChangedSomething = true;
 		}
 	}
 
-	if (InOutName.Len() > GetMaxNameLength())
+	if (SanitizedNameString.Len() > GetMaxNameLength())
 	{
-		InOutName.LeftChopInline(InOutName.Len() - GetMaxNameLength());
+		SanitizedNameString.LeftChopInline(SanitizedNameString.Len() - GetMaxNameLength());
+		bChangedSomething = true;
 	}
-}
 
-FName URigHierarchy::GetSanitizedName(const FString& InName)
-{
-	FString Name = InName;
-	SanitizeName(Name);
-
-	if (Name.IsEmpty())
+	if(bChangedSomething)
 	{
-		return NAME_None;
+		InOutName.SetName(SanitizedNameString);
 	}
-
-	return *Name;
 }
 
-bool URigHierarchy::IsNameAvailable(const FString& InPotentialNewName, ERigElementType InType, FString* OutErrorMessage) const
+FRigName URigHierarchy::GetSanitizedName(const FRigName& InName)
 {
-	const FString UnsanitizedName = InPotentialNewName;
-	if (UnsanitizedName.Len() > GetMaxNameLength())
-	{
-		if (OutErrorMessage)
-		{
-			*OutErrorMessage = TEXT("Name too long.");
-		}
-		return false;
-	}
-
-	if (UnsanitizedName == TEXT("None"))
-	{
-		if (OutErrorMessage)
-		{
-			*OutErrorMessage = TEXT("None is not a valid name.");
-		}
-		return false;
-	}
-
-	FString SanitizedName = UnsanitizedName;
-	SanitizeName(SanitizedName);
-
-	if (SanitizedName != UnsanitizedName)
-	{
-		if (OutErrorMessage)
-		{
-			*OutErrorMessage = TEXT("Name contains invalid characters.");
-		}
-		return false;
-	}
+	FRigName Name = InName;
+	SanitizeName(Name);
+	return Name;
+}
 
+bool URigHierarchy::IsNameAvailable(const FRigName& InPotentialNewName, ERigElementType InType, FString* OutErrorMessage) const
+{
 	// check for fixed keywords
-	const FRigElementKey PotentialKey(*InPotentialNewName, InType);
+	const FRigElementKey PotentialKey(InPotentialNewName.GetFName(), InType);
 	if(PotentialKey == URigHierarchy::GetDefaultParentKey())
 	{
 		return false;
@@ -1128,13 +1100,7 @@ bool URigHierarchy::IsNameAvailable(const FString& InPotentialNewName, ERigEleme
 		return false;
 	}
 
-	return true;
-}
-
-bool URigHierarchy::IsDisplayNameAvailable(const FRigElementKey& InParentElement,
-	const FString& InPotentialNewDisplayName, FString* OutErrorMessage) const
-{
-	const FString UnsanitizedName = InPotentialNewDisplayName;
+	const FRigName UnsanitizedName = InPotentialNewName;
 	if (UnsanitizedName.Len() > GetMaxNameLength())
 	{
 		if (OutErrorMessage)
@@ -1144,7 +1110,7 @@ bool URigHierarchy::IsDisplayNameAvailable(const FRigElementKey& InParentElement
 		return false;
 	}
 
-	if (UnsanitizedName == TEXT("None"))
+	if (UnsanitizedName.IsNone())
 	{
 		if (OutErrorMessage)
 		{
@@ -1153,7 +1119,7 @@ bool URigHierarchy::IsDisplayNameAvailable(const FRigElementKey& InParentElement
 		return false;
 	}
 
-	FString SanitizedName = UnsanitizedName;
+	FRigName SanitizedName = UnsanitizedName;
 	SanitizeName(SanitizedName);
 
 	if (SanitizedName != UnsanitizedName)
@@ -1165,14 +1131,20 @@ bool URigHierarchy::IsDisplayNameAvailable(const FRigElementKey& InParentElement
 		return false;
 	}
 
+	return true;
+}
+
+bool URigHierarchy::IsDisplayNameAvailable(const FRigElementKey& InParentElement,
+	const FRigName& InPotentialNewDisplayName, FString* OutErrorMessage) const
+{
 	if(InParentElement.IsValid())
 	{
 		const TArray<FRigElementKey> ChildKeys = GetChildren(InParentElement);
-		if(ChildKeys.ContainsByPredicate([SanitizedName, this](const FRigElementKey& InChildKey) -> bool
+		if(ChildKeys.ContainsByPredicate([&InPotentialNewDisplayName, this](const FRigElementKey& InChildKey) -> bool
 		{
 			if(const FRigBaseElement* BaseElement = Find(InChildKey))
 			{
-				if(BaseElement->GetDisplayName().ToString() == SanitizedName)
+				if(BaseElement->GetDisplayName() == InPotentialNewDisplayName.GetFName())
 				{
 					return true;
 				}
@@ -1188,33 +1160,64 @@ bool URigHierarchy::IsDisplayNameAvailable(const FRigElementKey& InParentElement
 		}
 	}
 
+	const FRigName UnsanitizedName = InPotentialNewDisplayName;
+	if (UnsanitizedName.Len() > GetMaxNameLength())
+	{
+		if (OutErrorMessage)
+		{
+			*OutErrorMessage = TEXT("Name too long.");
+		}
+		return false;
+	}
+
+	if (UnsanitizedName.IsNone())
+	{
+		if (OutErrorMessage)
+		{
+			*OutErrorMessage = TEXT("None is not a valid name.");
+		}
+		return false;
+	}
+
+	FRigName SanitizedName = UnsanitizedName;
+	SanitizeName(SanitizedName);
+
+	if (SanitizedName != UnsanitizedName)
+	{
+		if (OutErrorMessage)
+		{
+			*OutErrorMessage = TEXT("Name contains invalid characters.");
+		}
+		return false;
+	}
+
 	return true;
 }
 
-FName URigHierarchy::GetSafeNewName(const FString& InPotentialNewName, ERigElementType InType) const
+FRigName URigHierarchy::GetSafeNewName(const FRigName& InPotentialNewName, ERigElementType InType) const
 {
-	FString SanitizedName = InPotentialNewName;
+	FRigName SanitizedName = InPotentialNewName;
 	SanitizeName(SanitizedName);
-	FString Name = SanitizedName;
+	FRigName Name = SanitizedName;
 
 	int32 Suffix = 1;
 	while (!IsNameAvailable(Name, InType))
 	{
-		FString BaseString = SanitizedName;
+		FString BaseString = SanitizedName.GetName();
 		if (BaseString.Len() > GetMaxNameLength() - 4)
 		{
 			BaseString.LeftChopInline(BaseString.Len() - (GetMaxNameLength() - 4));
 		}
-		Name = *FString::Printf(TEXT("%s_%d"), *BaseString, ++Suffix);
+		Name.SetName(FString::Printf(TEXT("%s_%d"), *BaseString, ++Suffix));
 	}
-	return *Name;
+	return Name;
 }
 
-FName URigHierarchy::GetSafeNewDisplayName(const FRigElementKey& InParentElement, const FString& InPotentialNewDisplayName) const
+FRigName URigHierarchy::GetSafeNewDisplayName(const FRigElementKey& InParentElement, const FRigName& InPotentialNewDisplayName) const
 {
-	if(InPotentialNewDisplayName.IsEmpty())
+	if(InPotentialNewDisplayName.IsNone())
 	{
-		return NAME_None;
+		return FRigName();
 	}
 
 	TArray<FRigElementKey> KeysToCheck;
@@ -1239,9 +1242,9 @@ FName URigHierarchy::GetSafeNewDisplayName(const FRigElementKey& InParentElement
 		}
 	}
 
-	FString SanitizedName = InPotentialNewDisplayName;
+	FRigName SanitizedName = InPotentialNewDisplayName;
 	SanitizeName(SanitizedName);
-	FString Name = SanitizedName;
+	FRigName Name = SanitizedName;
 
 	TArray<FString> DisplayNames;
 	Algo::Transform(KeysToCheck, DisplayNames, [this](const FRigElementKey& InKey) -> FString
@@ -1254,17 +1257,17 @@ FName URigHierarchy::GetSafeNewDisplayName(const FRigElementKey& InParentElement
 	});
 
 	int32 Suffix = 1;
-	while (DisplayNames.Contains(Name))
+	while (DisplayNames.Contains(Name.GetName()))
 	{
-		FString BaseString = SanitizedName;
+		FString BaseString = SanitizedName.GetName();
 		if (BaseString.Len() > GetMaxNameLength() - 4)
 		{
 			BaseString.LeftChopInline(BaseString.Len() - (GetMaxNameLength() - 4));
 		}
-		Name = *FString::Printf(TEXT("%s_%d"), *BaseString, ++Suffix);
+		Name.SetName(FString::Printf(TEXT("%s_%d"), *BaseString, ++Suffix));
 	}
 
-	return *Name;
+	return Name;
 }
 
 int32 URigHierarchy::GetPoseVersion(const FRigElementKey& InKey) const
@@ -3102,8 +3105,8 @@ void URigHierarchy::SetTransform(FRigTransformElement* InTransformElement, const
 									static constexpr TCHAR MessageFormat[] = TEXT("Setting transform of parent (%s) after setting child (%s).\nThis may lead to unexpected results.");
 									const FString& Message = FString::Printf(
 										MessageFormat,
-										*InTransformElement->GetName().ToString(),
-										*Child->GetName().ToString());
+										*InTransformElement->GetName(),
+										*Child->GetName());
 									CRContext.GetLog()->Report(
 										EMessageSeverity::Info,
 										ExecuteContext->GetPublicData<>().GetFunctionName(),
@@ -3287,7 +3290,7 @@ FTransform URigHierarchy::GetControlOffsetTransform(FRigControlElement* InContro
 
 				checkf(FRigComputedTransform::Equals(GlobalTransform, ComputedTransform),
 					TEXT("Element '%s' Offset %s Cached vs Computed doesn't match. ('%s' <-> '%s')"),
-					*InControlElement->GetName().ToString(),
+					*InControlElement->GetName(),
 					*TransformTypeStrings[(int32)InTransformType],
 					*GlobalTransform.ToString(), *ComputedTransform.ToString());
 			}
@@ -3311,7 +3314,7 @@ FTransform URigHierarchy::GetControlOffsetTransform(FRigControlElement* InContro
 
 				checkf(FRigComputedTransform::Equals(LocalTransform, ComputedTransform),
 					TEXT("Element '%s' Offset %s Cached vs Computed doesn't match. ('%s' <-> '%s')"),
-					*InControlElement->GetName().ToString(),
+					*InControlElement->GetName(),
 					*TransformTypeStrings[(int32)InTransformType],
 					*LocalTransform.ToString(), *ComputedTransform.ToString());
 			}
@@ -3628,7 +3631,7 @@ void URigHierarchy::SetControlSettings(FRigControlElement* InControlElement, FRi
 		}
 		if (!BlueprintName.IsEmpty())
 		{
-			FString ControlNamePythonized = RigVMPythonUtils::PythonizeName(InControlElement->GetName().ToString());
+			FString ControlNamePythonized = RigVMPythonUtils::PythonizeName(InControlElement->GetName());
 			FString SettingsName = FString::Printf(TEXT("control_settings_%s"),
 				*ControlNamePythonized);
 			TArray<FString> Commands = ControlSettingsToPythonCommands(InControlElement->Settings, SettingsName);
@@ -4010,7 +4013,7 @@ void URigHierarchy::SetConnectorSettings(FRigConnectorElement* InConnectorElemen
 		}
 		if (!BlueprintName.IsEmpty())
 		{
-			FString ControlNamePythonized = RigVMPythonUtils::PythonizeName(InConnectorElement->GetName().ToString());
+			FString ControlNamePythonized = RigVMPythonUtils::PythonizeName(InConnectorElement->GetName());
 			FString SettingsName = FString::Printf(TEXT("connector_settings_%s"),
 				*ControlNamePythonized);
 			TArray<FString> Commands = ConnectorSettingsToPythonCommands(InConnectorElement->Settings, SettingsName);
@@ -5087,7 +5090,7 @@ void URigHierarchy::EnsureCacheValidityImpl()
 					const FTransform ComputedTransform = HierarchyForLambda->GetControlOffsetTransform(ControlElement, TransformType);
 					checkf(FRigComputedTransform::Equals(CachedTransform, ComputedTransform),
 						TEXT("Element '%s' Offset %s Cached vs Computed doesn't match. ('%s' <-> '%s')"),
-						*Element->GetName().ToString(),
+						*Element->GetName(),
 						*TransformTypeString,
 						*CachedTransform.ToString(),
 						*ComputedTransform.ToString());
@@ -5110,7 +5113,7 @@ void URigHierarchy::EnsureCacheValidityImpl()
 					const FTransform ComputedTransform = HierarchyForLambda->GetTransform(TransformElement, TransformType);
 					checkf(FRigComputedTransform::Equals(CachedTransform, ComputedTransform),
 						TEXT("Element '%s' Pose %s Cached vs Computed doesn't match. ('%s' <-> '%s')"),
-						*Element->GetName().ToString(),
+						*Element->GetName(),
 						*TransformTypeString,
 						*CachedTransform.ToString(), *ComputedTransform.ToString());
 				}
@@ -5132,7 +5135,7 @@ void URigHierarchy::EnsureCacheValidityImpl()
 					const FTransform ComputedTransform = HierarchyForLambda->GetControlShapeTransform(ControlElement, TransformType);
 					checkf(FRigComputedTransform::Equals(CachedTransform, ComputedTransform),
 						TEXT("Element '%s' Shape %s Cached vs Computed doesn't match. ('%s' <-> '%s')"),
-						*Element->GetName().ToString(),
+						*Element->GetName(),
 						*TransformTypeString,
 						*CachedTransform.ToString(), *ComputedTransform.ToString());
 				}
diff --git a/Engine/Plugins/Animation/ControlRig/Source/ControlRig/Private/Rigs/RigHierarchyController.cpp b/Engine/Plugins/Animation/ControlRig/Source/ControlRig/Private/Rigs/RigHierarchyController.cpp
--- a/Engine/Plugins/Animation/ControlRig/Source/ControlRig/Private/Rigs/RigHierarchyController.cpp
+++ b/Engine/Plugins/Animation/ControlRig/Source/ControlRig/Private/Rigs/RigHierarchyController.cpp
@@ -212,7 +212,7 @@ FRigElementKey URigHierarchyController::AddBone(FName InName, FRigElementKey InP
 	{
 		TGuardValue<bool> DisableCacheValidityChecks(Hierarchy->bEnableCacheValidityCheck, false);
 		NewElement->Key.Type = ERigElementType::Bone;
-		NewElement->Key.Name = Hierarchy->GetSafeNewName(InName.ToString(), NewElement->Key.Type);
+		NewElement->Key.Name = Hierarchy->GetSafeNewName(InName, NewElement->Key.Type);
 		NewElement->BoneType = InBoneType;
 		AddElement(NewElement, Hierarchy->Get(Hierarchy->GetIndex(InParent)), true);
 
@@ -274,7 +274,7 @@ FRigElementKey URigHierarchyController::AddNull(FName InName, FRigElementKey InP
 	{
 		TGuardValue<bool> DisableCacheValidityChecks(Hierarchy->bEnableCacheValidityCheck, false);		
 		NewElement->Key.Type = ERigElementType::Null;
-		NewElement->Key.Name = Hierarchy->GetSafeNewName(InName.ToString(), NewElement->Key.Type);
+		NewElement->Key.Name = Hierarchy->GetSafeNewName(InName, NewElement->Key.Type);
 		AddElement(NewElement, Hierarchy->Get(Hierarchy->GetIndex(InParent)), false);
 
 		if(bTransformInGlobal)
@@ -339,7 +339,7 @@ FRigElementKey URigHierarchyController::AddControl(
 	{
 		TGuardValue<bool> DisableCacheValidityChecks(Hierarchy->bEnableCacheValidityCheck, false);		
 		NewElement->Key.Type = ERigElementType::Control;
-		NewElement->Key.Name = Hierarchy->GetSafeNewName(InName.ToString(), NewElement->Key.Type);
+		NewElement->Key.Name = Hierarchy->GetSafeNewName(InName, NewElement->Key.Type);
 		NewElement->Settings = InSettings;
 		if(NewElement->Settings.LimitEnabled.IsEmpty())
 		{
@@ -348,7 +348,7 @@ FRigElementKey URigHierarchyController::AddControl(
 
 		if(!NewElement->Settings.DisplayName.IsNone())
 		{
-			NewElement->Settings.DisplayName = Hierarchy->GetSafeNewDisplayName(InParent, NewElement->Settings.DisplayName.ToString()); 
+			NewElement->Settings.DisplayName = Hierarchy->GetSafeNewDisplayName(InParent, NewElement->Settings.DisplayName); 
 		}
 		AddElement(NewElement, Hierarchy->Get(Hierarchy->GetIndex(InParent)), false);
 		
@@ -432,7 +432,7 @@ FRigElementKey URigHierarchyController::AddCurve(FName InName, float InValue, bo
 	{
 		TGuardValue<bool> DisableCacheValidityChecks(Hierarchy->bEnableCacheValidityCheck, false);		
 		NewElement->Key.Type = ERigElementType::Curve;
-		NewElement->Key.Name = Hierarchy->GetSafeNewName(InName.ToString(), NewElement->Key.Type);
+		NewElement->Key.Name = Hierarchy->GetSafeNewName(InName, NewElement->Key.Type);
 		NewElement->Value = InValue;
 		AddElement(NewElement, nullptr, false);
 	}
@@ -481,7 +481,7 @@ FRigElementKey URigHierarchyController::AddRigidBody(FName InName, FRigElementKe
 	{
 		TGuardValue<bool> DisableCacheValidityChecks(Hierarchy->bEnableCacheValidityCheck, false);		
 		NewElement->Key.Type = ERigElementType::RigidBody;
-		NewElement->Key.Name = Hierarchy->GetSafeNewName(InName.ToString(), NewElement->Key.Type);
+		NewElement->Key.Name = Hierarchy->GetSafeNewName(InName, NewElement->Key.Type);
 		NewElement->Settings = InSettings;
 		AddElement(NewElement, Hierarchy->Get(Hierarchy->GetIndex(InParent)), true);
 
@@ -533,7 +533,7 @@ FRigElementKey URigHierarchyController::AddReference(FName InName, FRigElementKe
 	{
 		TGuardValue<bool> DisableCacheValidityChecks(Hierarchy->bEnableCacheValidityCheck, false);		
 		NewElement->Key.Type = ERigElementType::Reference;
-		NewElement->Key.Name = Hierarchy->GetSafeNewName(InName.ToString(), NewElement->Key.Type);
+		NewElement->Key.Name = Hierarchy->GetSafeNewName(InName, NewElement->Key.Type);
 		NewElement->GetWorldTransformDelegate = InDelegate;
 		AddElement(NewElement, Hierarchy->Get(Hierarchy->GetIndex(InParent)), true);
 
@@ -571,7 +571,7 @@ FRigElementKey URigHierarchyController::AddConnector(FName InName, FRigElementKe
 	{
 		TGuardValue<bool> DisableCacheValidityChecks(Hierarchy->bEnableCacheValidityCheck, false);
 		NewElement->Key.Type = ERigElementType::Connector;
-		NewElement->Key.Name = Hierarchy->GetSafeNewName(InName.ToString(), NewElement->Key.Type);
+		NewElement->Key.Name = Hierarchy->GetSafeNewName(InName, NewElement->Key.Type);
 		NewElement->Settings = InSettings;
 		AddElement(NewElement, Hierarchy->Get(Hierarchy->GetIndex(InParent)), true);
 
@@ -712,7 +712,7 @@ TArray<FRigElementKey> URigHierarchyController::ImportBones(const FReferenceSkel
 		TArray<FRigBoneElement*> AllBones = GetHierarchy()->GetElementsOfType<FRigBoneElement>(true);
 		for(FRigBoneElement* BoneElement : AllBones)
 		{
-			BoneNameMap.Add(BoneElement->GetName(), BoneElement->GetName());
+			BoneNameMap.Add(BoneElement->GetFName(), BoneElement->GetFName());
 		}
 
 		for (int32 Index = 0; Index < BoneInfos.Num(); ++Index)
@@ -793,7 +793,7 @@ TArray<FRigElementKey> URigHierarchyController::ImportBones(const FReferenceSkel
 		TArray<FRigBoneElement*> AllBones = GetHierarchy()->GetElementsOfType<FRigBoneElement>(true);
 		for(FRigBoneElement* BoneElement : AllBones)
         {
-            if (!BoneNameToIndexInSkeleton.Contains(BoneElement->GetName()))
+            if (!BoneNameToIndexInSkeleton.Contains(BoneElement->GetFName()))
 			{
 				if (BoneElement->BoneType == ERigBoneType::Imported)
 				{
@@ -1280,7 +1280,7 @@ TArray<FRigElementKey> URigHierarchyController::ImportFromText(FString InContent
 			}
 		}
 		
-		NewElement->Key.Name = Hierarchy->GetSafeNewName(NewElement->Key.Name.ToString(), NewElement->Key.Type);
+		NewElement->Key.Name = Hierarchy->GetSafeNewName(NewElement->Key.Name, NewElement->Key.Type);
 		AddElement(NewElement, nullptr, true);
 
 		KeyMap.FindOrAdd(PerElementData.Key) = NewElement->Key;
@@ -1548,7 +1548,7 @@ TArray<FString> URigHierarchyController::GetAddBonePythonCommands(FRigBoneElemen
 	
 	// AddBone(FName InName, FRigElementKey InParent, FTransform InTransform, bool bTransformInGlobal = true, ERigBoneType InBoneType = ERigBoneType::User, bool bSetupUndo = false);
 	Commands.Add(FString::Printf(TEXT("hierarchy_controller.add_bone('%s', %s, %s, False, %s)"),
-		*Bone->GetName().ToString(),
+		*Bone->GetName(),
 		*ParentKeyStr,
 		*TransformStr,
 		*RigVMPythonUtils::EnumValueToPythonString<ERigBoneType>((int64)Bone->BoneType)
@@ -1569,7 +1569,7 @@ TArray<FString> URigHierarchyController::GetAddNullPythonCommands(FRigNullElemen
 		
 	// AddNull(FName InName, FRigElementKey InParent, FTransform InTransform, bool bTransformInGlobal = true, bool bSetupUndo = false);
 	return {FString::Printf(TEXT("hierarchy_controller.add_null('%s', %s, %s, False)"),
-		*Null->GetName().ToString(),
+		*Null->GetName(),
 		*ParentKeyStr,
 		*TransformStr)};
 }
@@ -1588,7 +1588,7 @@ TArray<FString> URigHierarchyController::GetAddControlPythonCommands(FRigControl
 	FRigControlSettings& Settings = Control->Settings;
 	FString SettingsStr;
 	{
-		FString ControlNamePythonized = RigVMPythonUtils::PythonizeName(Control->GetName().ToString());
+		FString ControlNamePythonized = RigVMPythonUtils::PythonizeName(Control->GetName());
 		SettingsStr = FString::Printf(TEXT("control_settings_%s"),
 			*ControlNamePythonized);
 			
@@ -1600,7 +1600,7 @@ TArray<FString> URigHierarchyController::GetAddControlPythonCommands(FRigControl
 	
 	// AddControl(FName InName, FRigElementKey InParent, FRigControlSettings InSettings, FRigControlValue InValue, bool bSetupUndo = true);
 	Commands.Add(FString::Printf(TEXT("hierarchy_controller.add_control('%s', %s, %s, %s)"),
-		*Control->GetName().ToString(),
+		*Control->GetName(),
 		*ParentKeyStr,
 		*SettingsStr,
 		*ValueStr));
@@ -1618,7 +1618,7 @@ TArray<FString> URigHierarchyController::GetAddCurvePythonCommands(FRigCurveElem
 {
 	// FRigElementKey AddCurve(FName InName, float InValue = 0.f, bool bSetupUndo = true);
 	return {FString::Printf(TEXT("hierarchy_controller.add_curve('%s', %f)"),
-		*Curve->GetName().ToString(),
+		*Curve->GetName(),
 		Hierarchy->GetCurveValue(Curve))};
 }
 
@@ -1633,7 +1633,7 @@ TArray<FString> URigHierarchyController::GetAddRigidBodyPythonCommands(FRigRigid
 		ParentKeyStr = RigidBody->ParentElement->GetKey().ToPythonString();
 	}
 
-	FString RigidBodyNamePythonized = RigVMPythonUtils::PythonizeName(RigidBody->GetName().ToString());
+	FString RigidBodyNamePythonized = RigVMPythonUtils::PythonizeName(RigidBody->GetName());
 	FRigRigidBodySettings& Settings = RigidBody->Settings;
 	FString SettingsStr;
 	{
@@ -1650,7 +1650,7 @@ TArray<FString> URigHierarchyController::GetAddRigidBodyPythonCommands(FRigRigid
 	
 	// FRigElementKey AddRigidBody(FName InName, FRigElementKey InParent, FRigRigidBodySettings InSettings, FTransform InLocalTransform, bool bSetupUndo = false);
 	Commands.Add(FString::Printf(TEXT("hierarchy_controller.add_rigid_body('%s', %s, %s, %s)"),
-		*RigidBody->GetName().ToString(),
+		*RigidBody->GetName(),
 		*ParentKeyStr,
 		*SettingsStr,
 		*TransformStr));
@@ -1672,7 +1672,7 @@ TArray<FString> URigHierarchyController::GetAddConnectorPythonCommands(FRigConne
 	FRigConnectorSettings& Settings = Connector->Settings;
 	FString SettingsStr;
 	{
-		FString ConnectorNamePythonized = RigVMPythonUtils::PythonizeName(Connector->GetName().ToString());
+		FString ConnectorNamePythonized = RigVMPythonUtils::PythonizeName(Connector->GetName());
 		SettingsStr = FString::Printf(TEXT("connector_settings_%s"),
 			*ConnectorNamePythonized);
 			
@@ -1681,7 +1681,7 @@ TArray<FString> URigHierarchyController::GetAddConnectorPythonCommands(FRigConne
 
 	// AddConnector(FName InName, FRigElementKey InParent, FTransform InTransform, bool bTransformInGlobal = true, FRigConnectorSettings InSettings = FRigConnectorSettings(), bool bSetupUndo = false);
 	Commands.Add(FString::Printf(TEXT("hierarchy_controller.add_connector('%s', %s, %s, False, %s)"),
-		*Connector->GetName().ToString(),
+		*Connector->GetName(),
 		*ParentKeyStr,
 		*TransformStr,
 		*SettingsStr
@@ -1764,7 +1764,7 @@ int32 URigHierarchyController::AddElement(FRigBaseElement* InElementToAdd, FRigB
 {
 	ensure(IsValid());
 
-	InElementToAdd->NameString = InElementToAdd->Key.Name.ToString();
+	InElementToAdd->NameString.Reset();
 	InElementToAdd->SubIndex = Hierarchy->Num(InElementToAdd->Key.Type);
 	InElementToAdd->Index = Hierarchy->Elements.Add(InElementToAdd);
 	Hierarchy->ElementsPerType[URigHierarchy::RigElementTypeToFlatIndex(InElementToAdd->GetKey().Type)].Add(InElementToAdd);
@@ -2114,7 +2114,7 @@ bool URigHierarchyController::RenameElement(FRigBaseElement* InElement, const FN
 		return false;
 	}
 
-	if (InElement->GetName().IsEqual(InName, ENameCase::CaseSensitive))
+	if (InElement->GetFName().IsEqual(InName, ENameCase::CaseSensitive))
 	{
 		return false;
 	}
@@ -2135,8 +2135,8 @@ bool URigHierarchyController::RenameElement(FRigBaseElement* InElement, const FN
 		TemporaryMap.Remove(OldKey);
    
 		TGuardValue<TMap<FRigElementKey, int32>> MapGuard(Hierarchy->IndexLookup, TemporaryMap);
-		InElement->Key.Name = Hierarchy->GetSafeNewName(InName.ToString(), InElement->GetType());
-		InElement->NameString = InElement->Key.Name.ToString();
+		InElement->Key.Name = Hierarchy->GetSafeNewName(InName, InElement->GetType());
+		InElement->NameString.Reset();
 	}
 	
 	const FRigElementKey NewKey = InElement->GetKey();
@@ -2278,8 +2278,7 @@ FName URigHierarchyController::SetDisplayName(FRigControlElement* InControlEleme
 		ParentElementKey = ParentElement->GetKey();
 	}
 
-	const FString DesiredDisplayName = InDisplayName.IsNone() ? FString() : InDisplayName.ToString();
-	const FName DisplayName = Hierarchy->GetSafeNewDisplayName(ParentElementKey, DesiredDisplayName);
+	const FName DisplayName = Hierarchy->GetSafeNewDisplayName(ParentElementKey, InDisplayName);
 	InControlElement->Settings.DisplayName = DisplayName;
 
 	Hierarchy->IncrementTopologyVersion();
@@ -2492,7 +2491,7 @@ bool URigHierarchyController::AddParent(FRigBaseElement* InChild, FRigBaseElemen
 				ControlElement->Settings.DisplayName =
 					Hierarchy->GetSafeNewDisplayName(
 						InParent->GetKey(),
-						ControlElement->Settings.DisplayName.ToString());
+						ControlElement->Settings.DisplayName);
 			}
 		}
 		
diff --git a/Engine/Plugins/Animation/ControlRig/Source/ControlRig/Private/Rigs/RigHierarchyDefines.cpp b/Engine/Plugins/Animation/ControlRig/Source/ControlRig/Private/Rigs/RigHierarchyDefines.cpp
--- a/Engine/Plugins/Animation/ControlRig/Source/ControlRig/Private/Rigs/RigHierarchyDefines.cpp
+++ b/Engine/Plugins/Animation/ControlRig/Source/ControlRig/Private/Rigs/RigHierarchyDefines.cpp
@@ -173,14 +173,14 @@ FRigElementKeyCollection FRigElementKeyCollection::MakeFromName(
 		return InHierarchy->GetKeysByPredicate([WildcardString, InElementTypes](const FRigBaseElement& InElement) -> bool
 		{
 			return InElement.IsTypeOf(static_cast<ERigElementType>(InElementTypes)) &&
-				   WildcardString.IsMatch(InElement.GetNameString());
+				   WildcardString.IsMatch(InElement.GetName());
 		}, bTraverse);
 	}
 	
 	return InHierarchy->GetKeysByPredicate([PartialNameString, InElementTypes](const FRigBaseElement& InElement) -> bool
 	{
 		return InElement.IsTypeOf(static_cast<ERigElementType>(InElementTypes)) &&
-			   InElement.GetNameString().Contains(PartialNameString);
+			   InElement.GetName().Contains(PartialNameString);
 	}, bTraverse);
 }
 
diff --git a/Engine/Plugins/Animation/ControlRig/Source/ControlRig/Private/Rigs/RigHierarchyElements.cpp b/Engine/Plugins/Animation/ControlRig/Source/ControlRig/Private/Rigs/RigHierarchyElements.cpp
--- a/Engine/Plugins/Animation/ControlRig/Source/ControlRig/Private/Rigs/RigHierarchyElements.cpp
+++ b/Engine/Plugins/Animation/ControlRig/Source/ControlRig/Private/Rigs/RigHierarchyElements.cpp
@@ -138,7 +138,7 @@ void FRigBaseElement::Load(FArchive& Ar, URigHierarchy* Hierarchy, ESerializatio
 		ensure(LoadedKey.Type == Key.Type);
 		Key = LoadedKey;
 
-		NameString = Key.Name.ToString();
+		NameString.Reset();
 
 		RemoveAllMetadata();
 		
@@ -200,18 +200,18 @@ bool FRigBaseElement::RemoveAllMetadata()
 {
 	if(!Metadata.IsEmpty())
 	{
-		TArray<FName> Names;
-		Names.Reserve(Metadata.Num());
+		TArray<FName> MetaDataNames;
+		MetaDataNames.Reserve(Metadata.Num());
 		for(FRigBaseMetadata* Md : Metadata)
 		{
-			Names.Add(Md->GetName());
+			MetaDataNames.Add(Md->GetName());
 			FRigBaseMetadata::DestroyMetadata(&Md);
 		}
 		Metadata.Reset();
 		MetadataNameToIndex.Reset();
-		for(const FName& Name: Names)
+		for(const FName& MetaDataName: MetaDataNames)
 		{
-			NotifyMetadataChanged(Name);
+			NotifyMetadataChanged(MetaDataName);
 		}
 		return true;
 	}
diff --git a/Engine/Plugins/Animation/ControlRig/Source/ControlRig/Private/Sequencer/MovieSceneControlRigParameterSection.cpp b/Engine/Plugins/Animation/ControlRig/Source/ControlRig/Private/Sequencer/MovieSceneControlRigParameterSection.cpp
--- a/Engine/Plugins/Animation/ControlRig/Source/ControlRig/Private/Sequencer/MovieSceneControlRigParameterSection.cpp
+++ b/Engine/Plugins/Animation/ControlRig/Source/ControlRig/Private/Sequencer/MovieSceneControlRigParameterSection.cpp
@@ -573,7 +573,7 @@ struct FParameterTransformChannelEditorData
 				auto GetTranslationFromTransform = [&](const FVector& InTranslation)
 				{
 					// switch translation to constraint space if needed
-					const uint32 ControlHash = UTransformableControlHandle::ComputeHash(ControlRig, ControlElement->GetName());
+					const uint32 ControlHash = UTransformableControlHandle::ComputeHash(ControlRig, ControlElement->GetFName());
 					TOptional<FTransform> ConstraintSpaceTransform = FTransformConstraintUtils::GetRelativeTransform(ControlRig->GetWorld(), ControlHash);
 					if (ConstraintSpaceTransform)
 					{
@@ -627,7 +627,7 @@ struct FParameterTransformChannelEditorData
 				if (ControlElement->Settings.ControlType == ERigControlType::EulerTransform)
 				{
 					// switch rotation to constraint space if needed
-					const uint32 ControlHash = UTransformableControlHandle::ComputeHash(ControlRig, ControlElement->GetName());
+					const uint32 ControlHash = UTransformableControlHandle::ComputeHash(ControlRig, ControlElement->GetFName());
 					TOptional<FTransform> ConstraintSpaceTransform = FTransformConstraintUtils::GetRelativeTransform(ControlRig->GetWorld(), ControlHash);
 					if (ConstraintSpaceTransform)
 					{
@@ -653,7 +653,7 @@ struct FParameterTransformChannelEditorData
 				auto GetScaleFromTransform = [&](const FVector& InScale3D)
 				{
 					// switch scale to constraint space if needed
-					const uint32 ControlHash = UTransformableControlHandle::ComputeHash(ControlRig, ControlElement->GetName());
+					const uint32 ControlHash = UTransformableControlHandle::ComputeHash(ControlRig, ControlElement->GetFName());
 					TOptional<FTransform> ConstraintSpaceTransform = FTransformConstraintUtils::GetRelativeTransform(ControlRig->GetWorld(), ControlHash);
 					if (ConstraintSpaceTransform)
 					{
@@ -950,7 +950,7 @@ void UMovieSceneControlRigParameterSection::SetBlendType(EMovieSceneBlendType In
 
 				case ERigControlType::Scale:
 				{
-					ChannelInfo = ControlChannelMap.Find(ControlElement->GetName());
+					ChannelInfo = ControlChannelMap.Find(ControlElement->GetFName());
 					if (ChannelInfo)
 					{
 						if (InBlendType == EMovieSceneBlendType::Absolute)
@@ -971,7 +971,7 @@ void UMovieSceneControlRigParameterSection::SetBlendType(EMovieSceneBlendType In
 				case ERigControlType::Transform:
 				case ERigControlType::EulerTransform:
 				{
-					ChannelInfo = ControlChannelMap.Find(ControlElement->GetName());
+					ChannelInfo = ControlChannelMap.Find(ControlElement->GetFName());
 					if (ChannelInfo)
 					{
 						if (InBlendType == EMovieSceneBlendType::Absolute)
@@ -1575,7 +1575,7 @@ bool UMovieSceneControlRigParameterSection::IsDifferentThanLastControlsUsedToRec
 	for (int32 Index = 0; Index < LastControlsUsedToReconstruct.Num(); ++Index)
 	{
 		//for the channel proxy we really just care about name and type, and if any are nullptr's
-		if (LastControlsUsedToReconstruct[Index].Key != NewControls[Index]->GetName() ||
+		if (LastControlsUsedToReconstruct[Index].Key != NewControls[Index]->GetFName() ||
 			LastControlsUsedToReconstruct[Index].Value != NewControls[Index]->Settings.ControlType)
 		{
 			return true;
@@ -1591,7 +1591,7 @@ void UMovieSceneControlRigParameterSection::StoreLastControlsUsedToReconstruct(c
 	{
 		if (NewControls[Index])
 		{
-			LastControlsUsedToReconstruct[Index].Key = NewControls[Index]->GetName();
+			LastControlsUsedToReconstruct[Index].Key = NewControls[Index]->GetFName();
 			LastControlsUsedToReconstruct[Index].Value = NewControls[Index]->Settings.ControlType;
 		}
 	}
@@ -1717,7 +1717,7 @@ void UMovieSceneControlRigParameterSection::ReconstructChannelProxy()
 			{
 				if(const FRigControlElement* ParentControlElement = Cast<FRigControlElement>(Hierarchy->GetFirstParent(ControlElement)))
 				{
-					ParentControlName = ParentControlElement->GetName();
+					ParentControlName = ParentControlElement->GetFName();
 					Group = FText::FromName(ParentControlElement->GetDisplayName());
 				}
 			}
@@ -1731,7 +1731,7 @@ void UMovieSceneControlRigParameterSection::ReconstructChannelProxy()
 				{
 					for (FScalarParameterNameAndCurve& Scalar : GetScalarParameterNamesAndCurves())
 					{
-						if (ControlElement->GetName() == Scalar.ParameterName)
+						if (ControlElement->GetFName() == Scalar.ParameterName)
 						{
 							if (Group.IsEmpty())
 							{
@@ -1764,7 +1764,7 @@ void UMovieSceneControlRigParameterSection::ReconstructChannelProxy()
 				{
 					for (FBoolParameterNameAndCurve& Bool : GetBoolParameterNamesAndCurves())
 					{
-						if (ControlElement->GetName() == Bool.ParameterName)
+						if (ControlElement->GetFName() == Bool.ParameterName)
 						{
 							if (Group.IsEmpty())
 							{
@@ -1801,7 +1801,7 @@ void UMovieSceneControlRigParameterSection::ReconstructChannelProxy()
 					{
 						for (FEnumParameterNameAndCurve& Enum : GetEnumParameterNamesAndCurves())
 						{
-							if (ControlElement->GetName() == Enum.ParameterName)
+							if (ControlElement->GetFName() == Enum.ParameterName)
 							{
 								if (Group.IsEmpty())
 								{
@@ -1835,7 +1835,7 @@ void UMovieSceneControlRigParameterSection::ReconstructChannelProxy()
 					{
 						for (FIntegerParameterNameAndCurve& Integer : GetIntegerParameterNamesAndCurves())
 						{
-							if (ControlElement->GetName() == Integer.ParameterName)
+							if (ControlElement->GetFName() == Integer.ParameterName)
 							{
 								if (Group.IsEmpty())
 								{
@@ -1872,7 +1872,7 @@ void UMovieSceneControlRigParameterSection::ReconstructChannelProxy()
 				{
 					for (FVector2DParameterNameAndCurves& Vector2D : GetVector2DParameterNamesAndCurves())
 					{
-						if (ControlElement->GetName() == Vector2D.ParameterName)
+						if (ControlElement->GetFName() == Vector2D.ParameterName)
 						{
 							if(Group.IsEmpty())
 							{
@@ -1908,7 +1908,7 @@ void UMovieSceneControlRigParameterSection::ReconstructChannelProxy()
 				{
 					for (FVectorParameterNameAndCurves& Vector : GetVectorParameterNamesAndCurves())
 					{
-						if (ControlElement->GetName() == Vector.ParameterName)
+						if (ControlElement->GetFName() == Vector.ParameterName)
 						{
 							if(Group.IsEmpty())
 							{
@@ -1987,9 +1987,9 @@ void UMovieSceneControlRigParameterSection::ReconstructChannelProxy()
 				{
 					for (FTransformParameterNameAndCurves& Transform : GetTransformParameterNamesAndCurves())
 					{
-						if (ControlElement->GetName() == Transform.ParameterName)
+						if (ControlElement->GetFName() == Transform.ParameterName)
 						{
-							const FName ControlName = ControlElement->GetName();
+							const FName ControlName = ControlElement->GetFName();
 							if(Group.IsEmpty())
 							{
 								ControlChannelMap.Add(Transform.ParameterName, FChannelMapInfo(ControlIndex, TotalIndex, FloatChannelIndex, INDEX_NONE, NAME_None, MaskIndex, CategoryIndex));
@@ -2088,7 +2088,7 @@ void UMovieSceneControlRigParameterSection::ReconstructChannelProxy()
 				{
 					for (FScalarParameterNameAndCurve& Scalar : GetScalarParameterNamesAndCurves())
 					{
-						if (ControlElement->GetName() == Scalar.ParameterName)
+						if (ControlElement->GetFName() == Scalar.ParameterName)
 						{
 							ControlChannelMap.Add(Scalar.ParameterName, FChannelMapInfo(ControlIndex, TotalIndex,FloatChannelIndex,INDEX_NONE, NAME_None,MaskIndex));
 							Channels.Add(Scalar.ParameterCurve);
@@ -2104,7 +2104,7 @@ void UMovieSceneControlRigParameterSection::ReconstructChannelProxy()
 				{
 					for (FBoolParameterNameAndCurve& Bool : GetBoolParameterNamesAndCurves())
 					{
-						if (ControlElement->GetName() == Bool.ParameterName)
+						if (ControlElement->GetFName() == Bool.ParameterName)
 						{
 							ControlChannelMap.Add(Bool.ParameterName, FChannelMapInfo(ControlIndex, TotalIndex, BoolChannelIndex, INDEX_NONE, NAME_None, MaskIndex));
 							Channels.Add(Bool.ParameterCurve);
@@ -2122,7 +2122,7 @@ void UMovieSceneControlRigParameterSection::ReconstructChannelProxy()
 					{
 						for (FEnumParameterNameAndCurve& Enum : GetEnumParameterNamesAndCurves())
 						{
-							if (ControlElement->GetName() == Enum.ParameterName)
+							if (ControlElement->GetFName() == Enum.ParameterName)
 							{
 								ControlChannelMap.Add(Enum.ParameterName, FChannelMapInfo(ControlIndex, TotalIndex, EnumChannelIndex, INDEX_NONE, NAME_None, MaskIndex));
 								Channels.Add(Enum.ParameterCurve);
@@ -2137,7 +2137,7 @@ void UMovieSceneControlRigParameterSection::ReconstructChannelProxy()
 					{
 						for (FIntegerParameterNameAndCurve& Integer : GetIntegerParameterNamesAndCurves())
 						{
-							if (ControlElement->GetName() == Integer.ParameterName)
+							if (ControlElement->GetFName() == Integer.ParameterName)
 							{
 								ControlChannelMap.Add(Integer.ParameterName, FChannelMapInfo(ControlIndex, TotalIndex, IntegerChannelIndex, INDEX_NONE, NAME_None, MaskIndex));
 								Channels.Add(Integer.ParameterCurve);
@@ -2154,7 +2154,7 @@ void UMovieSceneControlRigParameterSection::ReconstructChannelProxy()
 				{
 					for (FVector2DParameterNameAndCurves& Vector2D : GetVector2DParameterNamesAndCurves())
 					{
-						if (ControlElement->GetName() == Vector2D.ParameterName)
+						if (ControlElement->GetFName() == Vector2D.ParameterName)
 						{
 							ControlChannelMap.Add(Vector2D.ParameterName, FChannelMapInfo(ControlIndex, TotalIndex, FloatChannelIndex, INDEX_NONE, NAME_None, MaskIndex));
 							Channels.Add(Vector2D.XCurve);
@@ -2173,7 +2173,7 @@ void UMovieSceneControlRigParameterSection::ReconstructChannelProxy()
 				{
 					for (FVectorParameterNameAndCurves& Vector : GetVectorParameterNamesAndCurves())
 					{
-						if (ControlElement->GetName() == Vector.ParameterName)
+						if (ControlElement->GetFName() == Vector.ParameterName)
 						{
 							ControlChannelMap.Add(Vector.ParameterName, FChannelMapInfo(ControlIndex, TotalIndex, FloatChannelIndex, INDEX_NONE, NAME_None, MaskIndex));
 							bool bDoSpaceChannel = true;
@@ -2219,7 +2219,7 @@ void UMovieSceneControlRigParameterSection::ReconstructChannelProxy()
 				{
 					for (FTransformParameterNameAndCurves& Transform : GetTransformParameterNamesAndCurves())
 					{
-						if (ControlElement->GetName() == Transform.ParameterName)
+						if (ControlElement->GetFName() == Transform.ParameterName)
 						{
 							ControlChannelMap.Add(Transform.ParameterName, FChannelMapInfo(ControlIndex, TotalIndex,FloatChannelIndex, INDEX_NONE, NAME_None, MaskIndex));
 						
@@ -2446,7 +2446,7 @@ void UMovieSceneControlRigParameterSection::RecreateWithThisControlRig(UControlR
 		{
 			if (ControlElement->Settings.ControlType == ERigControlType::Float)
 			{
-				const FName TargetCurveName = UFKControlRig::GetControlTargetName(ControlElement->GetName(), ERigElementType::Curve);
+				const FName TargetCurveName = UFKControlRig::GetControlTargetName(ControlElement->GetFName(), ERigElementType::Curve);
 				const FRigElementKey CurveKey = FRigElementKey(TargetCurveName, ERigElementType::Curve);
 				// Ensure name is valid, and curve actually exists in the hierarchy,
 				// this means we could not be renaming some controls for which the curves do not exist anymore, which ties back to comment at the top op the function
@@ -2454,7 +2454,7 @@ void UMovieSceneControlRigParameterSection::RecreateWithThisControlRig(UControlR
 				if (TargetCurveName != NAME_None && Hierarchy->Find(CurveKey))
 				{
 					// Add mapping from old to new control naming scheme (previous was using uniform naming for both bones and curves)
-					CurveControlNameRemapping.Add(ControlElement->GetName(), UFKControlRig::GetControlName(TargetCurveName, ERigElementType::Bone));
+					CurveControlNameRemapping.Add(ControlElement->GetFName(), UFKControlRig::GetControlName(TargetCurveName, ERigElementType::Bone));
 				}
 			}
 		}
@@ -2473,7 +2473,7 @@ void UMovieSceneControlRigParameterSection::RecreateWithThisControlRig(UControlR
 			RenameParameterName(PreviousName, ControlElement->GetKey().Name);
 		}
 
-		if (const FName* OldCurveControlName = CurveControlNameRemapping.Find(ControlElement->GetName()))
+		if (const FName* OldCurveControlName = CurveControlNameRemapping.Find(ControlElement->GetFName()))
 		{
 			RenameParameterName(*OldCurveControlName, ControlElement->GetKey().Name);
 		}
@@ -2488,7 +2488,7 @@ void UMovieSceneControlRigParameterSection::RecreateWithThisControlRig(UControlR
 				//or use IntialValue?
 				DefaultValue = ControlRig->GetControlValue(ControlElement, ERigControlValueType::Current).Get<float>();
 			}
-			AddScalarParameter(ControlElement->GetName(), DefaultValue, false);
+			AddScalarParameter(ControlElement->GetFName(), DefaultValue, false);
 			break;
 		}
 		case ERigControlType::Bool:
@@ -2498,7 +2498,7 @@ void UMovieSceneControlRigParameterSection::RecreateWithThisControlRig(UControlR
 			if (bSetDefault)
 			{
 				DefaultValue = ControlRig->GetControlValue(ControlElement, ERigControlValueType::Current).Get<bool>();
-				AddBoolParameter(ControlElement->GetName(), DefaultValue, false);
+				AddBoolParameter(ControlElement->GetFName(), DefaultValue, false);
 
 			}
 			break;
@@ -2512,7 +2512,7 @@ void UMovieSceneControlRigParameterSection::RecreateWithThisControlRig(UControlR
 				if (bSetDefault)
 				{
 					DefaultValue = (uint8)ControlRig->GetControlValue(ControlElement, ERigControlValueType::Current).Get<int32>();
-					AddEnumParameter(ControlElement->GetName(), ControlElement->Settings.ControlEnum, DefaultValue, false);
+					AddEnumParameter(ControlElement->GetFName(), ControlElement->Settings.ControlEnum, DefaultValue, false);
 				}
 			}
 			else
@@ -2522,7 +2522,7 @@ void UMovieSceneControlRigParameterSection::RecreateWithThisControlRig(UControlR
 				if (bSetDefault)
 				{
 					DefaultValue = ControlRig->GetControlValue(ControlElement, ERigControlValueType::Current).Get<int32>();
-					AddIntegerParameter(ControlElement->GetName(), DefaultValue, false);
+					AddIntegerParameter(ControlElement->GetFName(), DefaultValue, false);
 				}
 			}
 			break;
@@ -2536,7 +2536,7 @@ void UMovieSceneControlRigParameterSection::RecreateWithThisControlRig(UControlR
 				const FVector3f TempValue = ControlRig->GetControlValue(ControlElement, ERigControlValueType::Current).Get<FVector3f>();
 				DefaultValue = FVector2D(TempValue.X, TempValue.Y);
 			}
-			AddVector2DParameter(ControlElement->GetName(), DefaultValue, false);
+			AddVector2DParameter(ControlElement->GetFName(), DefaultValue, false);
 			break;
 		}
 
@@ -2550,7 +2550,7 @@ void UMovieSceneControlRigParameterSection::RecreateWithThisControlRig(UControlR
 				//or use IntialValue?
 				DefaultValue = (FVector)ControlRig->GetControlValue(ControlElement, ERigControlValueType::Current).Get<FVector3f>();
 			}
-			AddVectorParameter(ControlElement->GetName(), DefaultValue, false);
+			AddVectorParameter(ControlElement->GetFName(), DefaultValue, false);
 			//mz todo specify rotator special so we can do quat interps
 			break;
 		}
@@ -2583,7 +2583,7 @@ void UMovieSceneControlRigParameterSection::RecreateWithThisControlRig(UControlR
 					DefaultValue = FEulerTransform(NoScale.Rotation.Rotator(), NoScale.Location, FVector::OneVector);
 				}
 			}
-			AddTransformParameter(ControlElement->GetName(), DefaultValue, false);
+			AddTransformParameter(ControlElement->GetFName(), DefaultValue, false);
 			break;
 		}
 
@@ -2882,7 +2882,7 @@ void UMovieSceneControlRigParameterSection::RecordControlRigKey(FFrameNumber Fra
 			{
 				continue;
 			}
-			FChannelMapInfo* pChannelIndex = ControlChannelMap.Find(ControlElement->GetName());
+			FChannelMapInfo* pChannelIndex = ControlChannelMap.Find(ControlElement->GetFName());
 			if (pChannelIndex == nullptr)
 			{
 				continue;
diff --git a/Engine/Plugins/Animation/ControlRig/Source/ControlRig/Private/Sequencer/MovieSceneControlRigParameterTemplate.cpp b/Engine/Plugins/Animation/ControlRig/Source/ControlRig/Private/Sequencer/MovieSceneControlRigParameterTemplate.cpp
--- a/Engine/Plugins/Animation/ControlRig/Source/ControlRig/Private/Sequencer/MovieSceneControlRigParameterTemplate.cpp
+++ b/Engine/Plugins/Animation/ControlRig/Source/ControlRig/Private/Sequencer/MovieSceneControlRigParameterTemplate.cpp
@@ -905,25 +905,25 @@ struct FControlRigParameterPreAnimatedTokenProducer : IMovieScenePreAnimatedToke
 					case ERigControlType::Bool:
 					{
 						const bool Val = ControlRig->GetControlValue(ControlElement, ERigControlValueType::Current).Get<bool>();
-						Token.BoolValues.Add(TNameAndValue<bool>{ ControlElement->GetName(), Val });
+						Token.BoolValues.Add(TNameAndValue<bool>{ ControlElement->GetFName(), Val });
 						break;
 					}
 					case ERigControlType::Integer:
 					{
 						const int32 Val = ControlRig->GetControlValue(ControlElement, ERigControlValueType::Current).Get<int32>();
-						Token.IntegerValues.Add(TNameAndValue<int32>{ ControlElement->GetName(), Val });
+						Token.IntegerValues.Add(TNameAndValue<int32>{ ControlElement->GetFName(), Val });
 						break;
 					}
 					case ERigControlType::Float:
 					{
 						const float Val = ControlRig->GetControlValue(ControlElement, ERigControlValueType::Current).Get<float>();
-						Token.ScalarValues.Add(TNameAndValue<float>{ ControlElement->GetName(), Val });
+						Token.ScalarValues.Add(TNameAndValue<float>{ ControlElement->GetFName(), Val });
 						break;
 					}
 					case ERigControlType::Vector2D:
 					{
 						const FVector3f Val = ControlRig->GetControlValue(ControlElement, ERigControlValueType::Current).Get<FVector3f>();
-						Token.Vector2DValues.Add(TNameAndValue<FVector2D>{ ControlElement->GetName(), FVector2D(Val.X, Val.Y) });
+						Token.Vector2DValues.Add(TNameAndValue<FVector2D>{ ControlElement->GetFName(), FVector2D(Val.X, Val.Y) });
 						break;
 					}
 					case ERigControlType::Position:
@@ -936,14 +936,14 @@ struct FControlRigParameterPreAnimatedTokenProducer : IMovieScenePreAnimatedToke
 						//FRigElementKey DefaultParent = RigHierarchy->GetFirstParent(ControlElement->GetKey());
 						SpaceValue.ControlRigElement = ControlElement->GetKey();
 						SpaceValue.SpaceType = EMovieSceneControlRigSpaceType::Parent;
-						Token.SpaceValues.Add(FControlSpaceAndValue(ControlElement->GetName(), SpaceValue));
+						Token.SpaceValues.Add(FControlSpaceAndValue(ControlElement->GetFName(), SpaceValue));
 						FVector3f Val = ControlRig->GetControlValue(ControlElement, ERigControlValueType::Current).Get<FVector3f>();
 						if (ControlElement->Settings.ControlType == ERigControlType::Rotator)
 						{
 							FVector Vector = RigHierarchy->GetControlSpecifiedEulerAngle(ControlElement);
 							Val = FVector3f(Vector.X, Vector.Y, Vector.Z);
 						}
-						Token.VectorValues.Add(TNameAndValue<FVector>{ ControlElement->GetName(), (FVector)Val });
+						Token.VectorValues.Add(TNameAndValue<FVector>{ ControlElement->GetFName(), (FVector)Val });
 						//mz todo specify rotator special so we can do quat interps
 						break;
 					}
@@ -955,12 +955,12 @@ struct FControlRigParameterPreAnimatedTokenProducer : IMovieScenePreAnimatedToke
 						//FRigElementKey DefaultParent = RigHierarchy->GetFirstParent(ControlElement->GetKey());
 						SpaceValue.ControlRigElement = ControlElement->GetKey();
 						SpaceValue.SpaceType = EMovieSceneControlRigSpaceType::Parent;
-						Token.SpaceValues.Add(FControlSpaceAndValue(ControlElement->GetName(), SpaceValue));
+						Token.SpaceValues.Add(FControlSpaceAndValue(ControlElement->GetFName(), SpaceValue));
 						const FTransform Val = ControlRig->GetControlValue(ControlElement, ERigControlValueType::Current).Get<FRigControlValue::FTransform_Float>().ToTransform();
 						FEulerTransform EulerTransform(Val);
 						FVector Vector = RigHierarchy->GetControlSpecifiedEulerAngle(ControlElement);
 						EulerTransform.Rotation = FRotator(Vector.Y, Vector.Z, Vector.X);
-						Token.TransformValues.Add(TNameAndValue<FEulerTransform>{ ControlElement->GetName(), EulerTransform });
+						Token.TransformValues.Add(TNameAndValue<FEulerTransform>{ ControlElement->GetFName(), EulerTransform });
 						break;
 					}
 					case ERigControlType::TransformNoScale:
@@ -971,7 +971,7 @@ struct FControlRigParameterPreAnimatedTokenProducer : IMovieScenePreAnimatedToke
 						FEulerTransform EulerTransform(NoScale.ToFTransform());
 						FVector Vector = RigHierarchy->GetControlSpecifiedEulerAngle(ControlElement);
 						EulerTransform.Rotation = FRotator(Vector.Y, Vector.Z, Vector.X);
-						Token.TransformValues.Add(TNameAndValue<FEulerTransform>{ ControlElement->GetName(), EulerTransform });
+						Token.TransformValues.Add(TNameAndValue<FEulerTransform>{ ControlElement->GetFName(), EulerTransform });
 						break;
 					}
 					case ERigControlType::EulerTransform:
@@ -981,7 +981,7 @@ struct FControlRigParameterPreAnimatedTokenProducer : IMovieScenePreAnimatedToke
 							->GetControlValue(ControlElement, ERigControlValueType::Current).Get<FRigControlValue::FEulerTransform_Float>().ToTransform();
 						FVector Vector = RigHierarchy->GetControlSpecifiedEulerAngle(ControlElement);
 						EulerTransform.Rotation = FRotator(Vector.Y, Vector.Z, Vector.X);
-						Token.TransformValues.Add(TNameAndValue<FEulerTransform>{ ControlElement->GetName(), EulerTransform });
+						Token.TransformValues.Add(TNameAndValue<FEulerTransform>{ ControlElement->GetFName(), EulerTransform });
 						break;
 					}
 				}
diff --git a/Engine/Plugins/Animation/ControlRig/Source/ControlRig/Private/Sequencer/MovieSceneControlRigParameterTrack.cpp b/Engine/Plugins/Animation/ControlRig/Source/ControlRig/Private/Sequencer/MovieSceneControlRigParameterTrack.cpp
--- a/Engine/Plugins/Animation/ControlRig/Source/ControlRig/Private/Sequencer/MovieSceneControlRigParameterTrack.cpp
+++ b/Engine/Plugins/Animation/ControlRig/Source/ControlRig/Private/Sequencer/MovieSceneControlRigParameterTrack.cpp
@@ -683,7 +683,7 @@ TOptional<EEulerRotationOrder> UMovieSceneControlRigParameterTrack::GetControlRo
 	TOptional<EEulerRotationOrder> Order;
 	if (bCurrent)
 	{
-		const FControlRotationOrder* RotationOrder = ControlsRotationOrder.Find(ControlElement->GetName());
+		const FControlRotationOrder* RotationOrder = ControlsRotationOrder.Find(ControlElement->GetFName());
 		if (RotationOrder)
 		{
 			Order = RotationOrder->RotationOrder;
@@ -717,7 +717,7 @@ TArray<FName> UMovieSceneControlRigParameterTrack::GetControlsWithDifferentRotat
 			TOptional<EEulerRotationOrder> Setting = GetControlRotationOrder(ControlElement, false);
 			if (Current != Setting)
 			{
-				Names.Add(ControlElement->GetName());
+				Names.Add(ControlElement->GetFName());
 			}
 
 		}
diff --git a/Engine/Plugins/Animation/ControlRig/Source/ControlRig/Private/Tools/ControlRigPose.cpp b/Engine/Plugins/Animation/ControlRig/Source/ControlRig/Private/Tools/ControlRigPose.cpp
--- a/Engine/Plugins/Animation/ControlRig/Source/ControlRig/Private/Tools/ControlRigPose.cpp
+++ b/Engine/Plugins/Animation/ControlRig/Source/ControlRig/Private/Tools/ControlRigPose.cpp
@@ -19,7 +19,7 @@ void FControlRigControlPose::SavePose(UControlRig* ControlRig, bool bUseAll)
 	
 	for (FRigControlElement* ControlElement : CurrentControls)
 	{
-		if (ControlRig->GetHierarchy()->IsAnimatable(ControlElement) && (bUseAll || ControlRig->IsControlSelected(ControlElement->GetName())))
+		if (ControlRig->GetHierarchy()->IsAnimatable(ControlElement) && (bUseAll || ControlRig->IsControlSelected(ControlElement->GetFName())))
 		{
 			//we store poses in default parent space so if not in that space we need to compensate it
 			bool bHasNonDefaultParent = false;
@@ -89,12 +89,12 @@ void FControlRigControlPose::PastePoseInternal(UControlRig* ControlRig, bool bDo
 
 	for (FRigControlElement* ControlElement : SortedControls)
 	{
-		if (!ControlRig->IsControlSelected(ControlElement->GetName()))
+		if (!ControlRig->IsControlSelected(ControlElement->GetFName()))
 		{
 			continue;
 		}
 		
-		FRigControlCopy* CopyRigControl = MirrorTable.GetControl(*this, ControlElement->GetName());
+		FRigControlCopy* CopyRigControl = MirrorTable.GetControl(*this, ControlElement->GetFName());
 		if (CopyRigControl)
 		{
 			//if not in default parent space we need to move it to default parent space first and then reset the global transforms
@@ -120,11 +120,11 @@ void FControlRigControlPose::PastePoseInternal(UControlRig* ControlRig, bool bDo
 				{
 					if (bDoLocal) // -V547  
 					{
-						ControlRig->SetControlLocalTransform(ControlElement->GetName(), CopyRigControl->LocalTransform, true,Context,bSetupUndo, true/* bFixEulerFlips*/);
+						ControlRig->SetControlLocalTransform(ControlElement->GetFName(), CopyRigControl->LocalTransform, true,Context,bSetupUndo, true/* bFixEulerFlips*/);
 					}
 					else
 					{
-						ControlRig->SetControlGlobalTransform(ControlElement->GetName(), CopyRigControl->GlobalTransform, true, Context, bSetupUndo, false /*bPrintPython*/, true/* bFixEulerFlips*/);
+						ControlRig->SetControlGlobalTransform(ControlElement->GetFName(), CopyRigControl->GlobalTransform, true, Context, bSetupUndo, false /*bPrintPython*/, true/* bFixEulerFlips*/);
 					}
 				}
 				else
@@ -133,32 +133,32 @@ void FControlRigControlPose::PastePoseInternal(UControlRig* ControlRig, bool bDo
 					FTransform LocalTransform;
 					bool bIsMatched = MirrorTable.IsMatched(CopyRigControl->Name);
 					MirrorTable.GetMirrorTransform(*CopyRigControl, bDoLocal,bIsMatched,GlobalTransform,LocalTransform);
-					SetControlMirrorTransform(bDoLocal,ControlRig, ControlElement->GetName(), bIsMatched, GlobalTransform,LocalTransform,true,Context,bSetupUndo);
+					SetControlMirrorTransform(bDoLocal,ControlRig, ControlElement->GetFName(), bIsMatched, GlobalTransform,LocalTransform,true,Context,bSetupUndo);
 				}				
 				break;
 			}
 			case ERigControlType::Float:
 			{
 				float Val = CopyRigControl->Value.Get<float>();
-				ControlRig->SetControlValue<float>(ControlElement->GetName(), Val, true, Context,bSetupUndo);
+				ControlRig->SetControlValue<float>(ControlElement->GetFName(), Val, true, Context,bSetupUndo);
 				break;
 			}
 			case ERigControlType::Bool:
 			{
 				bool Val = CopyRigControl->Value.Get<bool>();
-				ControlRig->SetControlValue<bool>(ControlElement->GetName(), Val, true, Context,bSetupUndo);
+				ControlRig->SetControlValue<bool>(ControlElement->GetFName(), Val, true, Context,bSetupUndo);
 				break;
 			}
 			case ERigControlType::Integer:
 			{
 				int32 Val = CopyRigControl->Value.Get<int32>();
-				ControlRig->SetControlValue<int32>(ControlElement->GetName(), Val, true, Context,bSetupUndo);
+				ControlRig->SetControlValue<int32>(ControlElement->GetFName(), Val, true, Context,bSetupUndo);
 				break;
 			}
 			case ERigControlType::Vector2D:
 			{
 				FVector3f Val = CopyRigControl->Value.Get<FVector3f>();
-				ControlRig->SetControlValue<FVector3f>(ControlElement->GetName(), Val, true, Context,bSetupUndo);
+				ControlRig->SetControlValue<FVector3f>(ControlElement->GetFName(), Val, true, Context,bSetupUndo);
 				break;
 			}
 			default:
@@ -170,7 +170,7 @@ void FControlRigControlPose::PastePoseInternal(UControlRig* ControlRig, bool bDo
 			{
 				FTransform GlobalTransform = ControlRig->GetHierarchy()->GetGlobalTransform(ControlElement->GetKey());
 				ControlRig->GetHierarchy()->SwitchToParent(ControlElement->GetKey(), SpaceKey);
-				ControlRig->SetControlGlobalTransform(ControlElement->GetName(), GlobalTransform, true, Context, bSetupUndo, false /*bPrintPython*/, true/* bFixEulerFlips*/);
+				ControlRig->SetControlGlobalTransform(ControlElement->GetFName(), GlobalTransform, true, Context, bSetupUndo, false /*bPrintPython*/, true/* bFixEulerFlips*/);
 			}
 		}
 	}
@@ -200,11 +200,11 @@ void FControlRigControlPose::BlendWithInitialPoses(FControlRigControlPose& Initi
 	const bool bSetupUndo = false;
 	for (FRigControlElement* ControlElement : SortedControls)
 	{
-		if (!ControlRig->IsControlSelected(ControlElement->GetName()))
+		if (!ControlRig->IsControlSelected(ControlElement->GetFName()))
 		{
 			continue;
 		}
-		FRigControlCopy* CopyRigControl = MirrorTable.GetControl(*this, ControlElement->GetName());
+		FRigControlCopy* CopyRigControl = MirrorTable.GetControl(*this, ControlElement->GetFName());
 		if (CopyRigControl)
 		{
 			FRigControlCopy* InitialFound = nullptr;
@@ -241,7 +241,7 @@ void FControlRigControlPose::BlendWithInitialPoses(FControlRigControlPose& Initi
 							Rotation = FQuat::Slerp(InitialVal.GetRotation(), Val.GetRotation(), BlendValue); //doing slerp here not fast lerp, can be slow this is for content creation
 							Scale = FMath::Lerp(InitialVal.GetScale3D(), Val.GetScale3D(), BlendValue);
 							Val = FTransform(Rotation, Translation, Scale);
-							ControlRig->SetControlLocalTransform(ControlElement->GetName(), Val, bDoKey,Context,bSetupUndo, true/* bFixEulerFlips*/);
+							ControlRig->SetControlLocalTransform(ControlElement->GetFName(), Val, bDoKey,Context,bSetupUndo, true/* bFixEulerFlips*/);
 						}
 						else
 						{
@@ -253,7 +253,7 @@ void FControlRigControlPose::BlendWithInitialPoses(FControlRigControlPose& Initi
 							Rotation = FQuat::Slerp(InitialVal.GetRotation(), Val.GetRotation(), BlendValue); //doing slerp here not fast lerp, can be slow this is for content creation
 							Scale = FMath::Lerp(InitialVal.GetScale3D(), Val.GetScale3D(), BlendValue);
 							Val = FTransform(Rotation, Translation, Scale);
-							ControlRig->SetControlGlobalTransform(ControlElement->GetName(), Val, bDoKey,Context,bSetupUndo, false /*bPrintPython*/, true/* bFixEulerFlips*/);
+							ControlRig->SetControlGlobalTransform(ControlElement->GetFName(), Val, bDoKey,Context,bSetupUndo, false /*bPrintPython*/, true/* bFixEulerFlips*/);
 						}
 					}
 					else
@@ -278,13 +278,13 @@ void FControlRigControlPose::BlendWithInitialPoses(FControlRigControlPose& Initi
 						LocalTransform.SetRotation(FQuat::Slerp(InitialLocationRotation, LocalTransform.GetRotation(), BlendValue)); //doing slerp here not fast lerp, can be slow this is for content creation
 						LocalTransform.SetScale3D(FMath::Lerp(InitialLocalScale, LocalTransform.GetScale3D(), BlendValue));
 
-						SetControlMirrorTransform(bDoLocal,ControlRig, ControlElement->GetName(), bIsMatched, GlobalTransform,LocalTransform,bDoKey,Context,bSetupUndo);							
+						SetControlMirrorTransform(bDoLocal,ControlRig, ControlElement->GetFName(), bIsMatched, GlobalTransform,LocalTransform,bDoKey,Context,bSetupUndo);							
 					}
 					if (bHasNonDefaultParent == true)
 					{
 						FTransform GlobalTransform = ControlRig->GetHierarchy()->GetGlobalTransform(ControlElement->GetKey());
 						ControlRig->GetHierarchy()->SwitchToParent(ControlElement->GetKey(), SpaceKey);
-						ControlRig->SetControlGlobalTransform(ControlElement->GetName(), GlobalTransform, true, Context, bSetupUndo, false /*bPrintPython*/, true/* bFixEulerFlips*/);
+						ControlRig->SetControlGlobalTransform(ControlElement->GetFName(), GlobalTransform, true, Context, bSetupUndo, false /*bPrintPython*/, true/* bFixEulerFlips*/);
 					}
 				}
 			}
diff --git a/Engine/Plugins/Animation/ControlRig/Source/ControlRig/Private/Tools/ControlRigPoseMirrorTable.cpp b/Engine/Plugins/Animation/ControlRig/Source/ControlRig/Private/Tools/ControlRigPoseMirrorTable.cpp
--- a/Engine/Plugins/Animation/ControlRig/Source/ControlRig/Private/Tools/ControlRigPoseMirrorTable.cpp
+++ b/Engine/Plugins/Animation/ControlRig/Source/ControlRig/Private/Tools/ControlRigPoseMirrorTable.cpp
@@ -15,7 +15,7 @@ void FControlRigPoseMirrorTable::SetUpMirrorTable(const UControlRig* ControlRig)
 		TArray<FRigControlElement*> CurrentControls = ControlRig->AvailableControls();
 		for (FRigControlElement* ControlElement : CurrentControls)
 		{
-			FString CurrentString = ControlElement->GetName().ToString();
+			FString CurrentString = ControlElement->GetName();
 			if (CurrentString.Contains(Settings->RightSide))
 			{
 				FString NewString = CurrentString.Replace(*Settings->RightSide, *Settings->LeftSide);
diff --git a/Engine/Plugins/Animation/ControlRig/Source/ControlRig/Private/Units/Execution/RigUnit_DynamicHierarchy.cpp b/Engine/Plugins/Animation/ControlRig/Source/ControlRig/Private/Units/Execution/RigUnit_DynamicHierarchy.cpp
--- a/Engine/Plugins/Animation/ControlRig/Source/ControlRig/Private/Units/Execution/RigUnit_DynamicHierarchy.cpp
+++ b/Engine/Plugins/Animation/ControlRig/Source/ControlRig/Private/Units/Execution/RigUnit_DynamicHierarchy.cpp
@@ -757,7 +757,7 @@ FRigUnit_HierarchyAddAnimationChannelBool_Execute()
 		ControlSettings.SetupLimitArrayForType(true, true, true);
 		ControlSettings.MinimumValue = FRigControlValue::Make<bool>(MinimumValue);
 		ControlSettings.MaximumValue = FRigControlValue::Make<bool>(MaximumValue);
-		ControlSettings.DisplayName = Controller->GetHierarchy()->GetSafeNewDisplayName(Parent, Name.ToString());
+		ControlSettings.DisplayName = Controller->GetHierarchy()->GetSafeNewDisplayName(Parent, Name);
 		const FRigControlValue Value = FRigControlValue::Make<bool>(InitialValue);
 
 		FRigHierarchyControllerInstructionBracket InstructionBracket(Controller, ExecuteContext.GetInstructionIndex());
@@ -796,7 +796,7 @@ FRigUnit_HierarchyAddAnimationChannelFloat_Execute()
 		ControlSettings.SetupLimitArrayForType(true, true, true);
 		ControlSettings.MinimumValue = FRigControlValue::Make<float>(MinimumValue);
 		ControlSettings.MaximumValue = FRigControlValue::Make<float>(MaximumValue);
-		ControlSettings.DisplayName = Controller->GetHierarchy()->GetSafeNewDisplayName(Parent, Name.ToString());
+		ControlSettings.DisplayName = Controller->GetHierarchy()->GetSafeNewDisplayName(Parent, Name);
 		const FRigControlValue Value = FRigControlValue::Make<float>(InitialValue);
 		
 		FRigHierarchyControllerInstructionBracket InstructionBracket(Controller, ExecuteContext.GetInstructionIndex());
@@ -835,7 +835,7 @@ FRigUnit_HierarchyAddAnimationChannelInteger_Execute()
 		ControlSettings.SetupLimitArrayForType(true, true, true);
 		ControlSettings.MinimumValue = FRigControlValue::Make<int32>(MinimumValue);
 		ControlSettings.MaximumValue = FRigControlValue::Make<int32>(MaximumValue);
-		ControlSettings.DisplayName = Controller->GetHierarchy()->GetSafeNewDisplayName(Parent, Name.ToString());
+		ControlSettings.DisplayName = Controller->GetHierarchy()->GetSafeNewDisplayName(Parent, Name);
 		const FRigControlValue Value = FRigControlValue::Make<int32>(InitialValue);
 		
 		FRigHierarchyControllerInstructionBracket InstructionBracket(Controller, ExecuteContext.GetInstructionIndex());
@@ -874,7 +874,7 @@ FRigUnit_HierarchyAddAnimationChannelVector2D_Execute()
 		ControlSettings.SetupLimitArrayForType(true, true, true);
 		ControlSettings.MinimumValue = FRigControlValue::Make<FVector3f>(FVector3f(MinimumValue.X, MinimumValue.Y, 0.f));
 		ControlSettings.MaximumValue = FRigControlValue::Make<FVector3f>(FVector3f(MaximumValue.X, MaximumValue.Y, 0.f));
-		ControlSettings.DisplayName = Controller->GetHierarchy()->GetSafeNewDisplayName(Parent, Name.ToString());
+		ControlSettings.DisplayName = Controller->GetHierarchy()->GetSafeNewDisplayName(Parent, Name);
 		const FRigControlValue Value = FRigControlValue::Make<FVector3f>(FVector3f(InitialValue.X, InitialValue.Y, 0.f));
 		
 		FRigHierarchyControllerInstructionBracket InstructionBracket(Controller, ExecuteContext.GetInstructionIndex());
@@ -913,7 +913,7 @@ FRigUnit_HierarchyAddAnimationChannelVector_Execute()
 		ControlSettings.SetupLimitArrayForType(true, true, true);
 		ControlSettings.MinimumValue = FRigControlValue::Make<FVector3f>(FVector3f(MinimumValue));
 		ControlSettings.MaximumValue = FRigControlValue::Make<FVector3f>(FVector3f(MaximumValue));
-		ControlSettings.DisplayName = Controller->GetHierarchy()->GetSafeNewDisplayName(Parent, Name.ToString());
+		ControlSettings.DisplayName = Controller->GetHierarchy()->GetSafeNewDisplayName(Parent, Name);
 		const FRigControlValue Value = FRigControlValue::Make<FVector3f>(FVector3f(InitialValue));
 		
 		FRigHierarchyControllerInstructionBracket InstructionBracket(Controller, ExecuteContext.GetInstructionIndex());
@@ -952,7 +952,7 @@ FRigUnit_HierarchyAddAnimationChannelRotator_Execute()
 		ControlSettings.SetupLimitArrayForType(true, true, true);
 		ControlSettings.MinimumValue = FRigControlValue::Make<FVector3f>(FVector3f(MinimumValue.Euler()));
 		ControlSettings.MaximumValue = FRigControlValue::Make<FVector3f>(FVector3f(MaximumValue.Euler()));
-		ControlSettings.DisplayName = Controller->GetHierarchy()->GetSafeNewDisplayName(Parent, Name.ToString());
+		ControlSettings.DisplayName = Controller->GetHierarchy()->GetSafeNewDisplayName(Parent, Name);
 		const FRigControlValue Value = FRigControlValue::Make<FVector3f>(FVector3f(InitialValue.Euler()));
 		
 		FRigHierarchyControllerInstructionBracket InstructionBracket(Controller, ExecuteContext.GetInstructionIndex());
diff --git a/Engine/Plugins/Animation/ControlRig/Source/ControlRig/Private/Units/Hierarchy/RigUnit_SetControlVisibility.cpp b/Engine/Plugins/Animation/ControlRig/Source/ControlRig/Private/Units/Hierarchy/RigUnit_SetControlVisibility.cpp
--- a/Engine/Plugins/Animation/ControlRig/Source/ControlRig/Private/Units/Hierarchy/RigUnit_SetControlVisibility.cpp
+++ b/Engine/Plugins/Animation/ControlRig/Source/ControlRig/Private/Units/Hierarchy/RigUnit_SetControlVisibility.cpp
@@ -45,7 +45,7 @@ FRigUnit_SetControlVisibility_Execute()
 		{
 			Hierarchy->ForEach<FRigControlElement>([&Keys, Pattern](FRigControlElement* ControlElement) -> bool
 			{
-				if (ControlElement->GetName().ToString().Contains(Pattern, ESearchCase::CaseSensitive))
+				if (ControlElement->GetFName().ToString().Contains(Pattern, ESearchCase::CaseSensitive))
 				{
 					Keys.Add(ControlElement->GetKey());
 				}
diff --git a/Engine/Plugins/Animation/ControlRig/Source/ControlRig/Private/Units/Highlevel/Hierarchy/RigUnit_MultiFABRIK.cpp b/Engine/Plugins/Animation/ControlRig/Source/ControlRig/Private/Units/Highlevel/Hierarchy/RigUnit_MultiFABRIK.cpp
--- a/Engine/Plugins/Animation/ControlRig/Source/ControlRig/Private/Units/Highlevel/Hierarchy/RigUnit_MultiFABRIK.cpp
+++ b/Engine/Plugins/Animation/ControlRig/Source/ControlRig/Private/Units/Highlevel/Hierarchy/RigUnit_MultiFABRIK.cpp
@@ -272,12 +272,12 @@ FRigUnit_MultiFABRIK_Execute()
 			    , ParentIndex(INDEX_NONE)
 				{
 					FRigBaseElement* BoneElement = InHierarchy->Get(InBoneIndex);
-					BoneName = BoneElement->GetName();
+					BoneName = BoneElement->GetFName();
 					BoneIndex = BoneElement->GetIndex();
 
 					if(FRigBaseElement* ParentBoneElement = InHierarchy->GetFirstParent(BoneElement))
 					{
-						ParentBoneName = ParentBoneElement->GetName();
+						ParentBoneName = ParentBoneElement->GetFName();
 						ParentIndex = ParentBoneElement->GetIndex();
 					}
 				}
diff --git a/Engine/Plugins/Animation/ControlRig/Source/ControlRig/Public/Rigs/RigHierarchy.h b/Engine/Plugins/Animation/ControlRig/Source/ControlRig/Public/Rigs/RigHierarchy.h
--- a/Engine/Plugins/Animation/ControlRig/Source/ControlRig/Public/Rigs/RigHierarchy.h
+++ b/Engine/Plugins/Animation/ControlRig/Source/ControlRig/Public/Rigs/RigHierarchy.h
@@ -5,6 +5,7 @@
 #include "CoreMinimal.h"
 #include "RigVMCore/RigVMExecuteContext.h"
 #include "RigVMCore/RigVM.h"
+#include "RigName.h"
 #include "RigHierarchyElements.h"
 #include "RigHierarchyCache.h"
 #include "RigHierarchyPose.h"
@@ -1556,14 +1557,14 @@ public:
 	 * Sanitizes a name by removing invalid characters.
 	 * @param InOutName The name to sanitize in place.
 	 */
-	static void SanitizeName(FString& InOutName);
+	static void SanitizeName(FRigName& InOutName);
 
 	/**
 	 * Sanitizes a name by removing invalid characters.
 	 * @param InName The name to sanitize.
 	 * @return The sanitized name.
  	 */
-	static FName GetSanitizedName(const FString& InName);
+	static FRigName GetSanitizedName(const FRigName& InName);
 
 	/**
 	 * Returns true if a given name is available.
@@ -1572,7 +1573,7 @@ public:
 	 * @param OutErrorMessage An optional pointer to return a potential error message 
 	 * @return Returns true if the name is available.
 	 */
-	bool IsNameAvailable(const FString& InPotentialNewName, ERigElementType InType, FString* OutErrorMessage = nullptr) const;
+	bool IsNameAvailable(const FRigName& InPotentialNewName, ERigElementType InType, FString* OutErrorMessage = nullptr) const;
 
 	/**
 	 * Returns true if a given display name is available.
@@ -1581,7 +1582,7 @@ public:
 	 * @param OutErrorMessage An optional pointer to return a potential error message 
 	 * @return Returns true if the name is available.
 	 */
-	bool IsDisplayNameAvailable(const FRigElementKey& InParentElement, const FString& InPotentialNewDisplayName, FString* OutErrorMessage = nullptr) const;
+	bool IsDisplayNameAvailable(const FRigElementKey& InParentElement, const FRigName& InPotentialNewDisplayName, FString* OutErrorMessage = nullptr) const;
 
 	/**
 	 * Returns a valid new name for a to-be-added element.
@@ -1589,7 +1590,7 @@ public:
 	 * @param InType The type of the to-be-added element
 	 * @return Returns the name to use for the to-be-added element.
 	 */
-	FName GetSafeNewName(const FString& InPotentialNewName, ERigElementType InType) const;
+	FRigName GetSafeNewName(const FRigName& InPotentialNewName, ERigElementType InType) const;
 
 	/**
 	 * Returns a valid new display name for a control
@@ -1597,7 +1598,7 @@ public:
 	 * @param InPotentialNewDisplayName The name to be sanitized and adjusted for availability
 	 * @return Returns the name to use for the to-be-added element.
 	 */
-	FName GetSafeNewDisplayName(const FRigElementKey& InParentElement, const FString& InPotentialNewDisplayName) const;
+	FRigName GetSafeNewDisplayName(const FRigElementKey& InParentElement, const FRigName& InPotentialNewDisplayName) const;
 
 	/**
 	 * Returns the modified event, which can be used to 
diff --git a/Engine/Plugins/Animation/ControlRig/Source/ControlRig/Public/Rigs/RigHierarchyElements.h b/Engine/Plugins/Animation/ControlRig/Source/ControlRig/Public/Rigs/RigHierarchyElements.h
--- a/Engine/Plugins/Animation/ControlRig/Source/ControlRig/Public/Rigs/RigHierarchyElements.h
+++ b/Engine/Plugins/Animation/ControlRig/Source/ControlRig/Public/Rigs/RigHierarchyElements.h
@@ -556,9 +556,9 @@ public:
 
 	FRigBaseElement()
     : Key()
-	, NameString()
     , Index(INDEX_NONE)
 	, SubIndex(INDEX_NONE)
+	, NameString()
 	, CreatedAtInstructionIndex(INDEX_NONE)
 	, OwnedInstances(0)
 	, TopologyVersion(0)
@@ -582,9 +582,6 @@ protected:
 	UPROPERTY(BlueprintReadOnly, VisibleAnywhere, Category = RigElement, meta = (AllowPrivateAccess = "true"))
 	FRigElementKey Key;
 
-	UPROPERTY(transient)
-	FString NameString;
-
 	UPROPERTY(BlueprintReadOnly, VisibleAnywhere, Category = RigElement, meta = (AllowPrivateAccess = "true"))
 	int32 Index;
 
@@ -599,6 +596,8 @@ protected:
 		return true;
 	}
 
+	mutable TOptional<FString> NameString;
+
 public:
 
 	UScriptStruct* GetElementStruct() const;
@@ -606,9 +605,16 @@ public:
 	virtual void Save(FArchive& Ar, URigHierarchy* Hierarchy, ESerializationPhase SerializationPhase);
 	virtual void Load(FArchive& Ar, URigHierarchy* Hierarchy, ESerializationPhase SerializationPhase);
 
-	const FName& GetName() const { return Key.Name; }
-	const FString& GetNameString() const { return NameString; }
-	virtual const FName& GetDisplayName() const { return GetName(); }
+	const FName& GetFName() const { return Key.Name; }
+	const FString& GetName() const
+	{
+		if(!NameString.IsSet())
+		{
+			NameString = Key.Name.ToString();
+		}
+		return NameString.GetValue();
+	}
+	virtual const FName& GetDisplayName() const { return GetFName(); }
 	ERigElementType GetType() const { return Key.Type; }
 	const FRigElementKey& GetKey() const { return Key; }
 	int32 GetIndex() const { return Index; }
diff --git a/Engine/Plugins/Animation/ControlRig/Source/ControlRig/Public/Rigs/RigName.h b/Engine/Plugins/Animation/ControlRig/Source/ControlRig/Public/Rigs/RigName.h
new file mode 100644
--- /dev/null
+++ b/Engine/Plugins/Animation/ControlRig/Source/ControlRig/Public/Rigs/RigName.h
@@ -0,0 +1,179 @@
+// Copyright Epic Games, Inc. All Rights Reserved.
+
+#pragma once
+
+#include "CoreMinimal.h"
+
+struct CONTROLRIG_API FRigName
+{
+public:
+	FRigName()
+	{}
+
+	FRigName(const FName& InName)
+		:FRigName()
+	{
+		SetFName(InName);
+	}
+
+	explicit FRigName(const FString& InNameString)
+		:FRigName()
+	{
+		SetName(InNameString);
+	}
+
+	FRigName(const FRigName& InOther)
+		: Name(InOther.Name)
+		, NameString(InOther.NameString)
+	{
+	}
+
+	FRigName& operator =(const FName& InName)
+	{
+		SetFName(InName);
+		return *this;
+	}
+
+	FRigName& operator =(const FString& InNameString)
+	{
+		SetName(InNameString);
+		return *this;
+	}
+
+	operator FName() const
+	{
+		return GetFName();
+	}
+	
+	bool IsValid() const
+	{
+		return Name.IsSet() || NameString.IsSet();
+	}
+
+	bool IsNone() const
+	{
+		if(!IsValid())
+		{
+			return true;
+		}
+		if(Name.IsSet())
+		{
+			return Name.GetValue().IsNone();
+		}
+		return GetName().IsEmpty();
+	}
+
+	bool operator ==(const FRigName& InOther) const
+	{
+		if(Name.IsSet() && InOther.Name.IsSet())
+		{
+			return GetFName().IsEqual(InOther.GetFName(), ENameCase::CaseSensitive);
+		}
+		
+		if(NameString.IsSet() && InOther.NameString.IsSet())
+		{
+			return GetName().Equals(InOther.GetName(), ESearchCase::CaseSensitive);
+		}
+
+		if(Name.IsSet() && InOther.NameString.IsSet())
+		{
+			return GetName().Equals(InOther.GetName(), ESearchCase::CaseSensitive);
+		}
+
+		if(NameString.IsSet() && InOther.Name.IsSet())
+		{
+			return GetName().Equals(InOther.GetName(), ESearchCase::CaseSensitive);
+		}
+
+		return IsNone() == InOther.IsNone();
+	}
+
+	bool operator !=(const FRigName& InOther) const
+	{
+		return !(*this == InOther);
+	}
+
+	int32 Len() const
+	{
+		if(NameString.IsSet())
+		{
+			return NameString->Len();
+		}
+		if(Name.IsSet())
+		{
+			return (int32)Name->GetStringLength();
+		}
+		return 0;
+	}
+
+	const FName& GetFName() const
+	{
+		if(Name.IsSet())
+		{
+			return Name.GetValue();
+		}
+		if(NameString.IsSet())
+		{
+			if(!NameString->IsEmpty())
+			{
+				Name = *NameString.GetValue();
+				return Name.GetValue();
+			}
+		}
+		static const FName NoneName(NAME_None);
+		return NoneName;
+	}
+
+	const FString& GetName() const
+	{
+		if(NameString.IsSet())
+		{
+			return NameString.GetValue();
+		}
+		if(Name.IsSet())
+		{
+			NameString = FString();
+			FString& ValueString = NameString.GetValue();
+			Name.GetValue().ToString(ValueString);
+			return ValueString;
+		}
+		static const FString EmptyString;
+		return EmptyString;
+	}
+
+	const FString& ToString() const
+	{
+		return GetName();
+	}
+
+	void SetFName(const FName& InName)
+	{
+		if(Name.IsSet())
+		{
+			if(Name.GetValue().IsEqual(InName, ENameCase::CaseSensitive))
+			{
+				return;
+			}
+		}
+		Name = InName;
+		NameString.Reset();
+	}
+
+	void SetName(const FString& InNameString)
+	{
+		if(NameString.IsSet())
+		{
+			if(NameString.GetValue().Equals(InNameString, ESearchCase::CaseSensitive))
+			{
+				return;
+			}
+		}
+		Name.Reset();
+		NameString = InNameString;
+	}
+
+private:
+
+	mutable TOptional<FName> Name;
+	mutable TOptional<FString> NameString;
+};
diff --git a/Engine/Plugins/Animation/ControlRig/Source/ControlRig/Public/Tools/ControlRigPose.h b/Engine/Plugins/Animation/ControlRig/Source/ControlRig/Public/Tools/ControlRigPose.h
--- a/Engine/Plugins/Animation/ControlRig/Source/ControlRig/Public/Tools/ControlRigPose.h
+++ b/Engine/Plugins/Animation/ControlRig/Source/ControlRig/Public/Tools/ControlRigPose.h
@@ -37,7 +37,7 @@ struct CONTROLRIG_API FRigControlCopy
 
 	FRigControlCopy(FRigControlElement* InControlElement, URigHierarchy* InHierarchy)
 	{
-		Name = InControlElement->GetName();
+		Name = InControlElement->GetFName();
 		ControlType = InControlElement->Settings.ControlType;
 		Value = InHierarchy->GetControlValue(InControlElement, ERigControlValueType::Current);
 		ParentKey = InHierarchy->GetFirstParent(InControlElement->GetKey());
diff --git a/Engine/Plugins/Animation/ControlRig/Source/ControlRigDeveloper/Private/AnimGraphNode_ControlRig.cpp b/Engine/Plugins/Animation/ControlRig/Source/ControlRigDeveloper/Private/AnimGraphNode_ControlRig.cpp
--- a/Engine/Plugins/Animation/ControlRig/Source/ControlRigDeveloper/Private/AnimGraphNode_ControlRig.cpp
+++ b/Engine/Plugins/Animation/ControlRig/Source/ControlRigDeveloper/Private/AnimGraphNode_ControlRig.cpp
@@ -208,7 +208,7 @@ void UAnimGraphNode_ControlRig::CreateCustomPins(TArray<UEdGraphPin*>* OldPins)
 				{
 					if (Hierarchy->IsAnimatable(ControlElement))
 					{
-						const FName ControlName = ControlElement->GetName();
+						const FName ControlName = ControlElement->GetFName();
 						BeginExposableNames.Remove(ControlName);
 
 						if (CurrentlyExposedNames.Contains(ControlName))
@@ -220,7 +220,7 @@ void UAnimGraphNode_ControlRig::CreateCustomPins(TArray<UEdGraphPin*>* OldPins)
 							}
 
 							UEdGraphPin* NewPin = CreatePin(EEdGraphPinDirection::EGPD_Input, PinType, ControlName);
-							NewPin->PinFriendlyName = FText::FromName(ControlElement->GetName());
+							NewPin->PinFriendlyName = FText::FromName(ControlElement->GetFName());
 
 							// Newly created pin does not have an auto generated default value, so we need to generate one here
 							// Missing the auto-gen default would cause UEdGraphPin::MovePersistentDataFromOldPin to override
@@ -378,7 +378,7 @@ void UAnimGraphNode_ControlRig::RebuildExposedProperties()
 				{
 					if (Hierarchy->IsAnimatable(ControlElement))
 					{
-						const FName ControlName = ControlElement->GetName();
+						const FName ControlName = ControlElement->GetFName();
 						CustomPinProperties.Add(MakeOptionalPin(ControlName));
 					}
 					return true;
@@ -644,7 +644,7 @@ void UAnimGraphNode_ControlRig::GetAvailableMapping(const FName& PathName, TArra
 					{
 						if (Hierarchy->IsAnimatable(ControlElement))
 						{
-							OutArray.Add(ControlElement->GetName());
+							OutArray.Add(ControlElement->GetFName());
 						}
 						return true;
 					});
@@ -718,8 +718,8 @@ void UAnimGraphNode_ControlRig::CreateVariableMapping(const FString& FilteredTex
 					{
 						if (Hierarchy->IsAnimatable(ControlElement))
 						{
-							const FName ControlName = ControlElement->GetName();
-							const FString& DisplayName = ControlName.ToString();
+							const FName ControlName = ControlElement->GetFName();
+							const FString& DisplayName = ControlElement->GetName();
 
 							if (!bDoFiltering || DisplayName.Contains(FilteredText))
 							{
diff --git a/Engine/Plugins/Animation/ControlRig/Source/ControlRigDeveloper/Private/ControlRigBlueprint.cpp b/Engine/Plugins/Animation/ControlRig/Source/ControlRigDeveloper/Private/ControlRigBlueprint.cpp
--- a/Engine/Plugins/Animation/ControlRig/Source/ControlRigDeveloper/Private/ControlRigBlueprint.cpp
+++ b/Engine/Plugins/Animation/ControlRig/Source/ControlRigDeveloper/Private/ControlRigBlueprint.cpp
@@ -240,7 +240,7 @@ void UControlRigBlueprint::UpdateExposedModuleConnectors()
 	Hierarchy->ForEach<FRigConnectorElement>([this](const FRigConnectorElement* ConnectorElement) -> bool
 	{
 		FRigModuleConnector ExposedConnector;
-		ExposedConnector.Name = ConnectorElement->GetNameString();
+		ExposedConnector.Name = ConnectorElement->GetName();
 		ExposedConnector.bIsRoot = Hierarchy->GetNumberOfParents(ConnectorElement) == 0;
 		ExposedConnector.Settings = ConnectorElement->Settings;
 		ModuleSettings.ExposedConnectors.Add(ExposedConnector);
@@ -1396,7 +1396,7 @@ UControlRigBlueprint::FControlValueScope::FControlValueScope(UControlRigBlueprin
 		TArray<FRigControlElement*> Controls = CR->AvailableControls();
 		for (FRigControlElement* ControlElement : Controls)
 		{
-			ControlValues.Add(ControlElement->GetName(), CR->GetControlValue(ControlElement->GetName()));
+			ControlValues.Add(ControlElement->GetFName(), CR->GetControlValue(ControlElement->GetFName()));
 		}
 	}
 #endif
diff --git a/Engine/Plugins/Animation/ControlRig/Source/ControlRigDeveloper/Private/ControlRigVisualGraphUtils.cpp b/Engine/Plugins/Animation/ControlRig/Source/ControlRigDeveloper/Private/ControlRigVisualGraphUtils.cpp
--- a/Engine/Plugins/Animation/ControlRig/Source/ControlRigDeveloper/Private/ControlRigVisualGraphUtils.cpp
+++ b/Engine/Plugins/Animation/ControlRig/Source/ControlRigDeveloper/Private/ControlRigVisualGraphUtils.cpp
@@ -181,7 +181,7 @@ FString FControlRigVisualGraphUtils::DumpRigHierarchyToDotGraph(URigHierarchy* I
 				}
 			}
 
-			NodeIndex = OutGraph.AddNode(NodeName, InElement->GetName(), Color, Shape);
+			NodeIndex = OutGraph.AddNode(NodeName, InElement->GetFName(), Color, Shape);
 
 			if(NodeIndex != INDEX_NONE)
 			{
diff --git a/Engine/Plugins/Animation/ControlRig/Source/ControlRigDeveloper/Public/Graph/ControlRigGraph.h b/Engine/Plugins/Animation/ControlRig/Source/ControlRigDeveloper/Public/Graph/ControlRigGraph.h
--- a/Engine/Plugins/Animation/ControlRig/Source/ControlRigDeveloper/Public/Graph/ControlRigGraph.h
+++ b/Engine/Plugins/Animation/ControlRig/Source/ControlRigDeveloper/Public/Graph/ControlRigGraph.h
@@ -88,7 +88,7 @@ private:
 			{
 				if(!bFilter || IncludeElementInNameList<T>(Cast<T>(Element)))
 				{
-					Names.Add(Element->GetName().ToString());
+					Names.Add(Element->GetName());
 				}
 			}
 		}
diff --git a/Engine/Plugins/Animation/ControlRig/Source/ControlRigEditor/Private/ControlRigElementDetails.cpp b/Engine/Plugins/Animation/ControlRig/Source/ControlRigEditor/Private/ControlRigElementDetails.cpp
--- a/Engine/Plugins/Animation/ControlRig/Source/ControlRigEditor/Private/ControlRigElementDetails.cpp
+++ b/Engine/Plugins/Animation/ControlRig/Source/ControlRigEditor/Private/ControlRigElementDetails.cpp
@@ -893,7 +893,7 @@ bool FRigBaseElementDetails::OnVerifyNameChanged(const FText& InText, FText& Out
 	}
 
 	FString OutErrorMessageStr;
-	if (!Hierarchy->IsNameAvailable(InText.ToString(), GetElementKey().Type, &OutErrorMessageStr))
+	if (!Hierarchy->IsNameAvailable(FRigName(InText.ToString()), GetElementKey().Type, &OutErrorMessageStr))
 	{
 		OutErrorMessage = FText::FromString(OutErrorMessageStr);
 		return false;
@@ -3501,14 +3501,14 @@ FReply FRigControlElementDetails::OnAddAnimationChannelClicked()
 	const FRigElementKey Key = PerElementInfos[0].GetElement()->GetKey();
 	URigHierarchy* HierarchyToChange = PerElementInfos[0].GetDefaultHierarchy();
 
-	static const FString ChannelName = TEXT("Channel");
+	static const FName ChannelName = TEXT("Channel");
 	FRigControlSettings Settings;
 	Settings.AnimationType = ERigControlAnimationType::AnimationChannel;
 	Settings.ControlType = ERigControlType::Float;
 	Settings.MinimumValue = FRigControlValue::Make<float>(0.f);
 	Settings.MaximumValue = FRigControlValue::Make<float>(1.f);
 	Settings.DisplayName = HierarchyToChange->GetSafeNewDisplayName(Key, ChannelName);
-	HierarchyToChange->GetController(true)->AddAnimationChannel(*ChannelName, Key, Settings, true, true);
+	HierarchyToChange->GetController(true)->AddAnimationChannel(ChannelName, Key, Settings, true, true);
 	HierarchyToChange->GetController(true)->SelectElement(Key);
 	return FReply::Handled();
 }
@@ -4349,7 +4349,7 @@ bool FRigControlElementDetails::OnVerifyDisplayNameChanged(const FText& InText,
 	if(const FRigBaseElement* ParentElement = Hierarchy->GetFirstParent(ControlElement))
 	{
 		FString OutErrorString;
-		if (!Hierarchy->IsDisplayNameAvailable(ParentElement->GetKey(), NewName, &OutErrorString))
+		if (!Hierarchy->IsDisplayNameAvailable(ParentElement->GetKey(), FRigName(NewName), &OutErrorString))
 		{
 			OutErrorMessage = FText::FromString(OutErrorString);
 			return false;
diff --git a/Engine/Plugins/Animation/ControlRig/Source/ControlRigEditor/Private/ControlRigSequencerEditorLibrary.cpp b/Engine/Plugins/Animation/ControlRig/Source/ControlRigEditor/Private/ControlRigSequencerEditorLibrary.cpp
--- a/Engine/Plugins/Animation/ControlRig/Source/ControlRigEditor/Private/ControlRigSequencerEditorLibrary.cpp
+++ b/Engine/Plugins/Animation/ControlRig/Source/ControlRigEditor/Private/ControlRigSequencerEditorLibrary.cpp
@@ -2931,7 +2931,7 @@ bool UControlRigSequencerEditorLibrary::GetControlsMask(UMovieSceneSection* InSe
 	int32 Index = 0;
 	for (const FRigControlElement* RigControl : Controls)
 	{
-		if (RigControl->GetName() == ControlName)
+		if (RigControl->GetFName() == ControlName)
 		{
 			return ParameterSection->GetControlsMask(Index);
 		}
@@ -2965,7 +2965,7 @@ void UControlRigSequencerEditorLibrary::SetControlsMask(UMovieSceneSection* InSe
 	int32 Index = 0;
 	for (const FRigControlElement* RigControl : Controls)
 	{
-		if (ControlNames.Contains(RigControl->GetName()))
+		if (ControlNames.Contains(RigControl->GetFName()))
 		{
 			ParameterSection->SetControlsMask(Index, bVisible);
 		}
diff --git a/Engine/Plugins/Animation/ControlRig/Source/ControlRigEditor/Private/ControlRigThumbnailRenderer.cpp b/Engine/Plugins/Animation/ControlRig/Source/ControlRigEditor/Private/ControlRigThumbnailRenderer.cpp
--- a/Engine/Plugins/Animation/ControlRig/Source/ControlRigEditor/Private/ControlRigThumbnailRenderer.cpp
+++ b/Engine/Plugins/Animation/ControlRig/Source/ControlRigEditor/Private/ControlRigThumbnailRenderer.cpp
@@ -157,7 +157,7 @@ void UControlRigThumbnailRenderer::AddAdditionalPreviewSceneContent(UObject* Obj
 								MeshComponent->SetMaterial(i, MaterialInstance);
 							}
 
-							ShapeActors.Add(ControlElement->GetName(), ShapeActor);
+							ShapeActors.Add(ControlElement->GetFName(), ShapeActor);
 
 							ShapeActor->GetStaticMeshComponent()->SetStaticMesh(StaticMesh);
 							ShapeActor->SetActorTransform(ShapeDef->Transform * ShapeGlobalTransform);
diff --git a/Engine/Plugins/Animation/ControlRig/Source/ControlRigEditor/Private/EditMode/ControlRigControlsProxy.cpp b/Engine/Plugins/Animation/ControlRig/Source/ControlRigEditor/Private/EditMode/ControlRigControlsProxy.cpp
--- a/Engine/Plugins/Animation/ControlRig/Source/ControlRigEditor/Private/EditMode/ControlRigControlsProxy.cpp
+++ b/Engine/Plugins/Animation/ControlRig/Source/ControlRigEditor/Private/EditMode/ControlRigControlsProxy.cpp
@@ -1002,7 +1002,7 @@ void UControlRigDetailPanelControlProxies::RecreateAllProxies(UControlRig* Contr
 	{
 		if(ControlElement->Settings.AnimationType != ERigControlAnimationType::VisualCue)
 		{
-			AddProxy(ControlRig,ControlElement->GetName(), ControlElement);
+			AddProxy(ControlRig,ControlElement->GetFName(), ControlElement);
 		}
 	}
 }
diff --git a/Engine/Plugins/Animation/ControlRig/Source/ControlRigEditor/Private/EditMode/ControlRigEditMode.cpp b/Engine/Plugins/Animation/ControlRig/Source/ControlRigEditor/Private/EditMode/ControlRigEditMode.cpp
--- a/Engine/Plugins/Animation/ControlRig/Source/ControlRigEditor/Private/EditMode/ControlRigEditMode.cpp
+++ b/Engine/Plugins/Animation/ControlRig/Source/ControlRigEditor/Private/EditMode/ControlRigEditMode.cpp
@@ -721,7 +721,7 @@ TSet<FName> FControlRigEditMode::GetActiveControlsFromSequencer(UControlRig* Con
 							{
 								if (Mask[Index])
 								{
-									ActiveControls.Add(ControlElement->GetName());
+									ActiveControls.Add(ControlElement->GetFName());
 								}
 								++Index;
 							}
@@ -798,7 +798,7 @@ void FControlRigEditMode::Render(const FSceneView* View, FViewport* Viewport, FP
 									FLinearColor Color = FLinearColor::White;
 									if (bHasFKRig)
 									{
-										FName ControlName = UFKControlRig::GetControlName(ParentTransformElement->GetName(), ParentTransformElement->GetType());
+										FName ControlName = UFKControlRig::GetControlName(ParentTransformElement->GetFName(), ParentTransformElement->GetType());
 										if (ActiveControlName.Num() > 0 && ActiveControlName.Contains(ControlName) == false)
 										{
 											continue;
@@ -814,7 +814,7 @@ void FControlRigEditMode::Render(const FSceneView* View, FViewport* Viewport, FP
 									{
 										if (bHitTesting)
 										{
-											PDI->SetHitProxy(new HFKRigBoneProxy(ParentTransformElement->GetName(), ControlRig));
+											PDI->SetHitProxy(new HFKRigBoneProxy(ParentTransformElement->GetFName(), ControlRig));
 										}
 										PDI->DrawLine(ComponentTransform.TransformPosition(Transform.GetLocation()), ComponentTransform.TransformPosition(ParentTransform.GetLocation()), Color, SDPG_Foreground);
 										if (bHitTesting)
@@ -828,7 +828,7 @@ void FControlRigEditMode::Render(const FSceneView* View, FViewport* Viewport, FP
 							FLinearColor Color = FLinearColor::White;
 							if (bHasFKRig)
 							{
-								FName ControlName = UFKControlRig::GetControlName(TransformElement->GetName(), TransformElement->GetType());
+								FName ControlName = UFKControlRig::GetControlName(TransformElement->GetFName(), TransformElement->GetType());
 								if (ActiveControlName.Num() > 0 && ActiveControlName.Contains(ControlName) == false)
 								{
 									return true;
@@ -843,7 +843,7 @@ void FControlRigEditMode::Render(const FSceneView* View, FViewport* Viewport, FP
 								const bool bHitTesting = PDI->IsHitTesting() && bBoolSetHitProxies && (TransformElement->GetType() == ERigElementType::Bone);
 								if (bHitTesting)
 								{
-									PDI->SetHitProxy(new HFKRigBoneProxy(TransformElement->GetName(), ControlRig));
+									PDI->SetHitProxy(new HFKRigBoneProxy(TransformElement->GetFName(), ControlRig));
 								}
 								PDI->DrawPoint(ComponentTransform.TransformPosition(Transform.GetLocation()), Color, 5.0f, SDPG_Foreground);
 
@@ -3308,7 +3308,7 @@ void FControlRigEditMode::RecreateControlShapeActors(const TArray<FRigElementKey
 						for (int32 Index = 0; Index < Controls.Num(); ++Index)
 						{
 							FRigControlElement* Element = Controls[Index];
-							if (Element && Element->GetName() == Actor->ControlName)
+							if (Element && Element->GetFName() == Actor->ControlName)
 							{
 								Controls.RemoveAtSwap(Index);
 								break;
@@ -3362,9 +3362,9 @@ void FControlRigEditMode::CreateShapeActors(UControlRig* ControlRig)
 			Param.ManipObj = ControlRig;
 			Param.ControlRigIndex = ControlRigIndex;
 			Param.ControlRig = ControlRig;
-			Param.ControlName = ControlElement->GetName();
+			Param.ControlName = ControlElement->GetFName();
 			Param.ShapeName = ControlElement->Settings.ShapeName;
-			Param.SpawnTransform = ControlRig->GetControlGlobalTransform(ControlElement->GetName());
+			Param.SpawnTransform = ControlRig->GetControlGlobalTransform(ControlElement->GetFName());
 			Param.ShapeTransform = ControlRig->GetHierarchy()->GetControlShapeTransform(ControlElement, ERigTransformType::CurrentLocal);
 			Param.bSelectable = ControlElement->Settings.IsSelectable(false);
 
@@ -3633,7 +3633,7 @@ void FControlRigEditMode::OnHierarchyModified(ERigHierarchyNotification InNotif,
 				if (FRigControlElement* ControlElement = InHierarchy->Find<FRigControlElement>(InElement->GetKey()))
 				{
 					UControlRig* ControlRig = InHierarchy->GetTypedOuter<UControlRig>();
-					if(ControlProxy->IsSelected(ControlRig, ControlElement->GetName()))
+					if(ControlProxy->IsSelected(ControlRig, ControlElement->GetFName()))
 					{
 						// reselect the control - to affect the details panel / sequencer
 						if(URigHierarchyController* Controller = InHierarchy->GetController())
@@ -3716,9 +3716,9 @@ void FControlRigEditMode::OnHierarchyModified(ERigHierarchyNotification InNotif,
 									{
 										if (const FRigControlElement* DrivenControl = ControlRig->GetHierarchy()->Find<FRigControlElement>(DrivenKey))
 										{
-											ControlProxy->SelectProxy(ControlRig, DrivenControl->GetName(), bSelected);
+											ControlProxy->SelectProxy(ControlRig, DrivenControl->GetFName(), bSelected);
 
-											if (AControlRigShapeActor* DrivenShapeActor = GetControlShapeFromControlName(ControlRig,DrivenControl->GetName()))
+											if (AControlRigShapeActor* DrivenShapeActor = GetControlShapeFromControlName(ControlRig,DrivenControl->GetFName()))
 											{
 												if(bSelected)
 												{
@@ -3743,7 +3743,7 @@ void FControlRigEditMode::OnHierarchyModified(ERigHierarchyNotification InNotif,
 													if(DrivenKeys.Contains(ParentControlElement->GetKey()) ||
 														ParentControlElement->GetKey() == ControlElement->GetKey())
 													{
-														ControlProxy->SelectProxy(ControlRig, AnimationChannelControl->GetName(), bSelected);
+														ControlProxy->SelectProxy(ControlRig, AnimationChannelControl->GetFName(), bSelected);
 													}
 												}
 											}
@@ -3834,7 +3834,7 @@ void FControlRigEditMode::OnHierarchyModified_AnyThread(ERigHierarchyNotificatio
 void FControlRigEditMode::OnControlModified(UControlRig* Subject, FRigControlElement* InControlElement, const FRigControlModifiedContext& Context)
 {
 	//this makes sure the details panel ui get's updated, don't remove
-	ControlProxy->ProxyChanged(Subject,InControlElement->GetName());
+	ControlProxy->ProxyChanged(Subject,InControlElement->GetFName());
 
 	/*
 	FScopedTransaction ScopedTransaction(LOCTEXT("ModifyControlTransaction", "Modify Control"),!GIsTransacting && Context.SetKey != EControlRigSetKey::Never);
diff --git a/Engine/Plugins/Animation/ControlRig/Source/ControlRigEditor/Private/EditMode/SControlRigDetails.cpp b/Engine/Plugins/Animation/ControlRig/Source/ControlRigEditor/Private/EditMode/SControlRigDetails.cpp
--- a/Engine/Plugins/Animation/ControlRig/Source/ControlRigEditor/Private/EditMode/SControlRigDetails.cpp
+++ b/Engine/Plugins/Animation/ControlRig/Source/ControlRigEditor/Private/EditMode/SControlRigDetails.cpp
@@ -149,15 +149,15 @@ void FControlRigEditModeGenericDetails::CustomizeDetails(class IDetailLayoutBuil
 				FRigControlElement* ParentControlElement = Cast<FRigControlElement>(Hierarchy->GetFirstParent(ControlElement));
 				if (ParentControlElement)
 				{
-					ParentControlName = ParentControlElement->GetName();
+					ParentControlName = ParentControlElement->GetFName();
 				}
 
-				if (ParentControlName == ControlElement->GetName())
+				if (ParentControlName == ControlElement->GetFName())
 				{
 					if (FControlRigEditMode* EditMode = static_cast<FControlRigEditMode*>(ModeTools->GetActiveMode(FControlRigEditMode::ModeName)))
 					{
 						UControlRig* ControlRig = Proxy->ControlRig.Get();
-						if (UObject* NestedProxy = EditMode->ControlProxy->FindProxy(ControlRig,ControlElement->GetName()))
+						if (UObject* NestedProxy = EditMode->ControlProxy->FindProxy(ControlRig,ControlElement->GetFName()))
 						{
 							FName PropertyName(NAME_None);
 							switch (ControlElement->Settings.ControlType)
diff --git a/Engine/Plugins/Animation/ControlRig/Source/ControlRigEditor/Private/EditMode/SCurveControlContainer.cpp b/Engine/Plugins/Animation/ControlRig/Source/ControlRigEditor/Private/EditMode/SCurveControlContainer.cpp
--- a/Engine/Plugins/Animation/ControlRig/Source/ControlRigEditor/Private/EditMode/SCurveControlContainer.cpp
+++ b/Engine/Plugins/Animation/ControlRig/Source/ControlRigEditor/Private/EditMode/SCurveControlContainer.cpp
@@ -245,10 +245,10 @@ void SCurveControlContainer::CreateCurveControlList( const FString& SearchText )
 		// Iterate through all curves..
 		Hierarchy->ForEach<FRigCurveElement>([CTRLName, this](FRigCurveElement* CurveElement) -> bool
         {
-			const FString Name = CurveElement->GetName().ToString();
+			const FString Name = CurveElement->GetName();
 			if (Name.Contains(CTRLName))
 			{
-				const FString CurveString = CurveElement->GetName().ToString();
+				const FString CurveString = CurveElement->GetName();
 
 				// See if we pass the search filter
 				if (!FilterText.IsEmpty() && !CurveString.Contains(*FilterText.ToString()))
@@ -256,7 +256,7 @@ void SCurveControlContainer::CreateCurveControlList( const FString& SearchText )
 					return true;
 				}
 
-				const TSharedRef<FDisplayedCurveControlInfo> NewItem = FDisplayedCurveControlInfo::Make(CurveElement->GetName());
+				const TSharedRef<FDisplayedCurveControlInfo> NewItem = FDisplayedCurveControlInfo::Make(CurveElement->GetFName());
 				CurveControlList.Add(NewItem);
 			}
 
@@ -354,7 +354,7 @@ void SCurveControlContainer::OnRigElementSelected(UControlRig* Subject, FRigCont
 {
 	for(const FDisplayedCurveControlInfoPtr& Item : CurveControlList)
 	{
-		if (Item->CurveName == ControlElement->GetName())
+		if (Item->CurveName == ControlElement->GetFName())
 		{
 			CurveControlListView->SetItemSelection(Item, bSelected);
 			break;
diff --git a/Engine/Plugins/Animation/ControlRig/Source/ControlRigEditor/Private/Editor/ControlRigEditor.cpp b/Engine/Plugins/Animation/ControlRig/Source/ControlRigEditor/Private/Editor/ControlRigEditor.cpp
--- a/Engine/Plugins/Animation/ControlRig/Source/ControlRigEditor/Private/Editor/ControlRigEditor.cpp
+++ b/Engine/Plugins/Animation/ControlRig/Source/ControlRigEditor/Private/Editor/ControlRigEditor.cpp
@@ -2742,7 +2742,7 @@ void FControlRigEditor::OnHierarchyModified(ERigHierarchyNotification InNotif, U
 
 			CacheNameLists();
 
-			const FString RemovedElementName = InElement->GetName().ToString();
+			const FString RemovedElementName = InElement->GetName();
 			const ERigElementType RemovedElementType = InElement->GetType();
 
 			TArray<UEdGraph*> EdGraphs;
@@ -2815,7 +2815,7 @@ void FControlRigEditor::OnHierarchyModified(ERigHierarchyNotification InNotif, U
 			}
 
 			const FString OldNameStr = InHierarchy->GetPreviousName(InElement->GetKey()).ToString();
-			const FString NewNameStr = InElement->GetName().ToString();
+			const FString NewNameStr = InElement->GetName();
 			const ERigElementType ElementType = InElement->GetType(); 
 
 			TArray<UEdGraph*> EdGraphs;
@@ -3058,7 +3058,7 @@ void FControlRigEditor::SynchronizeViewportBoneSelection()
 		TArray<const FRigBaseElement*> SelectedBones = GetHierarchyBeingDebugged()->GetSelectedElements(ERigElementType::Bone);
 		for (const FRigBaseElement* SelectedBone : SelectedBones)
 		{
- 			const int32 BoneIndex = EditorSkelComp->GetReferenceSkeleton().FindBoneIndex(SelectedBone->GetName());
+ 			const int32 BoneIndex = EditorSkelComp->GetReferenceSkeleton().FindBoneIndex(SelectedBone->GetFName());
 			if(BoneIndex != INDEX_NONE)
 			{
 				EditorSkelComp->BonesOfInterest.AddUnique(BoneIndex);
@@ -3561,7 +3561,7 @@ void FControlRigEditor::HandleMakeElementGetterSetter(ERigElementGetterSetterTyp
 				
 				if(const FRigControlElement* ParentControlElement = Cast<FRigControlElement>(Hierarchy->GetFirstParent(ControlElement)))
 				{
-					NameValue = ParentControlElement->GetName();
+					NameValue = ParentControlElement->GetFName();
 				}
 				else
 				{
@@ -4121,7 +4121,7 @@ void FControlRigEditor::HandleOnControlModified(UControlRig* Subject, FRigContro
 			}
 			else if (ElementKey.Type == ERigElementType::Null)
 			{
-				const FTransform GlobalTransform = GetControlRig()->GetControlGlobalTransform(ControlElement->GetName());
+				const FTransform GlobalTransform = GetControlRig()->GetControlGlobalTransform(ControlElement->GetFName());
 				Blueprint->Hierarchy->SetGlobalTransform(ElementKey, GlobalTransform);
 				Hierarchy->SetGlobalTransform(ElementKey, GlobalTransform);
 				if (IsConstructionModeEnabled())
diff --git a/Engine/Plugins/Animation/ControlRig/Source/ControlRigEditor/Private/Editor/ControlRigSkeletalMeshComponent.cpp b/Engine/Plugins/Animation/ControlRig/Source/ControlRigEditor/Private/Editor/ControlRigSkeletalMeshComponent.cpp
--- a/Engine/Plugins/Animation/ControlRig/Source/ControlRigEditor/Private/Editor/ControlRigSkeletalMeshComponent.cpp
+++ b/Engine/Plugins/Animation/ControlRig/Source/ControlRigEditor/Private/Editor/ControlRigSkeletalMeshComponent.cpp
@@ -117,7 +117,7 @@ void UControlRigSkeletalMeshComponent::RebuildDebugDrawSkeleton()
 					{
 						if(ParentName.IsNone())
 						{
-							ParentName = InBoneElement->GetName();
+							ParentName = InBoneElement->GetFName();
 						}
 						bContinue = false;
 					}
@@ -130,12 +130,12 @@ void UControlRigSkeletalMeshComponent::RebuildDebugDrawSkeleton()
 				}
 					
 				FMeshBoneInfo NewMeshBoneInfo;
-				NewMeshBoneInfo.Name = BoneElement->GetName();
+				NewMeshBoneInfo.Name = BoneElement->GetFName();
 				NewMeshBoneInfo.ParentIndex = ParentIndex; 
 				// give ref pose here
 				RefSkelModifier.Add(NewMeshBoneInfo, Hierarchy->GetInitialGlobalTransform(Index), true);
 
-				AddedBoneMap.FindOrAdd(BoneElement->GetName(), DebugDrawBones.Num());
+				AddedBoneMap.FindOrAdd(BoneElement->GetFName(), DebugDrawBones.Num());
 				DebugDrawBones.Add(DebugDrawBones.Num());
 				DebugDrawBoneIndexInHierarchy.Add(BoneElement->GetIndex());
 			}
diff --git a/Engine/Plugins/Animation/ControlRig/Source/ControlRigEditor/Private/Editor/SRigCurveContainer.cpp b/Engine/Plugins/Animation/ControlRig/Source/ControlRigEditor/Private/Editor/SRigCurveContainer.cpp
--- a/Engine/Plugins/Animation/ControlRig/Source/ControlRigEditor/Private/Editor/SRigCurveContainer.cpp
+++ b/Engine/Plugins/Animation/ControlRig/Source/ControlRigEditor/Private/Editor/SRigCurveContainer.cpp
@@ -392,7 +392,7 @@ void SRigCurveContainer::CreateRigCurveList( const FString& SearchText )
 		// Iterate through all curves..
 		Hierarchy->ForEach<FRigCurveElement>([this](FRigCurveElement* CurveElement) -> bool
 		{
-			const FString CurveString = CurveElement->GetName().ToString();
+			const FString CurveString = CurveElement->GetName();
 			
 			// See if we pass the search filter
             if (!FilterText.IsEmpty() && !CurveString.Contains(*FilterText.ToString()))
@@ -400,7 +400,7 @@ void SRigCurveContainer::CreateRigCurveList( const FString& SearchText )
                 return true;
             }
 
-            TSharedRef<FDisplayedRigCurveInfo> NewItem = FDisplayedRigCurveInfo::Make(CurveElement->GetName());
+            TSharedRef<FDisplayedRigCurveInfo> NewItem = FDisplayedRigCurveInfo::Make(CurveElement->GetFName());
             RigCurveList.Add(NewItem);
 
 			return true;
@@ -601,7 +601,7 @@ void SRigCurveContainer::OnHierarchyModified(ERigHierarchyNotification InNotif,
 				const bool bSelected = InNotif == ERigHierarchyNotification::ElementSelected;
 				for(const FDisplayedRigCurveInfoPtr& Item : RigCurveList)
 				{
-					if (Item->CurveName == InElement->GetName())
+					if (Item->CurveName == InElement->GetFName())
 					{
 						RigCurveListView->SetItemSelection(Item, bSelected);
 						break;
diff --git a/Engine/Plugins/Animation/ControlRig/Source/ControlRigEditor/Private/Editor/SRigHierarchy.cpp b/Engine/Plugins/Animation/ControlRig/Source/ControlRigEditor/Private/Editor/SRigHierarchy.cpp
--- a/Engine/Plugins/Animation/ControlRig/Source/ControlRigEditor/Private/Editor/SRigHierarchy.cpp
+++ b/Engine/Plugins/Animation/ControlRig/Source/ControlRigEditor/Private/Editor/SRigHierarchy.cpp
@@ -1870,7 +1870,7 @@ void SRigHierarchy::HandleNewItem(ERigElementType InElementType, bool bIsAnimati
 						Settings.ControlType = ERigControlType::Float;
 						Settings.MinimumValue = FRigControlValue::Make<float>(0.f);
 						Settings.MaximumValue = FRigControlValue::Make<float>(1.f);
-						Settings.DisplayName = Hierarchy->GetSafeNewDisplayName(ParentKey, *NewNameTemplate);
+						Settings.DisplayName = Hierarchy->GetSafeNewDisplayName(ParentKey, FRigName(NewNameTemplate));
 
 						NewItemKey = Controller->AddAnimationChannel(NewElementName, ParentKey, Settings, true, true);
 					}
@@ -2253,7 +2253,7 @@ URigHierarchy* SRigHierarchy::GetDefaultHierarchy() const
 
 FName SRigHierarchy::CreateUniqueName(const FName& InBaseName, ERigElementType InElementType) const
 {
-	return GetHierarchy()->GetSafeNewName(InBaseName.ToString(), InElementType);
+	return GetHierarchy()->GetSafeNewName(InBaseName, InElementType);
 }
 
 void SRigHierarchy::PostRedo(bool bSuccess) 
@@ -2471,9 +2471,8 @@ FName SRigHierarchy::HandleRenameElement(const FRigElementKey& OldKey, const FSt
 		URigHierarchyController* Controller = Hierarchy->GetController(true);
 		check(Controller);
 
-		FString SanitizedNameStr = NewName;
-		Hierarchy->SanitizeName(SanitizedNameStr);
-		const FName SanitizedName = *SanitizedNameStr;
+		FRigName SanitizedName(NewName);
+		Hierarchy->SanitizeName(SanitizedName);
 		FName ResultingName = NAME_None;
 
 		bool bUseDisplayName = false;
@@ -2546,7 +2545,7 @@ bool SRigHierarchy::HandleVerifyNameChanged(const FRigElementKey& OldKey, const
 				if(const FRigBaseElement* ParentElement = Hierarchy->GetFirstParent(ControlElement))
 				{
 					FString OutErrorString;
-					if (!Hierarchy->IsDisplayNameAvailable(ParentElement->GetKey(), NewName, &OutErrorString))
+					if (!Hierarchy->IsDisplayNameAvailable(ParentElement->GetKey(), FRigName(NewName), &OutErrorString))
 					{
 						OutErrorMessage = FText::FromString(OutErrorString);
 						return false;
@@ -2557,7 +2556,7 @@ bool SRigHierarchy::HandleVerifyNameChanged(const FRigElementKey& OldKey, const
 		else
 		{
 			FString OutErrorString;
-			if (!Hierarchy->IsNameAvailable(NewName, OldKey.Type, &OutErrorString))
+			if (!Hierarchy->IsNameAvailable(FRigName(NewName), OldKey.Type, &OutErrorString))
 			{
 				OutErrorMessage = FText::FromString(OutErrorString);
 				return false;
@@ -2847,7 +2846,7 @@ bool SRigHierarchy::FindClosestBone(const FVector& Point, FName& OutRigElementNa
             {
                 NearestDistance = CurDistance;
                 OutGlobalTransform = CurTransform;
-                OutRigElementName = Element->GetName();
+                OutRigElementName = Element->GetFName();
             }
             return true;
 		});
diff --git a/Engine/Plugins/Animation/ControlRig/Source/ControlRigEditor/Private/Sequencer/ControlRigParameterTrackEditor.cpp b/Engine/Plugins/Animation/ControlRig/Source/ControlRigEditor/Private/Sequencer/ControlRigParameterTrackEditor.cpp
--- a/Engine/Plugins/Animation/ControlRig/Source/ControlRigEditor/Private/Sequencer/ControlRigParameterTrackEditor.cpp
+++ b/Engine/Plugins/Animation/ControlRig/Source/ControlRigEditor/Private/Sequencer/ControlRigParameterTrackEditor.cpp
@@ -2566,7 +2566,7 @@ void FControlRigParameterTrackEditor::SelectRigsAndControls(UControlRig* Control
 									{
 										if (const TSet<FName>* Controls = RigsAndControls.Find(ControlRig))
 										{
-											if (Controls->Contains(ParentControlElement->GetName()))
+											if (Controls->Contains(ParentControlElement->GetFName()))
 											{
 												continue;
 											}
@@ -3413,7 +3413,7 @@ void FControlRigParameterTrackEditor::HandleControlSelected(UControlRig* Subject
 			{
 				UMovieSceneSection* Section = Track->GetSectionToKey();
 				UMovieSceneControlRigParameterSection* ParamSection = Cast<UMovieSceneControlRigParameterSection>(Section);
-				SelectSequencerNodeInSection(ParamSection, ControlElement->GetName(), bSelected);
+				SelectSequencerNodeInSection(ParamSection, ControlElement->GetFName(), bSelected);
 			}
 			else
 			{
@@ -3421,7 +3421,7 @@ void FControlRigParameterTrackEditor::HandleControlSelected(UControlRig* Subject
 				{
 					if (UMovieSceneControlRigParameterSection* ParamSection = Cast< UMovieSceneControlRigParameterSection>(BaseSection))
 					{
-						SelectSequencerNodeInSection(ParamSection, ControlElement->GetName(), bSelected);
+						SelectSequencerNodeInSection(ParamSection, ControlElement->GetFName(), bSelected);
 					}
 				}
 			}
@@ -3471,7 +3471,7 @@ void FControlRigParameterTrackEditor::HandleControlModified(UControlRig* Control
 	{
 		return;
 	}
-	FTransform  Transform = ControlRig->GetControlLocalTransform(ControlElement->GetName());
+	FTransform  Transform = ControlRig->GetControlLocalTransform(ControlElement->GetFName());
 	UMovieScene* MovieScene = GetSequencer()->GetFocusedMovieSceneSequence()->GetMovieScene();
 	const TArray<FMovieSceneBinding>& Bindings = MovieScene->GetBindings();
 	for (const FMovieSceneBinding& Binding : Bindings)
@@ -3493,7 +3493,7 @@ void FControlRigParameterTrackEditor::HandleControlModified(UControlRig* Control
 						{
 							KeyMode = ESequencerKeyMode::ManualKeyForced;
 						}
-						AddControlKeys(Component, ControlRig, Name, ControlElement->GetName(), (EControlRigContextChannelToKey)Context.KeyMask, 
+						AddControlKeys(Component, ControlRig, Name, ControlElement->GetFName(), (EControlRigContextChannelToKey)Context.KeyMask, 
 							KeyMode, Context.LocalTime);
 						ControlChangedDuringUndoBracket++;
 					}
@@ -3653,13 +3653,13 @@ void FControlRigParameterTrackEditor::GetControlRigKeys(
 			continue;
 		}
 
-		if (FChannelMapInfo* pChannelIndex = SectionToKey->ControlChannelMap.Find(ControlElement->GetName()))
+		if (FChannelMapInfo* pChannelIndex = SectionToKey->ControlChannelMap.Find(ControlElement->GetFName()))
 		{
 			int32 ChannelIndex = pChannelIndex->ChannelIndex;
 
 
 			bool bMaskKeyOut = (ControlIndex >= ControlsMask.Num() || ControlsMask[ControlIndex] == false);
-			bool bSetKey = ControlElement->GetName() == ParameterName && !bMaskKeyOut;
+			bool bSetKey = ControlElement->GetFName() == ParameterName && !bMaskKeyOut;
 
 			FRigControlValue ControlValue = InControlRig->GetControlValue(ControlElement, ERigControlValueType::Current);
 
@@ -3767,7 +3767,7 @@ void FControlRigParameterTrackEditor::GetControlRigKeys(
 
 				if (bInConstraintSpace)
 				{
-					const uint32 ControlHash = UTransformableControlHandle::ComputeHash(InControlRig, ControlElement->GetName());
+					const uint32 ControlHash = UTransformableControlHandle::ComputeHash(InControlRig, ControlElement->GetFName());
 					TOptional<FTransform> Transform = FTransformConstraintUtils::GetRelativeTransform(InControlRig->GetWorld(), ControlHash);
 					if (Transform)
 					{
@@ -4094,9 +4094,9 @@ bool FControlRigParameterTrackEditor::ModifyOurGeneratedKeysByCurrentAndWeight(U
 		{
 			for (const FFloatInterrogationData& Val : InterrogationData.Iterate<FFloatInterrogationData>(UMovieSceneControlRigParameterSection::GetFloatInterrogationKey()))
 			{
-				if ((Val.ParameterName == ControlElement->GetName()))
+				if ((Val.ParameterName == ControlElement->GetFName()))
 				{
-					pChannelIndex = Section->ControlChannelMap.Find(ControlElement->GetName());
+					pChannelIndex = Section->ControlChannelMap.Find(ControlElement->GetFName());
 					if (pChannelIndex && pChannelIndex->GeneratedKeyIndex != INDEX_NONE)
 					{
 						ChannelIndex = pChannelIndex->GeneratedKeyIndex;
@@ -4119,9 +4119,9 @@ bool FControlRigParameterTrackEditor::ModifyOurGeneratedKeysByCurrentAndWeight(U
 		{
 			for (const FVector2DInterrogationData& Val : InterrogationData.Iterate<FVector2DInterrogationData>(UMovieSceneControlRigParameterSection::GetVector2DInterrogationKey()))
 			{
-				if ((Val.ParameterName == ControlElement->GetName()))
+				if ((Val.ParameterName == ControlElement->GetFName()))
 				{
-					pChannelIndex = Section->ControlChannelMap.Find(ControlElement->GetName());
+					pChannelIndex = Section->ControlChannelMap.Find(ControlElement->GetFName());
 					if (pChannelIndex && pChannelIndex->GeneratedKeyIndex != INDEX_NONE)
 					{
 						ChannelIndex = pChannelIndex->GeneratedKeyIndex;
@@ -4141,9 +4141,9 @@ bool FControlRigParameterTrackEditor::ModifyOurGeneratedKeysByCurrentAndWeight(U
 		{
 			for (const FVectorInterrogationData& Val : InterrogationData.Iterate<FVectorInterrogationData>(UMovieSceneControlRigParameterSection::GetVectorInterrogationKey()))
 			{
-				if ((Val.ParameterName == ControlElement->GetName()))
+				if ((Val.ParameterName == ControlElement->GetFName()))
 				{
-					pChannelIndex = Section->ControlChannelMap.Find(ControlElement->GetName());
+					pChannelIndex = Section->ControlChannelMap.Find(ControlElement->GetFName());
 					if (pChannelIndex && pChannelIndex->GeneratedKeyIndex != INDEX_NONE)
 					{
 						ChannelIndex = pChannelIndex->GeneratedKeyIndex;
@@ -4167,9 +4167,9 @@ bool FControlRigParameterTrackEditor::ModifyOurGeneratedKeysByCurrentAndWeight(U
 			for (const FEulerTransformInterrogationData& Val : InterrogationData.Iterate<FEulerTransformInterrogationData>(UMovieSceneControlRigParameterSection::GetTransformInterrogationKey()))
 			{
 
-				if ((Val.ParameterName == ControlElement->GetName()))
+				if ((Val.ParameterName == ControlElement->GetFName()))
 				{
-					pChannelIndex = Section->ControlChannelMap.Find(ControlElement->GetName());
+					pChannelIndex = Section->ControlChannelMap.Find(ControlElement->GetFName());
 					if (pChannelIndex && pChannelIndex->GeneratedKeyIndex != INDEX_NONE)
 					{
 						ChannelIndex = pChannelIndex->GeneratedKeyIndex;
@@ -4948,7 +4948,7 @@ bool FControlRigParameterTrackEditor::CollapseAllLayers(TSharedPtr<ISequencer>&S
 							continue;
 						}
 						TPair<FName, TArray<FTransform>> NameTransforms;
-						NameTransforms.Key = ControlElement->GetName();
+						NameTransforms.Key = ControlElement->GetFName();
 						NameTransforms.Value.SetNum(Frames.Num());
 						ControlLocalTransforms.Add(NameTransforms);
 					}
@@ -5377,7 +5377,7 @@ void FControlRigParameterSection::ShowSelectedControlsChannels()
 		ControlRig->GetControlsInOrder(Controls);
 		for (const FRigControlElement* RigControl : Controls)
 		{
-			const FName RigName = RigControl->GetName();
+			const FName RigName = RigControl->GetFName();
 			if (ControlRig->IsControlSelected(RigName))
 			{
 				FChannelMapInfo* pChannelIndex = ParameterSection->ControlChannelMap.Find(RigName);
diff --git a/Engine/Plugins/Animation/ControlRig/Source/ControlRigEditor/Private/Sequencer/ControlRigSequencerFilter.cpp b/Engine/Plugins/Animation/ControlRig/Source/ControlRigEditor/Private/Sequencer/ControlRigSequencerFilter.cpp
--- a/Engine/Plugins/Animation/ControlRig/Source/ControlRigEditor/Private/Sequencer/ControlRigSequencerFilter.cpp
+++ b/Engine/Plugins/Animation/ControlRig/Source/ControlRigEditor/Private/Sequencer/ControlRigSequencerFilter.cpp
@@ -177,7 +177,7 @@ public:
 				TArray<const FRigBaseElement*> SelectedControls = ControlRig->GetHierarchy()->GetSelectedElements(ERigElementType::Control);
 				for (const FRigBaseElement* SelectedControl : SelectedControls)
 				{
-					if (Name == SelectedControl->GetName())
+					if (Name == SelectedControl->GetFName())
 					{
 						return true;
 					}
diff --git a/Engine/Plugins/Animation/ControlRig/Source/ControlRigEditor/Private/Sequencer/ControlRigSpaceChannelEditors.cpp b/Engine/Plugins/Animation/ControlRig/Source/ControlRigEditor/Private/Sequencer/ControlRigSpaceChannelEditors.cpp
--- a/Engine/Plugins/Animation/ControlRig/Source/ControlRigEditor/Private/Sequencer/ControlRigSpaceChannelEditors.cpp
+++ b/Engine/Plugins/Animation/ControlRig/Source/ControlRigEditor/Private/Sequencer/ControlRigSpaceChannelEditors.cpp
@@ -1303,7 +1303,7 @@ void FControlRigSpaceChannelHelpers::CompensateIfNeeded(
 		if(Control)// ac && Control->GetName() != ControlName)
 		{ 
 			//only if we have a channel
-			if (FSpaceControlNameAndChannel* Channel = Section->GetSpaceChannel(Control->GetName()))
+			if (FSpaceControlNameAndChannel* Channel = Section->GetSpaceChannel(Control->GetFName()))
 			{
 				const TArrayView<const FFrameNumber> FramesToCompensate = GetSpaceTimesToCompensate(Channel);
 				if (FramesToCompensate.Num() > 0)
@@ -1324,7 +1324,7 @@ void FControlRigSpaceChannelHelpers::CompensateIfNeeded(
 							TArray<FTransform> ControlWorldTransforms;
 							FControlRigSnapper Snapper;
 							Snapper.GetControlRigControlTransforms(
-								Sequencer, ControlRig, Control->GetName(),
+								Sequencer, ControlRig, Control->GetFName(),
 								{Time},
 								ControlRigParentWorldTransforms, ControlWorldTransforms);
 
@@ -1346,7 +1346,7 @@ void FControlRigSpaceChannelHelpers::CompensateIfNeeded(
 							//now set time -1 frame value
 							ControlRig->Evaluate_AnyThread();
 							KeyframeContext.LocalTime = TickResolution.AsSeconds(FFrameTime(Time - 1));
-							ControlRig->SetControlGlobalTransform(Control->GetName(), ControlWorldTransforms[0], true, KeyframeContext, false /*undo*/, false /*bPrintPython*/, true/* bFixEulerFlips*/);
+							ControlRig->SetControlGlobalTransform(Control->GetFName(), ControlWorldTransforms[0], true, KeyframeContext, false /*undo*/, false /*bPrintPython*/, true/* bFixEulerFlips*/);
 							
 							bDidIt = true;
 						}
diff --git a/Engine/Plugins/Animation/ControlRig/Source/ControlRigEditor/Private/Tools/ControlRigTweener.cpp b/Engine/Plugins/Animation/ControlRig/Source/ControlRigEditor/Private/Tools/ControlRigTweener.cpp
--- a/Engine/Plugins/Animation/ControlRig/Source/ControlRigEditor/Private/Tools/ControlRigTweener.cpp
+++ b/Engine/Plugins/Animation/ControlRig/Source/ControlRigEditor/Private/Tools/ControlRigTweener.cpp
@@ -380,13 +380,13 @@ bool FAnimSliderObjectSelection::Setup(const TArray<UControlRig*>& SelectedContr
 					TArray<FKeyHandle> Handles;
 					for (FRigControlElement* ControlElement : CurrentControls)
 					{
-						if (ControlRig->GetHierarchy()->IsAnimatable(ControlElement) &&  ControlRig->IsControlSelected(ControlElement->GetName()))
+						if (ControlRig->GetHierarchy()->IsAnimatable(ControlElement) &&  ControlRig->IsControlSelected(ControlElement->GetFName()))
 						{
 							FAnimSliderObjectSelection::FObjectChannels ObjectChannels;
 							ObjectChannels.Section = Section;
 							KeyTimes.SetNum(0);
 							Handles.SetNum(0);
-							FChannelMapInfo* pChannelIndex = Section->ControlChannelMap.Find(ControlElement->GetName());
+							FChannelMapInfo* pChannelIndex = Section->ControlChannelMap.Find(ControlElement->GetFName());
 							if (pChannelIndex == nullptr)
 							{
 								continue;
diff --git a/Engine/Plugins/Editor/SequencerAnimTools/Source/SequencerAnimTools/Private/Sequencer/SequencerTrailHierarchy.cpp b/Engine/Plugins/Editor/SequencerAnimTools/Source/SequencerAnimTools/Private/Sequencer/SequencerTrailHierarchy.cpp
--- a/Engine/Plugins/Editor/SequencerAnimTools/Source/SequencerAnimTools/Private/Sequencer/SequencerTrailHierarchy.cpp
+++ b/Engine/Plugins/Editor/SequencerAnimTools/Source/SequencerAnimTools/Private/Sequencer/SequencerTrailHierarchy.cpp
@@ -752,16 +752,16 @@ void FSequencerTrailHierarchy::RegisterControlRigDelegates(USkeletalMeshComponen
 
 			if (bSelected)
 			{
-				AddControlRigTrail(Component,ControlRig,CRParameterTrack, ControlElement->GetName());
+				AddControlRigTrail(Component,ControlRig,CRParameterTrack, ControlElement->GetFName());
 			}
 
 			if (ControlsTracked.Find(ControlRig) != nullptr)
 			{
-				if (ControlsTracked[ControlRig].Contains(ControlElement->GetName()))
+				if (ControlsTracked[ControlRig].Contains(ControlElement->GetFName()))
 				{
 					if (bSelected == false)
 					{
-						const FGuid TrailGuid = ControlsTracked[ControlRig][ControlElement->GetName()];
+						const FGuid TrailGuid = ControlsTracked[ControlRig][ControlElement->GetFName()];
 						VisibilityManager.ControlSelected.Remove(TrailGuid);
 						RemoveTrailIfNotAlwaysVisible(TrailGuid);
 					}
@@ -806,12 +806,12 @@ void FSequencerTrailHierarchy::RegisterControlRigDelegates(USkeletalMeshComponen
 			
 			if(InNotif == ERigHierarchyNotification::ElementRemoved)
 			{
-				if (!ControlsTracked.Contains(ControlRig) || !ControlsTracked[ControlRig].Contains(ControlElement->GetName())) // We only care about controls
+				if (!ControlsTracked.Contains(ControlRig) || !ControlsTracked[ControlRig].Contains(ControlElement->GetFName())) // We only care about controls
 				{
 					return;
 				}
 				
-				const FGuid TrailGuid = ControlsTracked[ControlRig][ControlElement->GetName()];
+				const FGuid TrailGuid = ControlsTracked[ControlRig][ControlElement->GetFName()];
 				RemoveTrail(TrailGuid);
 			}
             else if(InNotif == ERigHierarchyNotification::ElementRenamed)
@@ -825,7 +825,7 @@ void FSequencerTrailHierarchy::RegisterControlRigDelegates(USkeletalMeshComponen
 
 				const FGuid TempTrailGuid = ControlsTracked[ControlRig][OldName];
 				ControlsTracked[ControlRig].Remove(OldName);
-				ControlsTracked[ControlRig].Add(ControlElement->GetName(), TempTrailGuid);
+				ControlsTracked[ControlRig].Add(ControlElement->GetFName(), TempTrailGuid);
             }
         }
     );
diff --git a/Engine/Plugins/Editor/SequencerAnimTools/Source/SequencerAnimTools/Private/SequencerAnimEditPivotTool.cpp b/Engine/Plugins/Editor/SequencerAnimTools/Source/SequencerAnimTools/Private/SequencerAnimEditPivotTool.cpp
--- a/Engine/Plugins/Editor/SequencerAnimTools/Source/SequencerAnimTools/Private/SequencerAnimEditPivotTool.cpp
+++ b/Engine/Plugins/Editor/SequencerAnimTools/Source/SequencerAnimTools/Private/SequencerAnimEditPivotTool.cpp
@@ -632,7 +632,7 @@ void USequencerPivotTool::GizmoTransformStarted(UTransformProxy* Proxy)
 						//if we have a parent that's selected and we are a float/bool/int/enum skip
 						if (FRigControlElement* ParentControlElement = Cast<FRigControlElement>(ControlRig->GetHierarchy()->GetFirstParent(ControlElement)))
 						{
-							if (ControlRig->IsControlSelected(ParentControlElement->GetName()) && (
+							if (ControlRig->IsControlSelected(ParentControlElement->GetFName()) && (
 								ControlElement->Settings.ControlType == ERigControlType::Bool || 
 								ControlElement->Settings.ControlType == ERigControlType::Float || 
 								ControlElement->Settings.ControlType == ERigControlType::Integer)
diff --git a/Engine/Plugins/Experimental/FullBodyIK/Source/PBIK/Private/RigUnit_PBIK.cpp b/Engine/Plugins/Experimental/FullBodyIK/Source/PBIK/Private/RigUnit_PBIK.cpp
--- a/Engine/Plugins/Experimental/FullBodyIK/Source/PBIK/Private/RigUnit_PBIK.cpp
+++ b/Engine/Plugins/Experimental/FullBodyIK/Source/PBIK/Private/RigUnit_PBIK.cpp
@@ -69,7 +69,7 @@ FRigUnit_PBIK_Execute()
 		// create bones
 		for (int B = 0; B < BoneElements.Num(); ++B)
 		{
-			FName Name = BoneElements[B]->GetName();
+			FName Name = BoneElements[B]->GetFName();
 
 			// get the first parent that is not excluded
 			int32 ParentElementIndex = INDEX_NONE;
@@ -98,7 +98,7 @@ FRigUnit_PBIK_Execute()
 			const FTransform OrigTransform = Hierarchy->GetTransform(BoneElements[B], ERigTransformType::InitialGlobal);
 			const FVector InOrigPosition = OrigTransform.GetLocation();
 			const FQuat InOrigRotation = OrigTransform.GetRotation();
-			bool bIsRoot = BoneElements[B]->GetName() == Root;
+			bool bIsRoot = BoneElements[B]->GetFName() == Root;
 			WorkData.Solver.AddBone(Name, ParentIndex, InOrigPosition, InOrigRotation, bIsRoot);
 		}
 		
diff --git a/Engine/Plugins/Importers/USDImporter/Source/USDUtilities/Private/USDSkeletalDataConversion.cpp b/Engine/Plugins/Importers/USDImporter/Source/USDUtilities/Private/USDSkeletalDataConversion.cpp
--- a/Engine/Plugins/Importers/USDImporter/Source/USDUtilities/Private/USDSkeletalDataConversion.cpp
+++ b/Engine/Plugins/Importers/USDImporter/Source/USDUtilities/Private/USDSkeletalDataConversion.cpp
@@ -3328,7 +3328,7 @@ bool UnrealToUsd::ConvertControlRigSection(
 			RigJointIndexToRefSkeletonIndex.Reset();
 			for (FRigBoneElement* RigBone : Hierarchy->GetBones())
 			{
-				RigJointIndexToRefSkeletonIndex.Add(InRefSkeleton.FindBoneIndex(RigBone->GetName()));
+				RigJointIndexToRefSkeletonIndex.Add(InRefSkeleton.FindBoneIndex(RigBone->GetFName()));
 			}
 		};
 	RegenerateRigJointIndexToRefSkeletonIndex();
@@ -3451,7 +3451,7 @@ bool UnrealToUsd::ConvertControlRigSection(
 		for ( int32 Index = 0; Index < CurveElements.Num(); ++Index )
 		{
 			FRigCurveElement* Element = CurveElements[Index];
-			const FString& CurveNameString = Element->GetNameString();
+			const FString& CurveNameString = Element->GetName();
 
 			CurveNames.push_back( UnrealToUsd::ConvertToken( *CurveNameString ).Get() );
 		}
