ControlRig: Fix build breakage in RigName.h

diff --git a/Engine/Plugins/Animation/ControlRig/Source/ControlRig/Private/Rigs/RigName.cpp b/Engine/Plugins/Animation/ControlRig/Source/ControlRig/Private/Rigs/RigName.cpp
new file mode 100644
--- /dev/null
+++ b/Engine/Plugins/Animation/ControlRig/Source/ControlRig/Private/Rigs/RigName.cpp
@@ -0,0 +1,6 @@
+// Copyright Epic Games, Inc. All Rights Reserved.
+
+#include "Rigs/RigName.h"
+
+const FName FRigName::EmptyName(NAME_None);
+const FString FRigName::EmptyString;
diff --git a/Engine/Plugins/Animation/ControlRig/Source/ControlRig/Public/Rigs/RigName.h b/Engine/Plugins/Animation/ControlRig/Source/ControlRig/Public/Rigs/RigName.h
--- a/Engine/Plugins/Animation/ControlRig/Source/ControlRig/Public/Rigs/RigName.h
+++ b/Engine/Plugins/Animation/ControlRig/Source/ControlRig/Public/Rigs/RigName.h
@@ -120,8 +120,7 @@ public:
 				return Name.GetValue();
 			}
 		}
-		static const FName NoneName(NAME_None);
-		return NoneName;
+		return EmptyName;
 	}
 
 	const FString& GetName() const
@@ -137,7 +136,6 @@ public:
 			Name.GetValue().ToString(ValueString);
 			return ValueString;
 		}
-		static const FString EmptyString;
 		return EmptyString;
 	}
 
@@ -176,4 +174,7 @@ private:
 
 	mutable TOptional<FName> Name;
 	mutable TOptional<FString> NameString;
+
+	static const FString EmptyString;
+	static const FName EmptyName;
 };
