Control Rig: Fix linux build break with FRigName.h

diff --git a/Engine/Plugins/Animation/ControlRig/Source/ControlRig/Public/Rigs/RigName.h b/Engine/Plugins/Animation/ControlRig/Source/ControlRig/Public/Rigs/RigName.h
--- a/Engine/Plugins/Animation/ControlRig/Source/ControlRig/Public/Rigs/RigName.h
+++ b/Engine/Plugins/Animation/ControlRig/Source/ControlRig/Public/Rigs/RigName.h
@@ -10,6 +10,9 @@ public:
 	FRigName()
 	{}
 
+	FRigName(const FRigName& InOther) = default;
+	FRigName& operator =(const FRigName& InOther) = default;
+
 	FRigName(const FName& InName)
 		:FRigName()
 	{
@@ -22,12 +25,6 @@ public:
 		SetName(InNameString);
 	}
 
-	FRigName(const FRigName& InOther)
-		: Name(InOther.Name)
-		, NameString(InOther.NameString)
-	{
-	}
-
 	FRigName& operator =(const FName& InName)
 	{
 		SetFName(InName);
