Control Rig: Hierarchy performance improvements

diff --git a/Engine/Plugins/Animation/ControlRig/Source/ControlRig/Private/Rigs/RigHierarchy.cpp b/Engine/Plugins/Animation/ControlRig/Source/ControlRig/Private/Rigs/RigHierarchy.cpp
--- a/Engine/Plugins/Animation/ControlRig/Source/ControlRig/Private/Rigs/RigHierarchy.cpp
+++ b/Engine/Plugins/Animation/ControlRig/Source/ControlRig/Private/Rigs/RigHierarchy.cpp
@@ -4423,33 +4423,59 @@ void URigHierarchy::UpdateCachedChildren(const FRigBaseElement* InElement, bool
 	{
 		return;
 	}
-	
+
 	InElement->CachedChildren.Reset();
-	
+
+	for (int32 ElementIndex = 0; ElementIndex < Elements.Num(); ElementIndex++)
+	{
+		FRigBaseElement* Element = Elements[ElementIndex];
+		if(Element->IsA<FRigSingleParentElement>() || Element->IsA<FRigMultiParentElement>())
+		{
+			if(Element->TopologyVersion != TopologyVersion)
+			{
+				Element->CachedChildren.Reset();
+			}
+		}
+	}
+
+	// since we'll have to loop over all children anyway - it makes sense to update all of them
+	// at the same time.
 	for (int32 ElementIndex = 0; ElementIndex < Elements.Num(); ElementIndex++)
 	{
 		FRigBaseElement* Element = Elements[ElementIndex];
 		if(FRigSingleParentElement* SingleParentElement = Cast<FRigSingleParentElement>(Element))
 		{
-			if(SingleParentElement->ParentElement == InElement)
+			if(FRigTransformElement* ParentElement = SingleParentElement->ParentElement)
 			{
-				InElement->CachedChildren.Add(SingleParentElement);
+				if(ParentElement->TopologyVersion != TopologyVersion)
+				{
+					ParentElement->CachedChildren.Add(SingleParentElement);
+				}
 			}
 		}
 		else if(FRigMultiParentElement* MultiParentElement = Cast<FRigMultiParentElement>(Element))
 		{
 			for(const FRigElementParentConstraint& ParentConstraint : MultiParentElement->ParentConstraints)
 			{
-				if(ParentConstraint.ParentElement == InElement)
+				if(FRigTransformElement* ParentElement = ParentConstraint.ParentElement)
 				{
-					InElement->CachedChildren.Add(MultiParentElement);
-					break;
+					if(ParentElement->TopologyVersion != TopologyVersion)
+					{
+						ParentElement->CachedChildren.Add(MultiParentElement);
+					}
 				}
 			}
 		}
 	}
 
-	InElement->TopologyVersion = TopologyVersion;
+	for (int32 ElementIndex = 0; ElementIndex < Elements.Num(); ElementIndex++)
+	{
+		FRigBaseElement* Element = Elements[ElementIndex];
+		if(Element->IsA<FRigSingleParentElement>() || Element->IsA<FRigMultiParentElement>())
+		{
+			Element->TopologyVersion = TopologyVersion;
+		}
+	}
 }
 
 void URigHierarchy::UpdateAllCachedChildren() const
diff --git a/Engine/Plugins/Animation/ControlRig/Source/ControlRig/Private/Rigs/RigHierarchyElements.cpp b/Engine/Plugins/Animation/ControlRig/Source/ControlRig/Private/Rigs/RigHierarchyElements.cpp
--- a/Engine/Plugins/Animation/ControlRig/Source/ControlRig/Private/Rigs/RigHierarchyElements.cpp
+++ b/Engine/Plugins/Animation/ControlRig/Source/ControlRig/Private/Rigs/RigHierarchyElements.cpp
@@ -13,6 +13,8 @@
 // FRigBaseElement
 ////////////////////////////////////////////////////////////////////////////////
 
+const FRigBaseElement::EElementIndex FRigBaseElement::ElementTypeIndex = BaseElement;
+
 FRigBaseElement::FRigBaseElement(const FRigBaseElement& InOther)
 {
 	*this = InOther;
@@ -481,6 +483,8 @@ FRigBaseElement* FRigElementHandle::Get()
 // FRigTransformElement
 ////////////////////////////////////////////////////////////////////////////////
 
+const FRigBaseElement::EElementIndex FRigTransformElement::ElementTypeIndex = TransformElement;
+
 void FRigTransformElement::Save(FArchive& Ar, URigHierarchy* Hierarchy, ESerializationPhase SerializationPhase)
 {
 	Super::Save(Ar, Hierarchy, SerializationPhase);
@@ -543,6 +547,8 @@ void FRigTransformElement::CopyFrom(URigHierarchy* InHierarchy, FRigBaseElement*
 // FRigSingleParentElement
 ////////////////////////////////////////////////////////////////////////////////
 
+const FRigBaseElement::EElementIndex FRigSingleParentElement::ElementTypeIndex = SingleParentElement;
+
 void FRigSingleParentElement::Save(FArchive& Ar, URigHierarchy* Hierarchy, ESerializationPhase SerializationPhase)
 {
 	Super::Save(Ar, Hierarchy, SerializationPhase);
@@ -595,6 +601,8 @@ void FRigSingleParentElement::CopyFrom(URigHierarchy* InHierarchy, FRigBaseEleme
 // FRigMultiParentElement
 ////////////////////////////////////////////////////////////////////////////////
 
+const FRigBaseElement::EElementIndex FRigMultiParentElement::ElementTypeIndex = MultiParentElement;
+
 void FRigMultiParentElement::Save(FArchive& Ar, URigHierarchy* Hierarchy, ESerializationPhase SerializationPhase)
 {
 	Super::Save(Ar, Hierarchy, SerializationPhase);
@@ -746,6 +754,8 @@ void FRigMultiParentElement::CopyPose(FRigBaseElement* InOther, bool bCurrent, b
 static_assert(sizeof(FRigBoneElement) <= 736, "FRigBoneElement was optimized to fit into 736 bytes bin of MallocBinned3");
 #endif
 
+const FRigBaseElement::EElementIndex FRigBoneElement::ElementTypeIndex = BoneElement;
+
 void FRigBoneElement::Save(FArchive& Ar, URigHierarchy* Hierarchy, ESerializationPhase SerializationPhase)
 {
 	Super::Save(Ar, Hierarchy, SerializationPhase);
@@ -779,6 +789,12 @@ void FRigBoneElement::CopyFrom(URigHierarchy* InHierarchy, FRigBaseElement* InOt
 	BoneType = Source->BoneType;
 }
 
+////////////////////////////////////////////////////////////////////////////////
+// FRigNullElement
+////////////////////////////////////////////////////////////////////////////////
+
+const FRigBaseElement::EElementIndex FRigNullElement::ElementTypeIndex = NullElement;
+
 ////////////////////////////////////////////////////////////////////////////////
 // FRigControlSettings
 ////////////////////////////////////////////////////////////////////////////////
@@ -1231,6 +1247,8 @@ void FRigControlSettings::SetupLimitArrayForType(bool bLimitTranslation, bool bL
 // FRigControlElement
 ////////////////////////////////////////////////////////////////////////////////
 
+const FRigBaseElement::EElementIndex FRigControlElement::ElementTypeIndex = ControlElement;
+
 void FRigControlElement::Save(FArchive& Ar, URigHierarchy* Hierarchy, ESerializationPhase SerializationPhase)
 {
 	Super::Save(Ar, Hierarchy, SerializationPhase);
@@ -1302,6 +1320,8 @@ void FRigControlElement::CopyPose(FRigBaseElement* InOther, bool bCurrent, bool
 // FRigCurveElement
 ////////////////////////////////////////////////////////////////////////////////
 
+const FRigBaseElement::EElementIndex FRigCurveElement::ElementTypeIndex = CurveElement;
+
 void FRigCurveElement::Save(FArchive& Ar, URigHierarchy* Hierarchy, ESerializationPhase SerializationPhase)
 {
 	Super::Save(Ar, Hierarchy, SerializationPhase);
@@ -1376,6 +1396,8 @@ void FRigRigidBodySettings::Load(FArchive& Ar)
 // FRigRigidBodyElement
 ////////////////////////////////////////////////////////////////////////////////
 
+const FRigBaseElement::EElementIndex FRigRigidBodyElement::ElementTypeIndex = RigidBodyElement;
+
 void FRigRigidBodyElement::Save(FArchive& Ar, URigHierarchy* Hierarchy, ESerializationPhase SerializationPhase)
 {
 	Super::Save(Ar, Hierarchy, SerializationPhase);
@@ -1409,6 +1431,8 @@ void FRigRigidBodyElement::CopyFrom(URigHierarchy* InHierarchy, FRigBaseElement*
 // FRigReferenceElement
 ////////////////////////////////////////////////////////////////////////////////
 
+const FRigBaseElement::EElementIndex FRigReferenceElement::ElementTypeIndex = ReferenceElement;
+
 void FRigReferenceElement::Save(FArchive& Ar, URigHierarchy* Hierarchy, ESerializationPhase SerializationPhase)
 {
 	Super::Save(Ar, Hierarchy, SerializationPhase);
@@ -1527,6 +1551,8 @@ uint32 GetTypeHash(const FRigConnectorSettings& Settings)
 // FRigConnectorElement
 ////////////////////////////////////////////////////////////////////////////////
 
+const FRigBaseElement::EElementIndex FRigConnectorElement::ElementTypeIndex = ConnectorElement;
+
 void FRigConnectorElement::Save(FArchive& Ar, URigHierarchy* Hierarchy, ESerializationPhase SerializationPhase)
 {
 	Super::Save(Ar, Hierarchy, SerializationPhase);
diff --git a/Engine/Plugins/Animation/ControlRig/Source/ControlRig/Public/Rigs/RigHierarchyElements.h b/Engine/Plugins/Animation/ControlRig/Source/ControlRig/Public/Rigs/RigHierarchyElements.h
--- a/Engine/Plugins/Animation/ControlRig/Source/ControlRig/Public/Rigs/RigHierarchyElements.h
+++ b/Engine/Plugins/Animation/ControlRig/Source/ControlRig/Public/Rigs/RigHierarchyElements.h
@@ -38,7 +38,8 @@ template<typename T> \
 friend T* CastChecked(ElementType* InElement) \
 { \
 	return CastChecked<T>((FRigBaseElement*) InElement); \
-}
+} \
+virtual int32 GetElementTypeIndex() const override { return (int32)ElementType::ElementTypeIndex; }
 
 UENUM()
 namespace ERigTransformType
@@ -554,6 +555,25 @@ struct CONTROLRIG_API FRigBaseElement
 
 public:
 
+	enum EElementIndex
+	{
+		BaseElement,
+		TransformElement,
+		SingleParentElement,
+		MultiParentElement,
+		BoneElement,
+		NullElement,
+		ControlElement,
+		CurveElement,
+		RigidBodyElement,
+		ReferenceElement,
+		ConnectorElement,
+
+		Max
+	};
+
+	static const EElementIndex ElementTypeIndex;
+
 	FRigBaseElement()
     : Key()
     , Index(INDEX_NONE)
@@ -571,6 +591,9 @@ public:
 
 	virtual ~FRigBaseElement();
 
+	virtual int32 GetElementTypeIndex() const { return (int32)ElementTypeIndex; }
+	int32 GetElementTypeCount() const { return (int32)EElementIndex::Max; }
+
 	enum ESerializationPhase
 	{
 		StaticData,
@@ -753,6 +776,8 @@ public:
 	GENERATED_BODY()
 	DECLARE_RIG_ELEMENT_METHODS(FRigTransformElement)
 
+	static const EElementIndex ElementTypeIndex;
+
 	FRigTransformElement()
 		: FRigBaseElement()
 	{}
@@ -830,6 +855,8 @@ public:
 	GENERATED_BODY()
 	DECLARE_RIG_ELEMENT_METHODS(FRigSingleParentElement)
 
+	static const EElementIndex ElementTypeIndex;
+
 	FRigSingleParentElement()
 	: FRigTransformElement()
 	, ParentElement(nullptr)
@@ -979,6 +1006,8 @@ public:
 	GENERATED_BODY()
 	DECLARE_RIG_ELEMENT_METHODS(FRigMultiParentElement)
 
+	static const EElementIndex ElementTypeIndex;
+
     FRigMultiParentElement()
     : FRigTransformElement()
 	{}
@@ -1019,6 +1048,8 @@ public:
 	GENERATED_BODY()
 	DECLARE_RIG_ELEMENT_METHODS(FRigBoneElement)
 
+	static const EElementIndex ElementTypeIndex;
+
 	FRigBoneElement()
 		: FRigSingleParentElement()
 	{
@@ -1054,6 +1085,8 @@ public:
 	GENERATED_BODY()
 	DECLARE_RIG_ELEMENT_METHODS(FRigNullElement)
 
+	static const EElementIndex ElementTypeIndex;
+
 	FRigNullElement()
     : FRigMultiParentElement()
 	{
@@ -1368,6 +1401,8 @@ struct CONTROLRIG_API FRigControlElement : public FRigMultiParentElement
 	GENERATED_BODY()
 	DECLARE_RIG_ELEMENT_METHODS(FRigControlElement)
 
+	static const EElementIndex ElementTypeIndex;
+
 	FRigControlElement()
 		: FRigMultiParentElement()
 	{
@@ -1436,6 +1471,8 @@ public:
 	GENERATED_BODY()
 	DECLARE_RIG_ELEMENT_METHODS(FRigCurveElement)
 
+	static const EElementIndex ElementTypeIndex;
+
 	FRigCurveElement()
 		: FRigBaseElement()
 		, bIsValueSet(true)
@@ -1497,6 +1534,8 @@ public:
 	GENERATED_BODY()
 	DECLARE_RIG_ELEMENT_METHODS(FRigRigidBodyElement)
 
+	static const EElementIndex ElementTypeIndex;
+
     FRigRigidBodyElement()
         : FRigSingleParentElement()
 	{
@@ -1535,6 +1574,8 @@ public:
 	GENERATED_BODY()
 	DECLARE_RIG_ELEMENT_METHODS(FRigReferenceElement)
 
+	static const EElementIndex ElementTypeIndex;
+
     FRigReferenceElement()
         : FRigSingleParentElement()
 	{
@@ -1644,6 +1685,8 @@ public:
 	GENERATED_BODY()
 	DECLARE_RIG_ELEMENT_METHODS(FRigConnectorElement)
 
+	static const EElementIndex ElementTypeIndex;
+
 	FRigConnectorElement()
 		: FRigSingleParentElement()
 	{
