Control Rig: Fix ensure when undoing/redoing creation of rig
#rnx

diff --git a/Engine/Plugins/Animation/ControlRig/Source/ControlRig/Private/Rigs/RigHierarchy.cpp b/Engine/Plugins/Animation/ControlRig/Source/ControlRig/Private/Rigs/RigHierarchy.cpp
--- a/Engine/Plugins/Animation/ControlRig/Source/ControlRig/Private/Rigs/RigHierarchy.cpp
+++ b/Engine/Plugins/Animation/ControlRig/Source/ControlRig/Private/Rigs/RigHierarchy.cpp
@@ -230,6 +230,21 @@ void URigHierarchy::Load(FArchive& Ar)
 		Ar << SelectedKeys;
 	}
 
+	// If a controller is found where the outer is this hierarchy, make sure it is configured correctly
+	{
+		TArray<UObject*> ChildObjects;
+		GetObjectsWithOuter(this, ChildObjects, false);
+		ChildObjects = ChildObjects.FilterByPredicate([](UObject* Object)
+			{ return Object->IsA<URigHierarchyController>();});
+		if (!ChildObjects.IsEmpty())
+		{
+			ensure(ChildObjects.Num() == 1); // there should only be one controller
+			bIsControllerAvailable = true;
+			HierarchyController = Cast<URigHierarchyController>(ChildObjects[0]);
+			HierarchyController->SetHierarchy(this);
+		}
+	}
+	
 	Reset();
 
 	int32 ElementCount = 0;
