Control Rig: Improve performance of IsDependentOn

diff --git a/Engine/Plugins/Animation/ControlRig/Source/ControlRig/Private/Rigs/RigHierarchy.cpp b/Engine/Plugins/Animation/ControlRig/Source/ControlRig/Private/Rigs/RigHierarchy.cpp
--- a/Engine/Plugins/Animation/ControlRig/Source/ControlRig/Private/Rigs/RigHierarchy.cpp
+++ b/Engine/Plugins/Animation/ControlRig/Source/ControlRig/Private/Rigs/RigHierarchy.cpp
@@ -373,6 +373,7 @@ void URigHierarchy::Reset_Impl(bool bResetElements)
 	DefaultParentPerElement.Reset();
 	OrderedSelection.Reset();
 	PoseVersionPerElement.Reset();
+	ElementDependencyCache.Reset();
 
 	if(!IsGarbageCollecting())
 	{
@@ -4180,18 +4181,16 @@ FRigElementKey URigHierarchy::GetPreviousParent(const FRigElementKey& InKey) con
 
 bool URigHierarchy::IsParentedTo(FRigBaseElement* InChild, FRigBaseElement* InParent, const TElementDependencyMap& InDependencyMap) const
 {
-	TArray<bool> ElementsVisited;
-	return IsDependentOn(InChild, InParent, ElementsVisited, InDependencyMap);
-}
-
-bool URigHierarchy::IsDependentOn(FRigBaseElement* InDependent, FRigBaseElement* InDependency, TArray<bool>& InElementsVisited, const TElementDependencyMap& InDependencyMap) const
-{
-	if (InElementsVisited.Num() != Elements.Num())
+	ElementDependencyVisited.Reset();
+	if(!InDependencyMap.IsEmpty())
 	{
-		InElementsVisited.Reset();
-		InElementsVisited.AddZeroed(Elements.Num());
+		ElementDependencyVisited.SetNumZeroed(Elements.Num());
 	}
+	return IsDependentOn(InChild, InParent, InDependencyMap, true);
+}
 
+bool URigHierarchy::IsDependentOn(FRigBaseElement* InDependent, FRigBaseElement* InDependency, const TElementDependencyMap& InDependencyMap, bool bIsOnActualTopology) const
+{
 	if((InDependent == nullptr) || (InDependency == nullptr))
 	{
 		return false;
@@ -4203,30 +4202,73 @@ bool URigHierarchy::IsDependentOn(FRigBaseElement* InDependent, FRigBaseElement*
 	}
 
 	const int32 DependentElementIndex = InDependent->GetIndex();
+	const int32 DependencyElementIndex = InDependency->GetIndex();
+	const TTuple<int32,int32> CacheKey(DependentElementIndex, DependencyElementIndex);
 
-	if (!InElementsVisited.IsValidIndex(DependentElementIndex))
+		if(!ElementDependencyCache.IsValid(GetTopologyVersion()))
+		{
+			ElementDependencyCache.Set(TMap<TTuple<int32, int32>, bool>(), GetTopologyVersion());
+		}
+
+	// we'll only update the caches if we are following edges on the actual topology
+	if(const bool* bCachedResult = ElementDependencyCache.Get().Find(CacheKey))
 	{
-		return false;
+		return *bCachedResult;
 	}
-	
-	if (InElementsVisited[DependentElementIndex])
+
+	// check if the reverse dependency check has been stored before - if the dependency is dependent
+	// then we don't need to recurse any further.
+	const TTuple<int32,int32> ReverseCacheKey(DependencyElementIndex, DependentElementIndex);
+	if(const bool* bReverseCachedResult = ElementDependencyCache.Get().Find(ReverseCacheKey))
 	{
-		return false;
+		if(*bReverseCachedResult)
+		{
+			return false;
+		}
 	}
 
-	InElementsVisited[DependentElementIndex] = true;
+	if(!ElementDependencyVisited.IsEmpty())
+	{
+		// when running this with a provided dependency map
+		// we may run into a cycle / infinite recursion. this array
+		// keeps track of all elements being visited before.
+		if(!ElementDependencyVisited.IsValidIndex(DependentElementIndex))
+		{
+			return false;
+		}
+		if (ElementDependencyVisited[DependentElementIndex])
+		{
+			return false;
+		}
+		ElementDependencyVisited[DependentElementIndex] = true;
+	}
 
 	// collect all possible parents of the dependent
-	TArray<FRigBaseElement*> DependentParents;
 	if(const FRigSingleParentElement* SingleParentElement = Cast<FRigSingleParentElement>(InDependent))
 	{
-		 DependentParents.AddUnique(SingleParentElement->ParentElement);
+		if(IsDependentOn(SingleParentElement->ParentElement, InDependency, InDependencyMap, true))
+		{
+			// we'll only update the caches if we are following edges on the actual topology
+			if(bIsOnActualTopology)
+			{
+				ElementDependencyCache.Get().FindOrAdd(CacheKey, true);
+			}
+			return true;
+		}
 	}
 	else if(const FRigMultiParentElement* MultiParentElement = Cast<FRigMultiParentElement>(InDependent))
 	{
 		for(const FRigElementParentConstraint& ParentConstraint : MultiParentElement->ParentConstraints)
 		{
-			 DependentParents.AddUnique(ParentConstraint.ParentElement);
+			if(IsDependentOn(ParentConstraint.ParentElement, InDependency, InDependencyMap, true))
+			{
+				// we'll only update the caches if we are following edges on the actual topology
+				if(bIsOnActualTopology)
+				{
+					ElementDependencyCache.Get().FindOrAdd(CacheKey, true);
+				}
+				return true;
+			}
 		}
 	}
 
@@ -4237,18 +4279,23 @@ bool URigHierarchy::IsDependentOn(FRigBaseElement* InDependent, FRigBaseElement*
 		for(const int32 DependentIndex : DependentIndices)
 		{
 			ensure(Elements.IsValidIndex(DependentIndex));
-			DependentParents.AddUnique(Elements[DependentIndex]);
+			if(IsDependentOn(Elements[DependentIndex], InDependency, InDependencyMap, false))
+			{
+				// we'll only update the caches if we are following edges on the actual topology
+				if(bIsOnActualTopology)
+				{
+					ElementDependencyCache.Get().FindOrAdd(CacheKey, true);
+				}
+				return true;
+			}
 		}
 	}
 
-	for (FRigBaseElement* DependentParent :  DependentParents)
+	// we'll only update the caches if we are following edges on the actual topology
+	if(bIsOnActualTopology)
 	{
-		if(IsDependentOn(DependentParent, InDependency, InElementsVisited, InDependencyMap))
-		{
-			return true;
-		}
+		ElementDependencyCache.Get().FindOrAdd(CacheKey, false);
 	}
-
 	return false;
 }
 
diff --git a/Engine/Plugins/Animation/ControlRig/Source/ControlRig/Public/Rigs/RigHierarchy.h b/Engine/Plugins/Animation/ControlRig/Source/ControlRig/Public/Rigs/RigHierarchy.h
--- a/Engine/Plugins/Animation/ControlRig/Source/ControlRig/Public/Rigs/RigHierarchy.h
+++ b/Engine/Plugins/Animation/ControlRig/Source/ControlRig/Public/Rigs/RigHierarchy.h
@@ -93,6 +93,65 @@ struct FRigTransformStackEntry
 	TArray<FString> Callstack;
 };
 
+template<typename T>
+class CONTROLRIG_API THierarchyCache
+{
+public:
+
+	THierarchyCache()
+		: TopologyVersion(0)
+	{}
+
+	THierarchyCache(const T& InValue, uint32 InTopologyVersion)
+		: THierarchyCache()
+	{
+		Value = InValue;
+		TopologyVersion = InTopologyVersion;
+	}
+
+	bool IsValid(uint32 InTopologyVersion) const
+	{
+		return (TopologyVersion == InTopologyVersion) && Value.IsSet();
+	}
+
+	void Reset()
+	{
+		TopologyVersion = 0;
+		Value = TOptional<T>();
+	}
+
+	const T& Get() const
+	{
+		return Value.GetValue();
+	}
+
+	T& Get()
+	{
+		if(!Value.IsSet())
+		{
+			Value = T();
+		}
+		return Value.GetValue();
+	}
+
+	void Set(uint32 InTopologyVersion)
+	{
+		check(Value.IsSet());
+		TopologyVersion = InTopologyVersion;
+	}
+
+	void Set(const T& InValue, uint32 InTopologyVersion)
+	{
+		Value = InValue;
+		TopologyVersion = InTopologyVersion;
+	}
+
+private:
+
+	uint32 TopologyVersion;
+	TOptional<T> Value;
+};
+
 UCLASS(BlueprintType)
 class CONTROLRIG_API URigHierarchy : public UObject
 {
@@ -3807,16 +3866,18 @@ public:
 	 */
 	bool IsParentedTo(FRigBaseElement* InChild, FRigBaseElement* InParent, const TElementDependencyMap& InDependencyMap = TElementDependencyMap()) const;
 
+private:
 	/**
 	 * Returns true if an element is affected to another element
 	 * @param InDependent The dependent element to check for a dependency
 	 * @param InDependency The dependency element to check for
-	 * @param InElementsVisited An array to keep track of whether an element is visited to avoid infinite recursion
 	 * @param InDependencyMap An additional map of dependencies to respect
+	 * @param bIsOnActualTopology Indicates that the passed dependent and dependency are expected to be on the current topology (if false they are provided with the dependency map)
 	 * @return True if the given dependent is affected by the given dependency 
 	 */
-	bool IsDependentOn(FRigBaseElement* InDependent, FRigBaseElement* InDependency, TArray<bool>& InElementsVisited, const TElementDependencyMap& InDependencyMap = TElementDependencyMap()) const;
+	bool IsDependentOn(FRigBaseElement* InDependent, FRigBaseElement* InDependency, const TElementDependencyMap& InDependencyMap = TElementDependencyMap(), bool bIsOnActualTopology = true) const;
 
+public:
 	/**
 	 * Returns the index of an element given its key within its default parent (or root)
 	 * @param InElement The element to retrieve the index for
@@ -4116,7 +4177,10 @@ private:
 	int32 ResetPoseHash;
 	TArray<bool> ResetPoseIsFilteredOut;
 	TArray<int32> ElementsToRetainLocalTransform;
-	
+
+	mutable THierarchyCache<TMap<TTuple<int32, int32>, bool>> ElementDependencyCache;
+	mutable TArray<bool> ElementDependencyVisited;
+
 #if WITH_EDITOR
 
 	// this is mainly used for propagating changes between hierarchies in the direction of blueprint -> CDO -> other instances
