Control RIg: Additive control rig pickiing up rigid body simulation
#rnx

diff --git a/Engine/Plugins/Animation/ControlRig/Source/ControlRig/Private/Sequencer/ControlRigLayerInstanceProxy.cpp b/Engine/Plugins/Animation/ControlRig/Source/ControlRig/Private/Sequencer/ControlRigLayerInstanceProxy.cpp
--- a/Engine/Plugins/Animation/ControlRig/Source/ControlRig/Private/Sequencer/ControlRigLayerInstanceProxy.cpp
+++ b/Engine/Plugins/Animation/ControlRig/Source/ControlRig/Private/Sequencer/ControlRigLayerInstanceProxy.cpp
@@ -355,13 +355,13 @@ void FAnimNode_ControlRigInputPose::Update_AnyThread(const FAnimationUpdateConte
 	if (InputProxy)
 	{
 		FAnimationUpdateContext InputContext = Context.WithOtherProxy(InputProxy);
-		if (InputPose.GetLinkNode())
+		if (FAnimNode_Base* InputNode = InputPose.GetLinkNode())
 		{
-			InputPose.Update(InputContext);
+			InputProxy->UpdateAnimation_WithRoot(InputContext, InputNode, TEXT("AnimGraph"));
 		}
 		else
 		{
-			FControlRigLayerInstanceProxy::UpdateCustomProxy(InputProxy, InputContext);
+			InputProxy->UpdateAnimationNode(InputContext);
 		}
 	}
 }
@@ -373,22 +373,21 @@ void FAnimNode_ControlRigInputPose::Evaluate_AnyThread(FPoseContext& Output)
 		FBoneContainer& RequiredBones = InputProxy->GetRequiredBones();
 		if (RequiredBones.IsValid())
 		{
-			Output.Pose.SetBoneContainer(&RequiredBones);
-			FPoseContext InputContext(InputProxy, Output.ExpectsAdditivePose());
-
+			FPoseContext InnerOutput(InputProxy, Output.ExpectsAdditivePose());
+			
 			// if no linked node, just use Evaluate of proxy
-			if (InputPose.GetLinkNode())
+			if (FAnimNode_Base* InputNode = InputPose.GetLinkNode())
 			{
-				InputPose.Evaluate(InputContext);
+				InputProxy->EvaluateAnimation_WithRoot(InnerOutput, InputNode);
 			}
 			else
 			{
-				FControlRigLayerInstanceProxy::EvaluateCustomProxy(InputProxy, InputContext);
+				InputProxy->EvaluateAnimationNode(InnerOutput);
 			}
 
-			Output.Pose.MoveBonesFrom(InputContext.Pose);
-			Output.Curve.MoveFrom(InputContext.Curve);
-			Output.CustomAttributes.MoveFrom(InputContext.CustomAttributes);
+			Output.Pose.MoveBonesFrom(InnerOutput.Pose);
+			Output.Curve.MoveFrom(InnerOutput.Curve);
+			Output.CustomAttributes.MoveFrom(InnerOutput.CustomAttributes);
 			return;
 		}
 	}
diff --git a/Engine/Source/Runtime/Engine/Public/Animation/AnimInstanceProxy.h b/Engine/Source/Runtime/Engine/Public/Animation/AnimInstanceProxy.h
--- a/Engine/Source/Runtime/Engine/Public/Animation/AnimInstanceProxy.h
+++ b/Engine/Source/Runtime/Engine/Public/Animation/AnimInstanceProxy.h
@@ -560,6 +560,7 @@ public:
 	friend struct UE::Anim::FAnimSync;
 	friend class UE::Anim::FAnimSyncGroupScope;
 	friend class UE::Anim::FActiveStateMachineScope;
+	friend struct FAnimNode_ControlRigInputPose;
 	
 protected:
 	/** Called when our anim instance is being initialized */
