Control Rig: Fix control rig actors being recreated constantly

diff --git a/Engine/Plugins/Animation/ControlRig/Source/ControlRigEditor/Private/EditMode/ControlRigEditMode.cpp b/Engine/Plugins/Animation/ControlRig/Source/ControlRigEditor/Private/EditMode/ControlRigEditMode.cpp
--- a/Engine/Plugins/Animation/ControlRig/Source/ControlRigEditor/Private/EditMode/ControlRigEditMode.cpp
+++ b/Engine/Plugins/Animation/ControlRig/Source/ControlRigEditor/Private/EditMode/ControlRigEditMode.cpp
@@ -3618,8 +3618,11 @@ void FControlRigEditMode::OnHierarchyModified(ERigHierarchyNotification InNotif,
 				}
 			}
 
-			// if we can't deal with this lazily, let's fall back to recreating all control shape actors
-			RequestToRecreateControlShapeActors(ControlRig);
+			if(ControlRig != nullptr)
+			{
+				// if we can't deal with this lazily, let's fall back to recreating all control shape actors
+				RequestToRecreateControlShapeActors(ControlRig);
+			}
 			break;
 		}
 		case ERigHierarchyNotification::ControlDrivenListChanged:
