RigVM / ControlRig: Fix direct manipulation during PIE

diff --git a/Engine/Plugins/Animation/ControlRig/Source/ControlRig/Private/ControlRig.cpp b/Engine/Plugins/Animation/ControlRig/Source/ControlRig/Private/ControlRig.cpp
--- a/Engine/Plugins/Animation/ControlRig/Source/ControlRig/Private/ControlRig.cpp
+++ b/Engine/Plugins/Animation/ControlRig/Source/ControlRig/Private/ControlRig.cpp
@@ -1996,7 +1996,6 @@ FName UControlRig::AddTransientControl(const URigVMUnitNode* InNode, const FRigD
 	Info->ControlKey = ControlKey;
 	RigUnitManipulationInfos.Add(Info);
 	SetTransientControlValue(InNode, Info);
-	Info->bInitialized = true;
 
 	return Info->ControlKey.Name;
 }
@@ -2024,7 +2023,9 @@ bool UControlRig::SetTransientControlValue(const URigVMUnitNode* InNode, TShared
 	}
 
 	FControlRigExecuteContext& PublicContext = GetExtendedExecuteContext().GetPublicDataSafe<FControlRigExecuteContext>();
-	return UnitInstance->UpdateHierarchyForDirectManipulation(InNode, NodeInstance, PublicContext, InInfo);
+	const bool bResult = UnitInstance->UpdateHierarchyForDirectManipulation(InNode, NodeInstance, PublicContext, InInfo);
+	InInfo->bInitialized = true;
+	return bResult;
 }
 
 FName UControlRig::RemoveTransientControl(const URigVMUnitNode* InNode, const FRigDirectManipulationTarget& InTarget)
diff --git a/Engine/Plugins/Animation/ControlRig/Source/ControlRig/Private/ControlRigGizmoActor.cpp b/Engine/Plugins/Animation/ControlRig/Source/ControlRig/Private/ControlRigGizmoActor.cpp
--- a/Engine/Plugins/Animation/ControlRig/Source/ControlRig/Private/ControlRigGizmoActor.cpp
+++ b/Engine/Plugins/Animation/ControlRig/Source/ControlRig/Private/ControlRigGizmoActor.cpp
@@ -17,6 +17,7 @@ AControlRigShapeActor::AControlRigShapeActor(const FObjectInitializer& ObjectIni
 	, ControlName(NAME_None)
 	, ShapeName(NAME_None)
 	, OverrideColor(0, 0, 0, 0)
+	, OffsetTransform(FTransform::Identity)
 	, bSelected(false)
 	, bHovered(false)
 	, bSelectable(true)
diff --git a/Engine/Plugins/Animation/ControlRig/Source/ControlRig/Private/Rigs/RigHierarchy.cpp b/Engine/Plugins/Animation/ControlRig/Source/ControlRig/Private/Rigs/RigHierarchy.cpp
--- a/Engine/Plugins/Animation/ControlRig/Source/ControlRig/Private/Rigs/RigHierarchy.cpp
+++ b/Engine/Plugins/Animation/ControlRig/Source/ControlRig/Private/Rigs/RigHierarchy.cpp
@@ -2581,14 +2581,15 @@ void URigHierarchy::IncrementTopologyVersion()
 FRigPose URigHierarchy::GetPose(
 	bool bInitial,
 	ERigElementType InElementType,
-	const FRigElementKeyCollection& InItems 
+	const FRigElementKeyCollection& InItems,
+	bool bIncludeTransientControls
 ) const
 {
-	return GetPose(bInitial, InElementType, TArrayView<const FRigElementKey>(InItems.Keys.GetData(), InItems.Num()));
+	return GetPose(bInitial, InElementType, TArrayView<const FRigElementKey>(InItems.Keys.GetData(), InItems.Num()), bIncludeTransientControls);
 }
 
 FRigPose URigHierarchy::GetPose(bool bInitial, ERigElementType InElementType,
-	const TArrayView<const FRigElementKey>& InItems) const
+	const TArrayView<const FRigElementKey>& InItems, bool bIncludeTransientControls) const
 {
 	LLM_SCOPE_BYNAME(TEXT("Animation/ControlRig"));
 	FRigPose Pose;
@@ -2626,6 +2627,11 @@ FRigPose URigHierarchy::GetPose(bool bInitial, ERigElementType InElementType,
 			{
 				PoseElement.PreferredEulerAngle = GetControlPreferredEulerAngles(ControlElement,
 					GetControlPreferredEulerRotationOrder(ControlElement), bInitial);
+
+				if(!bIncludeTransientControls && ControlElement->Settings.bIsTransientControl)
+				{
+					continue;
+				}
 			}
 		}
 		else if(FRigCurveElement* CurveElement = Cast<FRigCurveElement>(Element))
diff --git a/Engine/Plugins/Animation/ControlRig/Source/ControlRig/Public/ControlRig.h b/Engine/Plugins/Animation/ControlRig/Source/ControlRig/Public/ControlRig.h
--- a/Engine/Plugins/Animation/ControlRig/Source/ControlRig/Public/ControlRig.h
+++ b/Engine/Plugins/Animation/ControlRig/Source/ControlRig/Public/ControlRig.h
@@ -716,6 +716,8 @@ public:
 	bool bRecordSelectionPoseForConstructionMode;
 	TMap<FRigElementKey, FTransform> SelectionPoseForConstructionMode;
 	bool bIsClearingTransientControls;
+
+	FRigPose InputPoseOnDebuggedRig;
 	
 #endif
 	
diff --git a/Engine/Plugins/Animation/ControlRig/Source/ControlRig/Public/ControlRigGizmoActor.h b/Engine/Plugins/Animation/ControlRig/Source/ControlRig/Public/ControlRigGizmoActor.h
--- a/Engine/Plugins/Animation/ControlRig/Source/ControlRig/Public/ControlRigGizmoActor.h
+++ b/Engine/Plugins/Animation/ControlRig/Source/ControlRig/Public/ControlRigGizmoActor.h
@@ -88,6 +88,8 @@ public:
 
 	FLinearColor OverrideColor;
 
+	FTransform OffsetTransform;
+
 	UFUNCTION(BlueprintSetter)
 	/** Set the control to be enabled/disabled */
 	virtual void SetEnabled(bool bInEnabled) { SetSelectable(bInEnabled); }
diff --git a/Engine/Plugins/Animation/ControlRig/Source/ControlRig/Public/Rigs/RigHierarchy.h b/Engine/Plugins/Animation/ControlRig/Source/ControlRig/Public/Rigs/RigHierarchy.h
--- a/Engine/Plugins/Animation/ControlRig/Source/ControlRig/Public/Rigs/RigHierarchy.h
+++ b/Engine/Plugins/Animation/ControlRig/Source/ControlRig/Public/Rigs/RigHierarchy.h
@@ -3397,13 +3397,15 @@ public:
 	 * Returns the current / initial pose of the hierarchy
 	 * @param bInitial If set to true the initial pose will be returned
 	 * @return The pose of the hierarchy
+	 * @param bIncludeTransientControls If true the transient controls will be included in the pose
 	 */
 	UFUNCTION(BlueprintCallable, Category = URigHierarchy)
 	FRigPose GetPose(
-		bool bInitial = false
+		bool bInitial = false,
+		bool bIncludeTransientControls = true
 	) const
 	{
-		return GetPose(bInitial, ERigElementType::All, FRigElementKeyCollection());
+		return GetPose(bInitial, ERigElementType::All, FRigElementKeyCollection(), bIncludeTransientControls);
 	}
 
 	/**
@@ -3411,12 +3413,14 @@ public:
 	 * @param bInitial If set to true the initial pose will be returned
 	 * @param InElementType The types of elements to get
 	 * @param InItems An optional list of items to get
+	 * @param bIncludeTransientControls If true the transient controls will be included in the pose
 	 * @return The pose of the hierarchy
 	 */
 	FRigPose GetPose(
 		bool bInitial,
 		ERigElementType InElementType,
-		const FRigElementKeyCollection& InItems 
+		const FRigElementKeyCollection& InItems ,
+		bool bIncludeTransientControls = true
 	) const;
 
 	/**
@@ -3424,12 +3428,14 @@ public:
 	 * @param bInitial If set to true the initial pose will be returned
 	 * @param InElementType The types of elements to get
 	 * @param InItems An optional list of items to get
+	 * @param bIncludeTransientControls If true the transient controls will be included in the pose
 	 * @return The pose of the hierarchy
 	 */
 	FRigPose GetPose(
 		bool bInitial,
 		ERigElementType InElementType,
-		const TArrayView<const FRigElementKey>& InItems 
+		const TArrayView<const FRigElementKey>& InItems,
+		bool bIncludeTransientControls = true
 	) const;
 
 	/**
diff --git a/Engine/Plugins/Animation/ControlRig/Source/ControlRigEditor/Private/EditMode/ControlRigEditMode.cpp b/Engine/Plugins/Animation/ControlRig/Source/ControlRigEditor/Private/EditMode/ControlRigEditMode.cpp
--- a/Engine/Plugins/Animation/ControlRig/Source/ControlRigEditor/Private/EditMode/ControlRigEditMode.cpp
+++ b/Engine/Plugins/Animation/ControlRig/Source/ControlRigEditor/Private/EditMode/ControlRigEditMode.cpp
@@ -141,6 +141,7 @@ namespace ControlRigSelectionConstants
 
 FControlRigEditMode::FControlRigEditMode()
 	: bIsChangingControlShapeTransform(false)
+	, bIsTracking(false)
 	, bManipulatorMadeChange(false)
 	, bSelecting(false)
 	, bSelectionChanged(false)
@@ -1027,6 +1028,7 @@ bool FControlRigEditMode::EndTracking(FEditorViewportClient* InViewportClient, F
 	const bool bWasInteracting = bManipulatorMadeChange && InteractionType != (uint8)EControlRigInteractionType::None;
 	
 	InteractionType = (uint8)EControlRigInteractionType::None;
+	bIsTracking = false;
 	
 	if (InteractionScopes.Num() > 0)
 	{		
@@ -1075,6 +1077,7 @@ bool FControlRigEditMode::StartTracking(FEditorViewportClient* InViewportClient,
 	}
 
 	InteractionType = GetInteractionType(InViewportClient);
+	bIsTracking = true;
 	
 	if (InteractionScopes.Num() == 0)
 	{
@@ -1851,6 +1854,8 @@ bool FControlRigEditMode::InputDelta(FEditorViewportClient* InViewportClient, FV
 					bool bCalcLocal = bDoLocal;
 					bool bFirstTime = true;
 					FTransform InOutLocal = FTransform::Identity;
+					
+					bool const bJustStartedManipulation = !bManipulatorMadeChange;
 
 					for (AControlRigShapeActor* ShapeActor : Pairs.Value)
 					{
@@ -1874,6 +1879,14 @@ bool FControlRigEditMode::InputDelta(FEditorViewportClient* InViewportClient, FV
 								}
 							}
 
+							if(bJustStartedManipulation)
+							{
+								if(const FRigControlElement* ControlElement = Pairs.Key->FindControl(ShapeActor->ControlName))
+								{
+									ShapeActor->OffsetTransform = Pairs.Key->GetHierarchy()->GetGlobalControlOffsetTransform(ControlElement->GetKey(), false);
+								}
+							}
+
 							MoveControlShape(ShapeActor, bDoTranslation, InDrag, bDoRotation, InRot, bDoScale, InScale, ComponentTransform,
 								bUseLocal, bDoLocal, InOutLocal);
 							bManipulatorMadeChange = true;
@@ -4074,6 +4087,13 @@ void FControlRigEditMode::MoveControlShape(AControlRigShapeActor* ShapeActor, co
 				FTransform CurrentTransform  = ControlRig->GetControlGlobalTransform(ShapeActor->ControlName);			// assumes it's attached to actor
 				CurrentTransform = ToWorldTransform * CurrentTransform;
 
+				// make the transform relative to the offset transform again.
+				// first we'll make it relative to the offset used at the time of starting the drag
+				// and then we'll make it absolute again based on the current offset. these two can be
+				// different if we are interacting on a control on an animated character
+				CurrentTransform = CurrentTransform.GetRelativeTransform(ShapeActor->OffsetTransform);
+				CurrentTransform = CurrentTransform * ControlRig->GetHierarchy()->GetGlobalControlOffsetTransform(ShapeActor->GetElementKey(), false);
+				
 				ShapeActor->SetGlobalTransform(CurrentTransform);
 
 				ControlRig->Evaluate_AnyThread();
@@ -4130,10 +4150,30 @@ void FControlRigEditMode::MoveControlShape(AControlRigShapeActor* ShapeActor, co
 					bPrintPythonCommands = World->IsPreviewWorld();
 				}
 
+				bool bIsTransientControl = false;
+				if(const FRigControlElement* ControlElement = ControlRig->FindControl(ShapeActor->ControlName))
+				{
+					bIsTransientControl = ControlElement->Settings.bIsTransientControl;
+				}
+
+				// if we are operating on a PIE instance which is playing we need to reapply the input pose
+				// since the hierarchy will also have been brought into the solved pose. by reapplying the
+				// input pose we avoid double transformation / double forward solve results.
+				if(bIsTransientControl && ControlRig->GetWorld()->IsPlayInEditor() && !ControlRig->GetWorld()->IsPaused())
+				{
+					ControlRig->GetHierarchy()->SetPose(ControlRig->InputPoseOnDebuggedRig);
+				}
+				
 				ControlRig->Evaluate_AnyThread();
 				SetControlShapeTransform(ShapeActor, NewTransform, ToWorldTransform, Context, bPrintPythonCommands);
 				NotifyDrivenControls(ControlRig, ShapeActor->GetElementKey());
-				ControlRig->Evaluate_AnyThread();
+				if(const FRigControlElement* ControlElement = ControlRig->FindControl(ShapeActor->ControlName))
+				{
+					if(!bIsTransientControl)
+					{
+						ControlRig->Evaluate_AnyThread();
+					}
+				}
 				ShapeActor->SetGlobalTransform(CurrentTransform);
 				if (bCalcLocal)
 				{
diff --git a/Engine/Plugins/Animation/ControlRig/Source/ControlRigEditor/Private/Editor/ControlRigEditor.cpp b/Engine/Plugins/Animation/ControlRig/Source/ControlRigEditor/Private/Editor/ControlRigEditor.cpp
--- a/Engine/Plugins/Animation/ControlRig/Source/ControlRigEditor/Private/Editor/ControlRigEditor.cpp
+++ b/Engine/Plugins/Animation/ControlRig/Source/ControlRigEditor/Private/Editor/ControlRigEditor.cpp
@@ -221,7 +221,7 @@ void FControlRigEditor::InitRigVMEditor(const EToolkitMode::Type Mode, const TSh
 		EditMode->OnGetContextMenu() = FOnGetContextMenu::CreateSP(this, &FControlRigEditor::HandleOnGetViewportContextMenuDelegate);
 		EditMode->OnContextMenuCommands() = FNewMenuCommandsDelegate::CreateSP(this, &FControlRigEditor::HandleOnViewportContextMenuCommandsDelegate);
 		EditMode->OnAnimSystemInitialized().Add(FSimpleMulticastDelegate::FDelegate::CreateRaw(this, &FControlRigEditor::OnAnimInitialized));
-		
+	
 		PersonaToolkit->GetPreviewScene()->SetRemoveAttachedComponentFilter(FOnRemoveAttachedComponentFilter::CreateSP(EditMode, &FControlRigEditMode::CanRemoveFromPreviewScene));
 	}
 
@@ -802,8 +802,10 @@ void FControlRigEditor::HandleSetObjectBeingDebugged(UObject* InObject)
 		if(!PreviouslyDebuggedControlRig->HasAnyFlags(RF_BeginDestroyed))
 		{
 			PreviouslyDebuggedControlRig->GetHierarchy()->OnModified().RemoveAll(this);
+			PreviouslyDebuggedControlRig->OnPreForwardsSolve_AnyThread().RemoveAll(this);
 			PreviouslyDebuggedControlRig->OnPreConstructionForUI_AnyThread().RemoveAll(this);
 			PreviouslyDebuggedControlRig->OnPostConstruction_AnyThread().RemoveAll(this);
+			PreviouslyDebuggedControlRig->ControlModified().RemoveAll(this);
 		}
 	}
 
@@ -855,8 +857,11 @@ void FControlRigEditor::HandleSetObjectBeingDebugged(UObject* InObject)
 		}
 
 		DebuggedControlRig->GetHierarchy()->OnModified().AddSP(this, &FControlRigEditor::OnHierarchyModified_AnyThread);
+		DebuggedControlRig->OnPreForwardsSolve_AnyThread().AddSP(this, &FControlRigEditor::OnPreForwardsSolve_AnyThread);
 		DebuggedControlRig->OnPreConstructionForUI_AnyThread().AddSP(this, &FControlRigEditor::OnPreConstruction_AnyThread);
 		DebuggedControlRig->OnPostConstruction_AnyThread().AddSP(this, &FControlRigEditor::OnPostConstruction_AnyThread);
+		DebuggedControlRig->ControlModified().AddSP(this, &FControlRigEditor::HandleOnControlModified);
+
 		LastHierarchyHash = INDEX_NONE;
 
 		if(EditorSkelComp)
@@ -1373,9 +1378,19 @@ void FControlRigEditor::HandleVMExecutedEvent(URigVMHost* InHost, const FName& I
 						{
 							continue;
 						}
+
+						// if we are not manipulating right now - reset the info so that it can follow the hierarchy
+						if (FControlRigEditorEditMode* EditMode = GetEditMode())
+						{
+							if(!EditMode->bIsTracking)
+							{
+								ManipulationInfo->Reset();
+							}
+						}
 				
 						FRigUnit* UnitInstance = UControlRig::GetRigUnitInstanceFromScope(NodeInstance);
 						UnitInstance->UpdateHierarchyForDirectManipulation(Node, NodeInstance, ExecuteContext, ManipulationInfo);
+						ManipulationInfo->bInitialized = true;
 						UnitInstance->PerformDebugDrawingForDirectManipulation(Node, NodeInstance, ExecuteContext, ManipulationInfo);
 					}
 				}
@@ -1913,6 +1928,13 @@ void FControlRigEditor::HandlePreviewSceneCreated(const TSharedRef<IPersonaPrevi
 	check(FloorMesh);
 	check(DefaultMaterial);
 
+	// leave some metadata on the world used for debug object labeling
+	if(FWorldContext* WorldContext = GEngine->GetWorldContextFromWorld(InPersonaPreviewScene->GetWorld()))
+	{
+		static constexpr TCHAR Format[] = TEXT("ControlRigEditor (%s)");
+		WorldContext->CustomDescription = FString::Printf(Format, *GetBlueprintObj()->GetName());
+	}
+
 	// create ground mesh actor
 	AStaticMeshActor* GroundActor = InPersonaPreviewScene->GetWorld()->SpawnActor<AStaticMeshActor>(AStaticMeshActor::StaticClass(), FTransform::Identity);
 	GroundActor->SetFlags(RF_Transient);
@@ -2005,6 +2027,7 @@ void FControlRigEditor::UpdateRigVMHost()
 				EditMode->SetObjects(ControlRig, EditorSkelComp,nullptr);
 			}
 
+			ControlRig->OnPreForwardsSolve_AnyThread().AddSP(this, &FControlRigEditor::OnPreForwardsSolve_AnyThread);
 			ControlRig->ControlModified().AddSP(this, &FControlRigEditor::HandleOnControlModified);
 		}
 	}
@@ -4004,7 +4027,8 @@ void FControlRigEditor::HandleMakeElementGetterSetter(ERigElementGetterSetterTyp
 
 void FControlRigEditor::HandleOnControlModified(UControlRig* Subject, FRigControlElement* ControlElement, const FRigControlModifiedContext& Context)
 {
-	if (Subject != GetControlRig())
+	UControlRig* DebuggedControlRig = Cast<UControlRig>(GetBlueprintObj()->GetObjectBeingDebugged());
+	if (Subject != DebuggedControlRig)
 	{
 		return;
 	}
@@ -4020,7 +4044,6 @@ void FControlRigEditor::HandleOnControlModified(UControlRig* Subject, FRigContro
 	if (ControlElement->Settings.bIsTransientControl && !GIsTransacting)
 	{
 		const URigVMUnitNode* UnitNode = nullptr;
-		UControlRig* DebuggedControlRig = Cast<UControlRig>(GetBlueprintObj()->GetObjectBeingDebugged());
 		const FString NodeName = UControlRig::GetNodeNameFromTransientControl(ControlElement->GetKey());
 		const FString PoseTarget = UControlRig::GetTargetFromTransientControl(ControlElement->GetKey());
 		TSharedPtr<FStructOnScope> NodeInstance;
@@ -4050,12 +4073,16 @@ void FControlRigEditor::HandleOnControlModified(UControlRig* Subject, FRigContro
 			FRigUnit* UnitInstance = DebuggedControlRig->GetRigUnitInstanceFromScope(NodeInstance);
 			check(UnitInstance);
 
+			const FRigPose Pose = DebuggedControlRig->GetHierarchy()->GetPose();
+
 			// update the node based on the incoming pose. once that is done we'll need to compare the node instance
 			// with the settings on the node in the graph and update them accordingly.
 			FControlRigExecuteContext& ExecuteContext = DebuggedControlRig->GetExtendedExecuteContext().GetPublicDataSafe<FControlRigExecuteContext>();
 			if(UnitInstance->UpdateDirectManipulationFromHierarchy(UnitNode, NodeInstance, ExecuteContext, ManipulationInfo))
 			{
 				UnitNode->UpdateHostFromStructInstance(DebuggedControlRig, NodeInstance);
+				DebuggedControlRig->GetHierarchy()->SetPose(Pose);
+				
 				URigVMController* Controller = Blueprint->GetOrCreateController(UnitNode->GetGraph());
 				TMap<FString, FString> PinPathToNewDefaultValue;
 				UnitNode->ComputePinValueDifferences(NodeInstance, PinPathToNewDefaultValue);
@@ -4177,6 +4204,26 @@ TSharedPtr<FUICommandList> FControlRigEditor::HandleOnViewportContextMenuCommand
 	return TSharedPtr<FUICommandList>();
 }
 
+void FControlRigEditor::OnPreForwardsSolve_AnyThread(UControlRig* InRig, const FName& InEventName)
+{
+	// if we are debugging a PIE instance, we need to remember the input pose on the
+	// rig so we can perform multiple evaluations. this is to avoid double transforms / double forward solve results.
+	if(InRig->GetWorld()->IsPlayInEditor())
+	{
+		if(!InRig->GetWorld()->IsPaused())
+		{
+			// store the pose while PIE is running
+			InRig->InputPoseOnDebuggedRig = InRig->GetHierarchy()->GetPose(false, false);
+		}
+		else
+		{
+			// reapply the pose as PIE is paused. during pause the rig won't be updated with the input pose
+			// from the animbp / client thus we need to reset the pose to avoid double transformation.
+			InRig->GetHierarchy()->SetPose(InRig->InputPoseOnDebuggedRig);
+		}
+	}
+}
+
 void FControlRigEditor::OnPreConstruction_AnyThread(UControlRig* InRig, const FName& InEventName)
 {
 	bIsConstructionEventRunning = true;
diff --git a/Engine/Plugins/Animation/ControlRig/Source/ControlRigEditor/Private/Editor/ControlRigEditor.h b/Engine/Plugins/Animation/ControlRig/Source/ControlRigEditor/Private/Editor/ControlRigEditor.h
--- a/Engine/Plugins/Animation/ControlRig/Source/ControlRigEditor/Private/Editor/ControlRigEditor.h
+++ b/Engine/Plugins/Animation/ControlRig/Source/ControlRigEditor/Private/Editor/ControlRigEditor.h
@@ -293,6 +293,7 @@ protected:
 	int32 RigHierarchyTabCount;
 	TWeakObjectPtr<AStaticMeshActor> WeakGroundActorPtr;
 
+	void OnPreForwardsSolve_AnyThread(UControlRig* InRig, const FName& InEventName);
 	void OnPreConstruction_AnyThread(UControlRig* InRig, const FName& InEventName);
 	void OnPostConstruction_AnyThread(UControlRig* InRig, const FName& InEventName);
 	FRigPose PreConstructionPose;
diff --git a/Engine/Plugins/Animation/ControlRig/Source/ControlRigEditor/Public/EditMode/ControlRigEditMode.h b/Engine/Plugins/Animation/ControlRig/Source/ControlRigEditor/Public/EditMode/ControlRigEditMode.h
--- a/Engine/Plugins/Animation/ControlRig/Source/ControlRigEditor/Public/EditMode/ControlRigEditMode.h
+++ b/Engine/Plugins/Animation/ControlRig/Source/ControlRigEditor/Public/EditMode/ControlRigEditMode.h
@@ -287,6 +287,9 @@ protected:
 	/** The scope for the interaction, one per manipulated Control rig */
 	TMap<UControlRig*,FControlRigInteractionScope*> InteractionScopes;
 
+	/** True if there's tracking going on right now */
+	bool bIsTracking;
+
 	/** Whether a manipulator actually made a change when transacting */
 	bool bManipulatorMadeChange;
 
diff --git a/Engine/Plugins/Runtime/RigVM/Source/RigVMEditor/Private/Editor/RigVMEditor.cpp b/Engine/Plugins/Runtime/RigVM/Source/RigVMEditor/Private/Editor/RigVMEditor.cpp
--- a/Engine/Plugins/Runtime/RigVM/Source/RigVMEditor/Private/Editor/RigVMEditor.cpp
+++ b/Engine/Plugins/Runtime/RigVM/Source/RigVMEditor/Private/Editor/RigVMEditor.cpp
@@ -76,6 +76,9 @@ FRigVMEditor::~FRigVMEditor()
 	{
 		FCoreUObjectDelegates::OnObjectPropertyChanged.Remove(PropertyChangedHandle);
 	}
+	
+	FEditorDelegates::EndPIE.RemoveAll(this);
+    FEditorDelegates::CancelPIE.RemoveAll(this);
 
 	if (RigVMBlueprint)
 	{
@@ -122,6 +125,9 @@ void FRigVMEditor::InitRigVMEditor(const EToolkitMode::Type Mode, const TSharedP
 		Toolbar = MakeShareable(new FBlueprintEditorToolbar(SharedThis(this)));
 	}
 
+	FEditorDelegates::EndPIE.AddRaw(this, &FRigVMEditor::OnPIEStopped);
+	FEditorDelegates::CancelPIE.AddRaw(this, &FRigVMEditor::OnPIEStopped, false);
+
 	// Build up a list of objects being edited in this asset editor
 	TArray<UObject*> ObjectsBeingEdited;
 	ObjectsBeingEdited.Add(InRigVMBlueprint);
@@ -1004,6 +1010,13 @@ void FRigVMEditor::Compile()
 			return;
 		}
 
+		// if we are running PIE - only compile the RigVM 
+		if(IsPIERunning())
+		{
+			RigVMBlueprint->RecompileVM();
+			return;
+		}
+
 		// force to disable the supended notif brackets
 		RigVMBlueprint->bSuspendModelNotificationsForOthers = false;
 		RigVMBlueprint->bSuspendModelNotificationsForSelf = false;
@@ -3287,6 +3300,11 @@ void FRigVMEditor::UpdateGraphCompilerErrors()
 
 }
 
+bool FRigVMEditor::IsPIERunning()
+{
+	return GEditor && (GEditor->PlayWorld != nullptr);
+}
+
 TArray<FName> FRigVMEditor::GetDefaultEventQueue() const
 {
 	return TArray<FName>();
@@ -3450,40 +3468,51 @@ void FRigVMEditor::GetCustomDebugObjects(TArray<FCustomDebugObject>& DebugList)
 			TArray<UObject*> ArchetypeInstances;
 			DefaultObject->GetArchetypeInstances(ArchetypeInstances);
 
-			for (UObject* Instance : ArchetypeInstances)
+			// run in two passes - find the PIE related objects first
+			for (int32 Pass = 0; Pass < 2 ; Pass++)
 			{
-				URigVMHost* InstancedHost = Cast<URigVMHost>(Instance);
-				if (InstancedHost && IsValid(InstancedHost) && InstancedHost != GetRigVMHost())
+				for (UObject* Instance : ArchetypeInstances)
 				{
-					if (InstancedHost->GetOuter() == nullptr)
+					URigVMHost* InstancedHost = Cast<URigVMHost>(Instance);
+					if (InstancedHost && IsValid(InstancedHost) && InstancedHost != GetRigVMHost())
 					{
-						continue;
-					}
+						if (InstancedHost->GetOuter() == nullptr)
+						{
+							continue;
+						}
 
-					UWorld* World = InstancedHost->GetWorld();
-					if (World == nullptr)
-					{
-						continue;
-					}
+						UWorld* World = InstancedHost->GetWorld();
+						if (World == nullptr)
+						{
+							continue;
+						}
 
-					// ensure to only allow preview actors in preview worlds
-					if (World->IsPreviewWorld())
-					{
-						if (!Local::OuterNameContainsRecursive(InstancedHost, TEXT("Preview")))
+						// during pass 0 only do PIE instances,
+						// and in pass 1 only do non PIE instances
+						if((Pass == 1) == (World->IsPlayInEditor()))
 						{
 							continue;
 						}
-					}
 
-					if (Local::IsPendingKillOrUnreachableRecursive(InstancedHost))
-					{
-						continue;
-					}
+						// ensure to only allow preview actors in preview worlds
+						if (World->IsPreviewWorld())
+						{
+							if (!Local::OuterNameContainsRecursive(InstancedHost, TEXT("Preview")))
+							{
+								continue;
+							}
+						}
+
+						if (Local::IsPendingKillOrUnreachableRecursive(InstancedHost))
+						{
+							continue;
+						}
 
-					FCustomDebugObject DebugObject;
-					DebugObject.Object = InstancedHost;
-					DebugObject.NameOverride = GetCustomDebugObjectLabel(InstancedHost);
-					DebugList.Add(DebugObject);
+						FCustomDebugObject DebugObject;
+						DebugObject.Object = InstancedHost;
+						DebugObject.NameOverride = GetCustomDebugObjectLabel(InstancedHost);
+						DebugList.Add(DebugObject);
+					}
 				}
 			}
 		}
@@ -3492,6 +3521,14 @@ void FRigVMEditor::GetCustomDebugObjects(TArray<FCustomDebugObject>& DebugList)
 
 void FRigVMEditor::HandleSetObjectBeingDebugged(UObject* InObject)
 {
+	if(URigVMHost* PreviouslyDebuggedHost = Cast<URigVMHost>(GetBlueprintObj()->GetObjectBeingDebugged()))
+	{
+		if(!PreviouslyDebuggedHost->HasAnyFlags(RF_BeginDestroyed))
+		{
+			PreviouslyDebuggedHost->OnExecuted_AnyThread().RemoveAll(this);
+		}
+	}
+	
 	URigVMHost* DebuggedHost = Cast<URigVMHost>(InObject);
 
 	if (DebuggedHost == nullptr)
@@ -3521,6 +3558,7 @@ void FRigVMEditor::HandleSetObjectBeingDebugged(UObject* InObject)
 	if(DebuggedHost)
 	{
 		DebuggedHost->SetLog(&RigVMLog);
+		DebuggedHost->OnExecuted_AnyThread().AddSP(this, &FRigVMEditor::HandleVMExecutedEvent);
 	}
 
 	RefreshDetailView();
@@ -3539,12 +3577,29 @@ FString FRigVMEditor::GetCustomDebugObjectLabel(UObject* ObjectBeingDebugged) co
 		return TEXT("Editor Preview");
 	}
 
-	if (AActor* ParentActor = ObjectBeingDebugged->GetTypedOuter<AActor>())
+	if (const AActor* ParentActor = ObjectBeingDebugged->GetTypedOuter<AActor>())
 	{
-		return FString::Printf(TEXT("%s in %s"), *GetBlueprintObj()->GetName(), *ParentActor->GetName());
+		if(const UWorld* World = ParentActor->GetWorld())
+		{
+			FString WorldLabel = GetDebugStringForWorld(World);
+			if(World->IsPlayInEditor())
+			{
+				static const FString PIEPrefix = TEXT("PIE");
+				WorldLabel = PIEPrefix;
+			}
+			return FString::Printf(TEXT("%s: %s in %s"), *WorldLabel, *GetBlueprintObj()->GetName(), *ParentActor->GetActorLabel());
+		}
 	}
 
 	return GetBlueprintObj()->GetName();
 }
 
+void FRigVMEditor::OnPIEStopped(bool bSimulation)
+{
+	if(URigVMBlueprint* Blueprint = GetRigVMBlueprint())
+	{
+		Blueprint->SetObjectBeingDebugged(GetRigVMHost());
+	}
+}
+
 #undef LOCTEXT_NAMESPACE 
diff --git a/Engine/Plugins/Runtime/RigVM/Source/RigVMEditor/Public/Editor/RigVMEditor.h b/Engine/Plugins/Runtime/RigVM/Source/RigVMEditor/Public/Editor/RigVMEditor.h
--- a/Engine/Plugins/Runtime/RigVM/Source/RigVMEditor/Public/Editor/RigVMEditor.h
+++ b/Engine/Plugins/Runtime/RigVM/Source/RigVMEditor/Public/Editor/RigVMEditor.h
@@ -294,8 +294,13 @@ protected:
 	/** Once the log is collected update the graph */
 	void UpdateGraphCompilerErrors();
 
+	/** Returns true if PIE is currently running */
+	static bool IsPIERunning();
+
 private:
 
+	void OnPIEStopped(bool bSimulation);
+
 	/** Our currently running rig vm instance */
 	TObjectPtr<URigVMHost> Host;
 
diff --git a/Engine/Source/Editor/AnimationBlueprintEditor/Private/AnimationBlueprintEditor.cpp b/Engine/Source/Editor/AnimationBlueprintEditor/Private/AnimationBlueprintEditor.cpp
--- a/Engine/Source/Editor/AnimationBlueprintEditor/Private/AnimationBlueprintEditor.cpp
+++ b/Engine/Source/Editor/AnimationBlueprintEditor/Private/AnimationBlueprintEditor.cpp
@@ -434,6 +434,13 @@ void FAnimationBlueprintEditor::InitAnimationBlueprintEditor(const EToolkitMode:
 	PersonaToolkit->GetPreviewScene()->SetDefaultAnimationMode(EPreviewSceneDefaultAnimationMode::AnimationBlueprint);
 	PersonaToolkit->GetPreviewScene()->RegisterOnPreviewMeshChanged(FOnPreviewMeshChanged::CreateSP(this, &FAnimationBlueprintEditor::HandlePreviewMeshChanged));
 
+	// leave some metadata on the world used for debug object labeling
+	if(FWorldContext* WorldContext = GEngine->GetWorldContextFromWorld(PersonaToolkit->GetPreviewScene()->GetWorld()))
+	{
+		static constexpr TCHAR Format[] = TEXT("AnimBPEditor (%s)");
+		WorldContext->CustomDescription = FString::Printf(Format, *InAnimBlueprint->GetName());
+	}
+
 	PersonaModule.RecordAssetOpened(InAnimBlueprint);
 
 	if(InAnimBlueprint->BlueprintType != BPTYPE_Interface && !InAnimBlueprint->bIsTemplate)
