Sequencer: Removed unused control rig spawner code

diff --git a/Engine/Plugins/Animation/ControlRig/Source/ControlRig/Private/ControlRigModule.cpp b/Engine/Plugins/Animation/ControlRig/Source/ControlRig/Private/ControlRigModule.cpp
--- a/Engine/Plugins/Animation/ControlRig/Source/ControlRig/Private/ControlRigModule.cpp
+++ b/Engine/Plugins/Animation/ControlRig/Source/ControlRig/Private/ControlRigModule.cpp
@@ -5,7 +5,6 @@
 #include "ControlRigGizmoActor.h"
 #include "TransformableRegistry.h"
 #include "Modules/ModuleManager.h"
-#include "Sequencer/ControlRigObjectSpawner.h"
 #include "ILevelSequenceModule.h"
 #include "ControlRigObjectVersion.h"
 #include "Constraints/ControlRigTransformableHandle.h"
@@ -19,9 +18,6 @@ FDevVersionRegistration GRegisterControlRigObjectVersion(FControlRigObjectVersio
 
 void FControlRigModule::StartupModule()
 {
-	ILevelSequenceModule& LevelSequenceModule = FModuleManager::LoadModuleChecked<ILevelSequenceModule>("LevelSequence");
-	OnCreateMovieSceneObjectSpawnerHandle = LevelSequenceModule.RegisterObjectSpawner(FOnCreateMovieSceneObjectSpawner::CreateStatic(&FControlRigObjectSpawner::CreateObjectSpawner));
-
 	ManipulatorMaterial = LoadObject<UMaterial>(nullptr, TEXT("/ControlRig/M_Manip.M_Manip"));
 
 	RegisterTransformableCustomization();
@@ -29,11 +25,6 @@ void FControlRigModule::StartupModule()
 
 void FControlRigModule::ShutdownModule()
 {
-	ILevelSequenceModule* LevelSequenceModule = FModuleManager::GetModulePtr<ILevelSequenceModule>("LevelSequence");
-	if (LevelSequenceModule)
-	{
-		LevelSequenceModule->UnregisterObjectSpawner(OnCreateMovieSceneObjectSpawnerHandle);
-	}
 }
 
 void FControlRigModule::RegisterTransformableCustomization() const
diff --git a/Engine/Plugins/Animation/ControlRig/Source/ControlRig/Private/Sequencer/ControlRigObjectSpawner.cpp b/Engine/Plugins/Animation/ControlRig/Source/ControlRig/Private/Sequencer/ControlRigObjectSpawner.cpp
deleted file mode 100644
diff --git a/Engine/Plugins/Animation/ControlRig/Source/ControlRig/Public/Sequencer/ControlRigObjectSpawner.h b/Engine/Plugins/Animation/ControlRig/Source/ControlRig/Public/Sequencer/ControlRigObjectSpawner.h
deleted file mode 100644
