Fixed EngineTest CIS issue at ControlRig / RigVMHost, due to deferencing a null pointer

#rnx

diff --git a/Engine/Plugins/Animation/ControlRig/Source/ControlRig/Private/ControlRig.cpp b/Engine/Plugins/Animation/ControlRig/Source/ControlRig/Private/ControlRig.cpp
--- a/Engine/Plugins/Animation/ControlRig/Source/ControlRig/Private/ControlRig.cpp
+++ b/Engine/Plugins/Animation/ControlRig/Source/ControlRig/Private/ControlRig.cpp
@@ -1215,7 +1215,9 @@ bool UControlRig::Execute_Internal(const FName& InEventName)
 			if(!VM->IsContextValidForExecution(Context) ||
 				!IsValidLowLevel() ||
 				!VM->IsValidLowLevel() ||
+				!GetLiteralMemory() ||
 				!GetLiteralMemory()->IsValidLowLevel() ||
+				!GetWorkMemory() ||
 				!GetWorkMemory()->IsValidLowLevel())
 			{
 				UE_LOG(LogControlRig, Warning, InvalidatedVMFormat, *GetClass()->GetName());
diff --git a/Engine/Plugins/Runtime/RigVM/Source/RigVM/Private/RigVMHost.cpp b/Engine/Plugins/Runtime/RigVM/Source/RigVM/Private/RigVMHost.cpp
--- a/Engine/Plugins/Runtime/RigVM/Source/RigVM/Private/RigVMHost.cpp
+++ b/Engine/Plugins/Runtime/RigVM/Source/RigVM/Private/RigVMHost.cpp
@@ -657,7 +657,9 @@ bool URigVMHost::Execute_Internal(const FName& InEventName)
 		if(!VM->IsContextValidForExecution(Context) ||
 			!IsValidLowLevel() ||
 			!VM->IsValidLowLevel() ||
-			!VM->GetLiteralMemory()->IsValidLowLevel() ||
+			!GetLiteralMemory() ||
+			!GetLiteralMemory()->IsValidLowLevel() ||
+			!GetWorkMemory() ||
 			!GetWorkMemory()->IsValidLowLevel())
 		{
 			UE_LOG(LogRigVM, Warning, InvalidatedVMFormat, *GetClass()->GetName());
