Control Rig: Fix anim preview settings for editor

diff --git a/Engine/Plugins/Animation/ControlRig/Source/ControlRig/Private/Sequencer/ControlRigLayerInstanceProxy.cpp b/Engine/Plugins/Animation/ControlRig/Source/ControlRig/Private/Sequencer/ControlRigLayerInstanceProxy.cpp
--- a/Engine/Plugins/Animation/ControlRig/Source/ControlRig/Private/Sequencer/ControlRigLayerInstanceProxy.cpp
+++ b/Engine/Plugins/Animation/ControlRig/Source/ControlRig/Private/Sequencer/ControlRigLayerInstanceProxy.cpp
@@ -4,6 +4,7 @@
 #include "AnimNode_ControlRig_ExternalSource.h"
 #include "Sequencer/ControlRigLayerInstance.h"
 #include "AnimSequencerInstance.h"
+#include "ControlRig.h"
 
 #include UE_INLINE_GENERATED_CPP_BY_NAME(ControlRigLayerInstanceProxy)
 
@@ -160,6 +161,8 @@ void FControlRigLayerInstanceProxy::AddControlRigTrack(int32 ControlRigID, UCont
 		}
 	}
 
+	InputPose.SetIsAdditive(InControlRig->IsAdditive());
+	
 	Node->SetControlRig(InControlRig);
 	Node->OnInitializeAnimInstance(this, CastChecked<UAnimInstance>(GetAnimInstanceObject()));
 	//mz removed this due to crash since Skeleton is not set up on a previous linked node 
@@ -348,6 +351,8 @@ void FAnimNode_ControlRigInputPose::CacheBones_AnyThread(const FAnimationCacheBo
 			FControlRigLayerInstanceProxy::CacheBonesCustomProxy(InputProxy);
 		}
 	}
+
+	bIsAdditive.Reset();
 }
 
 void FAnimNode_ControlRigInputPose::Update_AnyThread(const FAnimationUpdateContext& Context)
@@ -355,13 +360,27 @@ void FAnimNode_ControlRigInputPose::Update_AnyThread(const FAnimationUpdateConte
 	if (InputProxy)
 	{
 		FAnimationUpdateContext InputContext = Context.WithOtherProxy(InputProxy);
-		if (FAnimNode_Base* InputNode = InputPose.GetLinkNode())
+		if(IsAdditive())
 		{
-			InputProxy->UpdateAnimation_WithRoot(InputContext, InputNode, TEXT("AnimGraph"));
+			if (FAnimNode_Base* InputNode = InputPose.GetLinkNode())
+			{
+				InputProxy->UpdateAnimation_WithRoot(InputContext, InputNode, TEXT("AnimGraph"));
+			}
+			else
+			{
+				InputProxy->UpdateAnimationNode(InputContext);
+			}
 		}
 		else
 		{
-			InputProxy->UpdateAnimationNode(InputContext);
+			if (InputPose.GetLinkNode())
+			{
+				InputPose.Update(InputContext);
+			}
+			else
+			{
+				FControlRigLayerInstanceProxy::UpdateCustomProxy(InputProxy, InputContext);
+			}
 		}
 	}
 }
@@ -373,21 +392,36 @@ void FAnimNode_ControlRigInputPose::Evaluate_AnyThread(FPoseContext& Output)
 		FBoneContainer& RequiredBones = InputProxy->GetRequiredBones();
 		if (RequiredBones.IsValid())
 		{
-			FPoseContext InnerOutput(InputProxy, Output.ExpectsAdditivePose());
+			Output.Pose.SetBoneContainer(&RequiredBones);
+			FPoseContext InputContext(InputProxy, Output.ExpectsAdditivePose());
 			
 			// if no linked node, just use Evaluate of proxy
-			if (FAnimNode_Base* InputNode = InputPose.GetLinkNode())
+			if(IsAdditive())
 			{
-				InputProxy->EvaluateAnimation_WithRoot(InnerOutput, InputNode);
+				if (FAnimNode_Base* InputNode = InputPose.GetLinkNode())
+				{
+					InputProxy->EvaluateAnimation_WithRoot(InputContext, InputNode);
+				}
+				else
+				{
+					InputProxy->EvaluateAnimationNode(InputContext);
+				}
 			}
 			else
 			{
-				InputProxy->EvaluateAnimationNode(InnerOutput);
+				if (InputPose.GetLinkNode())
+				{
+					InputPose.Evaluate(InputContext);
+				}
+				else
+				{
+					FControlRigLayerInstanceProxy::EvaluateCustomProxy(InputProxy, InputContext);
+				}			
 			}
 
-			Output.Pose.MoveBonesFrom(InnerOutput.Pose);
-			Output.Curve.MoveFrom(InnerOutput.Curve);
-			Output.CustomAttributes.MoveFrom(InnerOutput.CustomAttributes);
+			Output.Pose.MoveBonesFrom(InputContext.Pose);
+			Output.Curve.MoveFrom(InputContext.Curve);
+			Output.CustomAttributes.MoveFrom(InputContext.CustomAttributes);
 			return;
 		}
 	}
@@ -433,6 +467,16 @@ void FAnimNode_ControlRigInputPose::Unlink()
 	InputProxy = nullptr;
 	InputAnimInstance = nullptr;
 	InputPose.SetLinkNode(nullptr);
+	bIsAdditive.Reset();
+}
+
+bool FAnimNode_ControlRigInputPose::IsAdditive() const
+{
+	if(!bIsAdditive.IsSet())
+	{
+		return false;
+	}
+	return bIsAdditive.GetValue();
 }
 
 
diff --git a/Engine/Plugins/Animation/ControlRig/Source/ControlRig/Private/Sequencer/ControlRigLayerInstanceProxy.h b/Engine/Plugins/Animation/ControlRig/Source/ControlRig/Private/Sequencer/ControlRigLayerInstanceProxy.h
--- a/Engine/Plugins/Animation/ControlRig/Source/ControlRig/Private/Sequencer/ControlRigLayerInstanceProxy.h
+++ b/Engine/Plugins/Animation/ControlRig/Source/ControlRig/Private/Sequencer/ControlRigLayerInstanceProxy.h
@@ -39,10 +39,20 @@ struct FAnimNode_ControlRigInputPose : public FAnimNode_Base
 	/** Called by linked instance nodes to dynamically unlink this to an outer graph */
 	void Unlink();
 
+	/** Returns true if the rig is additive */
+	bool IsAdditive() const;
+
+	/** Sets the input pose to be treated as additive */
+	void SetIsAdditive(bool bInAdditive)
+	{
+		bIsAdditive = bInAdditive;
+	}
+
 private:
 	/** The proxy to use when getting inputs, set when dynamically linked */
 	FAnimInstanceProxy* InputProxy;
 	UAnimInstance*		InputAnimInstance;
+	mutable TOptional<bool>	 bIsAdditive;
 };
 
 /** Proxy override for this UAnimInstance-derived class */
