Sequencer: Control Rig: Only modify if name has really changed to avoid dirtying files when validating.

diff --git a/Engine/Plugins/Animation/ControlRig/Source/ControlRig/Private/Sequencer/MovieSceneControlRigParameterSection.cpp b/Engine/Plugins/Animation/ControlRig/Source/ControlRig/Private/Sequencer/MovieSceneControlRigParameterSection.cpp
--- a/Engine/Plugins/Animation/ControlRig/Source/ControlRig/Private/Sequencer/MovieSceneControlRigParameterSection.cpp
+++ b/Engine/Plugins/Animation/ControlRig/Source/ControlRig/Private/Sequencer/MovieSceneControlRigParameterSection.cpp
@@ -838,17 +838,19 @@ void UMovieSceneControlRigParameterSection::PreSave(FObjectPreSaveContext SaveCo
 	}
 }
 
-
 bool UMovieSceneControlRigParameterSection::RenameParameterName(const FName& OldParameterName, const FName& NewParameterName)
 {
 	bool bWasReplaced = false;
-	Modify();
 	for (FScalarParameterNameAndCurve& ScalarParameterNameAndCurve : ScalarParameterNamesAndCurves)
 	{
 		if (ScalarParameterNameAndCurve.ParameterName == OldParameterName)
 		{
+			if (bWasReplaced == false)
+			{
+				Modify();
+				bWasReplaced = true;
+			}
 			ScalarParameterNameAndCurve.ParameterName = NewParameterName;
-			bWasReplaced = true;
 			break;
 		}
 	}
@@ -858,7 +860,6 @@ bool UMovieSceneControlRigParameterSection::RenameParameterName(const FName& Old
 		if (BoolParameterNameAndCurve.ParameterName == OldParameterName)
 		{
 			BoolParameterNameAndCurve.ParameterName = NewParameterName;
-			bWasReplaced = true;
 			break;
 		}
 	}
@@ -867,8 +868,12 @@ bool UMovieSceneControlRigParameterSection::RenameParameterName(const FName& Old
 	{
 		if (EnumParameterNameAndCurve.ParameterName == OldParameterName)
 		{
+			if (bWasReplaced == false)
+			{
+				Modify();
+				bWasReplaced = true;
+			}
 			EnumParameterNameAndCurve.ParameterName = NewParameterName;
-			bWasReplaced = true;
 			break;
 		}
 	}
@@ -877,8 +882,12 @@ bool UMovieSceneControlRigParameterSection::RenameParameterName(const FName& Old
 	{
 		if (IntegerParameterNameAndCurve.ParameterName == OldParameterName)
 		{
+			if (bWasReplaced == false)
+			{
+				Modify();
+				bWasReplaced = true;
+			}
 			IntegerParameterNameAndCurve.ParameterName = NewParameterName;
-			bWasReplaced = true;
 			break;
 		}
 	}
@@ -887,8 +896,12 @@ bool UMovieSceneControlRigParameterSection::RenameParameterName(const FName& Old
 	{
 		if (Vector2DParameterNameAndCurve.ParameterName == OldParameterName)
 		{
+			if (bWasReplaced == false)
+			{
+				Modify();
+				bWasReplaced = true;
+			}
 			Vector2DParameterNameAndCurve.ParameterName = NewParameterName;
-			bWasReplaced = true;
 			break;
 		}
 	}
@@ -897,8 +910,12 @@ bool UMovieSceneControlRigParameterSection::RenameParameterName(const FName& Old
 	{
 		if (VectorParameterNameAndCurve.ParameterName == OldParameterName)
 		{
+			if (bWasReplaced == false)
+			{
+				Modify();
+				bWasReplaced = true;
+			}
 			VectorParameterNameAndCurve.ParameterName = NewParameterName;
-			bWasReplaced = true;
 			break;
 		}
 	}
@@ -907,8 +924,12 @@ bool UMovieSceneControlRigParameterSection::RenameParameterName(const FName& Old
 	{
 		if (ColorParameterNameAndCurve.ParameterName == OldParameterName)
 		{
+			if (bWasReplaced == false)
+			{
+				Modify();
+				bWasReplaced = true;
+			}
 			ColorParameterNameAndCurve.ParameterName = NewParameterName;
-			bWasReplaced = true;
 			break;
 		}
 				
@@ -918,8 +939,12 @@ bool UMovieSceneControlRigParameterSection::RenameParameterName(const FName& Old
 	{
 		if (TransformParameterNamesAndCurve.ParameterName == OldParameterName)
 		{
+			if (bWasReplaced == false)
+			{
+				Modify();
+				bWasReplaced = true;
+			}
 			TransformParameterNamesAndCurve.ParameterName = NewParameterName;
-			bWasReplaced = true;
 			break;
 		}
 	}
@@ -1039,6 +1064,7 @@ void UMovieSceneControlRigParameterSection::PostLoad()
 	{
 		for (FConstraintAndActiveChannel& ConstraintChannel : ConstraintsChannels)
 		{
+			ConstraintChannel.Constraint.Reset(); //clear it out may be referencing wrong level
 			if (UTickableTransformConstraint* TransformConstraint = Cast<UTickableTransformConstraint>(ConstraintChannel.ConstraintCopyToSpawn))
 			{
 				if (UTransformableControlHandle* Handle = Cast<UTransformableControlHandle>(TransformConstraint->ChildTRSHandle))
diff --git a/Engine/Source/Runtime/MovieSceneTracks/Private/Sections/MovieScene3DTransformSection.cpp b/Engine/Source/Runtime/MovieSceneTracks/Private/Sections/MovieScene3DTransformSection.cpp
--- a/Engine/Source/Runtime/MovieSceneTracks/Private/Sections/MovieScene3DTransformSection.cpp
+++ b/Engine/Source/Runtime/MovieSceneTracks/Private/Sections/MovieScene3DTransformSection.cpp
@@ -1015,6 +1015,20 @@ TArray<FConstraintAndActiveChannel>& UMovieScene3DTransformSection::GetConstrain
 	}
 	return EmptyChannels;
 }
+
+void UMovieScene3DTransformSection::PostLoad()
+{
+	Super::PostLoad();
+	//for spawnables the control rig saved in our channels may have changed so we need to update thaem
+	if (Constraints)
+	{
+		for (FConstraintAndActiveChannel& ConstraintChannel : Constraints->ConstraintsChannels)
+		{
+			ConstraintChannel.Constraint.Reset(); //clear it out may be referencing wrong level
+		}
+	}
+}
+
 #if WITH_EDITOR
 bool UMovieScene3DTransformSection::Modify(bool bAlwaysMarkDirty)
 {
diff --git a/Engine/Source/Runtime/MovieSceneTracks/Public/Sections/MovieScene3DTransformSection.h b/Engine/Source/Runtime/MovieSceneTracks/Public/Sections/MovieScene3DTransformSection.h
--- a/Engine/Source/Runtime/MovieSceneTracks/Public/Sections/MovieScene3DTransformSection.h
+++ b/Engine/Source/Runtime/MovieSceneTracks/Public/Sections/MovieScene3DTransformSection.h
@@ -222,6 +222,8 @@ class UMovieScene3DTransformSection
 
 public:
 
+	virtual void PostLoad() override;
+
 #if WITH_EDITOR
 	/* From UObject*/
 	virtual bool Modify(bool bAlwaysMarkDirty = true) override;
