Control Rig: Add missing transactions when changing control rig spaces

diff --git a/Engine/Plugins/Animation/ControlRig/Source/ControlRigEditor/Private/EditMode/SControlRigEditModeTools.cpp b/Engine/Plugins/Animation/ControlRig/Source/ControlRigEditor/Private/EditMode/SControlRigEditModeTools.cpp
--- a/Engine/Plugins/Animation/ControlRig/Source/ControlRigEditor/Private/EditMode/SControlRigEditModeTools.cpp
+++ b/Engine/Plugins/Animation/ControlRig/Source/ControlRigEditor/Private/EditMode/SControlRigEditModeTools.cpp
@@ -836,10 +836,13 @@ void SControlRigEditModeTools::HandleActiveSpaceChanged(URigHierarchy* InHierarc
 void SControlRigEditModeTools::HandleSpaceListChanged(URigHierarchy* InHierarchy, const FRigElementKey& InControlKey,
 	const TArray<FRigElementKey>& InSpaceList)
 {
+	FScopedTransaction Transaction(LOCTEXT("ChangeControlRigSpace", "Change Control Rig Space"));
+
 	for (TWeakObjectPtr<UControlRig>& ControlRig : ControlRigs)
 	{
 		if (ControlRig.IsValid() && ControlRig->GetHierarchy() == InHierarchy)
 		{
+			ControlRig->Modify();
 
 			if (const FRigControlElement* ControlElement = InHierarchy->Find<FRigControlElement>(InControlKey))
 			{
