Control Rig: Enable Rotation order CVar by default

diff --git a/Engine/Plugins/Animation/ControlRig/Source/ControlRig/Private/Rigs/RigHierarchy.cpp b/Engine/Plugins/Animation/ControlRig/Source/ControlRig/Private/Rigs/RigHierarchy.cpp
--- a/Engine/Plugins/Animation/ControlRig/Source/ControlRig/Private/Rigs/RigHierarchy.cpp
+++ b/Engine/Plugins/Animation/ControlRig/Source/ControlRig/Private/Rigs/RigHierarchy.cpp
@@ -44,7 +44,7 @@ static TAutoConsoleVariable<int32> CVarControlRigHierarchyTraceAlways(TEXT("Cont
 static TAutoConsoleVariable<int32> CVarControlRigHierarchyTraceCallstack(TEXT("ControlRig.Hierarchy.TraceCallstack"), 0, TEXT("if nonzero we will record the callstack for any trace entry.\nOnly works if(ControlRig.Hierarchy.TraceEnabled != 0)"));
 static TAutoConsoleVariable<int32> CVarControlRigHierarchyTracePrecision(TEXT("ControlRig.Hierarchy.TracePrecision"), 3, TEXT("sets the number digits in a float when tracing hierarchies."));
 static TAutoConsoleVariable<int32> CVarControlRigHierarchyTraceOnSpawn(TEXT("ControlRig.Hierarchy.TraceOnSpawn"), 0, TEXT("sets the number of frames to trace when a new hierarchy is spawned"));
-TAutoConsoleVariable<bool> CVarControlRigHierarchyEnableRotationOrder(TEXT("ControlRig.Hierarchy.EnableRotationOrder"), false, TEXT("enables the rotation order for controls"));
+TAutoConsoleVariable<bool> CVarControlRigHierarchyEnableRotationOrder(TEXT("ControlRig.Hierarchy.EnableRotationOrder"), true, TEXT("enables the rotation order for controls"));
 TAutoConsoleVariable<bool> CVarControlRigHierarchyEnableModules(TEXT("ControlRig.Hierarchy.Modules"), false, TEXT("enables the modular rigging functionality"));
 static int32 sRigHierarchyLastTrace = INDEX_NONE;
 static TCHAR sRigHierarchyTraceFormat[16];
