Control Rig: Introduce caching to sphere trace

diff --git a/Engine/Plugins/Animation/ControlRig/Source/ControlRig/Private/Units/Collision/RigUnit_WorldCollision.cpp b/Engine/Plugins/Animation/ControlRig/Source/ControlRig/Private/Units/Collision/RigUnit_WorldCollision.cpp
--- a/Engine/Plugins/Animation/ControlRig/Source/ControlRig/Private/Units/Collision/RigUnit_WorldCollision.cpp
+++ b/Engine/Plugins/Animation/ControlRig/Source/ControlRig/Private/Units/Collision/RigUnit_WorldCollision.cpp
@@ -18,6 +18,25 @@ FRigUnit_SphereTraceWorld_Execute()
 
 	if(ExecuteContext.GetWorld() == nullptr)
 	{
+		WorkData.Reset();
+		return;
+	}
+
+	// Note: we are not adding the world transform to the hash since
+	// the world transform cannot change during one execution of the 
+	// rig. By adding NumExecutions to the hash we make sure to
+	// trace again once the world transform has changed.
+	uint32 Hash = GetTypeHash(ExecuteContext.GetNumExecutions());
+	Hash = HashCombine(Hash, GetTypeHash(Start));
+	Hash = HashCombine(Hash, GetTypeHash(End));
+	Hash = HashCombine(Hash, GetTypeHash((int32)Channel));
+	Hash = HashCombine(Hash, GetTypeHash(Radius));
+
+	if(WorkData.Hash == Hash)
+	{
+		bHit = WorkData.bHit;
+		HitLocation = WorkData.HitLocation;
+		HitNormal = WorkData.HitNormal;
 		return;
 	}
 	
@@ -47,6 +66,11 @@ FRigUnit_SphereTraceWorld_Execute()
 		HitLocation = ExecuteContext.ToVMSpace(HitResult.ImpactPoint);
 		HitNormal = ExecuteContext.GetToWorldSpaceTransform().InverseTransformVector(HitResult.ImpactNormal);
 	}
+
+	WorkData.Hash = Hash;
+	WorkData.bHit = bHit;
+	WorkData.HitLocation = HitLocation;
+	WorkData.HitNormal = HitNormal;
 }
 
 FRigVMStructUpgradeInfo FRigUnit_SphereTraceWorld::GetUpgradeInfo() const
@@ -69,6 +93,25 @@ FRigUnit_SphereTraceByTraceChannel_Execute()
 
 	if (ExecuteContext.GetWorld() == nullptr)
 	{
+		WorkData.Reset();
+		return;
+	}
+
+	// Note: we are not adding the world transform to the hash since
+	// the world transform cannot change during one execution of the 
+	// rig. By adding NumExecutions to the hash we make sure to
+	// trace again once the world transform has changed.
+	uint32 Hash = GetTypeHash(ExecuteContext.GetNumExecutions());
+	Hash = HashCombine(Hash, GetTypeHash(Start));
+	Hash = HashCombine(Hash, GetTypeHash(End));
+	Hash = HashCombine(Hash, GetTypeHash((int32)TraceChannel));
+	Hash = HashCombine(Hash, GetTypeHash(Radius));
+
+	if(WorkData.Hash == Hash)
+	{
+		bHit = WorkData.bHit;
+		HitLocation = WorkData.HitLocation;
+		HitNormal = WorkData.HitNormal;
 		return;
 	}
 
@@ -95,6 +138,11 @@ FRigUnit_SphereTraceByTraceChannel_Execute()
 		HitLocation = ExecuteContext.ToVMSpace(HitResult.ImpactPoint);
 		HitNormal = ExecuteContext.GetToWorldSpaceTransform().InverseTransformVector(HitResult.ImpactNormal);
 	}
+
+	WorkData.Hash = Hash;
+	WorkData.bHit = bHit;
+	WorkData.HitLocation = HitLocation;
+	WorkData.HitNormal = HitNormal;
 }
 
 FRigUnit_SphereTraceByObjectTypes_Execute()
@@ -107,6 +155,29 @@ FRigUnit_SphereTraceByObjectTypes_Execute()
 
 	if (ExecuteContext.GetWorld() == nullptr)
 	{
+		WorkData.Reset();
+		return;
+	}
+
+	// Note: we are not adding the world transform to the hash since
+	// the world transform cannot change during one execution of the 
+	// rig. By adding NumExecutions to the hash we make sure to
+	// trace again once the world transform has changed.
+	uint32 Hash = GetTypeHash(ExecuteContext.GetNumExecutions());
+	Hash = HashCombine(Hash, GetTypeHash(Start));
+	Hash = HashCombine(Hash, GetTypeHash(End));
+	Hash = HashCombine(Hash, GetTypeHash(ObjectTypes.Num()));
+	for(const TEnumAsByte<EObjectTypeQuery>& ObjectType : ObjectTypes)
+	{
+		Hash = HashCombine(Hash, GetTypeHash((int32)ObjectType));
+	}
+	Hash = HashCombine(Hash, GetTypeHash(Radius));
+
+	if(WorkData.Hash == Hash)
+	{
+		bHit = WorkData.bHit;
+		HitLocation = WorkData.HitLocation;
+		HitNormal = WorkData.HitNormal;
 		return;
 	}
 
@@ -150,5 +221,10 @@ FRigUnit_SphereTraceByObjectTypes_Execute()
 		HitLocation = ExecuteContext.ToVMSpace(HitResult.ImpactPoint);
 		HitNormal = ExecuteContext.GetToWorldSpaceTransform().InverseTransformVector(HitResult.ImpactNormal);
 	}
+
+	WorkData.Hash = Hash;
+	WorkData.bHit = bHit;
+	WorkData.HitLocation = HitLocation;
+	WorkData.HitNormal = HitNormal;
 }
 
diff --git a/Engine/Plugins/Animation/ControlRig/Source/ControlRig/Private/Units/Collision/RigUnit_WorldCollision.h b/Engine/Plugins/Animation/ControlRig/Source/ControlRig/Private/Units/Collision/RigUnit_WorldCollision.h
--- a/Engine/Plugins/Animation/ControlRig/Source/ControlRig/Private/Units/Collision/RigUnit_WorldCollision.h
+++ b/Engine/Plugins/Animation/ControlRig/Source/ControlRig/Private/Units/Collision/RigUnit_WorldCollision.h
@@ -6,6 +6,37 @@
 #include "Units/RigUnit.h"
 #include "RigUnit_WorldCollision.generated.h"
 
+USTRUCT()
+struct CONTROLRIG_API FRigUnit_SphereTrace_WorkData
+{
+	GENERATED_BODY()
+
+	FRigUnit_SphereTrace_WorkData()
+	: Hash(INDEX_NONE)
+	, bHit(false)
+	, HitLocation(EForceInit::ForceInitToZero)
+	, HitNormal(0.f, 0.f, 1.f)
+	{
+	}
+
+	void Reset()
+	{
+		Hash = INDEX_NONE;
+	}
+
+	UPROPERTY()
+	uint32 Hash;
+
+	UPROPERTY()
+	bool bHit;
+
+	UPROPERTY()
+	FVector HitLocation;
+	
+	UPROPERTY()
+	FVector HitNormal;
+};
+
 /**
  * Sweeps a sphere against the world and return the first blocking hit using a specific channel
  */
@@ -57,6 +88,10 @@ struct CONTROLRIG_API FRigUnit_SphereTraceWorld : public FRigUnit
 	UPROPERTY(meta = (Output))
 	FVector HitNormal;
 
+	/** Cache / workstate */
+	UPROPERTY()
+	FRigUnit_SphereTrace_WorkData WorkData;
+
 	RIGVM_METHOD()
 	virtual FRigVMStructUpgradeInfo GetUpgradeInfo() const override;
 };
@@ -112,6 +147,10 @@ struct CONTROLRIG_API FRigUnit_SphereTraceByTraceChannel : public FRigUnit
 	/** Hit normal in rig / global Space */
 	UPROPERTY(meta = (Output))
 	FVector HitNormal;
+	
+	/** Cache / workstate */
+	UPROPERTY()
+	FRigUnit_SphereTrace_WorkData WorkData;
 };
 
 /**
@@ -165,4 +204,8 @@ struct CONTROLRIG_API FRigUnit_SphereTraceByObjectTypes : public FRigUnit
 	/** Hit normal in rig / global Space */
 	UPROPERTY(meta = (Output))
 	FVector HitNormal;
+	
+	/** Cache / workstate */
+	UPROPERTY()
+	FRigUnit_SphereTrace_WorkData WorkData;
 };
diff --git a/Engine/Plugins/Runtime/RigVM/Source/RigVM/Private/RigVMCore/RigVM.cpp b/Engine/Plugins/Runtime/RigVM/Source/RigVM/Private/RigVMCore/RigVM.cpp
--- a/Engine/Plugins/Runtime/RigVM/Source/RigVM/Private/RigVMCore/RigVM.cpp
+++ b/Engine/Plugins/Runtime/RigVM/Source/RigVM/Private/RigVMCore/RigVM.cpp
@@ -1562,7 +1562,7 @@ ERigVMExecuteResult URigVM::Execute(FRigVMExtendedExecuteContext& Context, TArra
 
 	if(bIsRootEntry)
 	{
-		Context.NumExecutions++;
+		ContextPublicData.NumExecutions++;
 	}
 
 #if WITH_EDITOR
@@ -2091,7 +2091,7 @@ ERigVMExecuteResult URigVM::ExecuteBranch(FRigVMExtendedExecuteContext& Context,
 	// Maintain all settings on the context - to be reset once the branch has executed. 
 	FRigVMExecuteContext& PublicContext = Context.GetPublicData<>();
 	TGuardValue<double> LastExecutionMicroSecondsGuard(Context.LastExecutionMicroSeconds, Context.LastExecutionMicroSeconds);
-	TGuardValue<uint32> NumExecutionsGuard(Context.NumExecutions, Context.NumExecutions);
+	TGuardValue<uint32> NumExecutionsGuard(PublicContext.NumExecutions, PublicContext.NumExecutions);
 	TGuardValue<ERigVMExecuteResult> CurrentExecuteResultGuard(Context.CurrentExecuteResult, Context.CurrentExecuteResult);
 	TGuardValue<FName> CurrentEntryNameGuard(Context.CurrentEntryName, Context.CurrentEntryName);
 	TGuardValue<bool> bCurrentlyRunningRootEntryGuard(Context.bCurrentlyRunningRootEntry, Context.bCurrentlyRunningRootEntry);
diff --git a/Engine/Plugins/Runtime/RigVM/Source/RigVM/Private/RigVMCore/RigVMExecuteContext.cpp b/Engine/Plugins/Runtime/RigVM/Source/RigVM/Private/RigVMCore/RigVMExecuteContext.cpp
--- a/Engine/Plugins/Runtime/RigVM/Source/RigVM/Private/RigVMCore/RigVMExecuteContext.cpp
+++ b/Engine/Plugins/Runtime/RigVM/Source/RigVM/Private/RigVMCore/RigVMExecuteContext.cpp
@@ -74,7 +74,10 @@ void FRigVMExtendedExecuteContext::Reset()
 	LazyBranchInstanceData.Reset();
 	ExternalVariableRuntimeData.Reset();
 
-	NumExecutions = 0;
+	if(PublicDataScope.IsValid())
+	{
+		GetPublicData<>().NumExecutions = 0;
+	}
 
 	ExecutingThreadId = INDEX_NONE;
 
diff --git a/Engine/Plugins/Runtime/RigVM/Source/RigVM/Public/RigVMCore/RigVMExecuteContext.h b/Engine/Plugins/Runtime/RigVM/Source/RigVM/Public/RigVMCore/RigVMExecuteContext.h
--- a/Engine/Plugins/Runtime/RigVM/Source/RigVM/Public/RigVMCore/RigVMExecuteContext.h
+++ b/Engine/Plugins/Runtime/RigVM/Source/RigVM/Public/RigVMCore/RigVMExecuteContext.h
@@ -243,6 +243,7 @@ struct FRigVMExecuteContext
 		: EventName(NAME_None)
 		, FunctionName(NAME_None)
 		, InstructionIndex(0)
+		, NumExecutions(0)
 		, DeltaTime(0.0)
 		, AbsoluteTime(0.0)
 		, FramesPerSecond(1.0 / 60.0)
@@ -293,6 +294,8 @@ struct FRigVMExecuteContext
 
 	uint16 GetInstructionIndex() const { return InstructionIndex; }
 
+	uint32 GetNumExecutions() const { return NumExecutions; }
+
 	FName GetFunctionName() const { return FunctionName; }
 	
 	FName GetEventName() const { return EventName; }
@@ -403,6 +406,7 @@ struct FRigVMExecuteContext
 		EventName = InOtherContext->EventName;
 		FunctionName = InOtherContext->FunctionName;
 		InstructionIndex = InOtherContext->InstructionIndex;
+		NumExecutions = InOtherContext->NumExecutions;
 		RuntimeSettings = InOtherContext->RuntimeSettings;
 		NameCache = InOtherContext->NameCache;
 	}
@@ -425,6 +429,8 @@ protected:
 	
 	uint16 InstructionIndex;
 
+	uint32 NumExecutions;
+
 	double DeltaTime;
 
 	double AbsoluteTime;
@@ -694,7 +700,11 @@ struct RIGVM_API FRigVMExtendedExecuteContext
 
 	uint32 GetNumExecutions() const
 	{
-		return NumExecutions;
+		if(PublicDataScope.IsValid())
+		{
+			return GetPublicData<>().GetNumExecutions();
+		}
+		return 0;
 	}
 
 	UPROPERTY()
@@ -714,9 +724,6 @@ struct RIGVM_API FRigVMExtendedExecuteContext
 	const FRigVMDispatchFactory* Factory = nullptr;
 	FRigVMNameCache NameCache;
 
-	UPROPERTY(transient)
-	uint32 NumExecutions = 0;
-
 	TArray<FRigVMMemoryHandle> CachedMemoryHandles;
 	// changes to the layout of cached memory array should be reflected in GetContainerIndex()
 	TArray<URigVMMemoryStorage*> CachedMemory;
diff --git a/Engine/Plugins/Runtime/RigVM/Source/RigVM/Public/RigVMCore/RigVMNativized.h b/Engine/Plugins/Runtime/RigVM/Source/RigVM/Public/RigVMCore/RigVMNativized.h
--- a/Engine/Plugins/Runtime/RigVM/Source/RigVM/Public/RigVMCore/RigVMNativized.h
+++ b/Engine/Plugins/Runtime/RigVM/Source/RigVM/Public/RigVMCore/RigVMNativized.h
@@ -176,7 +176,7 @@ protected:
 		{
 			Context.ExecutionReachedExit().Broadcast(Context.GetPublicData<>().GetEventName());
 		}
-		Context.NumExecutions++;
+		Context.GetPublicDataSafe<>().NumExecutions++;
 	}
 
 	template<typename T>
