Control Rig: Deprecate interaction rig as it's no longer used.

The functions are still there, but marked deprecated for 5.4 in case someone was setting those properties and retrieving it for some other reason? Otherwise we could just delete it all...

diff --git a/Engine/Plugins/Animation/ControlRig/Source/ControlRig/Private/ControlRig.cpp b/Engine/Plugins/Animation/ControlRig/Source/ControlRig/Private/ControlRig.cpp
--- a/Engine/Plugins/Animation/ControlRig/Source/ControlRig/Private/ControlRig.cpp
+++ b/Engine/Plugins/Animation/ControlRig/Source/ControlRig/Private/ControlRig.cpp
@@ -99,7 +99,6 @@ void UControlRig::BeginDestroy()
 	PostConstructionEvent.Clear();
 	PreForwardsSolveEvent.Clear();
 	PostForwardsSolveEvent.Clear();
-	SetInteractionRig(nullptr);
 
 #if WITH_EDITOR
 	if (!HasAnyFlags(RF_ClassDefaultObject))
@@ -2487,75 +2486,6 @@ const FRigInfluenceMap* UControlRig::FindInfluenceMap(const FName& InEventName)
 	return nullptr;
 }
 
-void UControlRig::SetInteractionRig(UControlRig* InInteractionRig)
-{
-	if (InteractionRig == InInteractionRig)
-	{
-		return;
-	}
-
-	if (InteractionRig != nullptr)
-	{
-		InteractionRig->ControlModified().RemoveAll(this);
-		InteractionRig->OnInitialized_AnyThread().RemoveAll(this);
-		InteractionRig->OnExecuted_AnyThread().RemoveAll(this);
-		InteractionRig->ControlSelected().RemoveAll(this);
-		OnInitialized_AnyThread().RemoveAll(InteractionRig);
-		OnExecuted_AnyThread().RemoveAll(InteractionRig);
-		ControlSelected().RemoveAll(InteractionRig);
-	}
-
-	InteractionRig = InInteractionRig;
-
-	if (InteractionRig != nullptr)
-	{
-		SetInteractionRigClass(InteractionRig->GetClass());
-
-		InteractionRig->Initialize(true);
-		InteractionRig->CopyPoseFromOtherRig(this);
-		InteractionRig->RequestConstruction();
-		InteractionRig->Execute(FRigUnit_BeginExecution::EventName);
-
-		InteractionRig->ControlModified().AddUObject(this, &UControlRig::HandleInteractionRigControlModified);
-		InteractionRig->OnInitialized_AnyThread().AddUObject(this, &UControlRig::HandleInteractionRigInitialized);
-		InteractionRig->OnExecuted_AnyThread().AddUObject(this, &UControlRig::HandleInteractionRigExecuted);
-		InteractionRig->ControlSelected().AddUObject(this, &UControlRig::HandleInteractionRigControlSelected, false);
-		OnInitialized_AnyThread().AddUObject(ToRawPtr(InteractionRig), &UControlRig::HandleInteractionRigInitialized);
-		OnExecuted_AnyThread().AddUObject(ToRawPtr(InteractionRig), &UControlRig::HandleInteractionRigExecuted);
-		ControlSelected().AddUObject(ToRawPtr(InteractionRig), &UControlRig::HandleInteractionRigControlSelected, true);
-
-		FControlRigBracketScope BracketScope(InterRigSyncBracket);
-		InteractionRig->HandleInteractionRigExecuted(this, FRigUnit_BeginExecution::EventName);
-	}
-}
-
-void UControlRig::SetInteractionRigClass(TSubclassOf<UControlRig> InInteractionRigClass)
-{
-	if (InteractionRigClass == InInteractionRigClass)
-	{
-		return;
-	}
-
-	InteractionRigClass = InInteractionRigClass;
-
-	if(InteractionRigClass)
-	{
-		if(InteractionRig != nullptr)
-		{
-			if(InteractionRig->GetClass() != InInteractionRigClass)
-			{
-				SetInteractionRig(nullptr);
-			}
-		}
-
-		if(InteractionRig == nullptr)
-		{
-			UControlRig* NewInteractionRig = NewObject<UControlRig>(this, InteractionRigClass);
-			SetInteractionRig(NewInteractionRig);
-		}
-	}
-}
-
 #if WITH_EDITOR
 
 void UControlRig::PreEditChange(FProperty* PropertyAboutToChange)
@@ -2581,37 +2511,11 @@ void UControlRig::PreEditChange(FProperty* PropertyAboutToChange)
 	}
 	
 	Super::PreEditChange(PropertyAboutToChange);
-
-	if (PropertyAboutToChange && PropertyAboutToChange->GetFName() == GET_MEMBER_NAME_CHECKED(UControlRig, InteractionRig))
-	{
-		SetInteractionRig(nullptr);
-	}
-	else if (PropertyAboutToChange && PropertyAboutToChange->GetFName() == GET_MEMBER_NAME_CHECKED(UControlRig, InteractionRigClass))
-	{
-		SetInteractionRigClass(nullptr);
-	}
 }
 
 void UControlRig::PostEditChangeProperty(FPropertyChangedEvent& PropertyChangedEvent)
 {
 	Super::PostEditChangeProperty(PropertyChangedEvent);
-
-	if (PropertyChangedEvent.MemberProperty && PropertyChangedEvent.MemberProperty->GetFName() == GET_MEMBER_NAME_CHECKED(UControlRig, InteractionRig))
-	{
-		UControlRig* NewInteractionRig = InteractionRig;
-		SetInteractionRig(nullptr);
-		SetInteractionRig(NewInteractionRig);
-	}
-	else if (PropertyChangedEvent.MemberProperty && PropertyChangedEvent.MemberProperty->GetFName() == GET_MEMBER_NAME_CHECKED(UControlRig, InteractionRigClass))
-	{
-		TSubclassOf<UControlRig> NewInteractionRigClass = InteractionRigClass;
-		SetInteractionRigClass(nullptr);
-		SetInteractionRigClass(NewInteractionRigClass);
-		if (NewInteractionRigClass == nullptr)
-		{
-			SetInteractionRig(nullptr);
-		}
-	}
 }
 #endif
 
@@ -2721,183 +2625,6 @@ void UControlRig::CopyPoseFromOtherRig(UControlRig* Subject)
 	}
 }
 
-void UControlRig::HandleInteractionRigControlModified(UControlRig* Subject, FRigControlElement* Control, const FRigControlModifiedContext& Context)
-{
-	check(Subject);
-
-	if (IsSyncingWithOtherRig() || IsExecuting())
-	{
-		return;
-	}
-	FControlRigBracketScope BracketScope(InterRigSyncBracket);
-
-	if (Subject != InteractionRig)
-	{
-		return;
-	}
-
-	if (const FRigInfluenceMap* InfluenceMap = Subject->FindInfluenceMap(Context.EventName))
-	{
-		if (const FRigInfluenceEntry* InfluenceEntry = InfluenceMap->Find(Control->GetKey()))
-		{
-			for (const FRigElementKey& AffectedKey : *InfluenceEntry)
-			{
-				if (AffectedKey.Type == ERigElementType::Control)
-				{
-					if (FRigControlElement* AffectedControl = FindControl(AffectedKey.Name))
-					{
-						QueuedModifiedControls.Add(AffectedControl->GetKey());
-					}
-				}
-				else if (AffectedKey.Type == ERigElementType::Bone)
-				{
-					// special case controls with a CONTROL suffix
-					FName BoneControlName = *FString::Printf(TEXT("%s_CONTROL"), *AffectedKey.Name.ToString());
-					if (FRigControlElement* AffectedControl = FindControl(BoneControlName))
-					{
-						QueuedModifiedControls.Add(AffectedControl->GetKey());
-					}
-				}
-				else if(AffectedKey.Type == ERigElementType::Curve)
-				{
-					// special case controls with a CONTROL suffix
-					FName CurveControlName = *FString::Printf(TEXT("%s_CURVE_CONTROL"), *AffectedKey.Name.ToString());
-					if (FRigControlElement* AffectedControl = FindControl(CurveControlName))
-					{
-						QueuedModifiedControls.Add(AffectedControl->GetKey());
-					}
-				}
-			}
-		}
-	}
-
-}
-
-void UControlRig::HandleInteractionRigInitialized(URigVMHost* Subject, const FName& EventName)
-{
-	check(Subject);
-
-	if (IsSyncingWithOtherRig())
-	{
-		return;
-	}
-	FControlRigBracketScope BracketScope(InterRigSyncBracket);
-	RequestInit();
-}
-
-void UControlRig::HandleInteractionRigExecuted(URigVMHost* Subject, const FName& EventName)
-{
-	check(Subject);
-	UControlRig* SubjectRig = CastChecked<UControlRig>(Subject);
-
-	if (IsSyncingWithOtherRig() || IsExecuting())
-	{
-		return;
-	}
-	FControlRigBracketScope BracketScope(InterRigSyncBracket);
-
-	CopyPoseFromOtherRig(SubjectRig);
-	Execute(FRigUnit_InverseExecution::EventName);
-
-	FRigControlModifiedContext Context;
-	Context.EventName = FRigUnit_InverseExecution::EventName;
-	Context.SetKey = EControlRigSetKey::DoNotCare;
-
-	for (const FRigElementKey& QueuedModifiedControl : QueuedModifiedControls)
-	{
-		if(FRigControlElement* ControlElement = FindControl(QueuedModifiedControl.Name))
-		{
-			ControlModified().Broadcast(this, ControlElement, Context);
-		}
-	}
-}
-
-void UControlRig::HandleInteractionRigControlSelected(UControlRig* Subject, FRigControlElement* Control, bool bSelected, bool bInverted)
-{
-	check(Subject);
-
-	if (IsSyncingWithOtherRig() || IsExecuting())
-	{
-		return;
-	}
-	if (Subject->IsSyncingWithOtherRig() || Subject->IsExecuting())
-	{
-		return;
-	}
-	FControlRigBracketScope BracketScope(InterRigSyncBracket);
-
-	const FRigInfluenceMap* InfluenceMap = nullptr;
-	if (bInverted)
-	{
-		InfluenceMap = FindInfluenceMap(FRigUnit_BeginExecution::EventName);
-	}
-	else
-	{
-		InfluenceMap = Subject->FindInfluenceMap(FRigUnit_BeginExecution::EventName);
-	}
-
-	if (InfluenceMap)
-	{
-		FRigInfluenceMap InvertedMap;
-		if (bInverted)
-		{
-			InvertedMap = InfluenceMap->Inverse();
-			InfluenceMap = &InvertedMap;
-		}
-
-		struct Local
-		{
-			static void SelectAffectedElements(UControlRig* ThisRig, const FRigInfluenceMap* InfluenceMap, const FRigElementKey& InKey, bool bSelected, bool bInverted)
-			{
-				if (const FRigInfluenceEntry* InfluenceEntry = InfluenceMap->Find(InKey))
-				{
-					for (const FRigElementKey& AffectedKey : *InfluenceEntry)
-					{
-						if (AffectedKey.Type == ERigElementType::Control)
-						{
-							ThisRig->SelectControl(AffectedKey.Name, bSelected);
-						}
-
-						if (bInverted)
-						{
-							if (AffectedKey.Type == ERigElementType::Control)
-							{
-								ThisRig->SelectControl(AffectedKey.Name, bSelected);
-							}
-						}
-						else
-						{
-							if (AffectedKey.Type == ERigElementType::Control)
-							{
-								ThisRig->SelectControl(AffectedKey.Name, bSelected);
-							}
-							else if (AffectedKey.Type == ERigElementType::Bone ||
-								AffectedKey.Type == ERigElementType::Curve)
-							{
-								FName ControlName = *FString::Printf(TEXT("%s_CONTROL"), *AffectedKey.Name.ToString());
-								ThisRig->SelectControl(ControlName, bSelected);
-							}
-						}
-					}
-				}
-			}
-		};
-
-		Local::SelectAffectedElements(this, InfluenceMap, Control->GetKey(), bSelected, bInverted);
-
-		if (bInverted)
-		{
-			const FString ControlName = Control->GetFName().ToString();
-			if (ControlName.EndsWith(TEXT("_CONTROL")))
-			{
-				const FString BaseName = ControlName.Left(ControlName.Len() - 8);
-				Local::SelectAffectedElements(this, InfluenceMap, FRigElementKey(*BaseName, ERigElementType::Bone), bSelected, bInverted);
-				Local::SelectAffectedElements(this, InfluenceMap, FRigElementKey(*BaseName, ERigElementType::Curve), bSelected, bInverted);
-			}
-		}
-	}
-}
-
 void UControlRig::SetBoneInitialTransformsFromAnimInstance(UAnimInstance* InAnimInstance)
 {
 	FMemMark Mark(FMemStack::Get());
diff --git a/Engine/Plugins/Animation/ControlRig/Source/ControlRig/Public/ControlRig.h b/Engine/Plugins/Animation/ControlRig/Source/ControlRig/Public/ControlRig.h
--- a/Engine/Plugins/Animation/ControlRig/Source/ControlRig/Public/ControlRig.h
+++ b/Engine/Plugins/Animation/ControlRig/Source/ControlRig/Public/ControlRig.h
@@ -467,25 +467,43 @@ private:
 
 	const FRigInfluenceMap* FindInfluenceMap(const FName& InEventName);
 
-	UPROPERTY(transient, BlueprintGetter = GetInteractionRig, BlueprintSetter = SetInteractionRig, Category = "Interaction")
-	TObjectPtr<UControlRig> InteractionRig;
+#if WITH_EDITORONLY_DATA
+	UPROPERTY()
+	TObjectPtr<UControlRig> InteractionRig_DEPRECATED;
 
-	UPROPERTY(EditInstanceOnly, transient, BlueprintGetter = GetInteractionRigClass, BlueprintSetter = SetInteractionRigClass, Category = "Interaction", Meta=(DisplayName="Interaction Rig"))
-	TSubclassOf<UControlRig> InteractionRigClass;
+	UPROPERTY()
+	TSubclassOf<UControlRig> InteractionRigClass_DEPRECATED;
+#endif
 
 public:
 
-	UFUNCTION(BlueprintGetter)
-	UControlRig* GetInteractionRig() const { return InteractionRig; }
+	UE_DEPRECATED(5.4, "InteractionRig is no longer used")
+	UFUNCTION(BlueprintGetter, meta = (DeprecatedFunction, DeprecationMessage = "InteractionRig is no longer used"))
+	UControlRig* GetInteractionRig() const 
+	{
+#if WITH_EDITORONLY_DATA
+		return InteractionRig_DEPRECATED;
+#endif
+		return nullptr;
+	}
 
-	UFUNCTION(BlueprintSetter)
-	void SetInteractionRig(UControlRig* InInteractionRig);
+	UE_DEPRECATED(5.4, "InteractionRig is no longer used")
+	UFUNCTION(BlueprintSetter, meta = (DeprecatedFunction, DeprecationMessage = "InteractionRig is no longer used"))
+	void SetInteractionRig(UControlRig* InInteractionRig) {}
 
-	UFUNCTION(BlueprintGetter)
-	TSubclassOf<UControlRig> GetInteractionRigClass() const { return InteractionRigClass; }
+	UE_DEPRECATED(5.4, "InteractionRig is no longer used")
+	UFUNCTION(BlueprintGetter, meta = (DeprecatedFunction, DeprecationMessage = "InteractionRig is no longer used"))
+	TSubclassOf<UControlRig> GetInteractionRigClass() const 
+	{ 
+#if WITH_EDITORONLY_DATA
+		return InteractionRigClass_DEPRECATED;
+#endif
+		return nullptr;
+	}
 
-	UFUNCTION(BlueprintSetter)
-	void SetInteractionRigClass(TSubclassOf<UControlRig> InInteractionRigClass);
+	UE_DEPRECATED(5.4, "InteractionRig is no longer used")
+	UFUNCTION(BlueprintSetter, meta = (DeprecatedFunction, DeprecationMessage = "InteractionRig is no longer used"))
+	void SetInteractionRigClass(TSubclassOf<UControlRig> InInteractionRigClass) {}
 
 	// UObject interface
 #if WITH_EDITOR
@@ -506,11 +524,6 @@ protected:
 private:
 
 	void CopyPoseFromOtherRig(UControlRig* Subject);
-	void HandleInteractionRigControlModified(UControlRig* Subject, FRigControlElement* Control, const FRigControlModifiedContext& Context);
-	void HandleInteractionRigInitialized(URigVMHost* Subject, const FName& EventName);
-	void HandleInteractionRigExecuted(URigVMHost* Subject, const FName& EventName);
-	void HandleInteractionRigControlSelected(UControlRig* Subject, FRigControlElement* InControl, bool bSelected, bool bInverted);
-
 
 protected:
 	bool bCopyHierarchyBeforeConstruction;
