Control Rig: Added key reduction/optimize for bool, int, enum channels when baking to control rig

diff --git a/Engine/Plugins/Animation/ControlRig/Source/ControlRig/Private/Sequencer/MovieSceneControlRigParameterSection.cpp b/Engine/Plugins/Animation/ControlRig/Source/ControlRig/Private/Sequencer/MovieSceneControlRigParameterSection.cpp
--- a/Engine/Plugins/Animation/ControlRig/Source/ControlRig/Private/Sequencer/MovieSceneControlRigParameterSection.cpp
+++ b/Engine/Plugins/Animation/ControlRig/Source/ControlRig/Private/Sequencer/MovieSceneControlRigParameterSection.cpp
@@ -2764,6 +2764,9 @@ void UMovieSceneControlRigParameterSection::FixRotationWinding(const FName& Cont
 void UMovieSceneControlRigParameterSection::OptimizeSection(const FName& ControlName, const FKeyDataOptimizationParams& Params)
 {
 	TArrayView<FMovieSceneFloatChannel*> FloatChannels = GetChannelProxy().GetChannels<FMovieSceneFloatChannel>();
+	TArrayView<FMovieSceneBoolChannel*> BoolChannels = GetChannelProxy().GetChannels<FMovieSceneBoolChannel>();
+	TArrayView<FMovieSceneIntegerChannel*> IntegerChannels = GetChannelProxy().GetChannels<FMovieSceneIntegerChannel>();
+	TArrayView<FMovieSceneByteChannel*> EnumChannels = GetChannelProxy().GetChannels<FMovieSceneByteChannel>();
 	FChannelMapInfo* pChannelIndex = ControlChannelMap.Find(ControlName);
 	if (pChannelIndex != nullptr)
 	{
@@ -2805,6 +2808,23 @@ void UMovieSceneControlRigParameterSection::OptimizeSection(const FName& Control
 					break;
 
 				}
+				case ERigControlType::Bool:
+				{
+					BoolChannels[ChannelIndex]->Optimize(Params);
+					break;
+				}
+				case ERigControlType::Integer:
+				{
+					if (ControlElement->Settings.ControlEnum)
+					{
+						EnumChannels[ChannelIndex]->Optimize(Params);
+					}
+					else
+					{
+						IntegerChannels[ChannelIndex]->Optimize(Params);
+					}
+					break;
+				}
 				default:
 					break;
 			}
@@ -3265,6 +3285,24 @@ bool UMovieSceneControlRigParameterSection::LoadAnimSequenceIntoThisSection(UAni
 				return false;
 			}
 		}
+
+		TArrayView<FMovieSceneBoolChannel*> BoolChannels = GetChannelProxy().GetChannels<FMovieSceneBoolChannel>();
+		for (FMovieSceneBoolChannel* Channel : BoolChannels)
+		{
+			Channel->Optimize(Params);
+		}
+
+		TArrayView<FMovieSceneIntegerChannel*> IntegerChannels = GetChannelProxy().GetChannels<FMovieSceneIntegerChannel>();
+		for (FMovieSceneIntegerChannel* Channel : IntegerChannels)
+		{
+			Channel->Optimize(Params);
+		}
+
+		TArrayView<FMovieSceneByteChannel*> EnumChannels = GetChannelProxy().GetChannels<FMovieSceneByteChannel>();
+		for (FMovieSceneByteChannel* Channel : EnumChannels)
+		{
+			Channel->Optimize(Params);
+		}
 	}
 	
 	return true;
