Control Rig: Fix not able to add absolute control rig after an additive has been introduced
#rnx

diff --git a/Engine/Plugins/Animation/ControlRig/Source/ControlRig/Private/Sequencer/ControlRigLayerInstanceProxy.cpp b/Engine/Plugins/Animation/ControlRig/Source/ControlRig/Private/Sequencer/ControlRigLayerInstanceProxy.cpp
--- a/Engine/Plugins/Animation/ControlRig/Source/ControlRig/Private/Sequencer/ControlRigLayerInstanceProxy.cpp
+++ b/Engine/Plugins/Animation/ControlRig/Source/ControlRig/Private/Sequencer/ControlRigLayerInstanceProxy.cpp
@@ -4,6 +4,7 @@
 #include "AnimNode_ControlRig_ExternalSource.h"
 #include "Sequencer/ControlRigLayerInstance.h"
 #include "AnimSequencerInstance.h"
+#include "ControlRig.h"
 
 #include UE_INLINE_GENERATED_CPP_BY_NAME(ControlRigLayerInstanceProxy)
 
@@ -143,14 +144,23 @@ void FControlRigLayerInstanceProxy::AddControlRigTrack(int32 ControlRigID, UCont
 
 	if(!Node)
 	{
-		FAnimNode_ControlRig_ExternalSource* Parent = (ControlRigNodes.Num() > 0)? ControlRigNodes.Last().Get() : nullptr;
-		Node = ControlRigNodes.Add_GetRef(MakeShared<FAnimNode_ControlRig_ExternalSource>()).Get();
-		SequencerToControlRigNodeMap.FindOrAdd(ControlRigID) = Node;
-		if (Parent)
+		FAnimNode_ControlRig_ExternalSource* PreviousNode = nullptr;
+		FAnimNode_ControlRig_ExternalSource* NextNode = nullptr;
+		if (InControlRig->IsAdditive())
+		{
+			PreviousNode = (ControlRigNodes.Num() > 0)? ControlRigNodes.Last().Get() : nullptr;
+			Node = ControlRigNodes.Add_GetRef(MakeShared<FAnimNode_ControlRig_ExternalSource>()).Get();
+		}
+		else
 		{
-			FAnimNode_Base* LinkedNode = Parent->Source.GetLinkNode();
-			Parent->Source.SetLinkNode(Node);
-			Node->Source.SetLinkNode(LinkedNode);
+			NextNode = (ControlRigNodes.Num() > 0)? ControlRigNodes[0].Get() : nullptr;
+			Node = ControlRigNodes.Insert_GetRef(MakeShared<FAnimNode_ControlRig_ExternalSource>(), 0).Get();
+			ensure(!NextNode || NextNode->GetControlRig()->IsAdditive());
+		}
+
+		if (PreviousNode)
+		{
+			Node->Source.SetLinkNode(PreviousNode);
 		}
 		else
 		{
@@ -158,6 +168,13 @@ void FControlRigLayerInstanceProxy::AddControlRigTrack(int32 ControlRigID, UCont
 			CurrentRoot = Node;
 			Node->Source.SetLinkNode(&InputPose);
 		}
+		
+		if (NextNode)
+		{
+			NextNode->Source.SetLinkNode(Node);
+		}
+		
+		SequencerToControlRigNodeMap.FindOrAdd(ControlRigID) = Node;
 	}
 
 	Node->SetControlRig(InControlRig);
@@ -214,6 +231,7 @@ void FControlRigLayerInstanceProxy::RemoveControlRigTrack(int32 ControlRigID)
 					if (Child)
 					{
 						CurrentRoot = Child;
+						Child->Source.SetLinkNode(&InputPose);
 					}
 					else
 					{
@@ -224,11 +242,7 @@ void FControlRigLayerInstanceProxy::RemoveControlRigTrack(int32 ControlRigID)
 				{
 					if (Child)
 					{
-						Parent->Source.SetLinkNode(Child);
-					}
-					else
-					{
-						Parent->Source.SetLinkNode(&InputPose);
+						Child->Source.SetLinkNode(Parent);
 					}
 				}
 
diff --git a/Engine/Plugins/Animation/ControlRig/Source/ControlRigEditor/Private/Sequencer/ControlRigParameterTrackEditor.cpp b/Engine/Plugins/Animation/ControlRig/Source/ControlRigEditor/Private/Sequencer/ControlRigParameterTrackEditor.cpp
--- a/Engine/Plugins/Animation/ControlRig/Source/ControlRigEditor/Private/Sequencer/ControlRigParameterTrackEditor.cpp
+++ b/Engine/Plugins/Animation/ControlRig/Source/ControlRigEditor/Private/Sequencer/ControlRigParameterTrackEditor.cpp
@@ -1645,14 +1645,24 @@ void FControlRigParameterTrackEditor::BuildObjectBindingTrackMenu(FMenuBuilder&
 
 		if (Skeleton)
 		{
-			//if there are any other control rigs we don't allow it for now..
+			//if there are any other absolute control rigs we don't allow it for now..
 			//mz todo will allow later
 			UMovieScene* MovieScene = GetSequencer()->GetFocusedMovieSceneSequence()->GetMovieScene();
-			UMovieSceneControlRigParameterTrack* ExistingTrack = Cast<UMovieSceneControlRigParameterTrack>(MovieScene->FindTrack(UMovieSceneControlRigParameterTrack::StaticClass(), ObjectBindings[0], NAME_None));
-			if (!ExistingTrack)
+			TArray<UMovieSceneTrack*> ExistingTracks = MovieScene->FindTracks(UMovieSceneControlRigParameterTrack::StaticClass(), ObjectBindings[0], NAME_None);
+			ExistingTracks = ExistingTracks.FilterByPredicate([](UMovieSceneTrack* Track)
+			{
+				if (UMovieSceneControlRigParameterTrack* CRTrack = Cast<UMovieSceneControlRigParameterTrack>(Track))
+				{
+					if (UControlRig* ControlRig = CRTrack->GetControlRig())
+					{
+						return !ControlRig->IsAdditive();
+					}
+				}
+				return false;
+			});
+			if (ExistingTracks.IsEmpty())
 			{
 				UMovieSceneTrack* Track = nullptr;
-
 				MenuBuilder.AddSubMenu(LOCTEXT("ControlRigText", "Control Rig"), FText(), FNewMenuDelegate::CreateSP(this, &FControlRigParameterTrackEditor::HandleAddTrackSubMenu, ObjectBindings, Track));
 			}
 		}
