Control Rig: Fix compilation issue
#rnx

diff --git a/Engine/Plugins/Animation/ControlRig/Source/ControlRig/Private/Sequencer/ControlRigLayerInstanceProxy.cpp b/Engine/Plugins/Animation/ControlRig/Source/ControlRig/Private/Sequencer/ControlRigLayerInstanceProxy.cpp
--- a/Engine/Plugins/Animation/ControlRig/Source/ControlRig/Private/Sequencer/ControlRigLayerInstanceProxy.cpp
+++ b/Engine/Plugins/Animation/ControlRig/Source/ControlRig/Private/Sequencer/ControlRigLayerInstanceProxy.cpp
@@ -247,7 +247,7 @@ void FControlRigLayerInstanceProxy::RemoveControlRigTrack(int32 ControlRigID)
 			Parent = Current;
 		}
 
-		CurrentRoot = (ControlRigNodes.IsEmpty()) ? InputPose : ControlRigNodes.Last().Get();
+		CurrentRoot = (ControlRigNodes.IsEmpty()) ? &InputPose : ControlRigNodes.Last().Get();
 		SequencerToControlRigNodeMap.Remove(ControlRigID);
 	}
 }
