Control Rig: Fix proxy controls not behaving correctly
#rnx

diff --git a/Engine/Plugins/Animation/ControlRig/Source/ControlRig/Private/ControlRig.cpp b/Engine/Plugins/Animation/ControlRig/Source/ControlRig/Private/ControlRig.cpp
--- a/Engine/Plugins/Animation/ControlRig/Source/ControlRig/Private/ControlRig.cpp
+++ b/Engine/Plugins/Animation/ControlRig/Source/ControlRig/Private/ControlRig.cpp
@@ -1660,10 +1660,25 @@ bool UControlRig::SetControlGlobalTransform(const FName& InControlName, const FT
 		GlobalTransform = SetupControlFromGlobalTransform(InControlName, GlobalTransform);
 	}
 
+	FRigControlElement* Control = FindControl(InControlName);
+	if (IsAdditive())
+	{
+		if (Control)
+		{
+			// For additive control rigs, proxy controls should not be stored and applied after the backwards solve like other controls
+			// Instead, apply the transform as global, run the forward solve and let the driven controls request keying
+			if (Control->Settings.AnimationType == ERigControlAnimationType::ProxyControl)
+			{
+				DynamicHierarchy->SetTransform(Control, GlobalTransform, ERigTransformType::CurrentGlobal, true);
+				Execute(FRigUnit_BeginExecution::EventName);
+				return true;
+			}
+		}
+	}
+
 	FRigControlValue Value = GetControlValueFromGlobalTransform(InControlName, GlobalTransform, TransformType);
 	if (OnFilterControl.IsBound())
 	{
-		FRigControlElement* Control = FindControl(InControlName);
 		if (Control)
 		{
 			OnFilterControl.Broadcast(this, Control, Value);
