Control Rig: Avoid crash with SKM with out a skeleton

diff --git a/Engine/Plugins/Animation/ControlRig/Source/ControlRigEditor/Private/EditMode/ControlRigEditMode.cpp b/Engine/Plugins/Animation/ControlRig/Source/ControlRigEditor/Private/EditMode/ControlRigEditMode.cpp
--- a/Engine/Plugins/Animation/ControlRig/Source/ControlRigEditor/Private/EditMode/ControlRigEditMode.cpp
+++ b/Engine/Plugins/Animation/ControlRig/Source/ControlRigEditor/Private/EditMode/ControlRigEditMode.cpp
@@ -968,6 +968,11 @@ bool FControlRigEditMode::ProcessCapturedMouseMoves(FEditorViewportClient* InVie
 
 bool FControlRigEditMode::EndTracking(FEditorViewportClient* InViewportClient, FViewport* InViewport)
 {
+	if(RuntimeControlRigs.IsEmpty())
+	{
+		return false;
+	}
+
 	if (bisTrackingAnimToolDrag)
 	{
 		ResetAnimSlider();
@@ -1021,6 +1026,11 @@ bool FControlRigEditMode::EndTracking(FEditorViewportClient* InViewportClient, F
 
 bool FControlRigEditMode::StartTracking(FEditorViewportClient* InViewportClient, FViewport* InViewport)
 {
+	if(RuntimeControlRigs.IsEmpty())
+	{
+		return false;
+	}
+	
 	//break right out if doing anim slider scrub
 
 	if (IsDragAnimSliderToolPressed(InViewport))
