Control Rig: Maintain local pos / scl if parent scale is zero

diff --git a/Engine/Plugins/Animation/ControlRig/Source/ControlRig/Private/Rigs/RigHierarchy.cpp b/Engine/Plugins/Animation/ControlRig/Source/ControlRig/Private/Rigs/RigHierarchy.cpp
--- a/Engine/Plugins/Animation/ControlRig/Source/ControlRig/Private/Rigs/RigHierarchy.cpp
+++ b/Engine/Plugins/Animation/ControlRig/Source/ControlRig/Private/Rigs/RigHierarchy.cpp
@@ -2988,9 +2988,30 @@ FTransform URigHierarchy::GetTransform(FRigTransformElement* InTransformElement,
 		FTransform ParentTransform;
 		if(IsLocal(InTransformType))
 		{
+			// if we have a zero scale provided - and the parent also contains a zero scale,
+			// we'll keep the local translation and scale since otherwise we'll loose the values.
+			// we cannot compute the local from the global if the scale is 0 - since the local scale
+			// may be anything - any translation or scale multiplied with the parent's zero scale is zero. 
+			auto CompensateZeroScale = [this, InTransformElement, InTransformType](FTransform& Transform)
+			{
+				const FVector Scale = Transform.GetScale3D();
+				if(FMath::IsNearlyZero(Scale.X) || FMath::IsNearlyZero(Scale.Y) || FMath::IsNearlyZero(Scale.Z))
+				{
+					const FTransform ParentTransform =
+						GetParentTransform(InTransformElement, ERigTransformType::SwapLocalAndGlobal(InTransformType));
+					const FVector ParentScale = ParentTransform.GetScale3D();
+					if(FMath::IsNearlyZero(ParentScale.X) || FMath::IsNearlyZero(ParentScale.Y) || FMath::IsNearlyZero(ParentScale.Z))
+					{
+						Transform.SetTranslation(InTransformElement->Pose.Get(InTransformType).GetTranslation());
+						Transform.SetScale3D(InTransformElement->Pose.Get(InTransformType).GetScale3D());
+					}
+				}
+			};
+			
 			if(FRigControlElement* ControlElement = Cast<FRigControlElement>(InTransformElement))
 			{
-				const FTransform NewTransform = ComputeLocalControlValue(ControlElement, ControlElement->Pose.Get(OpposedType), GlobalType);
+				FTransform NewTransform = ComputeLocalControlValue(ControlElement, ControlElement->Pose.Get(OpposedType), GlobalType);
+				CompensateZeroScale(NewTransform);
 				ControlElement->Pose.Set(InTransformType, NewTransform);
 				/** from mikez we do not want geting a pose to set these preferred angles
 				switch(ControlElement->Settings.ControlType)
@@ -3015,9 +3036,10 @@ FTransform URigHierarchy::GetTransform(FRigTransformElement* InTransformElement,
 				// this is done for nulls and any element that can have more than one parent which 
 				// is not a control
 				const FTransform& GlobalTransform = MultiParentElement->Pose.Get(GlobalType);
-				const FTransform LocalTransform = InverseSolveParentConstraints(
+				FTransform LocalTransform = InverseSolveParentConstraints(
 					GlobalTransform, 
 					MultiParentElement->ParentConstraints, GlobalType, FTransform::Identity);
+				CompensateZeroScale(LocalTransform);
 				MultiParentElement->Pose.Set(InTransformType, LocalTransform);
 			}
 			else
@@ -3026,6 +3048,7 @@ FTransform URigHierarchy::GetTransform(FRigTransformElement* InTransformElement,
 
 				FTransform NewTransform = InTransformElement->Pose.Get(OpposedType).GetRelativeTransform(ParentTransform);
 				NewTransform.NormalizeRotation();
+				CompensateZeroScale(NewTransform);
 				InTransformElement->Pose.Set(InTransformType, NewTransform);
 			}
 		}
@@ -5783,11 +5806,13 @@ FTransform URigHierarchy::ComputeLocalControlValue(FRigControlElement* ControlEl
 	const FTransform OffsetTransform =
 		GetControlOffsetTransform(ControlElement, ERigTransformType::MakeLocal(InTransformType));
 
-	return InverseSolveParentConstraints(
+	FTransform Result = InverseSolveParentConstraints(
 		InGlobalTransform,
 		ControlElement->ParentConstraints,
 		InTransformType,
 		OffsetTransform);
+
+	return Result;
 }
 
 FTransform URigHierarchy::SolveParentConstraints(
