Control Rig: Hides additive control rigs behind CVar ControlRig.Sequencer.EnableAdditiveControlRigs
#rnx

diff --git a/Engine/Plugins/Animation/ControlRig/Source/ControlRigEditor/Private/Sequencer/ControlRigParameterTrackEditor.cpp b/Engine/Plugins/Animation/ControlRig/Source/ControlRigEditor/Private/Sequencer/ControlRigParameterTrackEditor.cpp
--- a/Engine/Plugins/Animation/ControlRig/Source/ControlRigEditor/Private/Sequencer/ControlRigParameterTrackEditor.cpp
+++ b/Engine/Plugins/Animation/ControlRig/Source/ControlRigEditor/Private/Sequencer/ControlRigParameterTrackEditor.cpp
@@ -99,6 +99,8 @@ TAutoConsoleVariable<bool> CVarSelectedKeysSelectControls(TEXT("ControlRig.Seque
 
 TAutoConsoleVariable<bool> CVarSelectedSectionSetsSectionToKey(TEXT("ControlRig.Sequencer.SelectedSectionSetsSectionToKey"), false, TEXT("When true when we select a channel in a section, if it's the only section selected we set it as the Section To Key, by default false."));
 
+TAutoConsoleVariable<bool> CVarEnableAdditiveControlRigs(TEXT("ControlRig.Sequencer.EnableAdditiveControlRigs"), false, TEXT("When true it is possible to add an additive control rig to a skeletal mesh component."));
+
 static USkeletalMeshComponent* AcquireSkeletalMeshFromObject(UObject* BoundObject, TSharedPtr<ISequencer> SequencerPtr)
 {
 	if (AActor* Actor = Cast<AActor>(BoundObject))
@@ -631,15 +633,18 @@ void FControlRigParameterTrackEditor::BuildObjectBindingContextMenu(FMenuBuilder
 					NAME_None,
 					EUserInterfaceActionType::Button);
 
-				MenuBuilder.AddMenuEntry(
-					LOCTEXT("AddAdditiveControlRig", "Add Additive Control Rig"),
-					LOCTEXT("AddAdditiveControlRigTooltip", "Add additive Control Rig and add a track for it"),
-					FSlateIcon(),
-					FUIAction(
-						FExecuteAction::CreateRaw(this, &FControlRigParameterTrackEditor::AddAdditiveControlRig, ObjectBindings[0], BoundObject, SkelMeshComp, Skeleton),
-						FCanExecuteAction::CreateRaw(this, &FControlRigParameterTrackEditor::CanAddAdditiveConrolRig)),
-					NAME_None,
-					EUserInterfaceActionType::Button);
+				if (CVarEnableAdditiveControlRigs->GetBool())
+				{
+					MenuBuilder.AddMenuEntry(
+					   LOCTEXT("AddAdditiveControlRig", "Add Additive Control Rig"),
+					   LOCTEXT("AddAdditiveControlRigTooltip", "Add additive Control Rig and add a track for it"),
+					   FSlateIcon(),
+					   FUIAction(
+						   FExecuteAction::CreateRaw(this, &FControlRigParameterTrackEditor::AddAdditiveControlRig, ObjectBindings[0], BoundObject, SkelMeshComp, Skeleton),
+						   FCanExecuteAction::CreateRaw(this, &FControlRigParameterTrackEditor::CanAddAdditiveConrolRig)),
+					   NAME_None,
+					   EUserInterfaceActionType::Button);
+				}
 
 				MenuBuilder.AddMenuEntry(
 					LOCTEXT("FilterAssetBySkeleton", "Filter Asset By Skeleton"),
