Constraint & additive ControlRig:
    In the context of additive rigs, the evaluation as to be forced after setting the control value to request keyframing thru the OnControlModified delegate.
    Not doing this results in the keys not being added (FControlRigParameterTrackEditor::HandleControlModified is not in game thread when evaluating the additive rig) so there's not new FMovieSceneRootEvaluationTemplateInstance created

NOTE
   Calling Evaluate_AnyThread here will force the update off every controls whereas we're just trying to add keys on a single one. Having a function taht could only compute the additive value on a set of keys could help reducing the cost but there might be some side effets as well.
   Leaving this as it is for now as additive rigs are still experimental atm.

#rnx

diff --git a/Engine/Plugins/Animation/ControlRig/Source/ControlRig/Private/Constraints/ControlRigTransformableHandle.cpp b/Engine/Plugins/Animation/ControlRig/Source/ControlRig/Private/Constraints/ControlRigTransformableHandle.cpp
--- a/Engine/Plugins/Animation/ControlRig/Source/ControlRig/Private/Constraints/ControlRigTransformableHandle.cpp
+++ b/Engine/Plugins/Animation/ControlRig/Source/ControlRig/Private/Constraints/ControlRigTransformableHandle.cpp
@@ -489,9 +489,19 @@ bool UTransformableControlHandle::AddTransformKeys(const TArray<FFrameNumber>& I
 
 		if (bLocal)
 		{
-			return InControlRig->SetControlLocalTransform(ControlName, InTransform, bNotify, InKeyframeContext, bUndo, bFixEuler);
+			InControlRig->SetControlLocalTransform(ControlName, InTransform, bNotify, InKeyframeContext, bUndo, bFixEuler);
+			if (InControlRig->IsAdditive())
+			{
+				InControlRig->Evaluate_AnyThread();
+			}
+			return;
 		}
+		
 		InControlRig->SetControlGlobalTransform(ControlName, InTransform, bNotify, InKeyframeContext, bUndo, bFixEuler);
+		if (InControlRig->IsAdditive())
+		{
+			InControlRig->Evaluate_AnyThread();
+		}
 	};
 
 	FRigControlModifiedContext KeyframeContext;
