Control Rig: Fix snapper not working with additive rigs
#rnx

diff --git a/Engine/Plugins/Animation/ControlRig/Source/ControlRig/Private/ControlRig.cpp b/Engine/Plugins/Animation/ControlRig/Source/ControlRig/Private/ControlRig.cpp
--- a/Engine/Plugins/Animation/ControlRig/Source/ControlRig/Private/ControlRig.cpp
+++ b/Engine/Plugins/Animation/ControlRig/Source/ControlRig/Private/ControlRig.cpp
@@ -409,6 +409,23 @@ void UControlRig::Evaluate_AnyThread()
 	}
 }
 
+bool UControlRig::EvaluateSkeletalMeshComponent(double InDeltaTime)
+{
+	if (USkeletalMeshComponent* SkelMeshComp = Cast<USkeletalMeshComponent>(GetObjectBinding()->GetBoundObject()))
+	{
+		SkelMeshComp->TickAnimation(InDeltaTime, false);
+
+		SkelMeshComp->RefreshBoneTransforms();
+		SkelMeshComp->RefreshFollowerComponents();
+		SkelMeshComp->UpdateComponentToWorld();
+		SkelMeshComp->FinalizeBoneTransform();
+		SkelMeshComp->MarkRenderTransformDirty();
+		SkelMeshComp->MarkRenderDynamicDataDirty();
+		return true;
+	}
+	return false;
+}
+
 void UControlRig::ResetControlValues()
 {
 	ControlValues.Reset();
diff --git a/Engine/Plugins/Animation/ControlRig/Source/ControlRig/Public/ControlRig.h b/Engine/Plugins/Animation/ControlRig/Source/ControlRig/Public/ControlRig.h
--- a/Engine/Plugins/Animation/ControlRig/Source/ControlRig/Public/ControlRig.h
+++ b/Engine/Plugins/Animation/ControlRig/Source/ControlRig/Public/ControlRig.h
@@ -107,6 +107,9 @@ public:
 	/** Evaluates the ControlRig */
 	virtual void Evaluate_AnyThread() override;
 
+	/** Ticks animation of the skeletal mesh component bound to this control rig */
+	bool EvaluateSkeletalMeshComponent(double InDeltaTime);
+
 	/** Removes any stored additive control values */
 	void ResetControlValues();
 
diff --git a/Engine/Plugins/Animation/ControlRig/Source/ControlRigEditor/Private/Tools/ControlRigSnapper.cpp b/Engine/Plugins/Animation/ControlRig/Source/ControlRigEditor/Private/Tools/ControlRigSnapper.cpp
--- a/Engine/Plugins/Animation/ControlRig/Source/ControlRigEditor/Private/Tools/ControlRigSnapper.cpp
+++ b/Engine/Plugins/Animation/ControlRig/Source/ControlRigEditor/Private/Tools/ControlRigSnapper.cpp
@@ -130,20 +130,30 @@ static bool LocalGetControlRigControlTransforms(IMovieScenePlayer* Player, const
 
 		FFrameRate TickResolution = MovieScene->GetTickResolution();
 		FFrameRate DisplayRate = MovieScene->GetDisplayRate();
-
+		
 		OutTransforms.SetNum(Frames.Num());
 		for (int32 Index = 0; Index < Frames.Num(); ++Index)
 		{
 			const FFrameNumber& FrameNumber = Frames[Index];
+			double DeltaTime = 0.0;
 			if (CurrentFrame.IsSet() == false || CurrentFrame.GetValue() != FrameNumber)
 			{
 				FFrameTime GlobalTime(FrameNumber);
 				GlobalTime = GlobalTime * RootToLocalTransform.InverseNoLooping(); //player evals in root time so need to go back to it.
 
 				FMovieSceneContext Context = FMovieSceneContext(FMovieSceneEvaluationRange(GlobalTime, TickResolution), Player->GetPlaybackStatus()).SetHasJumped(true);
+
+				DeltaTime = 1.0/Context.GetFrameRate().AsDecimal();
 				Player->GetEvaluationTemplate().EvaluateSynchronousBlocking(Context, *Player);
 			}
-			ControlRig->Evaluate_AnyThread();
+			if (ControlRig->IsAdditive())
+			{
+				ControlRig->EvaluateSkeletalMeshComponent(DeltaTime);
+			}
+			else
+			{
+				ControlRig->Evaluate_AnyThread();
+			}
 			OutTransforms[Index] = ControlRig->GetControlGlobalTransform(ControlName) * ParentTransforms[Index];
 		}
 	}
