Control Rig: Move loading of the control enum for rig element settings to postload.

#rnx

diff --git a/Engine/Plugins/Animation/ControlRig/Source/ControlRig/Private/Rigs/RigHierarchy.cpp b/Engine/Plugins/Animation/ControlRig/Source/ControlRig/Private/Rigs/RigHierarchy.cpp
--- a/Engine/Plugins/Animation/ControlRig/Source/ControlRig/Private/Rigs/RigHierarchy.cpp
+++ b/Engine/Plugins/Animation/ControlRig/Source/ControlRig/Private/Rigs/RigHierarchy.cpp
@@ -327,6 +327,11 @@ void URigHierarchy::PostLoad()
 
 	LLM_SCOPE_BYNAME(TEXT("Animation/ControlRig"));
 
+	for(FRigBaseElement* Element: Elements)
+	{
+		Element->PostLoad();
+	}
+
 	CleanupInvalidCaches();
 }
 
diff --git a/Engine/Plugins/Animation/ControlRig/Source/ControlRig/Private/Rigs/RigHierarchyElements.cpp b/Engine/Plugins/Animation/ControlRig/Source/ControlRig/Private/Rigs/RigHierarchyElements.cpp
--- a/Engine/Plugins/Animation/ControlRig/Source/ControlRig/Private/Rigs/RigHierarchyElements.cpp
+++ b/Engine/Plugins/Animation/ControlRig/Source/ControlRig/Private/Rigs/RigHierarchyElements.cpp
@@ -880,7 +880,6 @@ void FRigControlSettings::Load(FArchive& Ar)
 	static const UEnum* ControlAxisEnum = StaticEnum<ERigControlAxis>();
 
 	FName AnimationTypeName, ControlTypeName, ShapeVisibilityName, PrimaryAxisName;
-	FString ControlEnumPathName;
 
 	bool bLimitTranslation_DEPRECATED = false;
 	bool bLimitRotation_DEPRECATED = false;
@@ -955,7 +954,7 @@ void FRigControlSettings::Load(FArchive& Ar)
 	
 	Ar << ShapeColor;
 	Ar << bIsTransientControl;
-	Ar << ControlEnumPathName;
+	Ar << ControlEnumToLoadPathName;
 
 	AnimationType = (ERigControlAnimationType)AnimationTypeEnum->GetValueByName(AnimationTypeName);
 	PrimaryAxis = (ERigControlAxis)ControlAxisEnum->GetValueByName(PrimaryAxisName);
@@ -968,17 +967,6 @@ void FRigControlSettings::Load(FArchive& Ar)
 	}
 
 	ControlEnum = nullptr;
-	if(!ControlEnumPathName.IsEmpty())
-	{
-		if (IsInGameThread())
-		{
-			ControlEnum = LoadObject<UEnum>(nullptr, *ControlEnumPathName);
-		}
-		else
-		{			
-			ControlEnum = FindObject<UEnum>(nullptr, *ControlEnumPathName);
-		}
-	}
 
 	if (Ar.CustomVer(FControlRigObjectVersion::GUID) >= FControlRigObjectVersion::RigHierarchyControlSpaceFavorites)
 	{
@@ -1057,6 +1045,16 @@ void FRigControlSettings::Load(FArchive& Ar)
 
 }
 
+
+void FRigControlSettings::PostLoad()
+{
+	if(!ControlEnumToLoadPathName.IsEmpty())
+	{
+		ControlEnum = LoadObject<UEnum>(nullptr, *ControlEnumToLoadPathName);
+	}
+}
+
+
 uint32 GetTypeHash(const FRigControlSettings& Settings)
 {
 	uint32 Hash = GetTypeHash(Settings.ControlType);
@@ -1284,6 +1282,13 @@ void FRigControlElement::Load(FArchive& Ar, URigHierarchy* Hierarchy, ESerializa
 	}
 }
 
+void FRigControlElement::PostLoad()
+{
+	Super::PostLoad();
+
+	Settings.PostLoad();
+}
+
 void FRigControlElement::CopyFrom(URigHierarchy* InHierarchy, FRigBaseElement* InOther, URigHierarchy* InOtherHierarchy)
 {
 	Super::CopyFrom(InHierarchy, InOther, InOtherHierarchy);
diff --git a/Engine/Plugins/Animation/ControlRig/Source/ControlRig/Public/Rigs/RigHierarchyElements.h b/Engine/Plugins/Animation/ControlRig/Source/ControlRig/Public/Rigs/RigHierarchyElements.h
--- a/Engine/Plugins/Animation/ControlRig/Source/ControlRig/Public/Rigs/RigHierarchyElements.h
+++ b/Engine/Plugins/Animation/ControlRig/Source/ControlRig/Public/Rigs/RigHierarchyElements.h
@@ -627,6 +627,7 @@ public:
 	void Serialize(FArchive& Ar, URigHierarchy* Hierarchy, ESerializationPhase SerializationPhase);
 	virtual void Save(FArchive& Ar, URigHierarchy* Hierarchy, ESerializationPhase SerializationPhase);
 	virtual void Load(FArchive& Ar, URigHierarchy* Hierarchy, ESerializationPhase SerializationPhase);
+	virtual void PostLoad() {};
 
 	const FName& GetFName() const { return Key.Name; }
 	const FString& GetName() const
@@ -1140,6 +1141,7 @@ struct CONTROLRIG_API FRigControlSettings
 
 	void Save(FArchive& Ar);
 	void Load(FArchive& Ar);
+	void PostLoad();
 
 	friend uint32 GetTypeHash(const FRigControlSettings& Settings);
 
@@ -1250,6 +1252,11 @@ struct CONTROLRIG_API FRigControlSettings
 	UPROPERTY(BlueprintReadOnly, EditAnywhere, Category = Animation)
 	bool bUsePreferredRotationOrder;
 
+private:
+	/** The object path of the enum to load into ControlEnum in PostLoad */ 
+	FString ControlEnumToLoadPathName;
+	
+public:
 	/**
 	* The euler rotation order this control prefers for animation if it is active. If not set then we use the default UE rotator.
 	*/
@@ -1426,6 +1433,7 @@ struct CONTROLRIG_API FRigControlElement : public FRigMultiParentElement
 
 	virtual void Save(FArchive& A, URigHierarchy* Hierarchy, ESerializationPhase SerializationPhase) override;
 	virtual void Load(FArchive& Ar, URigHierarchy* Hierarchy, ESerializationPhase SerializationPhase) override;
+	virtual void PostLoad() override;
 
 private:
 
@@ -1444,7 +1452,7 @@ public:
 
 	UPROPERTY(BlueprintReadOnly, EditAnywhere, Category = RigElement)
 	FRigPreferredEulerAngles PreferredEulerAngles;
-	
+
 protected:
 
 	static bool IsClassOf(const FRigBaseElement* InElement)
