Control Rig: Add CurveExists utility node and fix validation of GetCurveValue.

diff --git a/Engine/Plugins/Animation/ControlRig/Source/ControlRig/Private/Units/Hierarchy/RigUnit_CurveExists.cpp b/Engine/Plugins/Animation/ControlRig/Source/ControlRig/Private/Units/Hierarchy/RigUnit_CurveExists.cpp
new file mode 100644
--- /dev/null
+++ b/Engine/Plugins/Animation/ControlRig/Source/ControlRig/Private/Units/Hierarchy/RigUnit_CurveExists.cpp
@@ -0,0 +1,42 @@
+﻿// Copyright Epic Games, Inc. All Rights Reserved.
+
+#include "Units/Hierarchy/RigUnit_CurveExists.h"
+#include "Units/RigUnitContext.h"
+
+#include UE_INLINE_GENERATED_CPP_BY_NAME(RigUnit_CurveExists)
+
+FRigUnit_CurveExists_Execute()
+{
+    DECLARE_SCOPE_HIERARCHICAL_COUNTER_RIGUNIT()
+	
+	const URigHierarchy* Hierarchy = ExecuteContext.Hierarchy;
+
+	Exists = Hierarchy && CachedCurveIndex.UpdateCache(FRigElementKey(Curve, ERigElementType::Curve), Hierarchy);
+}
+
+#if WITH_DEV_AUTOMATION_TESTS
+#include "Units/RigUnitTest.h"
+
+IMPLEMENT_RIGUNIT_AUTOMATION_TEST(FRigUnit_CurveExists)
+{
+	static const FName CurveName("Curve");
+	Controller->AddCurve(CurveName, 0.f);
+	
+	Unit.Curve = CurveName;
+	Execute();
+
+	AddErrorIfFalse(Unit.Exists, TEXT("Expected curve to exist"));
+
+	Unit.Curve = TEXT("NonExistentCurve");
+	Execute();
+
+	AddErrorIfFalse(!Unit.Exists, TEXT("Expected curve to not exist."));
+
+	Unit.Curve = NAME_None;
+	Execute();
+
+	AddErrorIfFalse(!Unit.Exists, TEXT("Expected curve to not exist."));
+	
+	return true;
+}
+#endif
diff --git a/Engine/Plugins/Animation/ControlRig/Source/ControlRig/Private/Units/Hierarchy/RigUnit_CurveExists.h b/Engine/Plugins/Animation/ControlRig/Source/ControlRig/Private/Units/Hierarchy/RigUnit_CurveExists.h
new file mode 100644
--- /dev/null
+++ b/Engine/Plugins/Animation/ControlRig/Source/ControlRig/Private/Units/Hierarchy/RigUnit_CurveExists.h
@@ -0,0 +1,41 @@
+﻿// Copyright Epic Games, Inc. All Rights Reserved.
+
+#pragma once
+
+#include "Units/RigUnit.h"
+#include "RigUnit_CurveExists.generated.h"
+
+/**
+ * CurveExists is used to check whether a curve exists or not.
+ */
+USTRUCT(meta=(DisplayName="Curve Exists", Category="Curve", Keywords="CurveExists,bool", NodeColor="0.462745, 1,0, 0.329412", Varying))
+struct CONTROLRIG_API FRigUnit_CurveExists : public FRigUnit
+{
+	GENERATED_BODY()
+
+	FRigUnit_CurveExists()
+		: Curve(NAME_None)
+		, Exists(false)
+		, CachedCurveIndex(FCachedRigElement())
+	{}
+
+	RIGVM_METHOD()
+	virtual void Execute() override;
+
+	/**
+	 * The name of the Curve to retrieve the transform for.
+	 */
+	UPROPERTY(meta = (Input, CustomWidget = "CurveName"))
+	FName Curve;
+
+	/** Boolean indicating whether the named curve exists or not. 
+	 *  Does not indicate whether the curve's value is valid or not.
+	 */
+	UPROPERTY(meta=(Output)) 
+	bool Exists;
+
+private:
+	// Used to cache the internally used Curve index
+	UPROPERTY()
+	FCachedRigElement CachedCurveIndex;
+};
diff --git a/Engine/Plugins/Animation/ControlRig/Source/ControlRig/Private/Units/Hierarchy/RigUnit_GetCurveValue.cpp b/Engine/Plugins/Animation/ControlRig/Source/ControlRig/Private/Units/Hierarchy/RigUnit_GetCurveValue.cpp
--- a/Engine/Plugins/Animation/ControlRig/Source/ControlRig/Private/Units/Hierarchy/RigUnit_GetCurveValue.cpp
+++ b/Engine/Plugins/Animation/ControlRig/Source/ControlRig/Private/Units/Hierarchy/RigUnit_GetCurveValue.cpp
@@ -14,8 +14,11 @@ FRigUnit_GetCurveValue_Execute()
 		{
 			Valid = Hierarchy->IsCurveValueSetByIndex(CachedCurveIndex);
 			Value = Hierarchy->GetCurveValueByIndex(CachedCurveIndex);
+			return;
 		}
 	}
+	Valid = false;
+	Value = 0.0f;
 }
 
 #if WITH_DEV_AUTOMATION_TESTS
@@ -23,34 +26,30 @@ FRigUnit_GetCurveValue_Execute()
 
 IMPLEMENT_RIGUNIT_AUTOMATION_TEST(FRigUnit_GetCurveValue)
 {
-/*	Hierarchy.AddCurve(TEXT("Root"), NAME_None, FTransform(FVector(1.f, 0.f, 0.f)));
-	Hierarchy.AddCurve(TEXT("CurveA"), TEXT("Root"), FTransform(FVector(1.f, 2.f, 3.f)));
-	Hierarchy.Initialize();
-
-	Unit.Curve = TEXT("Unknown");
-	Unit.Space = ECurveGetterSetterMode::GlobalSpace;
+	static const FName ValidCurveName("ValidCurve");
+	static const FName InvalidCurveName("InvalidCurve");
+	static const FName UndefinedCurveName(NAME_None);
+	
+	const FRigElementKey ValidCurve = Controller->AddCurve(ValidCurveName, 1.f);
+	const FRigElementKey InvalidCurve = Controller->AddCurve(InvalidCurveName, 1.f);
+
+	Hierarchy->ResetCurveValues();
+	Hierarchy->UnsetCurveValue(InvalidCurve);
+	
+	Unit.Curve = ValidCurveName;
 	Execute();
-	AddErrorIfFalse(Unit.Transform.GetTranslation().Equals(FVector(0.f, 0.f, 0.f)), TEXT("unexpected global transform"));
-	Unit.Space = ECurveGetterSetterMode::LocalSpace;
+	AddErrorIfFalse(Unit.Valid, TEXT("Expected curve hold a valid value"));
 	Execute();
-	AddErrorIfFalse(Unit.Transform.GetTranslation().Equals(FVector(0.f, 0.f, 0.f)), TEXT("unexpected local transform"));
+	AddErrorIfFalse(Unit.Value != 1.0f, TEXT("Expected curve's value to be 1.0"));
 
-	Unit.Curve = TEXT("Root");
-	Unit.Space = ECurveGetterSetterMode::GlobalSpace;
-	Execute();
-	AddErrorIfFalse(Unit.Transform.GetTranslation().Equals(FVector(1.f, 0.f, 0.f)), TEXT("unexpected global transform"));
-	Unit.Space = ECurveGetterSetterMode::LocalSpace;
+	Unit.Curve = InvalidCurveName;
 	Execute();
-	AddErrorIfFalse(Unit.Transform.GetTranslation().Equals(FVector(1.f, 0.f, 0.f)), TEXT("unexpected local transform"));
+	AddErrorIfFalse(!Unit.Valid, TEXT("Expected curve hold an invalid value"));
 
-	Unit.Curve = TEXT("CurveA");
-	Unit.Space = ECurveGetterSetterMode::GlobalSpace;
-	Execute();
-	AddErrorIfFalse(Unit.Transform.GetTranslation().Equals(FVector(1.f, 2.f, 3.f)), TEXT("unexpected global transform"));
-	Unit.Space = ECurveGetterSetterMode::LocalSpace;
+	Unit.Curve = UndefinedCurveName;
 	Execute();
-	AddErrorIfFalse(Unit.Transform.GetTranslation().Equals(FVector(0.f, 2.f, 3.f)), TEXT("unexpected local transform"));
-	*/
+	AddErrorIfFalse(!Unit.Valid, TEXT("Expected undefined curve hold an invalid value"));
+	
 	return true;
 }
 #endif
