ControlRig: Avoid redundant control visibility related warning

diff --git a/Engine/Plugins/Animation/ControlRig/Source/ControlRig/Private/Units/Hierarchy/RigUnit_SetControlVisibility.cpp b/Engine/Plugins/Animation/ControlRig/Source/ControlRig/Private/Units/Hierarchy/RigUnit_SetControlVisibility.cpp
--- a/Engine/Plugins/Animation/ControlRig/Source/ControlRig/Private/Units/Hierarchy/RigUnit_SetControlVisibility.cpp
+++ b/Engine/Plugins/Animation/ControlRig/Source/ControlRig/Private/Units/Hierarchy/RigUnit_SetControlVisibility.cpp
@@ -69,9 +69,10 @@ FRigUnit_SetControlVisibility_Execute()
 			if (CachedControlIndex.IsValid())
 			{
 #if WITH_EDITOR
-				if (FRigControlElement* ControlElement = Cast<FRigControlElement>(Hierarchy->Find(CachedControlIndex.GetKey())))
+				if (const FRigControlElement* ControlElement = Cast<FRigControlElement>(Hierarchy->Find(CachedControlIndex.GetKey())))
 				{
-					if (ControlElement->Settings.ShapeVisibility == ERigControlVisibility::BasedOnSelection)
+					if (ControlElement->Settings.AnimationType == ERigControlAnimationType::ProxyControl &&
+						ControlElement->Settings.ShapeVisibility == ERigControlVisibility::BasedOnSelection)
 					{
 						UE_CONTROLRIG_RIGUNIT_REPORT_WARNING(
 						TEXT("SetControlVisibility: visibility of '%s' is based on selection, and cannot be modified through SetControlVisibility. "
