Control Rig: Ui refactor for sequencer
#rnx

diff --git a/Engine/Plugins/Animation/ControlRig/Source/ControlRigEditor/Private/Sequencer/ControlRigParameterTrackEditor.cpp b/Engine/Plugins/Animation/ControlRig/Source/ControlRigEditor/Private/Sequencer/ControlRigParameterTrackEditor.cpp
--- a/Engine/Plugins/Animation/ControlRig/Source/ControlRigEditor/Private/Sequencer/ControlRigParameterTrackEditor.cpp
+++ b/Engine/Plugins/Animation/ControlRig/Source/ControlRigEditor/Private/Sequencer/ControlRigParameterTrackEditor.cpp
@@ -218,6 +218,7 @@ FControlRigParameterTrackEditor::FControlRigParameterTrackEditor(TSharedRef<ISeq
 	, bCurveDisplayTickIsPending(false)
 	, bIsDoingSelection(false)
 	, bSkipNextSelectionFromTimer(false)
+	, bIsAdditiveControlRig(false)
 	, bFilterAssetBySkeleton(true)
 	, bFilterAssetByAnimatableControls(false)
 	, ControlUndoBracket(0)
@@ -625,27 +626,6 @@ void FControlRigParameterTrackEditor::BuildObjectBindingContextMenu(FMenuBuilder
 		{
 			MenuBuilder.BeginSection("Control Rig", LOCTEXT("ControlRig", "Control Rig"));
 			{
-				MenuBuilder.AddMenuEntry(
-					LOCTEXT("EditWithFKControlRig", "Edit With FK Control Rig"),
-					LOCTEXT("ConvertToFKControlRigTooltip", "Convert to FK Control Rig and add a track for it"),
-					FSlateIcon(),
-					FUIAction(FExecuteAction::CreateRaw(this, &FControlRigParameterTrackEditor::ConvertToFKControlRig, ObjectBindings[0], BoundObject, SkelMeshComp, Skeleton)),
-					NAME_None,
-					EUserInterfaceActionType::Button);
-
-				if (CVarEnableAdditiveControlRigs->GetBool())
-				{
-					MenuBuilder.AddMenuEntry(
-					   LOCTEXT("AddAdditiveControlRig", "Add Additive Control Rig"),
-					   LOCTEXT("AddAdditiveControlRigTooltip", "Add additive Control Rig and add a track for it"),
-					   FSlateIcon(),
-					   FUIAction(
-						   FExecuteAction::CreateRaw(this, &FControlRigParameterTrackEditor::AddAdditiveControlRig, ObjectBindings[0], BoundObject, SkelMeshComp, Skeleton),
-						   FCanExecuteAction::CreateRaw(this, &FControlRigParameterTrackEditor::CanAddAdditiveConrolRig)),
-					   NAME_None,
-					   EUserInterfaceActionType::Button);
-				}
-
 				MenuBuilder.AddMenuEntry(
 					LOCTEXT("FilterAssetBySkeleton", "Filter Asset By Skeleton"),
 					LOCTEXT("FilterAssetBySkeletonTooltip", "Filters Control Rig assets to match current skeleton"),
@@ -792,11 +772,16 @@ public:
 			const bool bMatchesFlags = !InClass->HasAnyClassFlags(CLASS_Hidden | CLASS_HideDropDown | CLASS_Deprecated | CLASS_Abstract);
 			const bool bNotNative = !InClass->IsNative();
 
+			// ALlow any class contained in the extra picker common classes array
+			if (InInitOptions.ExtraPickerCommonClasses.Contains(InClass))
+			{
+				return true;
+			}
+			
 			if (bChildOfObjectClass && bMatchesFlags && bNotNative)
 			{
 				FAssetData AssetData(InClass);
 				return MatchesFilter(AssetData);
-
 			}
 		}
 		return false;
@@ -819,378 +804,6 @@ public:
 
 };
 
-void FControlRigParameterTrackEditor::ConvertToFKControlRig(FGuid ObjectBinding, UObject* BoundObject, USkeletalMeshComponent* SkelMeshComp, USkeleton* Skeleton)
-{
-	BakeToControlRig(UFKControlRig::StaticClass(), ObjectBinding, BoundObject, SkelMeshComp, Skeleton);
-}
-
-class SAddAdditiveControlRigOptionsWindow : public SCompoundWidget
-{
-public:
-	SLATE_BEGIN_ARGS(SAddAdditiveControlRigOptionsWindow)
-		: _bFilterBySkeleton(true)
-		, _ControlRigClass(nullptr)
-		, _Skeleton(nullptr)
-		, _WidgetWindow()
-	{}
-
-		SLATE_ARGUMENT(bool, bFilterBySkeleton)
-		SLATE_ARGUMENT(TSubclassOf<UControlRig>*, ControlRigClass)
-		SLATE_ARGUMENT(USkeleton*, Skeleton)
-		SLATE_ARGUMENT(TSharedPtr<SWindow>, WidgetWindow)
-		SLATE_END_ARGS()
-
-public:
-	void Construct(const FArguments& InArgs);
-
-	void OnAdd(UClass* InClass)
-	{
-		if (WidgetWindow.IsValid())
-		{
-			WidgetWindow.Pin()->RequestDestroyWindow();
-		}
-		*ControlRigClass = TSubclassOf<UControlRig>(InClass);
-		return;
-	}
-
-	FReply OnCancel()
-	{
-		if (WidgetWindow.IsValid())
-		{
-			WidgetWindow.Pin()->RequestDestroyWindow();
-		}
-		return FReply::Handled();
-	}
-
-	ECheckBoxState IsFilterBySkeletonChecked() const
-	{
-		return bFilterBySkeleton ? ECheckBoxState::Checked : ECheckBoxState::Unchecked;
-	}
-
-	void HandleFilterBySkeletonStateChanged( ECheckBoxState NewState )
-	{
-		bFilterBySkeleton = NewState == ECheckBoxState::Checked;
-		UpdateControlRigClasses();
-	}
-
-	void UpdateControlRigClasses()
-	{
-		FClassViewerInitializationOptions Options;
-		Options.bShowUnloadedBlueprints = true;
-		Options.NameTypeToDisplay = EClassViewerNameTypeToDisplay::DisplayName;
-
-		TSharedPtr<FControlRigClassFilter> ClassFilter = MakeShareable(new FControlRigClassFilter(bFilterBySkeleton, false, true, Skeleton));
-		Options.ClassFilters.Add(ClassFilter.ToSharedRef());
-		Options.bShowNoneOption = false;
-
-		FClassViewerModule& ClassViewerModule = FModuleManager::LoadModuleChecked<FClassViewerModule>("ClassViewer");
-		TSharedRef<SWidget> ClassViewer = ClassViewerModule.CreateClassViewer(Options, FOnClassPicked::CreateRaw(this, &SAddAdditiveControlRigOptionsWindow::OnAdd));
-		InspectorBox->SetContent(ClassViewer->AsShared());
-	}
-
-	SAddAdditiveControlRigOptionsWindow()
-		: bFilterBySkeleton(true)
-		, ControlRigClass(nullptr)
-		, Skeleton(nullptr)
-	{}
-
-private:
-	bool bFilterBySkeleton;
-	TSubclassOf<UControlRig>* ControlRigClass;
-	USkeleton* Skeleton;
-	TSharedPtr<SBox> InspectorBox;
-	TWeakPtr< SWindow > WidgetWindow;
-};
-
-
-void SAddAdditiveControlRigOptionsWindow::Construct(const FArguments& InArgs)
-{
-	bFilterBySkeleton = InArgs._bFilterBySkeleton;
-	ControlRigClass = InArgs._ControlRigClass;
-	Skeleton = InArgs._Skeleton;
-	WidgetWindow = InArgs._WidgetWindow;
-
-	FText CancelText = LOCTEXT("AnimSequenceOptions_Cancel", "Cancel");
-	FText CancelTooltipText = LOCTEXT("AnimSequenceOptions_Cancel_ToolTip", "Cancel control rig creation");
-
-	TSharedPtr<SBox> HeaderToolBox;
-	TSharedPtr<SHorizontalBox> AnimHeaderButtons;
-	this->ChildSlot
-	[
-		SNew(SBox)
-		[
-			SNew(SVerticalBox)
-			+ SVerticalBox::Slot()
-			.AutoHeight()
-			.Padding(2)
-			[
-				SAssignNew(HeaderToolBox, SBox)
-			]
-			+ SVerticalBox::Slot()
-			.AutoHeight()
-			.Padding(2)
-			[
-				SNew(SBorder)
-				.Padding(FMargin(3))
-				.BorderImage(FAppStyle::GetBrush("ToolPanel.GroupBorder"))
-				[
-					SNew(SHorizontalBox)
-					+ SHorizontalBox::Slot()
-					.AutoWidth()
-					[
-						SNew(SCheckBox)
-						.IsChecked(this, &SAddAdditiveControlRigOptionsWindow::IsFilterBySkeletonChecked)
-						.OnCheckStateChanged(this, &SAddAdditiveControlRigOptionsWindow::HandleFilterBySkeletonStateChanged)
-						.Padding(FMargin(6.0, 2.0))
-						[
-							SNew(STextBlock)
-								.Text(LOCTEXT("SkeletonFilterLabel", "Filter by Skeleton"))
-						]
-					]
-					+ SHorizontalBox::Slot()
-					.AutoWidth()
-					[
-						SNew(STextBlock)
-						.Font(FAppStyle::GetFontStyle("CurveEd.LabelFont"))
-						.Text(LOCTEXT("Export_CurrentFileTitle", "Current File: "))
-					]
-				]
-			]
-			+ SVerticalBox::Slot()
-			.FillHeight(1.0f)
-			.Padding(2)
-			[
-				SAssignNew(InspectorBox, SBox)
-			]
-		]
-	];
-
-	//todo move to .h for ue5
-	class FControlRigClassFilter : public IClassViewerFilter
-	{
-	public:
-		FControlRigClassFilter(bool bInCheckSkeleton, bool bInCheckAnimatable, bool bInCheckInversion, USkeleton* InSkeleton) :
-			bFilterAssetBySkeleton(bInCheckSkeleton),
-			bFilterExposesAnimatableControls(bInCheckAnimatable),
-			bFilterInversion(bInCheckInversion),
-			AssetRegistry(FModuleManager::GetModuleChecked<FAssetRegistryModule>(TEXT("AssetRegistry")).Get())
-		{
-			if (InSkeleton)
-			{
-				SkeletonName = FAssetData(InSkeleton).GetExportTextName();
-			}
-		}
-		bool bFilterAssetBySkeleton;
-		bool bFilterExposesAnimatableControls;
-		bool bFilterInversion;
-
-		FString SkeletonName;
-		const IAssetRegistry& AssetRegistry;
-
-		bool MatchesFilter(const FAssetData& AssetData)
-		{
-			bool bExposesAnimatableControls = AssetData.GetTagValueRef<bool>(TEXT("bExposesAnimatableControls"));
-			if (bFilterExposesAnimatableControls == true && bExposesAnimatableControls == false)
-			{
-				return false;
-			}
-			if (bFilterInversion)
-			{
-				bool bHasInversion = false;
-				FAssetDataTagMapSharedView::FFindTagResult Tag = AssetData.TagsAndValues.FindTag(TEXT("SupportedEventNames"));
-				if (Tag.IsSet())
-				{
-					FString EventString = FRigUnit_InverseExecution::EventName.ToString();
-					FString OldEventString = FString(TEXT("Inverse"));
-					TArray<FString> SupportedEventNames;
-					Tag.GetValue().ParseIntoArray(SupportedEventNames, TEXT(","), true);
-
-					for (const FString& Name : SupportedEventNames)
-					{
-						if (Name.Contains(EventString) || Name.Contains(OldEventString))
-						{
-							bHasInversion = true;
-							break;
-						}
-					}
-					if (bHasInversion == false)
-					{
-						return false;
-					}
-				}
-			}
-			if (bFilterAssetBySkeleton)
-			{
-				FString PreviewSkeletalMesh = AssetData.GetTagValueRef<FString>(TEXT("PreviewSkeletalMesh"));
-				if (PreviewSkeletalMesh.Len() > 0)
-				{
-					FAssetData SkelMeshData = AssetRegistry.GetAssetByObjectPath(FSoftObjectPath(PreviewSkeletalMesh));
-					FString PreviewSkeleton = SkelMeshData.GetTagValueRef<FString>(TEXT("Skeleton"));
-					if (PreviewSkeleton == SkeletonName)
-					{
-						return true;
-					}
-				}
-				FString PreviewSkeleton = AssetData.GetTagValueRef<FString>(TEXT("PreviewSkeleton"));
-				if (PreviewSkeleton == SkeletonName)
-				{
-					return true;
-				}
-				FString SourceHierarchyImport = AssetData.GetTagValueRef<FString>(TEXT("SourceHierarchyImport"));
-				if (SourceHierarchyImport == SkeletonName)
-				{
-					return true;
-				}
-				FString SourceCurveImport = AssetData.GetTagValueRef<FString>(TEXT("SourceCurveImport"));
-				if (SourceCurveImport == SkeletonName)
-				{
-					return true;
-				}
-				return false;
-			}
-			return true;
-
-		}
-		bool IsClassAllowed(const FClassViewerInitializationOptions& InInitOptions, const UClass* InClass, TSharedRef< FClassViewerFilterFuncs > InFilterFuncs) override
-		{
-			if (InClass)
-			{
-				const bool bChildOfObjectClass = InClass->IsChildOf(UControlRig::StaticClass());
-				const bool bMatchesFlags = !InClass->HasAnyClassFlags(CLASS_Hidden | CLASS_HideDropDown | CLASS_Deprecated | CLASS_Abstract);
-				const bool bNotNative = !InClass->IsNative();
-
-				if (bChildOfObjectClass && bMatchesFlags && bNotNative)
-				{
-					FAssetData AssetData(InClass);
-					return MatchesFilter(AssetData);
-				}
-			}
-			return false;
-		}
-
-		virtual bool IsUnloadedClassAllowed(const FClassViewerInitializationOptions& InInitOptions, const TSharedRef< const IUnloadedBlueprintData > InUnloadedClassData, TSharedRef< FClassViewerFilterFuncs > InFilterFuncs) override
-		{
-			const bool bChildOfObjectClass = InUnloadedClassData->IsChildOf(UControlRig::StaticClass());
-			const bool bMatchesFlags = !InUnloadedClassData->HasAnyClassFlags(CLASS_Hidden | CLASS_HideDropDown | CLASS_Deprecated | CLASS_Abstract);
-			if (bChildOfObjectClass && bMatchesFlags)
-			{
-				FString GeneratedClassPathString = InUnloadedClassData->GetClassPathName().ToString();
-				FString BlueprintPath = GeneratedClassPathString.LeftChop(2); // Chop off _C
-				FAssetData AssetData = AssetRegistry.GetAssetByObjectPath(FSoftObjectPath(BlueprintPath));
-				return MatchesFilter(AssetData);
-
-			}
-			return false;
-		}
-
-	};
-
-	UpdateControlRigClasses();
-}
-
-void FControlRigParameterTrackEditor::AddAdditiveControlRig(FGuid ObjectBinding, UObject* BoundActor, USkeletalMeshComponent* SkelMeshComp, USkeleton* Skeleton)
-{
-	FSlateApplication::Get().DismissAllMenus();
-
-	TSharedPtr<SWindow> ParentWindow;
-	if (FModuleManager::Get().IsModuleLoaded("MainFrame"))
-	{
-		IMainFrameModule& MainFrame = FModuleManager::LoadModuleChecked<IMainFrameModule>("MainFrame");
-		ParentWindow = MainFrame.GetParentWindow();
-	}
-
-	TSharedRef<SWindow> Window = SNew(SWindow)
-				.Title(LOCTEXT("AdditiveRigTitle", "Choose Additive Rig"))
-				.SizingRule(ESizingRule::UserSized)
-				.AutoCenter(EAutoCenter::PrimaryWorkArea)
-				.ClientSize(FVector2D(500, 445));
-
-	TSharedPtr<SAddAdditiveControlRigOptionsWindow> OptionWindow;
-	TSubclassOf<UControlRig> ControlRigClass;
-	Window->SetContent
-	(
-		SAssignNew(OptionWindow, SAddAdditiveControlRigOptionsWindow)
-		.ControlRigClass(&ControlRigClass)
-		.Skeleton(Skeleton)
-		.WidgetWindow(Window)
-	);
-
-	FSlateApplication::Get().AddModalWindow(Window, ParentWindow, false);
-	
-	
-	const TSharedPtr<ISequencer> ParentSequencer = GetSequencer();
-	if (ControlRigClass && ParentSequencer.IsValid())
-	{
-		UMovieSceneSequence* OwnerSequence = ParentSequencer->GetFocusedMovieSceneSequence();
-		UMovieScene* OwnerMovieScene = OwnerSequence->GetMovieScene();
-		{
-			OwnerMovieScene->Modify();
-			UMovieSceneControlRigParameterTrack* Track = Cast<UMovieSceneControlRigParameterTrack>(AddTrack(OwnerMovieScene, ObjectBinding, UMovieSceneControlRigParameterTrack::StaticClass(), NAME_None));
-			if (Track)
-			{
-				Track->Modify();
-
-				FString ObjectName = ControlRigClass->GetName();
-				ObjectName.RemoveFromEnd(TEXT("_C"));
-				UControlRig* ControlRig = NewObject<UControlRig>(Track, ControlRigClass, FName(*ObjectName), RF_Transactional);
-				ControlRig->SetIsAdditive(true);
-				if (!ControlRig->SupportsEvent(FRigUnit_InverseExecution::EventName))
-				{
-					OwnerMovieScene->RemoveTrack(*Track);
-					return;
-				}
-
-				FControlRigEditMode* ControlRigEditMode = GetEditMode();
-				if (!ControlRigEditMode)
-				{
-					ControlRigEditMode = GetEditMode(true);
-				}
-				
-				ControlRig->Modify();
-				ControlRig->SetObjectBinding(MakeShared<FControlRigObjectBinding>());
-				ControlRig->GetObjectBinding()->BindToObject(BoundActor);
-				ControlRig->GetDataSourceRegistry()->RegisterDataSource(UControlRig::OwnerComponent, ControlRig->GetObjectBinding()->GetBoundObject());
-				ControlRig->SetEventQueue({FRigUnit_InverseExecution::EventName});
-				ControlRig->Initialize();
-				ControlRig->RequestInit();
-				ControlRig->SetBoneInitialTransformsFromSkeletalMeshComponent(SkelMeshComp, true);
-				ControlRig->Evaluate_AnyThread();
-
-				const bool bSequencerOwnsControlRig = true;
-				UMovieSceneSection* NewSection = Track->CreateControlRigSection(0, ControlRig, bSequencerOwnsControlRig);
-				
-				const FString AdditiveObjectName = ObjectName + " (Additive)";
-				Track->SetTrackName(FName(*ObjectName));
-				Track->SetDisplayName(FText::FromString(AdditiveObjectName));
-				Track->SetColorTint(FColor(173, 151, 114));
-				
-
-				ParentSequencer->EmptySelection();
-				ParentSequencer->SelectSection(NewSection);
-				ParentSequencer->ThrobSectionSelection();
-				
-				//Finish Setup
-				if (ControlRigEditMode)
-				{
-					ControlRigEditMode->AddControlRigObject(ControlRig, ParentSequencer);
-				}
-				BindControlRig(ControlRig);
-
-				ParentSequencer->NotifyMovieSceneDataChanged(EMovieSceneDataChangeType::MovieSceneStructureItemAdded);
-			}
-		}
-	}
-}
-
-bool FControlRigParameterTrackEditor::CanAddAdditiveConrolRig()
-{
-	// if (FControlRigEditMode* ControlRigEditMode = GetEditMode())
-	// {
-	// 	return ControlRigEditMode->GetControlRigsArray(false).IsEmpty();
-	// }
-	return true;
-}
-
 void FControlRigParameterTrackEditor::BakeToControlRigSubMenu(FMenuBuilder& MenuBuilder, FGuid ObjectBinding, UObject* BoundObject, USkeletalMeshComponent* SkelMeshComp, USkeleton* Skeleton)
 {
 	const TSharedPtr<ISequencer> ParentSequencer = GetSequencer();
@@ -1203,6 +816,7 @@ void FControlRigParameterTrackEditor::BakeToControlRigSubMenu(FMenuBuilder& Menu
 		TSharedPtr<FControlRigClassFilter> ClassFilter = MakeShareable(new FControlRigClassFilter(bFilterAssetBySkeleton, false, true, Skeleton));
 		Options.ClassFilters.Add(ClassFilter.ToSharedRef());
 		Options.bShowNoneOption = false;
+		Options.ExtraPickerCommonClasses.Add(UFKControlRig::StaticClass());
 
 		FClassViewerModule& ClassViewerModule = FModuleManager::LoadModuleChecked<FClassViewerModule>("ClassViewer");
 
@@ -1677,15 +1291,20 @@ void FControlRigParameterTrackEditor::BuildObjectBindingTrackMenu(FMenuBuilder&
 
 void FControlRigParameterTrackEditor::HandleAddTrackSubMenu(FMenuBuilder& MenuBuilder, TArray<FGuid> ObjectBindings, UMovieSceneTrack* Track)
 {
-	MenuBuilder.AddMenuEntry(
-		LOCTEXT("AddFKControlRig", "FK Control Rig"),
-		LOCTEXT("AddFKControlRigTooltip", "Adds an FK Control Rig track"),
+	if (CVarEnableAdditiveControlRigs->GetBool())
+	{
+		MenuBuilder.AddMenuEntry(
+		LOCTEXT("IsAdditiveControlRig", "Additive"),
+		LOCTEXT("IsAdditiveControlRigTooltip", "Add an additive control rig"),
 		FSlateIcon(),
 		FUIAction(
-			FExecuteAction::CreateSP(this, &FControlRigParameterTrackEditor::AddFKControlRig, ObjectBindings),
-			FCanExecuteAction()
-		)
-	);
+			FExecuteAction::CreateSP(this, &FControlRigParameterTrackEditor::ToggleIsAdditiveControlRig),
+			FCanExecuteAction(),
+			FIsActionChecked::CreateSP(this, &FControlRigParameterTrackEditor::IsToggleIsAdditiveControlRig)
+		),
+		NAME_None,
+		EUserInterfaceActionType::ToggleButton);
+	}
 
 	MenuBuilder.AddMenuEntry(
 		LOCTEXT("FilterAssetBySkeleton", "Filter Asset By Skeleton"),
@@ -1718,6 +1337,16 @@ void FControlRigParameterTrackEditor::HandleAddTrackSubMenu(FMenuBuilder& MenuBu
 	);
 }
 
+void FControlRigParameterTrackEditor::ToggleIsAdditiveControlRig()
+{
+	bIsAdditiveControlRig = bIsAdditiveControlRig ? false : true;
+}
+
+bool FControlRigParameterTrackEditor::IsToggleIsAdditiveControlRig()
+{
+	return bIsAdditiveControlRig;
+}
+
 void FControlRigParameterTrackEditor::ToggleFilterAssetBySkeleton()
 {
 	bFilterAssetBySkeleton = bFilterAssetBySkeleton ? false : true;
@@ -1785,11 +1414,17 @@ void FControlRigParameterTrackEditor::HandleAddControlRigSubMenu(FMenuBuilder& M
 		FClassViewerInitializationOptions Options;
 		Options.bShowUnloadedBlueprints = true;
 		Options.NameTypeToDisplay = EClassViewerNameTypeToDisplay::DisplayName;
-		
-		TSharedPtr<FControlRigClassFilter> ClassFilter = MakeShareable(new FControlRigClassFilter(bFilterAssetBySkeleton, bFilterAssetByAnimatableControls, false, Skeleton));
+
+		const bool bCheckInversion = bIsAdditiveControlRig;
+		TSharedPtr<FControlRigClassFilter> ClassFilter = MakeShareable(new FControlRigClassFilter(bFilterAssetBySkeleton, bFilterAssetByAnimatableControls, bCheckInversion, Skeleton));
 		Options.ClassFilters.Add(ClassFilter.ToSharedRef());
 		Options.bShowNoneOption = false;
 
+		if (!bIsAdditiveControlRig)
+		{
+			Options.ExtraPickerCommonClasses.Add(UFKControlRig::StaticClass());
+		}
+
 		UMovieSceneSequence* Sequence = GetSequencer() ? GetSequencer()->GetFocusedMovieSceneSequence() : nullptr;
 		Options.AdditionalReferencingAssets.Add(FAssetData(Sequence));
 
@@ -1870,6 +1505,12 @@ void FControlRigParameterTrackEditor::AddControlRig(UClass* InClass, UObject* Bo
 
 	if (InClass && InClass->IsChildOf(UControlRig::StaticClass()) && SequencerParent.IsValid())
 	{
+		if (bIsAdditiveControlRig && !InClass->GetDefaultObject<UControlRig>()->SupportsEvent(FRigUnit_InverseExecution::EventName))
+		{
+			UE_LOG(LogControlRigEditor, Error, TEXT("Cannot add an additive control rig which does not contain a backwards solve event."));
+			return;
+		}
+		
 		UMovieSceneSequence* OwnerSequence = GetSequencer()->GetFocusedMovieSceneSequence();
 		UMovieScene* OwnerMovieScene = OwnerSequence->GetMovieScene();
 		FScopedTransaction AddControlRigTrackTransaction(LOCTEXT("AddControlRigTrack", "Add Control Rig Track"));
@@ -1891,6 +1532,7 @@ void FControlRigParameterTrackEditor::AddControlRig(UClass* InClass, UObject* Bo
 			}
 
 			ControlRig->Modify();
+			ControlRig->SetIsAdditive(bIsAdditiveControlRig);
 			ControlRig->SetObjectBinding(MakeShared<FControlRigObjectBinding>());
 			ControlRig->GetObjectBinding()->BindToObject(BoundActor);
 			ControlRig->GetDataSourceRegistry()->RegisterDataSource(UControlRig::OwnerComponent, ControlRig->GetObjectBinding()->GetBoundObject());
@@ -1907,9 +1549,19 @@ void FControlRigParameterTrackEditor::AddControlRig(UClass* InClass, UObject* Bo
 			UMovieSceneSection* NewSection = Track->CreateControlRigSection(0, ControlRig, bSequencerOwnsControlRig);
 			NewSection->Modify();
 
-			//mz todo need to have multiple rigs with same class
-			Track->SetTrackName(FName(*ObjectName));
-			Track->SetDisplayName(FText::FromString(ObjectName));
+			if (bIsAdditiveControlRig)
+			{
+				const FString AdditiveObjectName = ObjectName + " (Additive)";
+				Track->SetTrackName(FName(*ObjectName));
+				Track->SetDisplayName(FText::FromString(AdditiveObjectName));
+				Track->SetColorTint(FColor(173, 151, 114));
+			}
+			else
+			{
+				//mz todo need to have multiple rigs with same class
+				Track->SetTrackName(FName(*ObjectName));
+				Track->SetDisplayName(FText::FromString(ObjectName));
+			}
 
 			GetSequencer()->EmptySelection();
 			GetSequencer()->SelectSection(NewSection);
@@ -1930,6 +1582,10 @@ void FControlRigParameterTrackEditor::AddControlRig(UClass* InClass, UObject* Bo
 
 void FControlRigParameterTrackEditor::AddControlRig(UClass* InClass, UObject* BoundActor, FGuid ObjectBinding)
 {
+	if (InClass == UFKControlRig::StaticClass())
+	{
+		AcquireSkeletonFromObjectGuid(ObjectBinding, &BoundActor, GetSequencer());
+	}
 	AddControlRig(InClass, BoundActor, ObjectBinding, nullptr);
 }
 
@@ -1954,19 +1610,6 @@ void FControlRigParameterTrackEditor::AddControlRigFromComponent(FGuid InGuid)
 	}
 }
 
-void FControlRigParameterTrackEditor::AddFKControlRig(TArray<FGuid> ObjectBindings)
-{
-	for (const FGuid& ObjectBinding : ObjectBindings)
-	{
-		UObject* BoundObject = nullptr;
-		AcquireSkeletonFromObjectGuid(ObjectBinding, &BoundObject, GetSequencer());
-		if (BoundObject)
-		{
-			AddControlRig(UFKControlRig::StaticClass(), BoundObject, ObjectBinding);
-		}
-	}
-}
-
 bool FControlRigParameterTrackEditor::HasTransformKeyOverridePriority() const
 {
 	return CanAddTransformKeysForSelectedObjects();
diff --git a/Engine/Plugins/Animation/ControlRig/Source/ControlRigEditor/Private/Sequencer/ControlRigParameterTrackEditor.h b/Engine/Plugins/Animation/ControlRig/Source/ControlRigEditor/Private/Sequencer/ControlRigParameterTrackEditor.h
--- a/Engine/Plugins/Animation/ControlRig/Source/ControlRigEditor/Private/Sequencer/ControlRigParameterTrackEditor.h
+++ b/Engine/Plugins/Animation/ControlRig/Source/ControlRigEditor/Private/Sequencer/ControlRigParameterTrackEditor.h
@@ -89,6 +89,9 @@ private:
 	void HandleAddTrackSubMenu(FMenuBuilder& MenuBuilder, TArray<FGuid> ObjectBindings, UMovieSceneTrack* Track);
 	void HandleAddControlRigSubMenu(FMenuBuilder& MenuBuilder, TArray<FGuid> ObjectBindings, UMovieSceneTrack* Track);
 
+	void ToggleIsAdditiveControlRig();
+	bool IsToggleIsAdditiveControlRig();
+	
 	void ToggleFilterAssetBySkeleton();
 	bool IsToggleFilterAssetBySkeleton();
 
@@ -100,7 +103,6 @@ private:
 	void AddControlRig(UClass* InClass, UObject* BoundActor, FGuid ObjectBinding);
 	void AddControlRig(UClass* InClass, UObject* BoundActor, FGuid ObjectBinding, UControlRig* InExistingControlRig);
 	void AddControlRigFromComponent(FGuid InGuid);
-	void AddFKControlRig(TArray<FGuid> ObjectBindings);
 	
 	/** Delegate for Selection Changed Event */
 	void OnSelectionChanged(TArray<UMovieSceneTrack*> InTracks);
@@ -188,15 +190,6 @@ private:
 	/** Toggle FK Control Rig*/
 	void ToggleFKControlRig(UMovieSceneControlRigParameterTrack* Track, UFKControlRig* FKControlRig);
 
-	/** Convert to FK Control Rig*/
-	void ConvertToFKControlRig(FGuid ObjectBinding, UObject* BoundObject, USkeletalMeshComponent* SkelMeshComp, USkeleton* Skeleton);
-
-	/** Add additive Control Rig*/
-	void AddAdditiveControlRig(FGuid ObjectBinding, UObject* BoundObject, USkeletalMeshComponent* SkelMeshComp, USkeleton* Skeleton);
-
-	/** Returns true if an additive control rig can be added to the sequence */
-	bool CanAddAdditiveConrolRig();
-
 	/** Bake To Control Rig Sub Menu*/
 	void BakeToControlRigSubMenu(FMenuBuilder& MenuBuilder, FGuid ObjectBinding, UObject* BoundObject, USkeletalMeshComponent* SkelMeshComp,USkeleton* Skeleton);
 	
@@ -269,6 +262,8 @@ private:
 	/** A flag to determine if the next update coming from the timer should be skipped */
 	bool bSkipNextSelectionFromTimer;
 
+	bool bIsAdditiveControlRig;
+
 	/** Whether or not we should check Skeleton when filtering*/
 	bool bFilterAssetBySkeleton;
 
