Control Rig: Avoid double registration of onControlModified

diff --git a/Engine/Plugins/Animation/ControlRig/Source/ControlRigEditor/Private/Editor/ControlRigEditor.cpp b/Engine/Plugins/Animation/ControlRig/Source/ControlRigEditor/Private/Editor/ControlRigEditor.cpp
--- a/Engine/Plugins/Animation/ControlRig/Source/ControlRigEditor/Private/Editor/ControlRigEditor.cpp
+++ b/Engine/Plugins/Animation/ControlRig/Source/ControlRigEditor/Private/Editor/ControlRigEditor.cpp
@@ -856,6 +856,12 @@ void FControlRigEditor::HandleSetObjectBeingDebugged(UObject* InObject)
 			}
 		}
 
+		DebuggedControlRig->GetHierarchy()->OnModified().RemoveAll(this);
+		DebuggedControlRig->OnPreForwardsSolve_AnyThread().RemoveAll(this);
+		DebuggedControlRig->OnPreConstructionForUI_AnyThread().RemoveAll(this);
+		DebuggedControlRig->OnPostConstruction_AnyThread().RemoveAll(this);
+		DebuggedControlRig->ControlModified().RemoveAll(this);
+
 		DebuggedControlRig->GetHierarchy()->OnModified().AddSP(this, &FControlRigEditor::OnHierarchyModified_AnyThread);
 		DebuggedControlRig->OnPreForwardsSolve_AnyThread().AddSP(this, &FControlRigEditor::OnPreForwardsSolve_AnyThread);
 		DebuggedControlRig->OnPreConstructionForUI_AnyThread().AddSP(this, &FControlRigEditor::OnPreConstruction_AnyThread);
@@ -2027,6 +2033,9 @@ void FControlRigEditor::UpdateRigVMHost()
 				EditMode->SetObjects(ControlRig, EditorSkelComp,nullptr);
 			}
 
+			ControlRig->OnPreForwardsSolve_AnyThread().RemoveAll(this);
+			ControlRig->ControlModified().RemoveAll(this);
+
 			ControlRig->OnPreForwardsSolve_AnyThread().AddSP(this, &FControlRigEditor::OnPreForwardsSolve_AnyThread);
 			ControlRig->ControlModified().AddSP(this, &FControlRigEditor::HandleOnControlModified);
 		}
