Control Rig: Added missing comment for additive control rigs
#rnx

diff --git a/Engine/Plugins/Animation/ControlRig/Source/ControlRigEditor/Private/Sequencer/ControlRigParameterTrackEditor.h b/Engine/Plugins/Animation/ControlRig/Source/ControlRigEditor/Private/Sequencer/ControlRigParameterTrackEditor.h
--- a/Engine/Plugins/Animation/ControlRig/Source/ControlRigEditor/Private/Sequencer/ControlRigParameterTrackEditor.h
+++ b/Engine/Plugins/Animation/ControlRig/Source/ControlRigEditor/Private/Sequencer/ControlRigParameterTrackEditor.h
@@ -262,6 +262,7 @@ private:
 	/** A flag to determine if the next update coming from the timer should be skipped */
 	bool bSkipNextSelectionFromTimer;
 
+	/** Whether or not this rig will be used as an additive control rig */
 	bool bIsAdditiveControlRig;
 
 	/** Whether or not we should check Skeleton when filtering*/
