ControlRig: Only draw nulls / axes in the preview world / asset editor

diff --git a/Engine/Plugins/Animation/ControlRig/Source/ControlRigEditor/Private/EditMode/ControlRigEditMode.cpp b/Engine/Plugins/Animation/ControlRig/Source/ControlRigEditor/Private/EditMode/ControlRigEditMode.cpp
--- a/Engine/Plugins/Animation/ControlRig/Source/ControlRigEditor/Private/EditMode/ControlRigEditMode.cpp
+++ b/Engine/Plugins/Animation/ControlRig/Source/ControlRigEditor/Private/EditMode/ControlRigEditMode.cpp
@@ -858,7 +858,8 @@ void FControlRigEditMode::Render(const FSceneView* View, FViewport* Viewport, FP
 						});
 				}
 
-				if (Settings->bDisplayNulls || ControlRig->IsConstructionModeEnabled())
+				const bool bIsAssetEditor = Viewport->GetClient()->GetWorld()->IsPreviewWorld();
+				if (bIsAssetEditor && (Settings->bDisplayNulls || ControlRig->IsConstructionModeEnabled()))
 				{
 					TArray<FTransform> SpaceTransforms;
 					TArray<FTransform> SelectedSpaceTransforms;
@@ -879,7 +880,7 @@ void FControlRigEditMode::Render(const FSceneView* View, FViewport* Viewport, FP
 					ControlRig->DrawInterface.DrawAxes(FTransform::Identity, SelectedSpaceTransforms, FLinearColor(1.0f, 0.34f, 0.0f, 1.0f), Settings->AxisScale);
 				}
 
-				if (Settings->bDisplayAxesOnSelection && Settings->AxisScale > SMALL_NUMBER)
+				if (bIsAssetEditor && (Settings->bDisplayAxesOnSelection && Settings->AxisScale > SMALL_NUMBER))
 				{
 					if (ControlRig->GetWorld() && ControlRig->GetWorld()->IsPreviewWorld())
 					{
