Control Rig: Adds ScaleFloat as a control type

diff --git a/Engine/Plugins/Animation/ControlRig/Source/ControlRig/Private/AnimNode_ControlRig.cpp b/Engine/Plugins/Animation/ControlRig/Source/ControlRig/Private/AnimNode_ControlRig.cpp
--- a/Engine/Plugins/Animation/ControlRig/Source/ControlRig/Private/AnimNode_ControlRig.cpp
+++ b/Engine/Plugins/Animation/ControlRig/Source/ControlRig/Private/AnimNode_ControlRig.cpp
@@ -407,6 +407,7 @@ void FAnimNode_ControlRig::PropagateInputProperties(const UObject* InSourceInsta
 						break;
 					}
 					case ERigControlType::Float:
+					case ERigControlType::ScaleFloat:
 					{
 						if(ensure(CastField<FFloatProperty>(CallerProperty)))
 						{
diff --git a/Engine/Plugins/Animation/ControlRig/Source/ControlRig/Private/ControlRig.cpp b/Engine/Plugins/Animation/ControlRig/Source/ControlRig/Private/ControlRig.cpp
--- a/Engine/Plugins/Animation/ControlRig/Source/ControlRig/Private/ControlRig.cpp
+++ b/Engine/Plugins/Animation/ControlRig/Source/ControlRig/Private/ControlRig.cpp
@@ -988,6 +988,7 @@ bool UControlRig::Execute(const FName& InEventName)
 				switch (Settings.ControlType)
 				{
 					case ERigControlType::Float:
+					case ERigControlType::ScaleFloat:
 					{
 						if(Settings.LimitEnabled[0].IsOff())
 						{
@@ -2930,6 +2931,7 @@ UTransformableControlHandle* UControlRig::CreateTransformableControlHandle(
 		const FRigControlSettings& ControlSettings = ControlElement->Settings;
 		if (ControlSettings.ControlType == ERigControlType::Bool ||
 			ControlSettings.ControlType == ERigControlType::Float ||
+			ControlSettings.ControlType == ERigControlType::ScaleFloat ||
 			ControlSettings.ControlType == ERigControlType::Integer)
 		{
 			return false;
diff --git a/Engine/Plugins/Animation/ControlRig/Source/ControlRig/Private/ControlRigComponent.cpp b/Engine/Plugins/Animation/ControlRig/Source/ControlRig/Private/ControlRigComponent.cpp
--- a/Engine/Plugins/Animation/ControlRig/Source/ControlRig/Private/ControlRigComponent.cpp
+++ b/Engine/Plugins/Animation/ControlRig/Source/ControlRig/Private/ControlRigComponent.cpp
@@ -902,7 +902,7 @@ float UControlRigComponent::GetControlFloat(FName ControlName)
 	{
 		if(FRigControlElement* ControlElement = CR->GetHierarchy()->Find<FRigControlElement>(FRigElementKey(ControlName, ERigElementType::Control)))
 		{
-			if (ControlElement->Settings.ControlType == ERigControlType::Float)
+			if (ControlElement->Settings.ControlType == ERigControlType::Float || ControlElement->Settings.ControlType == ERigControlType::ScaleFloat)
 			{
 				return CR->GetHierarchy()->GetControlValue(ControlElement, ERigControlValueType::Current).Get<float>();
 			}
diff --git a/Engine/Plugins/Animation/ControlRig/Source/ControlRig/Private/ControlRigControlActor.cpp b/Engine/Plugins/Animation/ControlRig/Source/ControlRig/Private/ControlRigControlActor.cpp
--- a/Engine/Plugins/Animation/ControlRig/Source/ControlRig/Private/ControlRigControlActor.cpp
+++ b/Engine/Plugins/Animation/ControlRig/Source/ControlRig/Private/ControlRigControlActor.cpp
@@ -197,6 +197,7 @@ void AControlRigControlActor::Refresh()
 			switch (ControlElement->Settings.ControlType)
 			{
 				case ERigControlType::Float:
+				case ERigControlType::ScaleFloat:
 				case ERigControlType::Integer:
 				case ERigControlType::Vector2D:
 				case ERigControlType::Position:
diff --git a/Engine/Plugins/Animation/ControlRig/Source/ControlRig/Private/Rigs/RigHierarchy.cpp b/Engine/Plugins/Animation/ControlRig/Source/ControlRig/Private/Rigs/RigHierarchy.cpp
--- a/Engine/Plugins/Animation/ControlRig/Source/ControlRig/Private/Rigs/RigHierarchy.cpp
+++ b/Engine/Plugins/Animation/ControlRig/Source/ControlRig/Private/Rigs/RigHierarchy.cpp
@@ -1322,6 +1322,7 @@ FEdGraphPinType URigHierarchy::GetControlPinType(FRigControlElement* InControlEl
 			break;
 		}
 		case ERigControlType::Float:
+		case ERigControlType::ScaleFloat:
 		{
 			PinType.PinCategory = PC_Real;
 			PinType.PinSubCategory = PC_Float;
@@ -1377,6 +1378,7 @@ FString URigHierarchy::GetControlPinDefaultValue(FRigControlElement* InControlEl
 			return Value.ToString<bool>();
 		}
 		case ERigControlType::Float:
+		case ERigControlType::ScaleFloat:
 		{
 			return Value.ToString<float>();
 		}
diff --git a/Engine/Plugins/Animation/ControlRig/Source/ControlRig/Private/Rigs/RigHierarchyController.cpp b/Engine/Plugins/Animation/ControlRig/Source/ControlRig/Private/Rigs/RigHierarchyController.cpp
--- a/Engine/Plugins/Animation/ControlRig/Source/ControlRig/Private/Rigs/RigHierarchyController.cpp
+++ b/Engine/Plugins/Animation/ControlRig/Source/ControlRig/Private/Rigs/RigHierarchyController.cpp
@@ -1393,6 +1393,7 @@ TArray<FRigElementKey> URigHierarchyController::ImportFromHierarchyContainer(con
 		Settings.bGroupWithParentControl = Settings.IsAnimatable() && (
 			Settings.ControlType == ERigControlType::Bool ||
 			Settings.ControlType == ERigControlType::Float ||
+			Settings.ControlType == ERigControlType::ScaleFloat ||
 			Settings.ControlType == ERigControlType::Integer ||
 			Settings.ControlType == ERigControlType::Vector2D
 		);
diff --git a/Engine/Plugins/Animation/ControlRig/Source/ControlRig/Private/Rigs/RigHierarchyElements.cpp b/Engine/Plugins/Animation/ControlRig/Source/ControlRig/Private/Rigs/RigHierarchyElements.cpp
--- a/Engine/Plugins/Animation/ControlRig/Source/ControlRig/Private/Rigs/RigHierarchyElements.cpp
+++ b/Engine/Plugins/Animation/ControlRig/Source/ControlRig/Private/Rigs/RigHierarchyElements.cpp
@@ -1002,6 +1002,7 @@ void FRigControlSettings::Load(FArchive& Ar)
 		bGroupWithParentControl = IsAnimatable() && (
 			ControlType == ERigControlType::Bool ||
 			ControlType == ERigControlType::Float ||
+			ControlType == ERigControlType::ScaleFloat ||
 			ControlType == ERigControlType::Integer ||
 			ControlType == ERigControlType::Vector2D
 		);
@@ -1192,6 +1193,12 @@ void FRigControlSettings::SetupLimitArrayForType(bool bLimitTranslation, bool bL
 			LimitEnabled[0].Set(bLimitTranslation);
 			break;
 		}
+		case ERigControlType::ScaleFloat:
+		{
+			LimitEnabled.SetNum(1);
+			LimitEnabled[0].Set(bLimitScale);
+			break;
+		}
 		case ERigControlType::Vector2D:
 		{
 			LimitEnabled.SetNum(2);
diff --git a/Engine/Plugins/Animation/ControlRig/Source/ControlRig/Private/Sequencer/MovieSceneControlRigParameterSection.cpp b/Engine/Plugins/Animation/ControlRig/Source/ControlRig/Private/Sequencer/MovieSceneControlRigParameterSection.cpp
--- a/Engine/Plugins/Animation/ControlRig/Source/ControlRig/Private/Sequencer/MovieSceneControlRigParameterSection.cpp
+++ b/Engine/Plugins/Animation/ControlRig/Source/ControlRig/Private/Sequencer/MovieSceneControlRigParameterSection.cpp
@@ -1754,6 +1754,7 @@ void UMovieSceneControlRigParameterSection::ReconstructChannelProxy()
 			switch (ControlElement->Settings.ControlType)
 			{
 				case ERigControlType::Float:
+				case ERigControlType::ScaleFloat:
 				{
 					for (FScalarParameterNameAndCurve& Scalar : GetScalarParameterNamesAndCurves())
 					{
@@ -2507,6 +2508,7 @@ void UMovieSceneControlRigParameterSection::RecreateWithThisControlRig(UControlR
 		switch (ControlElement->Settings.ControlType)
 		{
 		case ERigControlType::Float:
+		case ERigControlType::ScaleFloat:
 		{
 			TOptional<float> DefaultValue;
 			if (bSetDefault)
@@ -2971,6 +2973,7 @@ void UMovieSceneControlRigParameterSection::RecordControlRigKey(FFrameNumber Fra
 					break;
 				}
 				case ERigControlType::Float:
+				case ERigControlType::ScaleFloat:
 				{
 					float Val = ControlRig->GetControlValue(ControlElement, ERigControlValueType::Current).Get<float>();
 					if (bSetDefault)
diff --git a/Engine/Plugins/Animation/ControlRig/Source/ControlRig/Private/Sequencer/MovieSceneControlRigParameterTemplate.cpp b/Engine/Plugins/Animation/ControlRig/Source/ControlRig/Private/Sequencer/MovieSceneControlRigParameterTemplate.cpp
--- a/Engine/Plugins/Animation/ControlRig/Source/ControlRig/Private/Sequencer/MovieSceneControlRigParameterTemplate.cpp
+++ b/Engine/Plugins/Animation/ControlRig/Source/ControlRig/Private/Sequencer/MovieSceneControlRigParameterTemplate.cpp
@@ -915,6 +915,7 @@ struct FControlRigParameterPreAnimatedTokenProducer : IMovieScenePreAnimatedToke
 						break;
 					}
 					case ERigControlType::Float:
+					case ERigControlType::ScaleFloat:
 					{
 						const float Val = ControlRig->GetControlValue(ControlElement, ERigControlValueType::Current).Get<float>();
 						Token.ScalarValues.Add(TNameAndValue<float>{ ControlElement->GetFName(), Val });
@@ -1382,7 +1383,7 @@ struct TControlRigParameterActuatorFloat : TMovieSceneBlendingActuator<FControlR
 			FRigControlElement* ControlElement = ControlRig->FindControl(ParameterName);
 			if (ControlElement && ControlElement->Settings.AnimationType != ERigControlAnimationType::ProxyControl &&
 				ControlElement->Settings.AnimationType != ERigControlAnimationType::VisualCue
-				&& ControlElement->Settings.ControlType == ERigControlType::Float)
+				&& (ControlElement->Settings.ControlType == ERigControlType::Float || ControlElement->Settings.ControlType == ERigControlType::ScaleFloat))
 			{
 				const float Val = ControlRig->GetControlValue(ControlElement, ERigControlValueType::Current).Get<float>();
 				return FControlRigTrackTokenFloat(Val);
@@ -1409,7 +1410,8 @@ struct TControlRigParameterActuatorFloat : TMovieSceneBlendingActuator<FControlR
 			{
 				FRigControlElement* ControlElement = ControlRig->FindControl(ParameterName);
 				if (ControlElement && ControlElement->Settings.AnimationType != ERigControlAnimationType::ProxyControl &&
-					ControlElement->Settings.AnimationType != ERigControlAnimationType::VisualCue && ControlElement->Settings.ControlType == ERigControlType::Float)
+					ControlElement->Settings.AnimationType != ERigControlAnimationType::VisualCue &&
+					(ControlElement->Settings.ControlType == ERigControlType::Float || ControlElement->Settings.ControlType == ERigControlType::ScaleFloat))
 				{
 					ControlRig->SetControlValue<float>(ParameterName, InFinalValue.Value, true, EControlRigSetKey::Never,bSetupUndo);
 				}
diff --git a/Engine/Plugins/Animation/ControlRig/Source/ControlRig/Private/Tools/ControlRigPose.cpp b/Engine/Plugins/Animation/ControlRig/Source/ControlRig/Private/Tools/ControlRigPose.cpp
--- a/Engine/Plugins/Animation/ControlRig/Source/ControlRig/Private/Tools/ControlRigPose.cpp
+++ b/Engine/Plugins/Animation/ControlRig/Source/ControlRig/Private/Tools/ControlRigPose.cpp
@@ -138,6 +138,7 @@ void FControlRigControlPose::PastePoseInternal(UControlRig* ControlRig, bool bDo
 				break;
 			}
 			case ERigControlType::Float:
+			case ERigControlType::ScaleFloat:
 			{
 				float Val = CopyRigControl->Value.Get<float>();
 				ControlRig->SetControlValue<float>(ControlElement->GetFName(), Val, true, Context,bSetupUndo);
diff --git a/Engine/Plugins/Animation/ControlRig/Source/ControlRig/Private/Units/Execution/RigUnit_DynamicHierarchy.cpp b/Engine/Plugins/Animation/ControlRig/Source/ControlRig/Private/Units/Execution/RigUnit_DynamicHierarchy.cpp
--- a/Engine/Plugins/Animation/ControlRig/Source/ControlRig/Private/Units/Execution/RigUnit_DynamicHierarchy.cpp
+++ b/Engine/Plugins/Animation/ControlRig/Source/ControlRig/Private/Units/Execution/RigUnit_DynamicHierarchy.cpp
@@ -415,7 +415,7 @@ void FRigUnit_HierarchyAddControlFloat_Settings::Configure(FRigControlSettings&
 {
 	Super::Configure(OutSettings);
 	
-	OutSettings.ControlType = ERigControlType::Float;
+	OutSettings.ControlType = bIsScale ? ERigControlType::ScaleFloat : ERigControlType::Float;
 	OutSettings.PrimaryAxis = PrimaryAxis;
 
 	Proxy.Configure(OutSettings);
@@ -814,6 +814,45 @@ FRigUnit_HierarchyAddAnimationChannelFloat_Execute()
 	}
 }
 
+FRigUnit_HierarchyAddAnimationChannelScaleFloat_Execute()
+{
+	FString ErrorMessage;
+	if(!FRigUnit_DynamicHierarchyBase::IsValidToRunInContext(ExecuteContext, true, &ErrorMessage))
+	{
+		if(!ErrorMessage.IsEmpty())
+		{
+			UE_CONTROLRIG_RIGUNIT_REPORT_ERROR(TEXT("%s"), *ErrorMessage);
+		}
+		return;
+	}
+
+	Item.Reset();
+
+	if(URigHierarchyController* Controller = ExecuteContext.Hierarchy->GetController(true))
+	{
+		FRigControlSettings ControlSettings;
+		ControlSettings.ControlType = ERigControlType::ScaleFloat;
+		ControlSettings.SetupLimitArrayForType(true, true, true);
+		ControlSettings.MinimumValue = FRigControlValue::Make<float>(MinimumValue);
+		ControlSettings.MaximumValue = FRigControlValue::Make<float>(MaximumValue);
+		ControlSettings.DisplayName = Controller->GetHierarchy()->GetSafeNewDisplayName(Parent, Name);
+		const FRigControlValue Value = FRigControlValue::Make<float>(InitialValue);
+		
+		FRigHierarchyControllerInstructionBracket InstructionBracket(Controller, ExecuteContext.GetInstructionIndex());
+		Item = Controller->AddAnimationChannel(Name, Parent, ControlSettings, false, false);
+
+		if(Item.IsValid())
+		{
+			ExecuteContext.Hierarchy->SetControlValue(Item, Value, ERigControlValueType::Initial, false, false);
+			ExecuteContext.Hierarchy->SetControlValue(Item, Value, ERigControlValueType::Current, false, false);
+		}
+		else
+		{
+			UE_CONTROLRIG_RIGUNIT_REPORT_ERROR(TEXT("Animation Channel could not be added. Is Parent valid?"));
+		}
+	}
+}
+
 FRigUnit_HierarchyAddAnimationChannelInteger_Execute()
 {
 	FString ErrorMessage;
@@ -931,6 +970,45 @@ FRigUnit_HierarchyAddAnimationChannelVector_Execute()
 	}
 }
 
+FRigUnit_HierarchyAddAnimationChannelScaleVector_Execute()
+{
+	FString ErrorMessage;
+	if(!FRigUnit_DynamicHierarchyBase::IsValidToRunInContext(ExecuteContext, true, &ErrorMessage))
+	{
+		if(!ErrorMessage.IsEmpty())
+		{
+			UE_CONTROLRIG_RIGUNIT_REPORT_ERROR(TEXT("%s"), *ErrorMessage);
+		}
+		return;
+	}
+
+	Item.Reset();
+
+	if(URigHierarchyController* Controller = ExecuteContext.Hierarchy->GetController(true))
+	{
+		FRigControlSettings ControlSettings;
+		ControlSettings.ControlType = ERigControlType::Scale;
+		ControlSettings.SetupLimitArrayForType(true, true, true);
+		ControlSettings.MinimumValue = FRigControlValue::Make<FVector3f>(FVector3f(MinimumValue));
+		ControlSettings.MaximumValue = FRigControlValue::Make<FVector3f>(FVector3f(MaximumValue));
+		ControlSettings.DisplayName = Controller->GetHierarchy()->GetSafeNewDisplayName(Parent, Name);
+		const FRigControlValue Value = FRigControlValue::Make<FVector3f>(FVector3f(InitialValue));
+		
+		FRigHierarchyControllerInstructionBracket InstructionBracket(Controller, ExecuteContext.GetInstructionIndex());
+		Item = Controller->AddAnimationChannel(Name, Parent, ControlSettings, false, false);
+
+		if(Item.IsValid())
+		{
+			ExecuteContext.Hierarchy->SetControlValue(Item, Value, ERigControlValueType::Initial, false, false);
+			ExecuteContext.Hierarchy->SetControlValue(Item, Value, ERigControlValueType::Current, false, false);
+		}
+		else
+		{
+			UE_CONTROLRIG_RIGUNIT_REPORT_ERROR(TEXT("Animation Channel could not be added. Is Parent valid?"));
+		}
+	}
+}
+
 FRigUnit_HierarchyAddAnimationChannelRotator_Execute()
 {
 	FString ErrorMessage;
diff --git a/Engine/Plugins/Animation/ControlRig/Source/ControlRig/Private/Units/Execution/RigUnit_DynamicHierarchy.h b/Engine/Plugins/Animation/ControlRig/Source/ControlRig/Private/Units/Execution/RigUnit_DynamicHierarchy.h
--- a/Engine/Plugins/Animation/ControlRig/Source/ControlRig/Private/Units/Execution/RigUnit_DynamicHierarchy.h
+++ b/Engine/Plugins/Animation/ControlRig/Source/ControlRig/Private/Units/Execution/RigUnit_DynamicHierarchy.h
@@ -520,6 +520,9 @@ struct CONTROLRIG_API FRigUnit_HierarchyAddControlFloat_Settings : public FRigUn
 
 	UPROPERTY(EditAnywhere, BlueprintReadWrite, Category = "Settings")
 	ERigControlAxis PrimaryAxis;
+
+	UPROPERTY(EditAnywhere, BlueprintReadWrite, Category = "Settings")
+	bool bIsScale;
 	
 	UPROPERTY(EditAnywhere, BlueprintReadWrite, Category = "Settings")
 	FRigUnit_HierarchyAddControlFloat_LimitSettings Limits;;
@@ -581,7 +584,7 @@ struct CONTROLRIG_API FRigUnit_HierarchyAddControlFloat : public FRigUnit_Hierar
 		InitialValue = 0.f;
 	}
 
-	virtual ERigControlType GetControlTypeToSpawn() const override { return ERigControlType::Float; }
+	virtual ERigControlType GetControlTypeToSpawn() const override { return Settings.bIsScale ? ERigControlType::ScaleFloat : ERigControlType::Float; }
 
 	/*
 	 * The initial value of the new control
@@ -1128,6 +1131,45 @@ struct CONTROLRIG_API FRigUnit_HierarchyAddAnimationChannelFloat : public FRigUn
 	virtual void Execute() override;
 };
 
+/**
+ * Adds a new animation channel to the hierarchy
+ * Note: This node only runs as part of the construction event.
+ */
+USTRUCT(meta=(DisplayName="Spawn Scale Float Animation Channel", TemplateName="SpawnScaleAnimationChannel", Keywords="Construction,Create,New,AddAnimationChannel,NewAnimationChannel,CreateAnimationChannel,AddChannel,NewChannel,CreateChannel,SpawnChannel", Varying))
+struct CONTROLRIG_API FRigUnit_HierarchyAddAnimationChannelScaleFloat : public FRigUnit_HierarchyAddElement
+{
+	GENERATED_BODY()
+
+	FRigUnit_HierarchyAddAnimationChannelScaleFloat()
+	{
+		Name = TEXT("NewChannel");
+		InitialValue = 1.f;
+		MinimumValue = 0.f;
+		MaximumValue = 10.f;
+	}
+
+	/*
+	 * The initial value of the new animation channel
+	 */
+	UPROPERTY(meta = (Input))
+	float InitialValue;
+
+	/*
+	 * The initial value of the new animation channel
+	 */
+	UPROPERTY(meta = (Input))
+	float MinimumValue;
+
+	/*
+	 * The maximum value for the animation channel
+	 */
+	UPROPERTY(meta = (Input))
+	float MaximumValue;
+
+	RIGVM_METHOD()
+	virtual void Execute() override;
+};
+
 /**
  * Adds a new animation channel to the hierarchy
  * Note: This node only runs as part of the construction event.
@@ -1245,6 +1287,45 @@ struct CONTROLRIG_API FRigUnit_HierarchyAddAnimationChannelVector : public FRigU
 	virtual void Execute() override;
 };
 
+/**
+ * Adds a new animation channel to the hierarchy
+ * Note: This node only runs as part of the construction event.
+ */
+USTRUCT(meta=(DisplayName="Spawn Vector Animation Channel", TemplateName="SpawnScaleAnimationChannel", Keywords="Construction,Create,New,AddAnimationChannel,NewAnimationChannel,CreateAnimationChannel,AddChannel,NewChannel,CreateChannel,SpawnChannel", Varying))
+struct CONTROLRIG_API FRigUnit_HierarchyAddAnimationChannelScaleVector : public FRigUnit_HierarchyAddElement
+{
+	GENERATED_BODY()
+
+	FRigUnit_HierarchyAddAnimationChannelScaleVector()
+	{
+		Name = TEXT("NewChannel");
+		InitialValue = FVector::OneVector;
+		MinimumValue = FVector::ZeroVector;
+		MaximumValue = FVector(10.f, 10.f, 10.f);
+	}
+
+	/*
+	 * The initial value of the new animation channel
+	 */
+	UPROPERTY(meta = (Input))
+	FVector InitialValue;
+
+	/*
+	 * The initial value of the new animation channel
+	 */
+	UPROPERTY(meta = (Input))
+	FVector MinimumValue;
+
+	/*
+	 * The maximum value for the animation channel
+	 */
+	UPROPERTY(meta = (Input))
+	FVector MaximumValue;
+
+	RIGVM_METHOD()
+	virtual void Execute() override;
+};
+
 /**
  * Adds a new animation channel to the hierarchy
  * Note: This node only runs as part of the construction event.
diff --git a/Engine/Plugins/Animation/ControlRig/Source/ControlRig/Private/Units/Hierarchy/RigUnit_ControlChannel.cpp b/Engine/Plugins/Animation/ControlRig/Source/ControlRig/Private/Units/Hierarchy/RigUnit_ControlChannel.cpp
--- a/Engine/Plugins/Animation/ControlRig/Source/ControlRig/Private/Units/Hierarchy/RigUnit_ControlChannel.cpp
+++ b/Engine/Plugins/Animation/ControlRig/Source/ControlRig/Private/Units/Hierarchy/RigUnit_ControlChannel.cpp
@@ -80,7 +80,7 @@ FRigUnit_GetFloatAnimationChannel_Execute()
 
 	if(const FRigControlElement* ChannelElement = ExecuteContext.Hierarchy->Find<FRigControlElement>(CachedChannelKey))
 	{
-		if(ChannelElement->Settings.ControlType == ERigControlType::Float)
+		if(ChannelElement->Settings.ControlType == ERigControlType::Float || ChannelElement->Settings.ControlType == ERigControlType::ScaleFloat)
 		{
 			const FRigControlValue StoredValue = ExecuteContext.Hierarchy->GetControlValueByIndex(ChannelElement->GetIndex(), bInitial ? ERigControlValueType::Initial : ERigControlValueType::Current);
 			Value = StoredValue.Get<float>();
@@ -242,7 +242,7 @@ FRigUnit_SetFloatAnimationChannel_Execute()
 
 	if(const FRigControlElement* ChannelElement = ExecuteContext.Hierarchy->Find<FRigControlElement>(CachedChannelKey))
 	{
-		if(ChannelElement->Settings.ControlType == ERigControlType::Float)
+		if(ChannelElement->Settings.ControlType == ERigControlType::Float || ChannelElement->Settings.ControlType == ERigControlType::ScaleFloat)
 		{
 			const FRigControlValue ValueToStore = FRigControlValue::Make<float>(Value);
 			ExecuteContext.Hierarchy->SetControlValueByIndex(ChannelElement->GetIndex(), ValueToStore, bInitial ? ERigControlValueType::Initial : ERigControlValueType::Current);
diff --git a/Engine/Plugins/Animation/ControlRig/Source/ControlRig/Public/Rigs/RigHierarchyDefines.h b/Engine/Plugins/Animation/ControlRig/Source/ControlRig/Public/Rigs/RigHierarchyDefines.h
--- a/Engine/Plugins/Animation/ControlRig/Source/ControlRig/Public/Rigs/RigHierarchyDefines.h
+++ b/Engine/Plugins/Animation/ControlRig/Source/ControlRig/Public/Rigs/RigHierarchyDefines.h
@@ -137,6 +137,7 @@ enum class ERigControlType : uint8
     Transform UMETA(Hidden),
     TransformNoScale UMETA(Hidden),
     EulerTransform,
+	ScaleFloat,
 };
 
 UENUM(BlueprintType)
@@ -456,14 +457,16 @@ public:
 		switch (InControlType)
 		{
 			case ERigControlType::Bool: ValueStr = FString::Printf(TEXT("unreal.RigHierarchy.make_control_value_from_bool(%s)"), Get<bool>() ? TEXT("True") : TEXT("False")); break;							
-			case ERigControlType::Float: ValueStr = FString::Printf(TEXT("unreal.RigHierarchy.make_control_value_from_float(%.6f)"), Get<float>()); break;
+			case ERigControlType::Float:
+			case ERigControlType::ScaleFloat:
+				ValueStr = FString::Printf(TEXT("unreal.RigHierarchy.make_control_value_from_float(%.6f)"), Get<float>()); break;
 			case ERigControlType::Integer: ValueStr = FString::Printf(TEXT("unreal.RigHierarchy.make_control_value_from_int(%d)"), Get<int>()); break;
-			case ERigControlType::Position: ValueStr = FString::Printf(TEXT("unreal.RigHierarchy.make_control_value_from_vector(unreal.Vector(%.6f, %.6f, %.6f))"),
-				Get<FVector3f>().X, Get<FVector3f>().Y, Get<FVector3f>().Z); break;
+			case ERigControlType::Position:
+			case ERigControlType::Scale:
+				ValueStr = FString::Printf(TEXT("unreal.RigHierarchy.make_control_value_from_vector(unreal.Vector(%.6f, %.6f, %.6f))"),
+					Get<FVector3f>().X, Get<FVector3f>().Y, Get<FVector3f>().Z); break;
 			case ERigControlType::Rotator: ValueStr = FString::Printf(TEXT("unreal.RigHierarchy.make_control_value_from_rotator(unreal.Rotator(pitch=%.6f, roll=%.6f, yaw=%.6f))"),
 				Get<FVector3f>().X, Get<FVector3f>().Z, Get<FVector3f>().Y); break;
-			case ERigControlType::Scale: ValueStr = FString::Printf(TEXT("unreal.RigHierarchy.make_control_value_from_vector(unreal.Vector(%.6f, %.6f, %.6f))"),
-				Get<FVector3f>().X, Get<FVector3f>().Y, Get<FVector3f>().Z); break;
 			case ERigControlType::Transform: ValueStr = FString::Printf(TEXT("unreal.RigHierarchy.make_control_value_from_euler_transform(unreal.EulerTransform(location=[%.6f,%.6f,%.6f],rotation=[%.6f,%.6f,%.6f],scale=[%.6f,%.6f,%.6f]))"),
 				Get<FTransform_Float>().TranslationX,
 				Get<FTransform_Float>().TranslationY,
@@ -541,6 +544,13 @@ public:
 				}
 				break;
 			}
+			case ERigControlType::ScaleFloat:
+			{
+				const float ValueToGet = Get<float>();
+				const FVector ScaleVector(ValueToGet, ValueToGet, ValueToGet);
+				Transform.SetScale3D(ScaleVector);
+				break;
+			}
 			case ERigControlType::Integer:
 			{
 				const int32 ValueToGet = Get<int32>();
@@ -663,6 +673,11 @@ public:
 				}
 				break;
 			}
+			case ERigControlType::ScaleFloat:
+			{
+				Set<float>(InTransform.GetScale3D().X);
+				break;
+			}
 			case ERigControlType::Integer:
 			{
 				switch (InPrimaryAxis)
@@ -785,6 +800,7 @@ public:
 		switch(InControlType)
 		{
 			case ERigControlType::Float:
+			case ERigControlType::ScaleFloat:
 			{
 				if (LimitEnabled[0].IsOn())
 				{
diff --git a/Engine/Plugins/Animation/ControlRig/Source/ControlRigDeveloper/Private/AnimGraphNode_ControlRig.cpp b/Engine/Plugins/Animation/ControlRig/Source/ControlRigDeveloper/Private/AnimGraphNode_ControlRig.cpp
--- a/Engine/Plugins/Animation/ControlRig/Source/ControlRigDeveloper/Private/AnimGraphNode_ControlRig.cpp
+++ b/Engine/Plugins/Animation/ControlRig/Source/ControlRigDeveloper/Private/AnimGraphNode_ControlRig.cpp
@@ -425,7 +425,7 @@ bool UAnimGraphNode_ControlRig::IsAvailableToMapToCurve(const FName& PropertyNam
 
 	if(const FRigControlElement* ControlElement = FindControlElement(PropertyName))
 	{
-		return ControlElement->Settings.ControlType == ERigControlType::Float;
+		return (ControlElement->Settings.ControlType == ERigControlType::Float || ControlElement->Settings.ControlType == ERigControlType::ScaleFloat);
 	}
 
 	return ensure(false);
diff --git a/Engine/Plugins/Animation/ControlRig/Source/ControlRigEditor/Private/ControlRigElementDetails.cpp b/Engine/Plugins/Animation/ControlRig/Source/ControlRigEditor/Private/ControlRigElementDetails.cpp
--- a/Engine/Plugins/Animation/ControlRig/Source/ControlRigEditor/Private/ControlRigElementDetails.cpp
+++ b/Engine/Plugins/Animation/ControlRig/Source/ControlRigEditor/Private/ControlRigElementDetails.cpp
@@ -2763,6 +2763,7 @@ void FRigControlElementDetails::CustomizeValue(IDetailLayoutBuilder& DetailBuild
 				break;
 			}
 			case ERigControlType::Float:
+			case ERigControlType::ScaleFloat:
 			{
 				CreateFloatValueWidgetRow(Keys, ValueCategory, Labels[Index], Tooltips[Index], ValueType, VisibilityAttribute);
 				break;
@@ -2917,6 +2918,7 @@ void FRigControlElementDetails::CustomizeControl(IDetailLayoutBuilder& DetailBui
 					ControlElement->Settings.bGroupWithParentControl =
 						ControlElement->Settings.ControlType == ERigControlType::Bool ||
 						ControlElement->Settings.ControlType == ERigControlType::Float ||
+						ControlElement->Settings.ControlType == ERigControlType::ScaleFloat ||
 						ControlElement->Settings.ControlType == ERigControlType::Integer ||
 						ControlElement->Settings.ControlType == ERigControlType::Vector2D;
 
@@ -3023,6 +3025,7 @@ void FRigControlElementDetails::CustomizeControl(IDetailLayoutBuilder& DetailBui
 	if(bSupportsShape &&
 		!(IsAnyControlNotOfValueType(ERigControlType::Integer) &&
 		IsAnyControlNotOfValueType(ERigControlType::Float) &&
+		IsAnyControlNotOfValueType(ERigControlType::ScaleFloat) &&
 		IsAnyControlNotOfValueType(ERigControlType::Vector2D)))
 	{
 		const TSharedPtr<IPropertyHandle> PrimaryAxisHandle = SettingsHandle->GetChildHandle(TEXT("PrimaryAxis"));
@@ -3328,6 +3331,7 @@ void FRigControlElementDetails::CustomizeAnimationChannels(IDetailLayoutBuilder&
 						break;
 					}
 					case ERigControlType::Float:
+					case ERigControlType::ScaleFloat:
 					{
 						WidgetRow = &CreateFloatValueWidgetRow(ChildElementKeys, Category, Label, FText(), ERigControlValueType::Current, Visibility, NameContent);
 						break;
@@ -3603,6 +3607,15 @@ void FRigControlElementDetails::HandleControlTypeChanged(ERigControlType Control
 				ControlElement->Settings.bGroupWithParentControl = ControlElement->Settings.IsAnimatable();
 				break;
 			}
+			case ERigControlType::ScaleFloat:
+			{
+				ValueToSet = FRigControlValue::Make<float>(1.f);
+				ControlElement->Settings.SetupLimitArrayForType(false);
+				ControlElement->Settings.MinimumValue = FRigControlValue::Make<float>(0.f);
+				ControlElement->Settings.MaximumValue = FRigControlValue::Make<float>(10.f);
+				ControlElement->Settings.bGroupWithParentControl = ControlElement->Settings.IsAnimatable();
+				break;
+			}
 			case ERigControlType::Integer:
 			{
 				ValueToSet = FRigControlValue::Make<int32>(0);
diff --git a/Engine/Plugins/Animation/ControlRig/Source/ControlRigEditor/Private/ControlRigSequencerEditorLibrary.cpp b/Engine/Plugins/Animation/ControlRig/Source/ControlRigEditor/Private/ControlRigSequencerEditorLibrary.cpp
--- a/Engine/Plugins/Animation/ControlRig/Source/ControlRigEditor/Private/ControlRigSequencerEditorLibrary.cpp
+++ b/Engine/Plugins/Animation/ControlRig/Source/ControlRigEditor/Private/ControlRigSequencerEditorLibrary.cpp
@@ -1389,7 +1389,7 @@ float UControlRigSequencerEditorLibrary::GetLocalControlRigFloat(ULevelSequence*
 	{
 		if (FRigControlElement* Element = ControlRig->FindControl(ControlName))
 		{
-			if (Element->Settings.ControlType != ERigControlType::Float)
+			if (Element->Settings.ControlType != ERigControlType::Float || Element->Settings.ControlType != ERigControlType::ScaleFloat)
 			{
 				UE_LOG(LogControlRig, Error, TEXT("Control Rig Wrong Type"));
 				return Value;
@@ -1413,7 +1413,7 @@ TArray<float> UControlRigSequencerEditorLibrary::GetLocalControlRigFloats(ULevel
 	{
 		if (FRigControlElement* Element = ControlRig->FindControl(ControlName))
 		{
-			if (Element->Settings.ControlType != ERigControlType::Float)
+			if (Element->Settings.ControlType != ERigControlType::Float || Element->Settings.ControlType != ERigControlType::ScaleFloat)
 			{
 				UE_LOG(LogControlRig, Error, TEXT("Control Rig Wrong Type"));
 				return Values;
@@ -1442,7 +1442,7 @@ void UControlRigSequencerEditorLibrary::SetLocalControlRigFloat(ULevelSequence*
 	}
 	if (FRigControlElement* Element = ControlRig->FindControl(ControlName))
 	{
-		if (Element->Settings.ControlType != ERigControlType::Float)
+		if (Element->Settings.ControlType != ERigControlType::Float || Element->Settings.ControlType != ERigControlType::ScaleFloat)
 		{
 			UE_LOG(LogControlRig, Error, TEXT("Control Rig Wrong Type"));
 			return;
@@ -1471,7 +1471,7 @@ void UControlRigSequencerEditorLibrary::SetLocalControlRigFloats(ULevelSequence*
 	}
 	if (FRigControlElement* Element = ControlRig->FindControl(ControlName))
 	{
-		if (Element->Settings.ControlType != ERigControlType::Float)
+		if (Element->Settings.ControlType != ERigControlType::Float || Element->Settings.ControlType != ERigControlType::ScaleFloat)
 		{
 			UE_LOG(LogControlRig, Error, TEXT("Control Rig Wrong Type"));
 			return;
diff --git a/Engine/Plugins/Animation/ControlRig/Source/ControlRigEditor/Private/ControlRigThumbnailRenderer.cpp b/Engine/Plugins/Animation/ControlRig/Source/ControlRigEditor/Private/ControlRigThumbnailRenderer.cpp
--- a/Engine/Plugins/Animation/ControlRig/Source/ControlRigEditor/Private/ControlRigThumbnailRenderer.cpp
+++ b/Engine/Plugins/Animation/ControlRig/Source/ControlRigEditor/Private/ControlRigThumbnailRenderer.cpp
@@ -112,6 +112,7 @@ void UControlRigThumbnailRenderer::AddAdditionalPreviewSceneContent(UObject* Obj
 				switch (ControlElement->Settings.ControlType)
 				{
 					case ERigControlType::Float:
+					case ERigControlType::ScaleFloat:
 					case ERigControlType::Integer:
 					case ERigControlType::Vector2D:
 					case ERigControlType::Position:
diff --git a/Engine/Plugins/Animation/ControlRig/Source/ControlRigEditor/Private/EditMode/ControlRigControlsProxy.cpp b/Engine/Plugins/Animation/ControlRig/Source/ControlRigEditor/Private/EditMode/ControlRigControlsProxy.cpp
--- a/Engine/Plugins/Animation/ControlRig/Source/ControlRigEditor/Private/EditMode/ControlRigControlsProxy.cpp
+++ b/Engine/Plugins/Animation/ControlRig/Source/ControlRigEditor/Private/EditMode/ControlRigControlsProxy.cpp
@@ -870,6 +870,7 @@ void UControlRigDetailPanelControlProxies::AddProxy(UControlRig* ControlRig, con
 				break;
 			}
 			case ERigControlType::Float:
+			case ERigControlType::ScaleFloat:
 			{
 				Proxy = NewObject<UControlRigFloatControlProxy>(GetTransientPackage(), NAME_None);
 				break;
diff --git a/Engine/Plugins/Animation/ControlRig/Source/ControlRigEditor/Private/EditMode/ControlRigEditMode.cpp b/Engine/Plugins/Animation/ControlRig/Source/ControlRigEditor/Private/EditMode/ControlRigEditMode.cpp
--- a/Engine/Plugins/Animation/ControlRig/Source/ControlRigEditor/Private/EditMode/ControlRigEditMode.cpp
+++ b/Engine/Plugins/Animation/ControlRig/Source/ControlRigEditor/Private/EditMode/ControlRigEditMode.cpp
@@ -624,6 +624,7 @@ void FControlRigEditMode::Tick(FEditorViewportClient* ViewportClient, float Delt
 										break;
 									}
 									case ERigControlType::Scale:
+									case ERigControlType::ScaleFloat:
 									{
 										ViewportClient->SetWidgetMode(UE::Widget::WM_Scale);
 										break;
@@ -3233,6 +3234,7 @@ static bool IsSupportedControlType(const ERigControlType ControlType)
 	switch (ControlType)
 	{
 	case ERigControlType::Float:
+	case ERigControlType::ScaleFloat:
 	case ERigControlType::Integer:
 	case ERigControlType::Vector2D:
 	case ERigControlType::Position:
@@ -4288,6 +4290,7 @@ bool FControlRigEditMode::ModeSupportedByShapeActor(const AControlRigShapeActor*
 						switch (ControlElement->Settings.ControlType)
 						{
 							case ERigControlType::Scale:
+							case ERigControlType::ScaleFloat:
 							case ERigControlType::Transform:
 							case ERigControlType::EulerTransform:
 							{
diff --git a/Engine/Plugins/Animation/ControlRig/Source/ControlRigEditor/Private/EditMode/SControlRigDetails.cpp b/Engine/Plugins/Animation/ControlRig/Source/ControlRigEditor/Private/EditMode/SControlRigDetails.cpp
--- a/Engine/Plugins/Animation/ControlRig/Source/ControlRigEditor/Private/EditMode/SControlRigDetails.cpp
+++ b/Engine/Plugins/Animation/ControlRig/Source/ControlRigEditor/Private/EditMode/SControlRigDetails.cpp
@@ -99,7 +99,8 @@ void FControlRigEditModeGenericDetails::CustomizeDetails(class IDetailLayoutBuil
 		}
 
 		FName ValuePropertyName = TEXT("Transform");
-		if (ControlElement->Settings.ControlType == ERigControlType::Float)
+		if (ControlElement->Settings.ControlType == ERigControlType::Float ||
+			ControlElement->Settings.ControlType == ERigControlType::ScaleFloat)
 		{
 			ValuePropertyName = TEXT("Float");
 		}
@@ -168,6 +169,7 @@ void FControlRigEditModeGenericDetails::CustomizeDetails(class IDetailLayoutBuil
 								break;
 							}
 							case ERigControlType::Float:
+							case ERigControlType::ScaleFloat:
 							{
 								PropertyName = TEXT("Float");
 								break;
diff --git a/Engine/Plugins/Animation/ControlRig/Source/ControlRigEditor/Private/Editor/ControlRigEditor.cpp b/Engine/Plugins/Animation/ControlRig/Source/ControlRigEditor/Private/Editor/ControlRigEditor.cpp
--- a/Engine/Plugins/Animation/ControlRig/Source/ControlRigEditor/Private/Editor/ControlRigEditor.cpp
+++ b/Engine/Plugins/Animation/ControlRig/Source/ControlRigEditor/Private/Editor/ControlRigEditor.cpp
@@ -3621,6 +3621,7 @@ void FControlRigEditor::HandleMakeElementGetterSetter(ERigElementGetterSetterTyp
 						break;
 					}
 					case ERigControlType::Float:
+					case ERigControlType::ScaleFloat:
 					{
 						if(bIsGetter)
 						{
@@ -3742,6 +3743,7 @@ void FControlRigEditor::HandleMakeElementGetterSetter(ERigElementGetterSetterTyp
 								break;
 							}
 							case ERigControlType::Float:
+							case ERigControlType::ScaleFloat:
 							{
 								NamePins.Add(TEXT("Control"));
 								StructTemplate = FRigUnit_GetControlFloat::StaticStruct();
@@ -3830,6 +3832,7 @@ void FControlRigEditor::HandleMakeElementGetterSetter(ERigElementGetterSetterTyp
 								break;
 							}
 							case ERigControlType::Float:
+							case ERigControlType::ScaleFloat:
 							{
 								NamePins.Add(TEXT("Control"));
 								StructTemplate = FRigUnit_SetControlFloat::StaticStruct();
diff --git a/Engine/Plugins/Animation/ControlRig/Source/ControlRigEditor/Private/Editor/SRigHierarchyTreeView.cpp b/Engine/Plugins/Animation/ControlRig/Source/ControlRigEditor/Private/Editor/SRigHierarchyTreeView.cpp
--- a/Engine/Plugins/Animation/ControlRig/Source/ControlRigEditor/Private/Editor/SRigHierarchyTreeView.cpp
+++ b/Engine/Plugins/Animation/ControlRig/Source/ControlRigEditor/Private/Editor/SRigHierarchyTreeView.cpp
@@ -894,6 +894,7 @@ FLinearColor SRigHierarchyItem::GetColorForControlType(ERigControlType InControl
 			break;
 		}
 		case ERigControlType::Float:
+		case ERigControlType::ScaleFloat:
 		{
 			PinType = RigVMTypeUtils::PinTypeFromCPPType(RigVMTypeUtils::FloatTypeName, nullptr);
 			break;
diff --git a/Engine/Plugins/Animation/ControlRig/Source/ControlRigEditor/Private/Editor/SRigSpacePickerWidget.cpp b/Engine/Plugins/Animation/ControlRig/Source/ControlRigEditor/Private/Editor/SRigSpacePickerWidget.cpp
--- a/Engine/Plugins/Animation/ControlRig/Source/ControlRigEditor/Private/Editor/SRigSpacePickerWidget.cpp
+++ b/Engine/Plugins/Animation/ControlRig/Source/ControlRigEditor/Private/Editor/SRigSpacePickerWidget.cpp
@@ -244,6 +244,7 @@ void SRigSpacePickerWidget::SetControls(
 			}
 			if (ControlElement->Settings.ControlType == ERigControlType::Bool ||
 				ControlElement->Settings.ControlType == ERigControlType::Float ||
+				ControlElement->Settings.ControlType == ERigControlType::ScaleFloat ||
 				ControlElement->Settings.ControlType == ERigControlType::Integer)
 			{
 				//if it has a channel and has a parent bail
diff --git a/Engine/Plugins/Animation/ControlRig/Source/ControlRigEditor/Private/Sequencer/ControlRigParameterTrackEditor.cpp b/Engine/Plugins/Animation/ControlRig/Source/ControlRigEditor/Private/Sequencer/ControlRigParameterTrackEditor.cpp
--- a/Engine/Plugins/Animation/ControlRig/Source/ControlRigEditor/Private/Sequencer/ControlRigParameterTrackEditor.cpp
+++ b/Engine/Plugins/Animation/ControlRig/Source/ControlRigEditor/Private/Sequencer/ControlRigParameterTrackEditor.cpp
@@ -1764,6 +1764,7 @@ static void EvaluateThisControl(UMovieSceneControlRigParameterSection* Section,
 
 			}
 			case ERigControlType::Float:
+			case ERigControlType::ScaleFloat:
 			{
 				TOptional <float> Value = Section->EvaluateScalarParameter(FrameTime, ControlName);
 				if (Value.IsSet())
@@ -2215,6 +2216,7 @@ void FControlRigParameterTrackEditor::SelectRigsAndControls(UControlRig* Control
 						
 							if (ControlElement->Settings.ControlType == ERigControlType::Bool ||
 								ControlElement->Settings.ControlType == ERigControlType::Float ||
+								ControlElement->Settings.ControlType == ERigControlType::ScaleFloat ||
 								ControlElement->Settings.ControlType == ERigControlType::Integer)
 							{
 								if (ControlElement->Settings.SupportsShape() || !Hierarchy->IsAnimatable(ControlElement))
@@ -3012,6 +3014,7 @@ void FControlRigParameterTrackEditor::HandleControlSelected(UControlRig* Subject
 				{
 					if (OtherControlElement->Settings.ControlType == ERigControlType::Bool ||
 						OtherControlElement->Settings.ControlType == ERigControlType::Float ||
+						OtherControlElement->Settings.ControlType == ERigControlType::ScaleFloat ||
 						OtherControlElement->Settings.ControlType == ERigControlType::Integer)
 					{
 						if(OtherControlElement->Settings.SupportsShape() || !Hierarchy->IsAnimatable(OtherControlElement))
@@ -3347,6 +3350,7 @@ void FControlRigParameterTrackEditor::GetControlRigKeys(
 				break;
 			}
 			case ERigControlType::Float:
+			case ERigControlType::ScaleFloat:
 			{
 				float Val = ControlValue.Get<float>();
 				pChannelIndex->GeneratedKeyIndex = OutGeneratedKeys.Num();
@@ -3749,6 +3753,7 @@ bool FControlRigParameterTrackEditor::ModifyOurGeneratedKeysByCurrentAndWeight(U
 		switch (ControlElement->Settings.ControlType)
 		{
 		case ERigControlType::Float:
+		case ERigControlType::ScaleFloat:
 		{
 			for (const FFloatInterrogationData& Val : InterrogationData.Iterate<FFloatInterrogationData>(UMovieSceneControlRigParameterSection::GetFloatInterrogationKey()))
 			{
diff --git a/Engine/Plugins/Animation/ControlRig/Source/ControlRigEditor/Private/Tools/ControlRigTweener.cpp b/Engine/Plugins/Animation/ControlRig/Source/ControlRigEditor/Private/Tools/ControlRigTweener.cpp
--- a/Engine/Plugins/Animation/ControlRig/Source/ControlRigEditor/Private/Tools/ControlRigTweener.cpp
+++ b/Engine/Plugins/Animation/ControlRig/Source/ControlRigEditor/Private/Tools/ControlRigTweener.cpp
@@ -395,6 +395,7 @@ bool FAnimSliderObjectSelection::Setup(const TArray<UControlRig*>& SelectedContr
 							switch (ControlElement->Settings.ControlType)
 							{
 								case ERigControlType::Float:
+								case ERigControlType::ScaleFloat:
 								{
 									NumChannels = 1;
 								}
diff --git a/Engine/Plugins/Editor/SequencerAnimTools/Source/SequencerAnimTools/Private/SequencerAnimEditPivotTool.cpp b/Engine/Plugins/Editor/SequencerAnimTools/Source/SequencerAnimTools/Private/SequencerAnimEditPivotTool.cpp
--- a/Engine/Plugins/Editor/SequencerAnimTools/Source/SequencerAnimTools/Private/SequencerAnimEditPivotTool.cpp
+++ b/Engine/Plugins/Editor/SequencerAnimTools/Source/SequencerAnimTools/Private/SequencerAnimEditPivotTool.cpp
@@ -635,6 +635,7 @@ void USequencerPivotTool::GizmoTransformStarted(UTransformProxy* Proxy)
 							if (ControlRig->IsControlSelected(ParentControlElement->GetFName()) && (
 								ControlElement->Settings.ControlType == ERigControlType::Bool || 
 								ControlElement->Settings.ControlType == ERigControlType::Float || 
+								ControlElement->Settings.ControlType == ERigControlType::ScaleFloat || 
 								ControlElement->Settings.ControlType == ERigControlType::Integer)
 								)
 							{
