Control Rig: FIx static analysis error
#rnx

diff --git a/Engine/Plugins/Animation/ControlRig/Source/ControlRigEditor/Private/ControlRigSequencerEditorLibrary.cpp b/Engine/Plugins/Animation/ControlRig/Source/ControlRigEditor/Private/ControlRigSequencerEditorLibrary.cpp
--- a/Engine/Plugins/Animation/ControlRig/Source/ControlRigEditor/Private/ControlRigSequencerEditorLibrary.cpp
+++ b/Engine/Plugins/Animation/ControlRig/Source/ControlRigEditor/Private/ControlRigSequencerEditorLibrary.cpp
@@ -1389,7 +1389,7 @@ float UControlRigSequencerEditorLibrary::GetLocalControlRigFloat(ULevelSequence*
 	{
 		if (FRigControlElement* Element = ControlRig->FindControl(ControlName))
 		{
-			if (Element->Settings.ControlType != ERigControlType::Float || Element->Settings.ControlType != ERigControlType::ScaleFloat)
+			if (Element->Settings.ControlType != ERigControlType::Float && Element->Settings.ControlType != ERigControlType::ScaleFloat)
 			{
 				UE_LOG(LogControlRig, Error, TEXT("Control Rig Wrong Type"));
 				return Value;
@@ -1413,7 +1413,7 @@ TArray<float> UControlRigSequencerEditorLibrary::GetLocalControlRigFloats(ULevel
 	{
 		if (FRigControlElement* Element = ControlRig->FindControl(ControlName))
 		{
-			if (Element->Settings.ControlType != ERigControlType::Float || Element->Settings.ControlType != ERigControlType::ScaleFloat)
+			if (Element->Settings.ControlType != ERigControlType::Float && Element->Settings.ControlType != ERigControlType::ScaleFloat)
 			{
 				UE_LOG(LogControlRig, Error, TEXT("Control Rig Wrong Type"));
 				return Values;
@@ -1442,7 +1442,7 @@ void UControlRigSequencerEditorLibrary::SetLocalControlRigFloat(ULevelSequence*
 	}
 	if (FRigControlElement* Element = ControlRig->FindControl(ControlName))
 	{
-		if (Element->Settings.ControlType != ERigControlType::Float || Element->Settings.ControlType != ERigControlType::ScaleFloat)
+		if (Element->Settings.ControlType != ERigControlType::Float && Element->Settings.ControlType != ERigControlType::ScaleFloat)
 		{
 			UE_LOG(LogControlRig, Error, TEXT("Control Rig Wrong Type"));
 			return;
@@ -1471,7 +1471,7 @@ void UControlRigSequencerEditorLibrary::SetLocalControlRigFloats(ULevelSequence*
 	}
 	if (FRigControlElement* Element = ControlRig->FindControl(ControlName))
 	{
-		if (Element->Settings.ControlType != ERigControlType::Float || Element->Settings.ControlType != ERigControlType::ScaleFloat)
+		if (Element->Settings.ControlType != ERigControlType::Float && Element->Settings.ControlType != ERigControlType::ScaleFloat)
 		{
 			UE_LOG(LogControlRig, Error, TEXT("Control Rig Wrong Type"));
 			return;
