Control Rig: Fix bool controls not working properly in additive control rigs
#rnx

diff --git a/Engine/Plugins/Animation/ControlRig/Source/ControlRig/Private/ControlRig.cpp b/Engine/Plugins/Animation/ControlRig/Source/ControlRig/Private/ControlRig.cpp
--- a/Engine/Plugins/Animation/ControlRig/Source/ControlRig/Private/ControlRig.cpp
+++ b/Engine/Plugins/Animation/ControlRig/Source/ControlRig/Private/ControlRig.cpp
@@ -370,24 +370,33 @@ void UControlRig::Evaluate_AnyThread()
 				{
 					FRigSetControlValueInfo& Info = Value.Value;
 
-					// Transform from animation
-					const FRigControlValue PreviousValue = Hierarchy->GetControlValue(Control, ERigControlValueType::Current, false);
-					const FTransform PreviousTransform = PreviousValue.GetAsTransform(Control->Settings.ControlType, Control->Settings.PrimaryAxis);
+					// A bool value is not an additive property. We just overwrite the value.
+					if (Control->Settings.ControlType == ERigControlType::Bool)
+					{
+						const bool bSetupUndo = false; // Rely on the sequencer track to handle undo/redo
+						Hierarchy->SetControlValue(Control, Info.Value, ERigControlValueType::Current, bSetupUndo, false, Info.bPrintPythonCommnds, false);
+					}
+					else
+					{
+						// Transform from animation
+						const FRigControlValue PreviousValue = Hierarchy->GetControlValue(Control, ERigControlValueType::Current, false);
+						const FTransform PreviousTransform = PreviousValue.GetAsTransform(Control->Settings.ControlType, Control->Settings.PrimaryAxis);
 
-					// Additive transform from controls
-					const FRigControlValue& AdditiveValue = Info.Value;
-					const FTransform AdditiveTransform = AdditiveValue.GetAsTransform(Control->Settings.ControlType, Control->Settings.PrimaryAxis);
+						// Additive transform from controls
+						const FRigControlValue& AdditiveValue = Info.Value;
+						const FTransform AdditiveTransform = AdditiveValue.GetAsTransform(Control->Settings.ControlType, Control->Settings.PrimaryAxis);
 
-					// Add them to find the final value
-					FTransform FinalTransform = AdditiveTransform * PreviousTransform;
+						// Add them to find the final value
+						FTransform FinalTransform = AdditiveTransform * PreviousTransform;
 
-					FRigControlValue FinalValue;
-					FinalValue.SetFromTransform(FinalTransform, Control->Settings.ControlType, Control->Settings.PrimaryAxis);
+						FRigControlValue FinalValue;
+						FinalValue.SetFromTransform(FinalTransform, Control->Settings.ControlType, Control->Settings.PrimaryAxis);
+
+						const bool bSetupUndo = false; // Rely on the sequencer track to handle undo/redo
+						Hierarchy->SetControlValue(Control, FinalValue, ERigControlValueType::Current, bSetupUndo, false, Info.bPrintPythonCommnds, false);
+						Hierarchy->SetPreferredEulerAnglesFromValue(Control, AdditiveValue, ERigControlValueType::Current, Info.bFixEulerFlips);
+					}
 
-					const bool bSetupUndo = false; // Rely on the sequencer track to handle undo/redo
-					Hierarchy->SetControlValue(Control, FinalValue, ERigControlValueType::Current, bSetupUndo, false, Info.bPrintPythonCommnds, false);
-					Hierarchy->SetPreferredEulerAnglesFromValue(Control, AdditiveValue, ERigControlValueType::Current, Info.bFixEulerFlips);
-					
 					if (Info.bNotify && OnControlModified.IsBound())
 					{
 						OnControlModified.Broadcast(this, Control, Info.Context);
@@ -1337,6 +1346,11 @@ bool UControlRig::IsConstructionRequired() const
 	return IsRunOnceEvent(FRigUnit_PrepareForExecution::EventName);
 }
 
+bool UControlRig::SupportsBackwardsSolve() const
+{
+	return SupportsEvent(FRigUnit_InverseExecution::EventName);
+}
+
 void UControlRig::AdaptEventQueueForEvaluate(TArray<FName>& InOutEventQueueToRun)
 {
 	Super::AdaptEventQueueForEvaluate(InOutEventQueueToRun);
@@ -1625,6 +1639,12 @@ FRigControlValue UControlRig::GetControlValue(FRigControlElement* InControl, con
 		const int32 ControlIndex = ControlsAfterBackwardsSolve.GetIndex(InControl->GetKey());
 		if (ControlIndex != INDEX_NONE)
 		{
+			// Booleans are not additive properties, just return the current value
+			if (InControl->Settings.ControlType == ERigControlType::Bool)
+			{
+				return GetHierarchy()->GetControlValue(InControl, InValueType);
+			}
+			
 			// return local space control value (the one to be added after backwards solve)
 			const FRigPoseElement& AnimPose = ControlsAfterBackwardsSolve[ControlIndex];
 			const FRigControlValue& CurrentValue = GetHierarchy()->GetControlValue(InControl, InValueType);
