Control Rig: FIx compilation error
#rnx

diff --git a/Engine/Plugins/Animation/ControlRig/Source/ControlRig/Public/ControlRig.h b/Engine/Plugins/Animation/ControlRig/Source/ControlRig/Public/ControlRig.h
--- a/Engine/Plugins/Animation/ControlRig/Source/ControlRig/Public/ControlRig.h
+++ b/Engine/Plugins/Animation/ControlRig/Source/ControlRig/Public/ControlRig.h
@@ -171,6 +171,10 @@ public:
 
 	bool IsConstructionRequired() const;
 
+	/** Contains a backwards solve event */
+	UFUNCTION(BlueprintCallable, Category = "Control Rig")
+	bool SupportsBackwardsSolve() const;
+
 	virtual void AdaptEventQueueForEvaluate(TArray<FName>& InOutEventQueueToRun) override;
 
 	/** INodeMappingInterface implementation */
