Fixed CIS error due to uninitialized var

diff --git a/Engine/Plugins/Animation/ControlRig/Source/ControlRig/Private/Units/Execution/RigUnit_DynamicHierarchy.h b/Engine/Plugins/Animation/ControlRig/Source/ControlRig/Private/Units/Execution/RigUnit_DynamicHierarchy.h
--- a/Engine/Plugins/Animation/ControlRig/Source/ControlRig/Private/Units/Execution/RigUnit_DynamicHierarchy.h
+++ b/Engine/Plugins/Animation/ControlRig/Source/ControlRig/Private/Units/Execution/RigUnit_DynamicHierarchy.h
@@ -513,6 +513,7 @@ struct CONTROLRIG_API FRigUnit_HierarchyAddControlFloat_Settings : public FRigUn
 	FRigUnit_HierarchyAddControlFloat_Settings()
 		: FRigUnit_HierarchyAddControl_Settings()
 		, PrimaryAxis(ERigControlAxis::X)
+		, bIsScale(false)
 	{}
 	virtual ~FRigUnit_HierarchyAddControlFloat_Settings() override {}
 
