RigHierarchy: Improve performance of UpdateCachedChildren

diff --git a/Engine/Plugins/Animation/ControlRig/Source/ControlRig/Private/Rigs/RigHierarchy.cpp b/Engine/Plugins/Animation/ControlRig/Source/ControlRig/Private/Rigs/RigHierarchy.cpp
--- a/Engine/Plugins/Animation/ControlRig/Source/ControlRig/Private/Rigs/RigHierarchy.cpp
+++ b/Engine/Plugins/Animation/ControlRig/Source/ControlRig/Private/Rigs/RigHierarchy.cpp
@@ -293,7 +293,7 @@ void URigHierarchy::Load(FArchive& Ar)
 		}
 	}
 
-	UpdateAllCachedChildren();
+	UpdateAllCachedChildren(true);
 
 	if(Ar.IsTransacting())
 	{
@@ -559,7 +559,7 @@ void URigHierarchy::CopyHierarchy(URigHierarchy* InHierarchy)
 	TopologyVersion = InHierarchy->GetTopologyVersion();
 	MetadataVersion = InHierarchy->GetMetadataVersion();
 
-	UpdateAllCachedChildren();
+	UpdateAllCachedChildren(true);
 	EnsureCacheValidity();
 }
 
@@ -4507,40 +4507,47 @@ void URigHierarchy::ResetCachedChildren()
 	}
 }
 
-void URigHierarchy::UpdateCachedChildren(const FRigBaseElement* InElement, bool bForce) const
+bool URigHierarchy::UpdateCachedChildren(const FRigBaseElement* InElement, bool bForce) const
 {
 	LLM_SCOPE_BYNAME(TEXT("Animation/ControlRig"));
 	check(InElement);
 
 	if(InElement->TopologyVersion == TopologyVersion && !bForce)
 	{
-		return;
+		return false;
 	}
 
-	InElement->CachedChildren.Reset();
+	// we collect the transforms into a separate array to loop faster,
+	// since some hierarchies have a very large number of curves / non transform elements.
+	TArray<FRigTransformElement*> TransformElements;
+	TransformElements.Reserve(Elements.Num());
 
 	for (int32 ElementIndex = 0; ElementIndex < Elements.Num(); ElementIndex++)
 	{
 		FRigBaseElement* Element = Elements[ElementIndex];
-		if(Element->IsA<FRigSingleParentElement>() || Element->IsA<FRigMultiParentElement>())
+		if(bForce || (Element->TopologyVersion != TopologyVersion))
 		{
-			if(Element->TopologyVersion != TopologyVersion)
+			if(FRigTransformElement* TransformElement = Cast<FRigTransformElement>(Element))
 			{
-				Element->CachedChildren.Reset();
+				TransformElements.Add(TransformElement);
+				TransformElement->CachedChildren.Reset();
+			}
+			else
+			{
+				Element->TopologyVersion = TopologyVersion;
 			}
 		}
 	}
 
 	// since we'll have to loop over all children anyway - it makes sense to update all of them
 	// at the same time.
-	for (int32 ElementIndex = 0; ElementIndex < Elements.Num(); ElementIndex++)
+	for (FRigTransformElement* Element : TransformElements)
 	{
-		FRigBaseElement* Element = Elements[ElementIndex];
 		if(FRigSingleParentElement* SingleParentElement = Cast<FRigSingleParentElement>(Element))
 		{
 			if(FRigTransformElement* ParentElement = SingleParentElement->ParentElement)
 			{
-				if(ParentElement->TopologyVersion != TopologyVersion)
+				if(bForce || (ParentElement->TopologyVersion != TopologyVersion))
 				{
 					ParentElement->CachedChildren.Add(SingleParentElement);
 				}
@@ -4552,7 +4559,7 @@ void URigHierarchy::UpdateCachedChildren(const FRigBaseElement* InElement, bool
 			{
 				if(FRigTransformElement* ParentElement = ParentConstraint.ParentElement)
 				{
-					if(ParentElement->TopologyVersion != TopologyVersion)
+					if(bForce || (ParentElement->TopologyVersion != TopologyVersion))
 					{
 						ParentElement->CachedChildren.Add(MultiParentElement);
 					}
@@ -4561,53 +4568,25 @@ void URigHierarchy::UpdateCachedChildren(const FRigBaseElement* InElement, bool
 		}
 	}
 
-	for (int32 ElementIndex = 0; ElementIndex < Elements.Num(); ElementIndex++)
+	for (FRigTransformElement* Element : TransformElements)
 	{
-		FRigBaseElement* Element = Elements[ElementIndex];
-		if(Element->IsA<FRigSingleParentElement>() || Element->IsA<FRigMultiParentElement>())
-		{
-			Element->TopologyVersion = TopologyVersion;
-		}
+		Element->TopologyVersion = TopologyVersion;
 	}
+
+	return true;
 }
 
-void URigHierarchy::UpdateAllCachedChildren() const
+void URigHierarchy::UpdateAllCachedChildren(bool bForce) const
 {
 	LLM_SCOPE_BYNAME(TEXT("Animation/ControlRig"));
-	TArray<bool> ParentVisited;
-	ParentVisited.AddZeroed(Elements.Num());
 	
-	for (int32 ElementIndex = 0; ElementIndex < Elements.Num(); ElementIndex++)
+	for (const FRigBaseElement* Element : Elements)
 	{
-		FRigBaseElement* Element = Elements[ElementIndex];
-		Element->TopologyVersion = TopologyVersion;
-		
-		if(FRigSingleParentElement* SingleParentElement = Cast<FRigSingleParentElement>(Element))
-		{
-			if(FRigTransformElement* ParentElement = SingleParentElement->ParentElement)
-			{
-				if(!ParentVisited[ParentElement->Index])
-				{
-					ParentElement->CachedChildren.Reset();
-					ParentVisited[ParentElement->Index] = true;
-				}
-				ParentElement->CachedChildren.Add(Element);
-			}
-		}
-		else if(FRigMultiParentElement* MultiParentElement = Cast<FRigMultiParentElement>(Element))
+		// as soon as one element was updated,
+		// as a side effect all other elements are updated.
+		if(UpdateCachedChildren(Element, bForce))
 		{
-			for(const FRigElementParentConstraint& ParentConstraint : MultiParentElement->ParentConstraints)
-			{
-				if(ParentConstraint.ParentElement)
-				{
-					if(!ParentVisited[ParentConstraint.ParentElement->Index])
-					{
-						ParentConstraint.ParentElement->CachedChildren.Reset();
-						ParentVisited[ParentConstraint.ParentElement->Index] = true;
-					}
-					ParentConstraint.ParentElement->CachedChildren.Add(Element);
-				}
-			}
+			return;
 		}
 	}
 }
diff --git a/Engine/Plugins/Animation/ControlRig/Source/ControlRig/Private/Rigs/RigHierarchyController.cpp b/Engine/Plugins/Animation/ControlRig/Source/ControlRig/Private/Rigs/RigHierarchyController.cpp
--- a/Engine/Plugins/Animation/ControlRig/Source/ControlRig/Private/Rigs/RigHierarchyController.cpp
+++ b/Engine/Plugins/Animation/ControlRig/Source/ControlRig/Private/Rigs/RigHierarchyController.cpp
@@ -2248,12 +2248,6 @@ bool URigHierarchyController::ReorderElement(FRigBaseElement* InElement, int32 I
 		Hierarchy->Elements[GlobalIndex] = Element;
 		Element->Index = GlobalIndex;
 		Hierarchy->IndexLookup.FindOrAdd(Element->Key) = GlobalIndex;
-		Element->TopologyVersion = Hierarchy->GetTopologyVersion();
-	}
-
-	if(ParentElement)
-	{
-		Hierarchy->UpdateCachedChildren(ParentElement, true);
 	}
 
 	Notify(ERigHierarchyNotification::ElementReordered, InElement);
diff --git a/Engine/Plugins/Animation/ControlRig/Source/ControlRig/Public/Rigs/RigHierarchy.h b/Engine/Plugins/Animation/ControlRig/Source/ControlRig/Public/Rigs/RigHierarchy.h
--- a/Engine/Plugins/Animation/ControlRig/Source/ControlRig/Public/Rigs/RigHierarchy.h
+++ b/Engine/Plugins/Animation/ControlRig/Source/ControlRig/Public/Rigs/RigHierarchy.h
@@ -3969,14 +3969,15 @@ private:
 	 * Updates the transient cached children table for a given element if needed (or if bForce == true).
 	 * @param InElement The element to update the children table for
 	 * @param bForce If set to true the table will always be updated
+	 * @return Returns true if a change was performed
 	 */
-	void UpdateCachedChildren(const FRigBaseElement* InElement, bool bForce = false) const;
+	bool UpdateCachedChildren(const FRigBaseElement* InElement, bool bForce = false) const;
 
 	/**
-	* Updates the transient cached children table for all elements if needed (or if bForce == true).
-	* @param bForce If set to true the table will always be updated
+	* Updates the transient cached children table for all elements if needed.
+	 * @param bForce If set to true the table will always be updated
 	*/
-	void UpdateAllCachedChildren() const;
+	void UpdateAllCachedChildren(bool bForce = false) const;
 
 	/**
 	 * Corrects a parent element key for space switching
