Control Rig: Allow adding additive fk control rig directly from the + context menu
#rnx

diff --git a/Engine/Plugins/Animation/ControlRig/Source/ControlRig/Private/Sequencer/MovieSceneControlRigParameterTrack.cpp b/Engine/Plugins/Animation/ControlRig/Source/ControlRig/Private/Sequencer/MovieSceneControlRigParameterTrack.cpp
--- a/Engine/Plugins/Animation/ControlRig/Source/ControlRig/Private/Sequencer/MovieSceneControlRigParameterTrack.cpp
+++ b/Engine/Plugins/Animation/ControlRig/Source/ControlRig/Private/Sequencer/MovieSceneControlRigParameterTrack.cpp
@@ -17,6 +17,9 @@
 
 #define LOCTEXT_NAMESPACE "MovieSceneParameterControlRigTrack"
 
+FColor UMovieSceneControlRigParameterTrack::AbsoluteRigTrackColor = FColor(65, 89, 194, 65);
+FColor UMovieSceneControlRigParameterTrack::AdditiveRigTrackColor = FColor(173, 151, 114);
+
 FControlRotationOrder::FControlRotationOrder()
 	:RotationOrder(EEulerRotationOrder::YZX),
 	bOverrideSetting(false)
@@ -30,7 +33,7 @@ UMovieSceneControlRigParameterTrack::UMovieSceneControlRigParameterTrack(const F
 	, ControlRig(nullptr)
 {
 #if WITH_EDITORONLY_DATA
-	TrackTint = FColor(65, 89, 194, 65);
+	TrackTint = AbsoluteRigTrackColor;
 #endif
 
 	SupportedBlendTypes = FMovieSceneBlendTypeField::None();
diff --git a/Engine/Plugins/Animation/ControlRig/Source/ControlRig/Public/Sequencer/MovieSceneControlRigParameterTrack.h b/Engine/Plugins/Animation/ControlRig/Source/ControlRig/Public/Sequencer/MovieSceneControlRigParameterTrack.h
--- a/Engine/Plugins/Animation/ControlRig/Source/ControlRig/Public/Sequencer/MovieSceneControlRigParameterTrack.h
+++ b/Engine/Plugins/Animation/ControlRig/Source/ControlRig/Public/Sequencer/MovieSceneControlRigParameterTrack.h
@@ -199,6 +199,9 @@ private:
 	UPROPERTY()
 	TMap<FName,FControlRotationOrder> ControlsRotationOrder;
 
+public:
+	static CONTROLRIG_API FColor AbsoluteRigTrackColor;
+	static CONTROLRIG_API FColor AdditiveRigTrackColor;
 };
 
 
diff --git a/Engine/Plugins/Animation/ControlRig/Source/ControlRigEditor/Private/Sequencer/ControlRigParameterTrackEditor.cpp b/Engine/Plugins/Animation/ControlRig/Source/ControlRigEditor/Private/Sequencer/ControlRigParameterTrackEditor.cpp
--- a/Engine/Plugins/Animation/ControlRig/Source/ControlRigEditor/Private/Sequencer/ControlRigParameterTrackEditor.cpp
+++ b/Engine/Plugins/Animation/ControlRig/Source/ControlRigEditor/Private/Sequencer/ControlRigParameterTrackEditor.cpp
@@ -1419,11 +1419,7 @@ void FControlRigParameterTrackEditor::HandleAddControlRigSubMenu(FMenuBuilder& M
 		TSharedPtr<FControlRigClassFilter> ClassFilter = MakeShareable(new FControlRigClassFilter(bFilterAssetBySkeleton, bFilterAssetByAnimatableControls, bCheckInversion, Skeleton));
 		Options.ClassFilters.Add(ClassFilter.ToSharedRef());
 		Options.bShowNoneOption = false;
-
-		if (!bIsAdditiveControlRig)
-		{
-			Options.ExtraPickerCommonClasses.Add(UFKControlRig::StaticClass());
-		}
+		Options.ExtraPickerCommonClasses.Add(UFKControlRig::StaticClass());
 
 		UMovieSceneSequence* Sequence = GetSequencer() ? GetSequencer()->GetFocusedMovieSceneSequence() : nullptr;
 		Options.AdditionalReferencingAssets.Add(FAssetData(Sequence));
@@ -1505,7 +1501,7 @@ void FControlRigParameterTrackEditor::AddControlRig(UClass* InClass, UObject* Bo
 
 	if (InClass && InClass->IsChildOf(UControlRig::StaticClass()) && SequencerParent.IsValid())
 	{
-		if (bIsAdditiveControlRig && !InClass->GetDefaultObject<UControlRig>()->SupportsEvent(FRigUnit_InverseExecution::EventName))
+		if (bIsAdditiveControlRig && InClass != UFKControlRig::StaticClass() && !InClass->GetDefaultObject<UControlRig>()->SupportsEvent(FRigUnit_InverseExecution::EventName))
 		{
 			UE_LOG(LogControlRigEditor, Error, TEXT("Cannot add an additive control rig which does not contain a backwards solve event."));
 			return;
@@ -1532,7 +1528,17 @@ void FControlRigParameterTrackEditor::AddControlRig(UClass* InClass, UObject* Bo
 			}
 
 			ControlRig->Modify();
-			ControlRig->SetIsAdditive(bIsAdditiveControlRig);
+			if (UFKControlRig* FKControlRig = Cast<UFKControlRig>(ControlRig))
+			{
+				if (bIsAdditiveControlRig)
+				{
+					FKControlRig->SetApplyMode(EControlRigFKRigExecuteMode::Additive);
+				}
+			}
+			else
+			{
+				ControlRig->SetIsAdditive(bIsAdditiveControlRig);
+			}
 			ControlRig->SetObjectBinding(MakeShared<FControlRigObjectBinding>());
 			ControlRig->GetObjectBinding()->BindToObject(BoundActor);
 			ControlRig->GetDataSourceRegistry()->RegisterDataSource(UControlRig::OwnerComponent, ControlRig->GetObjectBinding()->GetBoundObject());
@@ -1551,16 +1557,17 @@ void FControlRigParameterTrackEditor::AddControlRig(UClass* InClass, UObject* Bo
 
 			if (bIsAdditiveControlRig)
 			{
-				const FString AdditiveObjectName = ObjectName + " (Additive)";
+				const FString AdditiveObjectName = ObjectName + TEXT(" (Additive)");
 				Track->SetTrackName(FName(*ObjectName));
 				Track->SetDisplayName(FText::FromString(AdditiveObjectName));
-				Track->SetColorTint(FColor(173, 151, 114));
+				Track->SetColorTint(UMovieSceneControlRigParameterTrack::AdditiveRigTrackColor);
 			}
 			else
 			{
 				//mz todo need to have multiple rigs with same class
 				Track->SetTrackName(FName(*ObjectName));
 				Track->SetDisplayName(FText::FromString(ObjectName));
+				Track->SetColorTint(UMovieSceneControlRigParameterTrack::AbsoluteRigTrackColor);
 			}
 
 			GetSequencer()->EmptySelection();
@@ -4011,6 +4018,18 @@ void FControlRigParameterTrackEditor::ToggleFKControlRig(UMovieSceneControlRigPa
 	FKControlRig->Modify();
 	Track->Modify();
 	FKControlRig->ToggleApplyMode();
+	if (FKControlRig->GetApplyMode() == EControlRigFKRigExecuteMode::Additive)
+	{
+		const FString AdditiveObjectName = Track->GetTrackName().ToString() + TEXT(" (Additive)");
+		Track->SetDisplayName(FText::FromString(AdditiveObjectName));
+		Track->SetColorTint(UMovieSceneControlRigParameterTrack::AdditiveRigTrackColor);
+	}
+	else
+	{
+		Track->SetDisplayName(FText::FromName(Track->GetTrackName()));
+		Track->SetColorTint(UMovieSceneControlRigParameterTrack::AdditiveRigTrackColor);
+		Track->SetColorTint(UMovieSceneControlRigParameterTrack::AbsoluteRigTrackColor);
+	}
 	for (UMovieSceneSection* Section : Track->GetAllSections())
 	{
 		if (Section)
