Control Rig: Allow asset user data per anim node

diff --git a/Engine/Plugins/Animation/ControlRig/Source/ControlRig/Private/AnimNode_ControlRig.cpp b/Engine/Plugins/Animation/ControlRig/Source/ControlRig/Private/AnimNode_ControlRig.cpp
--- a/Engine/Plugins/Animation/ControlRig/Source/ControlRig/Private/AnimNode_ControlRig.cpp
+++ b/Engine/Plugins/Animation/ControlRig/Source/ControlRig/Private/AnimNode_ControlRig.cpp
@@ -142,6 +142,9 @@ void FAnimNode_ControlRig::CacheBones_AnyThread(const FAnimationCacheBonesContex
 {
 	DECLARE_SCOPE_HIERARCHICAL_COUNTER_FUNC()
 
+	// we also need access to the properties when running construction event
+	PropagateInputProperties(Context.AnimInstanceProxy->GetAnimInstanceObject());
+
 	FAnimNode_ControlRigBase::CacheBones_AnyThread(Context);
 
 	FBoneContainer& RequiredBones = Context.AnimInstanceProxy->GetRequiredBones();
@@ -591,6 +594,13 @@ void FAnimNode_ControlRig::PropagateInputProperties(const UObject* InSourceInsta
 						ArrayProperty->CopyCompleteValue(Variable.Memory, SrcPtr);
 					}
 				}
+				else if(FObjectProperty* ObjectProperty = CastField<FObjectProperty>(CallerProperty))
+				{
+					if(ensure(ObjectProperty->SameType(Variable.Property)))
+					{
+						ObjectProperty->CopyCompleteValue(Variable.Memory, SrcPtr);
+					}
+				}
 				else
 				{
 					ensureMsgf(false, TEXT("Property %s type %s not recognized"), *CallerProperty->GetName(), *CallerProperty->GetCPPType());
diff --git a/Engine/Plugins/Animation/ControlRig/Source/ControlRig/Private/AnimNode_ControlRigBase.cpp b/Engine/Plugins/Animation/ControlRig/Source/ControlRig/Private/AnimNode_ControlRigBase.cpp
--- a/Engine/Plugins/Animation/ControlRig/Source/ControlRig/Private/AnimNode_ControlRigBase.cpp
+++ b/Engine/Plugins/Animation/ControlRig/Source/ControlRig/Private/AnimNode_ControlRigBase.cpp
@@ -41,7 +41,6 @@ FAnimNode_ControlRigBase::FAnimNode_ControlRigBase()
 	, bControlRigRequiresInitialization(true)
 	, LastBonesSerialNumberForCacheBones(0)
 {
-
 }
 
 void FAnimNode_ControlRigBase::OnInitializeAnimInstance(const FAnimInstanceProxy* InProxy, const UAnimInstance* InAnimInstance)
@@ -50,6 +49,8 @@ void FAnimNode_ControlRigBase::OnInitializeAnimInstance(const FAnimInstanceProxy
 
 	FAnimNode_CustomProperty::OnInitializeAnimInstance(InProxy, InAnimInstance);
 
+	WeakAnimInstanceObject = TWeakObjectPtr<const UAnimInstance>(InAnimInstance);
+
 	USkeletalMeshComponent* Component = InAnimInstance->GetOwningComponent();
 	UControlRig* ControlRig = GetControlRig();
 	if (Component && Component->GetSkeletalMeshAsset() && ControlRig)
@@ -66,6 +67,7 @@ void FAnimNode_ControlRigBase::OnInitializeAnimInstance(const FAnimInstanceProxy
 
 		// register skeletalmesh component for now
 		ControlRig->GetDataSourceRegistry()->RegisterDataSource(UControlRig::OwnerComponent, InAnimInstance->GetOwningComponent());
+		UpdateGetAssetUserDataDelegate(ControlRig);
 	}
 }
 
@@ -488,8 +490,9 @@ void FAnimNode_ControlRigBase::ExecuteControlRig(FPoseContext& InOutput)
 			{
 				ControlRig->ClearPoseBeforeBackwardsSolve();
 			}
-			
+
 			// evaluate control rig
+			UpdateGetAssetUserDataDelegate(ControlRig);
 			ControlRig->Evaluate_AnyThread();
 
 #if ENABLE_ANIM_DEBUG 
@@ -572,6 +575,7 @@ void FAnimNode_ControlRigBase::CacheBones_AnyThread(const FAnimationCacheBonesCo
 		if(ControlRig->IsConstructionModeEnabled() ||
 			(ControlRig->IsConstructionRequired() && (bControlRigRequiresInitialization || bIsLODChange)))
 		{
+			UpdateGetAssetUserDataDelegate(ControlRig);
 			ControlRig->Execute(FRigUnit_PrepareForExecution::EventName);
 			bControlRigRequiresInitialization = false;
 		}
@@ -822,4 +826,30 @@ void FAnimNode_ControlRigBase::QueueControlRigDrawInstructions(UControlRig* Cont
 	}
 }
 
+void FAnimNode_ControlRigBase::UpdateGetAssetUserDataDelegate(UControlRig* InControlRig) const
+{
+	if(GetAssetUserData().IsEmpty() || !WeakAnimInstanceObject.IsValid())
+	{
+		InControlRig->GetExternalAssetUserDataDelegate.Unbind();
+		return;
+	}
+	
+	// due to the re-instancing of the anim nodes we have to set this up for every run
+	// since the delegate may go stale quickly. to guard against destroyed anim nodes
+	// we'll rely on the anim instance to provide an indication if the memory is still valid. 
+	TWeakObjectPtr<const UAnimInstance> LocalWeakAnimInstance = WeakAnimInstanceObject;
+	InControlRig->GetExternalAssetUserDataDelegate = UControlRig::FGetExternalAssetUserData::CreateLambda([InControlRig, LocalWeakAnimInstance, this]
+	{
+		if(LocalWeakAnimInstance.IsValid())
+		{
+			return this->GetAssetUserData();
+		}
+		if(IsValid(InControlRig))
+		{
+			InControlRig->GetExternalAssetUserDataDelegate.Unbind();
+		}
+		return TArray<TObjectPtr<UAssetUserData>>();
+	});
+}
+
 
diff --git a/Engine/Plugins/Animation/ControlRig/Source/ControlRig/Private/ControlRig.cpp b/Engine/Plugins/Animation/ControlRig/Source/ControlRig/Private/ControlRig.cpp
--- a/Engine/Plugins/Animation/ControlRig/Source/ControlRig/Private/ControlRig.cpp
+++ b/Engine/Plugins/Animation/ControlRig/Source/ControlRig/Private/ControlRig.cpp
@@ -35,6 +35,7 @@
 #include "Constraints/ControlRigTransformableHandle.h"
 #include "RigVMCore/RigVMNativized.h"
 #include "UObject/UObjectIterator.h"
+#include "RigVMCore/RigVMAssetUserData.h"
 
 #include UE_INLINE_GENERATED_CPP_BY_NAME(ControlRig)
 
@@ -485,6 +486,8 @@ void UControlRig::InitializeFromCDO()
 		{
 			HierarchySettings.ProceduralElementLimit += CDOHierarchy->Num();
 		}
+
+		ExternalVariableDataAssetLinks.Reset();
 	}
 }
 
@@ -2608,6 +2611,50 @@ const TArray<UAssetUserData*>* UControlRig::GetAssetUserDataArray() const
 	{
 		CombinedAssetUserData.Append(AssetUserData);
 	}
+	
+	if(GetExternalAssetUserDataDelegate.IsBound())
+	{
+		CombinedAssetUserData.Append(GetExternalAssetUserDataDelegate.Execute());
+	}
+
+	// find the data assets on the external variable list
+	TArray<FRigVMExternalVariable> ExternalVariables = GetExternalVariablesImpl(false);
+	for(const FRigVMExternalVariable& ExternalVariable : ExternalVariables)
+	{
+		if(ExternalVariable.Memory != nullptr)
+		{
+			if(const UClass* Class = Cast<UClass>(ExternalVariable.TypeObject))
+			{
+				if(Class->IsChildOf(UDataAsset::StaticClass()))
+				{
+					TObjectPtr<UDataAsset>& DataAsset = *(TObjectPtr<UDataAsset>*)ExternalVariable.Memory;
+					if(IsValid(DataAsset) && !DataAsset->GetFName().IsNone())
+					{
+						if(const TObjectPtr<UDataAssetLink>* ExistingDataAssetLink =
+							ExternalVariableDataAssetLinks.Find(ExternalVariable.Name))
+						{
+							if(!IsValid(*ExistingDataAssetLink) ||
+								(*ExistingDataAssetLink)->HasAnyFlags(EObjectFlags::RF_BeginDestroyed))
+							{
+								ExternalVariableDataAssetLinks.Remove(ExternalVariable.Name);
+							}
+						}
+						if(!ExternalVariableDataAssetLinks.Contains(ExternalVariable.Name))
+						{
+							ExternalVariableDataAssetLinks.Add(
+								ExternalVariable.Name,
+								NewObject<UDataAssetLink>((UObject*)this));
+						}
+
+						TObjectPtr<UDataAssetLink>& DataAssetLink = ExternalVariableDataAssetLinks.FindChecked(ExternalVariable.Name);
+						DataAssetLink->NameSpace = ExternalVariable.Name.ToString();
+						DataAssetLink->SetDataAsset(DataAsset);
+						CombinedAssetUserData.Add(DataAssetLink);
+					}
+				}
+			}
+		}
+	}
 
 	if(OuterSceneComponent.IsValid())
 	{
diff --git a/Engine/Plugins/Animation/ControlRig/Source/ControlRig/Public/AnimNode_ControlRigBase.h b/Engine/Plugins/Animation/ControlRig/Source/ControlRig/Public/AnimNode_ControlRigBase.h
--- a/Engine/Plugins/Animation/ControlRig/Source/ControlRig/Public/AnimNode_ControlRigBase.h
+++ b/Engine/Plugins/Animation/ControlRig/Source/ControlRig/Public/AnimNode_ControlRigBase.h
@@ -4,9 +4,9 @@
 
 #include "Animation/AnimBulkCurves.h"
 #include "Animation/AnimNode_CustomProperty.h"
+#include "ControlRig.h"
 #include "AnimNode_ControlRigBase.generated.h"
 
-class UControlRig;
 class UNodeMappingContainer;
 
 /** Struct defining the settings to override when driving a control rig */
@@ -127,6 +127,9 @@ protected:
 	UPROPERTY(EditAnywhere, AdvancedDisplay, Category = Settings)
 	TArray<FBoneReference> OutputBonesToTransfer;
 
+	UPROPERTY(EditAnywhere, AdvancedDisplay, Category = Settings)
+	TArray<TObjectPtr<UAssetUserData>> AssetUserData;
+
 	/** Complete mapping from skeleton to control rig bone index */
 	TArray<TPair<uint16, uint16>> ControlRigBoneInputMappingByIndex;
 	TArray<TPair<uint16, uint16>> ControlRigBoneOutputMappingByIndex;
@@ -175,9 +178,13 @@ protected:
 	void ExecuteControlRig(FPoseContext& InOutput);
 
 	void QueueControlRigDrawInstructions(UControlRig* ControlRig, FAnimInstanceProxy* Proxy) const;
-	
+
+	TArray<TObjectPtr<UAssetUserData>> GetAssetUserData() const { return AssetUserData; }
+	void UpdateGetAssetUserDataDelegate(UControlRig* InControlRig) const;
+
 	bool bControlRigRequiresInitialization;
 	uint16 LastBonesSerialNumberForCacheBones;
+	TWeakObjectPtr<const UAnimInstance> WeakAnimInstanceObject;
 
 	friend struct FControlRigSequencerAnimInstanceProxy;
 	friend struct FControlRigLayerInstanceProxy;
diff --git a/Engine/Plugins/Animation/ControlRig/Source/ControlRig/Public/ControlRig.h b/Engine/Plugins/Animation/ControlRig/Source/ControlRig/Public/ControlRig.h
--- a/Engine/Plugins/Animation/ControlRig/Source/ControlRig/Public/ControlRig.h
+++ b/Engine/Plugins/Animation/ControlRig/Source/ControlRig/Public/ControlRig.h
@@ -528,6 +528,12 @@ public:
 protected:
 	mutable TArray<TObjectPtr<UAssetUserData>> CombinedAssetUserData;
 
+	UPROPERTY(Transient, DuplicateTransient)
+	mutable TMap<FName, TObjectPtr<UDataAssetLink>> ExternalVariableDataAssetLinks;
+
+	DECLARE_DELEGATE_RetVal(TArray<TObjectPtr<UAssetUserData>>, FGetExternalAssetUserData);
+	FGetExternalAssetUserData GetExternalAssetUserDataDelegate;
+
 private:
 
 	void CopyPoseFromOtherRig(UControlRig* Subject);
@@ -757,6 +763,7 @@ private:
 	friend class FControlRigInteractionScope;
 	friend class UControlRigValidator;
 	friend struct FAnimNode_ControlRig;
+	friend struct FAnimNode_ControlRigBase;
 	friend class URigHierarchy;
 	friend class UFKControlRig;
 	friend class UControlRigGraph;
