ControlRig: Allow GetUserData node in all events

diff --git a/Engine/Plugins/Animation/ControlRig/Source/ControlRig/Private/Units/Core/RigUnit_UserData.cpp b/Engine/Plugins/Animation/ControlRig/Source/ControlRig/Private/Units/Core/RigUnit_UserData.cpp
--- a/Engine/Plugins/Animation/ControlRig/Source/ControlRig/Private/Units/Core/RigUnit_UserData.cpp
+++ b/Engine/Plugins/Animation/ControlRig/Source/ControlRig/Private/Units/Core/RigUnit_UserData.cpp
@@ -79,9 +79,8 @@ void FRigDispatch_GetUserData::Execute(FRigVMExtendedExecuteContext& InContext,
 	// allow to run only in construction event
 	if(ControlRigContext.GetEventName() != FRigUnit_PrepareForExecution::EventName)
 	{
-		static constexpr TCHAR ErrorMessageFormat[] = TEXT("Node can only run in %s Event");
-		ControlRigContext.Logf(EMessageSeverity::Error, ErrorMessageFormat, *FRigUnit_PrepareForExecution::EventName.ToString());
-		return;
+		static constexpr TCHAR ErrorMessageFormat[] = TEXT("It's recommended to use this node in the %s event. When using it in other events please consider caching the results.");
+		ControlRigContext.Logf(EMessageSeverity::Info, ErrorMessageFormat, *FRigUnit_PrepareForExecution::EventName.ToString());
 	}
 
 	const FProperty* PropertyDefault = Handles[2].GetResolvedProperty(); 
@@ -188,9 +187,8 @@ FRigUnit_SetupShapeLibraryFromUserData_Execute()
 	// allow to run only in construction event
 	if(ExecuteContext.GetEventName() != FRigUnit_PrepareForExecution::EventName)
 	{
-		static constexpr TCHAR ErrorMessageFormat[] = TEXT("Node can only run in %s Event");
-		Log(FString::Printf(ErrorMessageFormat, *FRigUnit_PrepareForExecution::EventName.ToString()));
-		return;
+		static constexpr TCHAR MessageFormat[] = TEXT("It's recommended to use this node in the %s event. When using it in other events please consider caching the results.");
+		ExecuteContext.Logf(EMessageSeverity::Info, MessageFormat, *FRigUnit_PrepareForExecution::EventName.ToString());
 	}
 
 	if(const UNameSpacedUserData* UserDataObject = ExecuteContext.FindUserData(NameSpace))
