Control Rig: Direct manipulation support for MakeAbsolute

diff --git a/Engine/Plugins/Animation/ControlRig/Source/ControlRig/Private/Units/RigUnit.cpp b/Engine/Plugins/Animation/ControlRig/Source/ControlRig/Private/Units/RigUnit.cpp
--- a/Engine/Plugins/Animation/ControlRig/Source/ControlRig/Private/Units/RigUnit.cpp
+++ b/Engine/Plugins/Animation/ControlRig/Source/ControlRig/Private/Units/RigUnit.cpp
@@ -1,6 +1,8 @@
 // Copyright Epic Games, Inc. All Rights Reserved.
 
 #include "Units/RigUnit.h"
+#include "RigVMFunctions/Math/RigVMFunction_MathTransform.h"
+#include "RigVMFunctions/Math/RigVMFunction_MathVector.h"
 
 #include UE_INLINE_GENERATED_CPP_BY_NAME(RigUnit)
 
@@ -152,6 +154,46 @@ bool FRigUnit::UpdateHierarchyForDirectManipulation(const URigVMUnitNode* InNode
 		return false;
 	}
 
+	URigHierarchy* Hierarchy = InContext.Hierarchy;
+	if (Hierarchy == nullptr)
+	{
+		return false;
+	}
+
+	// deal with special cases here
+	if(Struct == FRigVMFunction_MathTransformMakeAbsolute::StaticStruct())
+	{
+		if(InInfo->Target.Name.Equals(GET_MEMBER_NAME_STRING_CHECKED(FRigVMFunction_MathTransformMakeAbsolute, Local), ESearchCase::CaseSensitive))
+		{
+			const FRigVMFunction_MathTransformMakeAbsolute* MakeAbsolute =
+				reinterpret_cast<const FRigVMFunction_MathTransformMakeAbsolute*>(InInstance->GetStructMemory());
+			Hierarchy->SetControlOffsetTransform(InInfo->ControlKey, MakeAbsolute->Parent, false);
+			Hierarchy->SetLocalTransform(InInfo->ControlKey, MakeAbsolute->Local, false);
+
+			if(!InInfo->bInitialized)
+			{
+				Hierarchy->SetLocalTransform(InInfo->ControlKey, MakeAbsolute->Local, true);
+			}
+			return true;
+		}
+	}
+	else if(Struct == FRigVMFunction_MathVectorMakeAbsolute::StaticStruct())
+	{
+		if(InInfo->Target.Name.Equals(GET_MEMBER_NAME_STRING_CHECKED(FRigVMFunction_MathVectorMakeAbsolute, Global), ESearchCase::CaseSensitive))
+		{
+			const FRigVMFunction_MathVectorMakeAbsolute* MakeAbsolute =
+				reinterpret_cast<const FRigVMFunction_MathVectorMakeAbsolute*>(InInstance->GetStructMemory());
+			Hierarchy->SetControlOffsetTransform(InInfo->ControlKey, FTransform(MakeAbsolute->Parent), false);
+			Hierarchy->SetLocalTransform(InInfo->ControlKey, FTransform(MakeAbsolute->Local), false);
+
+			if(!InInfo->bInitialized)
+			{
+				Hierarchy->SetLocalTransform(InInfo->ControlKey, FTransform(MakeAbsolute->Local), true);
+			}
+			return true;
+		}
+	}
+
 	TTuple<const FStructProperty*, uint8*> StructPropertyAndTargetMemory =
 		FindStructPropertyAndTargetMemory(InInstance, Struct, InInfo->Target.Name);
 
@@ -162,13 +204,6 @@ bool FRigUnit::UpdateHierarchyForDirectManipulation(const URigVMUnitNode* InNode
 		return false;
 	}
 
-
-	URigHierarchy* Hierarchy = InContext.Hierarchy;
-	if (Hierarchy == nullptr)
-	{
-		return false;
-	}
-
 	FTransform Transform = Hierarchy->GetGlobalTransform(InInfo->ControlKey, false);
 
 	if(StructProperty->Struct == TBaseStructure<FTransform>::Get())
@@ -218,6 +253,36 @@ bool FRigUnit::UpdateDirectManipulationFromHierarchy(const URigVMUnitNode* InNod
 		return false;
 	}
 
+	URigHierarchy* Hierarchy = InContext.Hierarchy;
+	if (Hierarchy == nullptr)
+	{
+		return false;
+	}
+
+	const FTransform Transform = Hierarchy->GetLocalTransform(InInfo->ControlKey, false);
+
+	// deal with special cases here
+	if(Struct == FRigVMFunction_MathTransformMakeAbsolute::StaticStruct())
+	{
+		if(InInfo->Target.Name.Equals(GET_MEMBER_NAME_STRING_CHECKED(FRigVMFunction_MathTransformMakeAbsolute, Local), ESearchCase::CaseSensitive))
+		{
+			FRigVMFunction_MathTransformMakeAbsolute* MakeAbsolute =
+				reinterpret_cast<FRigVMFunction_MathTransformMakeAbsolute*>(InInstance->GetStructMemory());
+			MakeAbsolute->Local = Transform;
+			return true;
+		}
+	}
+	if(Struct == FRigVMFunction_MathVectorMakeAbsolute::StaticStruct())
+	{
+		if(InInfo->Target.Name.Equals(GET_MEMBER_NAME_STRING_CHECKED(FRigVMFunction_MathVectorMakeAbsolute, Local), ESearchCase::CaseSensitive))
+		{
+			FRigVMFunction_MathVectorMakeAbsolute* MakeAbsolute =
+				reinterpret_cast<FRigVMFunction_MathVectorMakeAbsolute*>(InInstance->GetStructMemory());
+			MakeAbsolute->Local = Transform.GetTranslation();
+			return true;
+		}
+	}
+	
 	TTuple<const FStructProperty*, uint8*> StructPropertyAndTargetMemory =
 		FindStructPropertyAndTargetMemory(InInstance, Struct, InInfo->Target.Name);
 
@@ -228,14 +293,6 @@ bool FRigUnit::UpdateDirectManipulationFromHierarchy(const URigVMUnitNode* InNod
 		return false;
 	}
 
-	URigHierarchy* Hierarchy = InContext.Hierarchy;
-	if (Hierarchy == nullptr)
-	{
-		return false;
-	}
-	
-	const FTransform Transform = Hierarchy->GetLocalTransform(InInfo->ControlKey, false);
-
 	if(StructProperty->Struct == TBaseStructure<FTransform>::Get())
 	{
 		FTransform& Result = *(FTransform*)Memory;
