Control Rig: Control Rig Modules WIP

diff --git a/Engine/Plugins/Animation/ControlRig/Source/ControlRig/Private/ControlRig.cpp b/Engine/Plugins/Animation/ControlRig/Source/ControlRig/Private/ControlRig.cpp
--- a/Engine/Plugins/Animation/ControlRig/Source/ControlRig/Private/ControlRig.cpp
+++ b/Engine/Plugins/Animation/ControlRig/Source/ControlRig/Private/ControlRig.cpp
@@ -459,27 +459,29 @@ void UControlRig::InitializeFromCDO()
 		// where CDO is initialized from BP there,
 		// we initialize all other instances of Control Rig from the CDO here
 		UControlRig* CDO = GetClass()->GetDefaultObject<UControlRig>();
+		URigHierarchy* Hierarchy = GetHierarchy();
 
 		// copy hierarchy
 		{
-			FRigHierarchyValidityBracket ValidityBracketA(GetHierarchy());
+			FRigHierarchyValidityBracket ValidityBracketA(Hierarchy);
 			FRigHierarchyValidityBracket ValidityBracketB(CDO->GetHierarchy());
 			
-			TGuardValue<bool> Guard(GetHierarchy()->GetSuspendNotificationsFlag(), true);
-			GetHierarchy()->CopyHierarchy(CDO->GetHierarchy());
-			GetHierarchy()->ResetPoseToInitial(ERigElementType::All);
+			TGuardValue<bool> Guard(Hierarchy->GetSuspendNotificationsFlag(), true);
+			Hierarchy->CopyHierarchy(CDO->GetHierarchy());
+			Hierarchy->ResetPoseToInitial(ERigElementType::All);
 		}
 
 #if WITH_EDITOR
 		// current hierarchy should always mirror CDO's hierarchy whenever a change of interest happens
-		CDO->GetHierarchy()->RegisterListeningHierarchy(GetHierarchy());
+		CDO->GetHierarchy()->RegisterListeningHierarchy(Hierarchy);
 #endif
 
 		// notify clients that the hierarchy has changed
-		GetHierarchy()->Notify(ERigHierarchyNotification::HierarchyReset, nullptr);
+		Hierarchy->Notify(ERigHierarchyNotification::HierarchyReset, nullptr);
 
 		// copy hierarchy settings
 		HierarchySettings = CDO->HierarchySettings;
+		ElementKeyRedirector = FRigElementKeyRedirector(CDO->ElementKeyRedirector, Hierarchy); 
 		
 		// increment the procedural limit based on the number of elements in the CDO
 		if(const URigHierarchy* CDOHierarchy = CDO->GetHierarchy())
@@ -567,6 +569,7 @@ bool UControlRig::Execute(const FName& InEventName)
 	PublicContext.SetDeltaTime(DeltaTime);
 	PublicContext.SetAbsoluteTime(AbsoluteTime);
 	PublicContext.SetFramesPerSecond(GetCurrentFramesPerSecond());
+
 #if UE_RIGVM_DEBUG_EXECUTION
 	PublicContext.bDebugExecution = bDebugExecutionEnabled;
 #endif
@@ -1313,6 +1316,12 @@ bool UControlRig::Execute_Internal(const FName& InEventName)
 #endif
 		FRigHierarchyExecuteContextBracket HierarchyContextGuard(Hierarchy, &Context);
 
+		// setup the module information
+		FControlRigExecuteContext& PublicContext = Context.GetPublicDataSafe<FControlRigExecuteContext>();
+		TGuardValue<FName> ModuleInstanceNameGuard(PublicContext.ModuleInstanceName, GetModuleInstanceName());
+		TGuardValue<FString> ModuleInstanceNameSpaceGuard(PublicContext.ModuleInstanceNameSpace, GetRigModuleNameSpace());
+		FRigHierarchyRedirectorGuard ElementRedirectorGuard(DynamicHierarchy, GetElementKeyRedirector());
+
 		TArray<TRigVMMemoryStorage*> LocalMemory = VM->GetLocalMemoryArray(Context);
 		const bool bSuccess = VM->Execute(Context, LocalMemory, InEventName) != ERigVMExecuteResult::Failed;
 
@@ -1527,6 +1536,82 @@ void UControlRig::PostLoad()
 	Influences = NewInfluences;
 }
 
+const FRigModuleSettings& UControlRig::GetRigModuleSettings() const
+{
+	if(HasAnyFlags(RF_ClassDefaultObject))
+	{
+		return RigModuleSettings;
+	}
+	if (const UControlRig* CDO = Cast<UControlRig>(GetClass()->GetDefaultObject()))
+	{
+		return CDO->GetRigModuleSettings();
+	}
+	return RigModuleSettings;
+}
+
+bool UControlRig::IsRigModule() const
+{
+	return GetRigModuleSettings().IsValidModule();
+}
+
+bool UControlRig::IsRigModuleInstance() const
+{
+	if(IsRigModule())
+	{
+		return GetParentRig() != nullptr;
+	}
+	return false;
+}
+
+UControlRig* UControlRig::GetParentRig() const
+{
+	return GetTypedOuter<UControlRig>();
+}
+
+FName UControlRig::GetModuleInstanceName() const
+{
+	if(IsRigModuleInstance())
+	{
+		return GetFName();
+	}
+	return NAME_None;
+}
+
+const FString& UControlRig::GetRigModuleNameSpace() const
+{
+	if(IsRigModule())
+	{
+		if(const UControlRig* ParentRig = GetParentRig())
+		{
+			const FString& ParentNameSpace = ParentRig->GetRigModuleNameSpace();
+			static constexpr TCHAR JoinFormat[] = TEXT("%s%s::");
+			RigModuleNameSpace = FString::Printf(JoinFormat, *ParentNameSpace, *GetName());
+			return RigModuleNameSpace;
+		}
+	}
+
+	static const FString EmptyNameSpace;
+	return EmptyNameSpace;
+}
+
+FRigElementKeyRedirector& UControlRig::GetElementKeyRedirector()
+{
+	// if we are an instance on a modular rig, use our local info
+	if(IsRigModuleInstance())
+	{
+		return ElementKeyRedirector;
+	}
+
+	// if we are a rig module ( but not an instance on a rig)
+	if(IsRigModule())
+	{
+		return ElementKeyRedirector;
+	}
+
+	static FRigElementKeyRedirector EmptyRedirector = FRigElementKeyRedirector();
+	return EmptyRedirector;
+}
+
 TArray<FRigControlElement*> UControlRig::AvailableControls() const
 {
 	if(DynamicHierarchy)
@@ -2980,6 +3065,11 @@ void UControlRig::PostInitInstance(URigVMHost* InCDO)
 		}
 	}
 
+	if(UControlRig* CDOControlRig = Cast<UControlRig>(InCDO))
+	{
+		ElementKeyRedirector = FRigElementKeyRedirector(CDOControlRig->ElementKeyRedirector, DynamicHierarchy);
+	}
+
 	RequestInit();
 }
 
diff --git a/Engine/Plugins/Animation/ControlRig/Source/ControlRig/Private/Rigs/RigHierarchy.cpp b/Engine/Plugins/Animation/ControlRig/Source/ControlRig/Private/Rigs/RigHierarchy.cpp
--- a/Engine/Plugins/Animation/ControlRig/Source/ControlRig/Private/Rigs/RigHierarchy.cpp
+++ b/Engine/Plugins/Animation/ControlRig/Source/ControlRig/Private/Rigs/RigHierarchy.cpp
@@ -108,6 +108,7 @@ URigHierarchy::URigHierarchy()
 , bIsInteracting(false)
 , LastInteractedKey()
 , bSuspendNotifications(false)
+, bIsRunningGetIndex(false)
 , HierarchyController(nullptr)
 , bIsControllerAvailable(true)
 , ResetPoseHash(INDEX_NONE)
@@ -124,7 +125,7 @@ URigHierarchy::URigHierarchy()
 #if WITH_EDITOR
 , bRecordTransformsAtRuntime(true)
 #endif
-, ElementRedirector(nullptr)
+, ElementKeyRedirector(nullptr)
 {
 	Reset();
 #if WITH_EDITOR
@@ -959,6 +960,41 @@ bool URigHierarchy::IsProcedural(const FRigBaseElement* InElement) const
 	return InElement->IsProcedural();
 }
 
+TArray<FRigConnectorInfo> URigHierarchy::GetConnectorInfos() const
+{
+	const TArray<FRigElementKey> Keys = GetConnectorKeys(true);
+	TArray<FRigConnectorInfo> Infos;
+	Infos.Reserve(Keys.Num());
+	for(const FRigElementKey& Key : Keys)
+	{
+		const FRigConnectorElement* Connector = FindChecked<FRigConnectorElement>(Key);
+		Infos.Add(Connector->GetConnectorInfo(this));
+	}
+	return Infos;
+}
+
+TArray<FRigElementKey> URigHierarchy::RestoreConnectorsFromInfos(TArray<FRigConnectorInfo> InInfos, bool bSetupUndoRedo)
+{
+	TArray<FRigElementKey> Keys;
+	for(const FRigConnectorInfo& Info : InInfos)
+	{
+		FRigElementKey Key(Info.Name, ERigElementType::Connector);
+
+		if(const FRigConnectorElement* Connector = Find<FRigConnectorElement>(Key))
+		{
+			SetInitialLocalTransform(Connector->Index, Info.LocalTransform, true, bSetupUndoRedo, false);
+			SetConnectorSettings(Key, Info.Settings, bSetupUndoRedo, false, false);
+		}
+		else
+		{
+			Key = GetController()->AddConnector(Info.Name, Info.LocalTransform, false, Info.Settings, bSetupUndoRedo, false);
+		}
+
+		Keys.Add(Key);
+	}
+	return Keys;
+}
+
 TArray<FName> URigHierarchy::GetMetadataNames(FRigElementKey InItem)
 {
 	TArray<FName> Names;
@@ -2408,6 +2444,19 @@ bool URigHierarchy::CanConnect(const FRigConnectionInfo* InConnectionInfo, FStri
 	return true;
 }
 
+const FRigElementKey& URigHierarchy::GetResolvedTarget(const FRigElementKey& InConnectorKey) const
+{
+	if(ElementKeyRedirector)
+	{
+		if(const FCachedRigElement* Target = ElementKeyRedirector->Find(InConnectorKey))
+		{
+			return Target->GetKey();
+		}
+	}
+	static const FRigElementKey InValidKey;
+	return InValidKey;
+}
+
 bool URigHierarchy::Undo()
 {
 #if WITH_EDITOR
diff --git a/Engine/Plugins/Animation/ControlRig/Source/ControlRig/Private/Rigs/RigHierarchyCache.cpp b/Engine/Plugins/Animation/ControlRig/Source/ControlRig/Private/Rigs/RigHierarchyCache.cpp
--- a/Engine/Plugins/Animation/ControlRig/Source/ControlRig/Private/Rigs/RigHierarchyCache.cpp
+++ b/Engine/Plugins/Animation/ControlRig/Source/ControlRig/Private/Rigs/RigHierarchyCache.cpp
@@ -9,7 +9,7 @@ bool FCachedRigElement::UpdateCache(const URigHierarchy* InHierarchy)
 {
 	if(InHierarchy)
 	{
-		if(!IsValid() || InHierarchy->GetTopologyVersion() != ContainerVersion)
+		if(!IsValid() || InHierarchy->GetTopologyVersionHash() != ContainerVersion)
 		{
 			return UpdateCache(GetKey(), InHierarchy);
 		}
@@ -39,14 +39,12 @@ bool FCachedRigElement::UpdateCache(const FRigElementKey& InKey, const URigHiera
 					{
 						Key = KeyToResolve;
 						Element = PreviousElement;
-						ContainerVersion = InHierarchy->GetTopologyVersion();
+						ContainerVersion = InHierarchy->GetTopologyVersionHash();
 						return IsValid();
 					}
 				}
 			}
 
-			Reset();
-
 			int32 Idx = InHierarchy->GetIndex(KeyToResolve);
 			if(Idx != INDEX_NONE)
 			{
@@ -54,8 +52,13 @@ bool FCachedRigElement::UpdateCache(const FRigElementKey& InKey, const URigHiera
 				Index = (uint16)Idx;
 				Element = InHierarchy->Get(Index);
 			}
+			else
+			{
+				Reset();
+				Key = KeyToResolve;
+			}
 
-			ContainerVersion = InHierarchy->GetTopologyVersion();
+			ContainerVersion = InHierarchy->GetTopologyVersionHash();
 		}
 		return IsValid();
 	}
@@ -64,7 +67,7 @@ bool FCachedRigElement::UpdateCache(const FRigElementKey& InKey, const URigHiera
 
 bool FCachedRigElement::IsIdentical(const FRigElementKey& InKey, const URigHierarchy* InHierarchy)
 {
-	return InKey == Key && InHierarchy->GetTopologyVersion() == ContainerVersion;
+	return InKey == Key && InHierarchy->GetTopologyVersionHash() == ContainerVersion;
 }
 
 FRigElementKeyRedirector::FRigElementKeyRedirector(const TMap<FRigElementKey, FRigElementKey>& InMap, const URigHierarchy* InHierarchy)
@@ -72,11 +75,30 @@ FRigElementKeyRedirector::FRigElementKeyRedirector(const TMap<FRigElementKey, FR
 	check(InHierarchy);
 	InternalKeyToExternalKey.Reserve(InMap.Num());
 	ExternalKeys.Reserve(InMap.Num());
+
+	Hash = 0;
 	for(const TPair<FRigElementKey, FRigElementKey>& Pair : InMap)
 	{
 		check(Pair.Key.IsValid());
-		InternalKeyToExternalKey.Add(Pair.Key, FCachedRigElement(Pair.Value, InHierarchy));
+		InternalKeyToExternalKey.Add(Pair.Key, FCachedRigElement(Pair.Value, InHierarchy, true));
 		ExternalKeys.Add(Pair.Value);
+		Hash = HashCombine(Hash, HashCombine(GetTypeHash(Pair.Key), GetTypeHash(Pair.Value)));
+	}
+}
+
+FRigElementKeyRedirector::FRigElementKeyRedirector(const FRigElementKeyRedirector& InOther, const URigHierarchy* InHierarchy)
+{
+	check(InHierarchy);
+	InternalKeyToExternalKey.Reserve(InOther.InternalKeyToExternalKey.Num());
+	ExternalKeys.Reserve(InOther.ExternalKeys.Num());
+
+	Hash = 0;
+	for(const TPair<FRigElementKey, FCachedRigElement>& Pair : InOther.InternalKeyToExternalKey)
+	{
+		check(Pair.Key.IsValid());
+		InternalKeyToExternalKey.Add(Pair.Key, FCachedRigElement(Pair.Value.GetKey(), InHierarchy, true));
+		ExternalKeys.Add(Pair.Value.GetKey());
+		Hash = HashCombine(Hash, HashCombine(GetTypeHash(Pair.Key), GetTypeHash(Pair.Value.GetKey())));
 	}
 }
 
diff --git a/Engine/Plugins/Animation/ControlRig/Source/ControlRig/Private/Rigs/RigHierarchyController.cpp b/Engine/Plugins/Animation/ControlRig/Source/ControlRig/Private/Rigs/RigHierarchyController.cpp
--- a/Engine/Plugins/Animation/ControlRig/Source/ControlRig/Private/Rigs/RigHierarchyController.cpp
+++ b/Engine/Plugins/Animation/ControlRig/Source/ControlRig/Private/Rigs/RigHierarchyController.cpp
@@ -550,7 +550,7 @@ FRigElementKey URigHierarchyController::AddReference(FName InName, FRigElementKe
 	return NewElement->Key;
 }
 
-FRigElementKey URigHierarchyController::AddConnector(FName InName, FRigElementKey InParent, FTransform InTransform,
+FRigElementKey URigHierarchyController::AddConnector(FName InName, FTransform InTransform,
 	bool bTransformInGlobal, FRigConnectorSettings InSettings, bool bSetupUndo, bool bPrintPythonCommand)
 {
 	if(!IsValid())
@@ -573,7 +573,7 @@ FRigElementKey URigHierarchyController::AddConnector(FName InName, FRigElementKe
 		NewElement->Key.Type = ERigElementType::Connector;
 		NewElement->Key.Name = Hierarchy->GetSafeNewName(InName, NewElement->Key.Type);
 		NewElement->Settings = InSettings;
-		AddElement(NewElement, Hierarchy->Get(Hierarchy->GetIndex(InParent)), true);
+		AddElement(NewElement, nullptr, true);
 
 		if(bTransformInGlobal)
 		{
@@ -1681,9 +1681,8 @@ TArray<FString> URigHierarchyController::GetAddConnectorPythonCommands(FRigConne
 	}
 
 	// AddConnector(FName InName, FRigElementKey InParent, FTransform InTransform, bool bTransformInGlobal = true, FRigConnectorSettings InSettings = FRigConnectorSettings(), bool bSetupUndo = false);
-	Commands.Add(FString::Printf(TEXT("hierarchy_controller.add_connector('%s', %s, %s, False, %s)"),
+	Commands.Add(FString::Printf(TEXT("hierarchy_controller.add_connector('%s', %s, False, %s)"),
 		*Connector->GetName(),
-		*ParentKeyStr,
 		*TransformStr,
 		*SettingsStr
 	));
diff --git a/Engine/Plugins/Animation/ControlRig/Source/ControlRig/Private/Rigs/RigHierarchyElements.cpp b/Engine/Plugins/Animation/ControlRig/Source/ControlRig/Private/Rigs/RigHierarchyElements.cpp
--- a/Engine/Plugins/Animation/ControlRig/Source/ControlRig/Private/Rigs/RigHierarchyElements.cpp
+++ b/Engine/Plugins/Animation/ControlRig/Source/ControlRig/Private/Rigs/RigHierarchyElements.cpp
@@ -1666,6 +1666,16 @@ bool FRigConnectorElement::CanConnect(const FRigConnectionInfo* InConnectionInfo
 	return true;
 }
 
+FRigConnectorInfo FRigConnectorElement::GetConnectorInfo(const URigHierarchy* InHierarchy) const
+{
+	FRigConnectorInfo Info;
+	Info.Name = Key.Name;
+	Info.ResolvedTarget = InHierarchy->GetResolvedTarget(Key);
+	Info.LocalTransform = InHierarchy->GetInitialLocalTransform(Key);
+	Info.Settings = Settings;
+	return Info;
+}
+
 void FRigConnectorElement::CopyFrom(URigHierarchy* InHierarchy, FRigBaseElement* InOther, URigHierarchy* InOtherHierarchy)
 {
 	Super::CopyFrom(InHierarchy, InOther, InOtherHierarchy);
diff --git a/Engine/Plugins/Animation/ControlRig/Source/ControlRig/Private/Units/Execution/RigUnit_DynamicHierarchy.cpp b/Engine/Plugins/Animation/ControlRig/Source/ControlRig/Private/Units/Execution/RigUnit_DynamicHierarchy.cpp
--- a/Engine/Plugins/Animation/ControlRig/Source/ControlRig/Private/Units/Execution/RigUnit_DynamicHierarchy.cpp
+++ b/Engine/Plugins/Animation/ControlRig/Source/ControlRig/Private/Units/Execution/RigUnit_DynamicHierarchy.cpp
@@ -383,11 +383,24 @@ FTransform FRigUnit_HierarchyAddControlElement::ProjectOffsetTransform(const FTr
 	return InOffsetTransform;
 }
 
+void FRigUnit_HierarchyAddControl_Settings::ConfigureFrom(const FRigControlElement* InControlElement, const FRigControlSettings& InSettings)
+{
+	DisplayName = InSettings.DisplayName;
+}
+
 void FRigUnit_HierarchyAddControl_Settings::Configure(FRigControlSettings& OutSettings) const
 {
 	OutSettings.DisplayName = DisplayName;
 }
 
+void FRigUnit_HierarchyAddControl_ShapeSettings::ConfigureFrom(const FRigControlElement* InControlElement, const FRigControlSettings& InSettings)
+{
+	bVisible = InSettings.bShapeVisible;
+	Name = InSettings.ShapeName;
+	Color = InSettings.ShapeColor;
+	Transform = InControlElement->Shape.Get(ERigTransformType::InitialLocal);
+}
+
 void FRigUnit_HierarchyAddControl_ShapeSettings::Configure(FRigControlSettings& OutSettings) const
 {
 	OutSettings.bShapeVisible = bVisible;
@@ -395,6 +408,13 @@ void FRigUnit_HierarchyAddControl_ShapeSettings::Configure(FRigControlSettings&
 	OutSettings.ShapeColor = Color;
 }
 
+void FRigUnit_HierarchyAddControl_ProxySettings::ConfigureFrom(const FRigControlElement* InControlElement, const FRigControlSettings& InSettings)
+{
+	bIsProxy = InSettings.AnimationType == ERigControlAnimationType::ProxyControl;
+	DrivenControls = InSettings.DrivenControls;
+	ShapeVisibility = InSettings.ShapeVisibility;
+}
+
 void FRigUnit_HierarchyAddControl_ProxySettings::Configure(FRigControlSettings& OutSettings) const
 {
 	OutSettings.AnimationType = bIsProxy ? ERigControlAnimationType::ProxyControl : ERigControlAnimationType::AnimationControl;
@@ -402,6 +422,14 @@ void FRigUnit_HierarchyAddControl_ProxySettings::Configure(FRigControlSettings&
 	OutSettings.ShapeVisibility = ShapeVisibility;
 }
 
+void FRigUnit_HierarchyAddControlFloat_LimitSettings::ConfigureFrom(const FRigControlElement* InControlElement, const FRigControlSettings& InSettings)
+{
+	Limit = InSettings.LimitEnabled[0];
+	MinValue = InSettings.MinimumValue.Get<float>();
+	MaxValue = InSettings.MaximumValue.Get<float>();
+	bDrawLimits = InSettings.bDrawLimits;
+}
+
 void FRigUnit_HierarchyAddControlFloat_LimitSettings::Configure(FRigControlSettings& OutSettings) const
 {
 	OutSettings.SetupLimitArrayForType(false, false, false);
@@ -411,6 +439,18 @@ void FRigUnit_HierarchyAddControlFloat_LimitSettings::Configure(FRigControlSetti
 	OutSettings.bDrawLimits = bDrawLimits;
 }
 
+void FRigUnit_HierarchyAddControlFloat_Settings::ConfigureFrom(const FRigControlElement* InControlElement, const FRigControlSettings& InSettings)
+{
+	FRigUnit_HierarchyAddControl_Settings::ConfigureFrom(InControlElement, InSettings);
+
+	bIsScale = InSettings.ControlType == ERigControlType::ScaleFloat;
+	PrimaryAxis = InSettings.PrimaryAxis;
+
+	Proxy.ConfigureFrom(InControlElement, InSettings);
+	Limits.ConfigureFrom(InControlElement, InSettings);
+	Shape.ConfigureFrom(InControlElement, InSettings);
+}
+
 void FRigUnit_HierarchyAddControlFloat_Settings::Configure(FRigControlSettings& OutSettings) const
 {
 	Super::Configure(OutSettings);
@@ -452,6 +492,14 @@ FRigUnit_HierarchyAddControlFloat_Execute()
 	}
 }
 
+void FRigUnit_HierarchyAddControlInteger_LimitSettings::ConfigureFrom(const FRigControlElement* InControlElement, const FRigControlSettings& InSettings)
+{
+	Limit = InSettings.LimitEnabled[0];
+	MinValue = InSettings.MinimumValue.Get<int32>();
+	MaxValue = InSettings.MaximumValue.Get<int32>();
+	bDrawLimits = InSettings.bDrawLimits;
+}
+
 void FRigUnit_HierarchyAddControlInteger_LimitSettings::Configure(FRigControlSettings& OutSettings) const
 {
 	OutSettings.SetupLimitArrayForType(false, false, false);
@@ -461,6 +509,17 @@ void FRigUnit_HierarchyAddControlInteger_LimitSettings::Configure(FRigControlSet
 	OutSettings.bDrawLimits = bDrawLimits;
 }
 
+void FRigUnit_HierarchyAddControlInteger_Settings::ConfigureFrom(const FRigControlElement* InControlElement, const FRigControlSettings& InSettings)
+{
+	FRigUnit_HierarchyAddControl_Settings::ConfigureFrom(InControlElement, InSettings);
+
+	PrimaryAxis = InSettings.PrimaryAxis;
+
+	Proxy.ConfigureFrom(InControlElement, InSettings);
+	Limits.ConfigureFrom(InControlElement, InSettings);
+	Shape.ConfigureFrom(InControlElement, InSettings);
+}
+
 void FRigUnit_HierarchyAddControlInteger_Settings::Configure(FRigControlSettings& OutSettings) const
 {
 	Super::Configure(OutSettings);
@@ -502,6 +561,15 @@ FRigUnit_HierarchyAddControlInteger_Execute()
 	}
 }
 
+void FRigUnit_HierarchyAddControlVector2D_LimitSettings::ConfigureFrom(const FRigControlElement* InControlElement, const FRigControlSettings& InSettings)
+{
+	LimitX = InSettings.LimitEnabled[0];
+	LimitY = InSettings.LimitEnabled[1];
+	MinValue = FVector2D(InSettings.MinimumValue.Get<FVector3f>().X, InSettings.MinimumValue.Get<FVector3f>().Y);
+	MaxValue = FVector2D(InSettings.MaximumValue.Get<FVector3f>().X, InSettings.MaximumValue.Get<FVector3f>().Y);
+	bDrawLimits = InSettings.bDrawLimits;
+}
+
 void FRigUnit_HierarchyAddControlVector2D_LimitSettings::Configure(FRigControlSettings& OutSettings) const
 {
 	OutSettings.SetupLimitArrayForType(false, false, false);
@@ -512,6 +580,18 @@ void FRigUnit_HierarchyAddControlVector2D_LimitSettings::Configure(FRigControlSe
 	OutSettings.bDrawLimits = bDrawLimits;
 }
 
+void FRigUnit_HierarchyAddControlVector2D_Settings::ConfigureFrom(const FRigControlElement* InControlElement, const FRigControlSettings& InSettings)
+{
+	FRigUnit_HierarchyAddControl_Settings::ConfigureFrom(InControlElement, InSettings);
+
+	PrimaryAxis = InSettings.PrimaryAxis;
+	FilteredChannels = InSettings.FilteredChannels;
+
+	Proxy.ConfigureFrom(InControlElement, InSettings);
+	Limits.ConfigureFrom(InControlElement, InSettings);
+	Shape.ConfigureFrom(InControlElement, InSettings);
+}
+
 void FRigUnit_HierarchyAddControlVector2D_Settings::Configure(FRigControlSettings& OutSettings) const
 {
 	Super::Configure(OutSettings);
@@ -554,6 +634,16 @@ FRigUnit_HierarchyAddControlVector2D_Execute()
 	}
 }
 
+void FRigUnit_HierarchyAddControlVector_LimitSettings::ConfigureFrom(const FRigControlElement* InControlElement, const FRigControlSettings& InSettings)
+{
+	LimitX = InSettings.LimitEnabled[0];
+	LimitY = InSettings.LimitEnabled[1];
+	LimitZ = InSettings.LimitEnabled[2];
+	MinValue = FVector(InSettings.MinimumValue.Get<FVector3f>());
+	MaxValue = FVector(InSettings.MaximumValue.Get<FVector3f>());
+	bDrawLimits = InSettings.bDrawLimits;
+}
+
 void FRigUnit_HierarchyAddControlVector_LimitSettings::Configure(FRigControlSettings& OutSettings) const
 {
 	OutSettings.SetupLimitArrayForType(false, false, false);
@@ -565,6 +655,18 @@ void FRigUnit_HierarchyAddControlVector_LimitSettings::Configure(FRigControlSett
 	OutSettings.bDrawLimits = bDrawLimits;
 }
 
+void FRigUnit_HierarchyAddControlVector_Settings::ConfigureFrom(const FRigControlElement* InControlElement, const FRigControlSettings& InSettings)
+{
+	FRigUnit_HierarchyAddControl_Settings::ConfigureFrom(InControlElement, InSettings);
+
+	bIsPosition = InSettings.ControlType == ERigControlType::Position;
+	FilteredChannels = InSettings.FilteredChannels;
+
+	Proxy.ConfigureFrom(InControlElement, InSettings);
+	Limits.ConfigureFrom(InControlElement, InSettings);
+	Shape.ConfigureFrom(InControlElement, InSettings);
+}
+
 void FRigUnit_HierarchyAddControlVector_Settings::Configure(FRigControlSettings& OutSettings) const
 {
 	Super::Configure(OutSettings);
@@ -622,6 +724,16 @@ FRigUnit_HierarchyAddControlVector_Execute()
 	}
 }
 
+void FRigUnit_HierarchyAddControlRotator_LimitSettings::ConfigureFrom(const FRigControlElement* InControlElement, const FRigControlSettings& InSettings)
+{
+	LimitPitch = InSettings.LimitEnabled[0];
+	LimitYaw = InSettings.LimitEnabled[1];
+	LimitRoll = InSettings.LimitEnabled[2];
+	MinValue = FRotator::MakeFromEuler(FVector(InSettings.MinimumValue.Get<FVector3f>()));
+	MaxValue = FRotator::MakeFromEuler(FVector(InSettings.MaximumValue.Get<FVector3f>()));
+	bDrawLimits = InSettings.bDrawLimits;
+}
+
 void FRigUnit_HierarchyAddControlRotator_LimitSettings::Configure(FRigControlSettings& OutSettings) const
 {
 	OutSettings.SetupLimitArrayForType(false, false, false);
@@ -633,6 +745,17 @@ void FRigUnit_HierarchyAddControlRotator_LimitSettings::Configure(FRigControlSet
 	OutSettings.bDrawLimits = bDrawLimits;
 }
 
+void FRigUnit_HierarchyAddControlRotator_Settings::ConfigureFrom(const FRigControlElement* InControlElement, const FRigControlSettings& InSettings)
+{
+	FRigUnit_HierarchyAddControl_Settings::ConfigureFrom(InControlElement, InSettings);
+
+	FilteredChannels = InSettings.FilteredChannels;
+
+	Proxy.ConfigureFrom(InControlElement, InSettings);
+	Limits.ConfigureFrom(InControlElement, InSettings);
+	Shape.ConfigureFrom(InControlElement, InSettings);
+}
+
 void FRigUnit_HierarchyAddControlRotator_Settings::Configure(FRigControlSettings& OutSettings) const
 {
 	Super::Configure(OutSettings);
@@ -684,6 +807,18 @@ FRigUnit_HierarchyAddControlRotator_Execute()
 	}
 }
 
+void FRigUnit_HierarchyAddControlTransform_Settings::ConfigureFrom(const FRigControlElement* InControlElement, const FRigControlSettings& InSettings)
+{
+	FRigUnit_HierarchyAddControl_Settings::ConfigureFrom(InControlElement, InSettings);
+
+	FilteredChannels = InSettings.FilteredChannels;
+	bUsePreferredRotationOrder = InSettings.bUsePreferredRotationOrder;
+	PreferredRotationOrder = InSettings.PreferredRotationOrder;
+
+	Proxy.ConfigureFrom(InControlElement, InSettings);
+	Shape.ConfigureFrom(InControlElement, InSettings);
+}
+
 void FRigUnit_HierarchyAddControlTransform_Settings::Configure(FRigControlSettings& OutSettings) const
 {
 	Super::Configure(OutSettings);
diff --git a/Engine/Plugins/Animation/ControlRig/Source/ControlRig/Private/Units/Execution/RigUnit_DynamicHierarchy.h b/Engine/Plugins/Animation/ControlRig/Source/ControlRig/Private/Units/Execution/RigUnit_DynamicHierarchy.h
--- a/Engine/Plugins/Animation/ControlRig/Source/ControlRig/Private/Units/Execution/RigUnit_DynamicHierarchy.h
+++ b/Engine/Plugins/Animation/ControlRig/Source/ControlRig/Private/Units/Execution/RigUnit_DynamicHierarchy.h
@@ -423,6 +423,7 @@ struct CONTROLRIG_API FRigUnit_HierarchyAddControl_Settings
 	{}
 	virtual ~FRigUnit_HierarchyAddControl_Settings(){}
 
+	virtual void ConfigureFrom(const FRigControlElement* InControlElement, const FRigControlSettings& InSettings);
 	virtual void Configure(FRigControlSettings& OutSettings) const;
 
 	UPROPERTY(EditAnywhere, BlueprintReadWrite, Category = "Settings")
@@ -441,6 +442,7 @@ struct CONTROLRIG_API FRigUnit_HierarchyAddControl_ShapeSettings
 		, Transform(FTransform::Identity)
 	{}
 
+	void ConfigureFrom(const FRigControlElement* InControlElement, const FRigControlSettings& InSettings);
 	void Configure(FRigControlSettings& OutSettings) const;
 
 	UPROPERTY(EditAnywhere, BlueprintReadWrite, Category = "Settings")
@@ -466,6 +468,7 @@ struct CONTROLRIG_API FRigUnit_HierarchyAddControl_ProxySettings
 		, ShapeVisibility(ERigControlVisibility::BasedOnSelection)
 	{}
 
+	void ConfigureFrom(const FRigControlElement* InControlElement, const FRigControlSettings& InSettings);
 	void Configure(FRigControlSettings& OutSettings) const;
 
 	UPROPERTY(EditAnywhere, BlueprintReadWrite, Category = "Settings")
@@ -490,6 +493,7 @@ struct CONTROLRIG_API FRigUnit_HierarchyAddControlFloat_LimitSettings
 		, bDrawLimits(true)
 	{}
 
+	void ConfigureFrom(const FRigControlElement* InControlElement, const FRigControlSettings& InSettings);
 	void Configure(FRigControlSettings& OutSettings) const;
 
 	UPROPERTY(EditAnywhere, BlueprintReadWrite, Category = "Settings")
@@ -517,6 +521,7 @@ struct CONTROLRIG_API FRigUnit_HierarchyAddControlFloat_Settings : public FRigUn
 	{}
 	virtual ~FRigUnit_HierarchyAddControlFloat_Settings() override {}
 
+	virtual void ConfigureFrom(const FRigControlElement* InControlElement, const FRigControlSettings& InSettings) override;
 	virtual void Configure(FRigControlSettings& OutSettings) const override;
 
 	UPROPERTY(EditAnywhere, BlueprintReadWrite, Category = "Settings")
@@ -615,6 +620,7 @@ struct CONTROLRIG_API FRigUnit_HierarchyAddControlInteger_LimitSettings
 		, bDrawLimits(true)
 	{}
 
+	void ConfigureFrom(const FRigControlElement* InControlElement, const FRigControlSettings& InSettings);
 	void Configure(FRigControlSettings& OutSettings) const;
 
 	UPROPERTY(EditAnywhere, BlueprintReadWrite, Category = "Settings")
@@ -641,6 +647,7 @@ struct CONTROLRIG_API FRigUnit_HierarchyAddControlInteger_Settings : public FRig
 	{}
 	virtual ~FRigUnit_HierarchyAddControlInteger_Settings() override {}
 
+	virtual void ConfigureFrom(const FRigControlElement* InControlElement, const FRigControlSettings& InSettings) override;
 	virtual void Configure(FRigControlSettings& OutSettings) const override;
 
 	UPROPERTY(EditAnywhere, BlueprintReadWrite, Category = "Settings")
@@ -701,6 +708,7 @@ struct CONTROLRIG_API FRigUnit_HierarchyAddControlVector2D_LimitSettings
 		, bDrawLimits(true)
 	{}
 
+	void ConfigureFrom(const FRigControlElement* InControlElement, const FRigControlSettings& InSettings);
 	void Configure(FRigControlSettings& OutSettings) const;
 
 	UPROPERTY(EditAnywhere, BlueprintReadWrite, Category = "Settings")
@@ -732,6 +740,7 @@ struct CONTROLRIG_API FRigUnit_HierarchyAddControlVector2D_Settings : public FRi
 	}
 	virtual ~FRigUnit_HierarchyAddControlVector2D_Settings() override {}
 
+	virtual void ConfigureFrom(const FRigControlElement* InControlElement, const FRigControlSettings& InSettings) override;
 	virtual void Configure(FRigControlSettings& OutSettings) const override;
 
 	UPROPERTY(EditAnywhere, BlueprintReadWrite, Category = "Settings")
@@ -796,6 +805,7 @@ struct CONTROLRIG_API FRigUnit_HierarchyAddControlVector_LimitSettings
 		, bDrawLimits(true)
 	{}
 
+	void ConfigureFrom(const FRigControlElement* InControlElement, const FRigControlSettings& InSettings);
 	void Configure(FRigControlSettings& OutSettings) const;
 
 	UPROPERTY(EditAnywhere, BlueprintReadWrite, Category = "Settings")
@@ -831,6 +841,7 @@ struct CONTROLRIG_API FRigUnit_HierarchyAddControlVector_Settings : public FRigU
 	}
 	virtual ~FRigUnit_HierarchyAddControlVector_Settings() override {}
 
+	virtual void ConfigureFrom(const FRigControlElement* InControlElement, const FRigControlSettings& InSettings) override;
 	virtual void Configure(FRigControlSettings& OutSettings) const override;
 
 	UPROPERTY(EditAnywhere, BlueprintReadWrite, Category = "Settings")
@@ -901,6 +912,7 @@ struct CONTROLRIG_API FRigUnit_HierarchyAddControlRotator_LimitSettings
 		, bDrawLimits(true)
 	{}
 
+	void ConfigureFrom(const FRigControlElement* InControlElement, const FRigControlSettings& InSettings);
 	void Configure(FRigControlSettings& OutSettings) const;
 
 	UPROPERTY(EditAnywhere, BlueprintReadWrite, Category = "Settings")
@@ -935,6 +947,7 @@ struct CONTROLRIG_API FRigUnit_HierarchyAddControlRotator_Settings : public FRig
 	}
 	virtual ~FRigUnit_HierarchyAddControlRotator_Settings() override {}
 
+	virtual void ConfigureFrom(const FRigControlElement* InControlElement, const FRigControlSettings& InSettings) override;
 	virtual void Configure(FRigControlSettings& OutSettings) const override;
 
 	UPROPERTY(EditAnywhere, BlueprintReadWrite, Category = "Settings")
@@ -1000,6 +1013,7 @@ struct CONTROLRIG_API FRigUnit_HierarchyAddControlTransform_Settings : public FR
 	}
 	virtual ~FRigUnit_HierarchyAddControlTransform_Settings() override {}
 
+	virtual void ConfigureFrom(const FRigControlElement* InControlElement, const FRigControlSettings& InSettings) override;
 	virtual void Configure(FRigControlSettings& OutSettings) const override;
 
 	UPROPERTY(EditAnywhere, BlueprintReadWrite, Category = "Settings")
diff --git a/Engine/Plugins/Animation/ControlRig/Source/ControlRig/Private/Units/RigUnitContext.cpp b/Engine/Plugins/Animation/ControlRig/Source/ControlRig/Private/Units/RigUnitContext.cpp
new file mode 100644
--- /dev/null
+++ b/Engine/Plugins/Animation/ControlRig/Source/ControlRig/Private/Units/RigUnitContext.cpp
@@ -0,0 +1,48 @@
+// Copyright Epic Games, Inc. All Rights Reserved.
+
+#include "Units/RigUnitContext.h"
+
+#include UE_INLINE_GENERATED_CPP_BY_NAME(RigUnitContext)
+
+FName FControlRigExecuteContext::AddRigModuleNameSpace(const FName& InName) const
+{
+	if(ModuleInstanceName.IsNone())
+	{
+		return InName;
+	}
+	return *AddRigModuleNameSpace(InName.ToString());
+}
+
+FString FControlRigExecuteContext::AddRigModuleNameSpace(const FString& InName) const
+{
+	if(ModuleInstanceName.IsNone())
+	{
+		return InName;
+	}
+	check(!ModuleInstanceNameSpace.IsEmpty());
+	return ModuleInstanceNameSpace + InName;
+}
+
+FName FControlRigExecuteContext::RemoveRigModuleNameSpace(const FName& InName) const
+{
+	if(ModuleInstanceName.IsNone())
+	{
+		return InName;
+	}
+	return *RemoveRigModuleNameSpace(InName.ToString());
+}
+
+FString FControlRigExecuteContext::RemoveRigModuleNameSpace(const FString& InName) const
+{
+	if(ModuleInstanceName.IsNone())
+	{
+		return InName;
+	}
+	check(!ModuleInstanceNameSpace.IsEmpty());
+
+	if(InName.StartsWith(ModuleInstanceNameSpace, ESearchCase::CaseSensitive))
+	{
+		return InName.Mid(ModuleInstanceNameSpace.Len());
+	}
+	return InName;
+}
diff --git a/Engine/Plugins/Animation/ControlRig/Source/ControlRig/Public/ControlRig.h b/Engine/Plugins/Animation/ControlRig/Source/ControlRig/Public/ControlRig.h
--- a/Engine/Plugins/Animation/ControlRig/Source/ControlRig/Public/ControlRig.h
+++ b/Engine/Plugins/Animation/ControlRig/Source/ControlRig/Public/ControlRig.h
@@ -18,6 +18,7 @@
 #include "Animation/AnimInstance.h"
 #include "Animation/AnimInstanceProxy.h"
 #include "Animation/AttributesRuntime.h"
+#include "Rigs/RigModuleDefines.h"
 
 #if WITH_EDITOR
 #include "RigVMModel/RigVMPin.h"
@@ -84,6 +85,28 @@ public:
 	virtual void PostLoad() override;
 	virtual UScriptStruct* GetPublicContextStruct() const override { return FControlRigExecuteContext::StaticStruct(); }
 
+	// Returns the settings of the module this instance belongs to
+	const FRigModuleSettings& GetRigModuleSettings() const;
+
+	// Returns true if the rig is defined as a rig module
+	bool IsRigModule() const;
+
+	// Returns true if this rig is an instance module. Rigs may be a module but not instance
+	// when being interacted with the asset editor
+	bool IsRigModuleInstance() const;
+
+	// Returns the parent rig hosting this module instance
+	UControlRig* GetParentRig() const;
+
+	// Returns the name of this module instance
+	FName GetModuleInstanceName() const;
+
+	// Returns the namespace of this module (for example ArmModule::)
+	const FString& GetRigModuleNameSpace() const;
+
+	// Returns the redirector from key to key for this rig
+	FRigElementKeyRedirector& GetElementKeyRedirector();
+
 	/** Creates a transformable control handle for the specified control to be used by the constraints system. Should use the UObject from 
 	ConstraintsScriptingLibrary::GetManager(UWorld* InWorld)*/
 	UFUNCTION(BlueprintCallable, Category = "Control Rig | Constraints")
@@ -482,6 +505,8 @@ private:
 	TSubclassOf<UControlRig> InteractionRigClass_DEPRECATED;
 #endif
 
+	FRigElementKeyRedirector ElementKeyRedirector;
+
 public:
 
 	UE_DEPRECATED(5.4, "InteractionRig is no longer used")
@@ -673,6 +698,12 @@ private:
 		ERigTransformType::Type TransformType;
 	};
 
+	UPROPERTY()
+	FRigModuleSettings RigModuleSettings;
+
+	UPROPERTY(transient)
+	mutable FString RigModuleNameSpace;
+
 #if WITH_EDITOR
 
 	// Class used to temporarily cache current transient controls
diff --git a/Engine/Plugins/Animation/ControlRig/Source/ControlRig/Public/Rigs/RigHierarchy.h b/Engine/Plugins/Animation/ControlRig/Source/ControlRig/Public/Rigs/RigHierarchy.h
--- a/Engine/Plugins/Animation/ControlRig/Source/ControlRig/Public/Rigs/RigHierarchy.h
+++ b/Engine/Plugins/Animation/ControlRig/Source/ControlRig/Public/Rigs/RigHierarchy.h
@@ -198,6 +198,11 @@ public:
 	UFUNCTION(BlueprintCallable, Category = URigHierarchy)
 	void CopyHierarchy(URigHierarchy* InHierarchy);
 
+	/**
+	 * Returns true if the hierarchy currently has an execute context / the rig is running
+	 */
+	bool HasExecuteContext() const { return ExecuteContext != nullptr; }
+
 	/**
 	 * Returns a hash for the hierarchy representing all names
 	 * as well as the topology version.
@@ -397,17 +402,18 @@ public:
 	 */
 	int32 GetIndex(const FRigElementKey& InKey) const
 	{
-		if(ElementRedirector)
+		if(ElementKeyRedirector)
 		{
-			if(FCachedRigElement* CachedRigElement = ElementRedirector->Find(InKey))
+			if(FCachedRigElement* CachedRigElement = ElementKeyRedirector->Find(InKey))
 			{
+				TGuardValue<bool> GuardRecursion(bIsRunningGetIndex, true);
 				if(CachedRigElement->UpdateCache(this))
 				{
 					return CachedRigElement->GetIndex();
 				}
 				return INDEX_NONE;
 			}
-			if(ElementRedirector->ContainsExternalKey(InKey))
+			if(!bIsRunningGetIndex && ElementKeyRedirector->ContainsExternalKey(InKey))
 			{
 				return INDEX_NONE;
 			}
@@ -967,6 +973,19 @@ public:
 	{
 		return GetKeysOfType<FRigConnectorElement>(bTraverse);
 	}
+
+	/**
+	 * Returns all of the connectors' infos
+	 */
+	UFUNCTION(BlueprintPure, Category = URigHierarchy)
+	TArray<FRigConnectorInfo> GetConnectorInfos() const;
+
+	/**
+	 * Try to restore the connectors from the info structs
+	 */
+	UFUNCTION(BlueprintCallable, Category = URigHierarchy)
+	TArray<FRigElementKey> RestoreConnectorsFromInfos(TArray<FRigConnectorInfo> InInfos, bool bSetupUndoRedo = false);
+	
 	/**
 	 * Returns all root elements
 	 */
@@ -3280,6 +3299,11 @@ public:
 	}
 	bool CanConnect(const FRigConnectionInfo* InConnectionInfo, FString* OutFailureReason = nullptr) const;
 
+	/**
+	 * Returns the currently resolved target for given connector key
+	 */
+	const FRigElementKey& GetResolvedTarget(const FRigElementKey& InConnectorKey) const;
+
 	/**
 	 * Performs undo for one transform change
 	 */
@@ -3362,6 +3386,21 @@ public:
 	 */
 	uint16 GetTopologyVersion() const { return TopologyVersion; }
 
+	/**
+	 * Returns the hash of this hierarchy used for cached element keys
+	 */
+	uint32 GetTopologyVersionHash() const
+	{
+		const uint32 Hash = HashCombine(
+			(uint32)reinterpret_cast<long long>(this),
+			GetTypeHash(TopologyVersion));
+		if(ElementKeyRedirector)
+		{
+			return HashCombine(Hash, ElementKeyRedirector->GetHash());
+		}
+		return Hash;
+	}
+
 	/**
 	 * Increments the topology version
 	 */
@@ -4162,6 +4201,11 @@ private:
 	 */
 	bool bSuspendNotifications;
 
+	/**
+	 * If set to true the hierarchy is currently running GetIndex - flag to avoid infinite recursion
+	 */
+	mutable bool bIsRunningGetIndex;
+
 	/**
 	 * The event fired during undo / redo
 	 */
@@ -4522,7 +4566,7 @@ protected:
 	mutable TMap<FRigElementKey, FRigElementKey> DefaultParentPerElement;
 
 	bool bUpdatePreferedEulerAngleWhenSettingTransform;
-	
+
 private:
 	
 	void EnsureCacheValidityImpl();
@@ -4564,7 +4608,7 @@ private:
 		}
 	}
 
-	FRigElementKeyRedirector* ElementRedirector;
+	FRigElementKeyRedirector* ElementKeyRedirector;
 
 	void UpdateVisibilityOnProxyControls();
 
@@ -4716,7 +4760,7 @@ struct CONTROLRIG_API FRigHierarchyRedirectorGuard
 {
 public:
 	FRigHierarchyRedirectorGuard(URigHierarchy* InHierarchy, FRigElementKeyRedirector& InRedirector)
-		: Guard(InHierarchy->ElementRedirector, &InRedirector)
+		: Guard(InHierarchy->ElementKeyRedirector, &InRedirector)
 	{
 	}
 
diff --git a/Engine/Plugins/Animation/ControlRig/Source/ControlRig/Public/Rigs/RigHierarchyCache.h b/Engine/Plugins/Animation/ControlRig/Source/ControlRig/Public/Rigs/RigHierarchyCache.h
--- a/Engine/Plugins/Animation/ControlRig/Source/ControlRig/Public/Rigs/RigHierarchyCache.h
+++ b/Engine/Plugins/Animation/ControlRig/Source/ControlRig/Public/Rigs/RigHierarchyCache.h
@@ -23,12 +23,16 @@ public:
 		, ContainerVersion(INDEX_NONE)
 	{}
 
-	FCachedRigElement(const FRigElementKey& InKey, const URigHierarchy* InHierarchy)
+	FCachedRigElement(const FRigElementKey& InKey, const URigHierarchy* InHierarchy, bool bForceStoreKey = false)
 		: Key()
 		, Index(UINT16_MAX)
 		, ContainerVersion(INDEX_NONE)
 	{
 		UpdateCache(InKey, InHierarchy);
+		if(bForceStoreKey)
+		{
+			Key = InKey;
+		}
 	}
 
 	bool IsValid() const
@@ -157,17 +161,21 @@ public:
 
 	FRigElementKeyRedirector()
 		: InternalKeyToExternalKey()
+		, Hash(UINT32_MAX)
 	{}
 
 	FRigElementKeyRedirector(const TMap<FRigElementKey, FRigElementKey>& InMap, const URigHierarchy* InHierarchy);
-	
+	FRigElementKeyRedirector(const FRigElementKeyRedirector& InOther, const URigHierarchy* InHierarchy);
+
 	bool Contains(const FRigElementKey& InKey) const { return InternalKeyToExternalKey.Contains(InKey); }
 	bool ContainsExternalKey(const FRigElementKey& InKey) const { return ExternalKeys.Contains(InKey); }
 	const FCachedRigElement* Find(const FRigElementKey& InKey) const { return InternalKeyToExternalKey.Find(InKey); }
 	FCachedRigElement* Find(const FRigElementKey& InKey) { return InternalKeyToExternalKey.Find(InKey); }
-
+	uint32 GetHash() const { return Hash; }
+	
 private:
 
 	TMap<FRigElementKey, FCachedRigElement> InternalKeyToExternalKey;
 	TSet<FRigElementKey> ExternalKeys;
+	uint32 Hash;
 };
\ No newline at end of file
diff --git a/Engine/Plugins/Animation/ControlRig/Source/ControlRig/Public/Rigs/RigHierarchyController.h b/Engine/Plugins/Animation/ControlRig/Source/ControlRig/Public/Rigs/RigHierarchyController.h
--- a/Engine/Plugins/Animation/ControlRig/Source/ControlRig/Public/Rigs/RigHierarchyController.h
+++ b/Engine/Plugins/Animation/ControlRig/Source/ControlRig/Public/Rigs/RigHierarchyController.h
@@ -240,7 +240,6 @@ public:
 	/**
 	 * Adds a connector to the hierarchy
 	 * @param InName The suggested name of the new connector - will eventually be corrected by the namespace
-	 * @param InParent The (optional) parent of the new connector. If you don't need a parent, pass FRigElementKey()
 	 * @param InTransform The transform for the new connector - either in local or global space, based on bTransformInGlobal
 	 * @param bTransformInGlobal Set this to true if the Transform passed is expressed in global space, false for local space.
 	 * @param InSettings All of the connector's settings
@@ -249,7 +248,7 @@ public:
 	 * @return The key for the newly created bone.
 	 */
 	UFUNCTION(BlueprintCallable, Category = URigHierarchyController)
-	FRigElementKey AddConnector(FName InName, FRigElementKey InParent, FTransform InTransform, bool bTransformInGlobal = true, FRigConnectorSettings InSettings = FRigConnectorSettings(), bool bSetupUndo = false, bool bPrintPythonCommand = false);
+	FRigElementKey AddConnector(FName InName, FTransform InTransform, bool bTransformInGlobal = true, FRigConnectorSettings InSettings = FRigConnectorSettings(), bool bSetupUndo = false, bool bPrintPythonCommand = false);
 
 	/**
 	 * Returns the control settings of a given control
diff --git a/Engine/Plugins/Animation/ControlRig/Source/ControlRig/Public/Rigs/RigHierarchyElements.h b/Engine/Plugins/Animation/ControlRig/Source/ControlRig/Public/Rigs/RigHierarchyElements.h
--- a/Engine/Plugins/Animation/ControlRig/Source/ControlRig/Public/Rigs/RigHierarchyElements.h
+++ b/Engine/Plugins/Animation/ControlRig/Source/ControlRig/Public/Rigs/RigHierarchyElements.h
@@ -766,6 +766,7 @@ protected:
 	friend struct FRigUnit_SetMetadataTag;
 	friend struct FRigUnit_SetMetadataTagArray;
 	friend struct FRigUnit_RemoveMetadataTag;
+	friend class FControlRigEditor;
 };
 
 USTRUCT(BlueprintType)
@@ -1631,15 +1632,15 @@ public:
 
 	// The keys of the connectors in the source hierarchy
 	// mapping to the to-be-linked keys of the elements in the target hierarchy
-	UPROPERTY(BlueprintReadWrite, EditAnywhere, Category = Limit)
+	UPROPERTY(BlueprintReadWrite, EditAnywhere, Category = Connections)
 	TMap<FRigElementKey, FRigElementKey> ConnectionMap;
 
 	// The hierarchy owning the connectors
-	UPROPERTY(BlueprintReadWrite, EditAnywhere, Category = Limit)
+	UPROPERTY(BlueprintReadWrite, EditAnywhere, Category = Connections)
 	URigHierarchy* SourceHierarchy; 
 
 	// The hierarchy to be linked into
-	UPROPERTY(BlueprintReadWrite, EditAnywhere, Category = Limit)
+	UPROPERTY(BlueprintReadWrite, EditAnywhere, Category = Connections)
 	URigHierarchy* TargetHierarchy; 
 };
 
@@ -1655,10 +1656,10 @@ struct CONTROLRIG_API FRigConnectorSettings
 
 	friend uint32 GetTypeHash(const FRigConnectorSettings& Settings);
 
-	UPROPERTY(BlueprintReadWrite, EditAnywhere, Category = Control)
+	UPROPERTY(BlueprintReadWrite, EditAnywhere, Category = Connector)
 	FString Description;
 
-	UPROPERTY(BlueprintReadWrite, EditAnywhere, Category = Control)
+	UPROPERTY(BlueprintReadWrite, EditAnywhere, Category = Connector)
 	TArray<FRigConnectionRuleStash> Rules;
 
 	bool operator == (const FRigConnectorSettings& InOther) const;
@@ -1677,6 +1678,31 @@ struct CONTROLRIG_API FRigConnectorSettings
 	uint32 GetRulesHash() const;
 };
 
+USTRUCT(BlueprintType)
+struct CONTROLRIG_API FRigConnectorInfo
+{
+	GENERATED_BODY()
+	
+	FRigConnectorInfo()
+		: Name(NAME_None)
+		, ResolvedTarget()
+		, LocalTransform(FTransform::Identity)
+		, Settings()
+	{}
+	
+	UPROPERTY(BlueprintReadWrite, EditAnywhere, Category = Connector)
+	FName Name;
+
+	UPROPERTY(BlueprintReadWrite, EditAnywhere, Category = Connector)
+	FRigElementKey ResolvedTarget;
+
+	UPROPERTY(BlueprintReadWrite, EditAnywhere, Category = Connector)
+	FTransform LocalTransform;
+
+	UPROPERTY(BlueprintReadWrite, EditAnywhere, Category = Connector)
+	FRigConnectorSettings Settings;
+};
+
 USTRUCT(BlueprintType)
 struct CONTROLRIG_API FRigConnectorElement : public FRigSingleParentElement
 {
@@ -1699,6 +1725,8 @@ public:
 	virtual void Load(FArchive& Ar, URigHierarchy* Hierarchy, ESerializationPhase SerializationPhase) override;
 
 	bool CanConnect(const FRigConnectionInfo* InConnectionInfo, FString* OutFailureReason) const;
+
+	FRigConnectorInfo GetConnectorInfo(const URigHierarchy* InHierarchy) const;
 	
 private:
 
diff --git a/Engine/Plugins/Animation/ControlRig/Source/ControlRig/Public/Rigs/RigModuleDefines.h b/Engine/Plugins/Animation/ControlRig/Source/ControlRig/Public/Rigs/RigModuleDefines.h
--- a/Engine/Plugins/Animation/ControlRig/Source/ControlRig/Public/Rigs/RigModuleDefines.h
+++ b/Engine/Plugins/Animation/ControlRig/Source/ControlRig/Public/Rigs/RigModuleDefines.h
@@ -6,6 +6,27 @@
 #include "RigHierarchyElements.h"
 #include "RigModuleDefines.generated.h"
 
+USTRUCT(BlueprintType)
+struct CONTROLRIG_API FRigModuleIdentifier
+{
+	GENERATED_BODY()
+	
+	FRigModuleIdentifier()
+		: Name()
+		, Type(TEXT("Module"))
+	{}
+
+	// The name of the module used to find it in the module library
+	UPROPERTY(BlueprintReadWrite, EditAnywhere, Category = Connector)
+	FString Name;
+
+	// The kind of module this is (for example "Arm")
+	UPROPERTY(BlueprintReadWrite, EditAnywhere, Category = Connector)
+	FString Type;
+
+	bool IsValid() const { return !Name.IsEmpty(); }
+};
+
 USTRUCT(BlueprintType)
 struct CONTROLRIG_API FRigModuleConnector
 {
@@ -33,12 +54,22 @@ struct CONTROLRIG_API FRigModuleSettings
 	FRigModuleSettings()
 	{}
 
-	bool IsValidModule() const { return !ExposedConnectors.IsEmpty(); }
+	bool IsValidModule() const
+	{
+		return
+			Identifier.IsValid() &&
+			!ExposedConnectors.IsEmpty();
+	}
+
+	// The identifier used to retrieve the module in the module library
+	UPROPERTY(BlueprintReadWrite, EditAnywhere, Category = Module)
+	FRigModuleIdentifier Identifier;
 
+	// The icon used for the module in the module library
 	UPROPERTY(BlueprintReadWrite, EditAnywhere, Category = Module,  meta = (AllowedClasses = "/Script/Engine.Texture2D"))
 	FSoftObjectPath Icon;
 
-	UPROPERTY(BlueprintReadOnly, VisibleAnywhere, Category = Module)
+	UPROPERTY(BlueprintReadOnly, Category = Module)
 	TArray<FRigModuleConnector> ExposedConnectors;
 };
 
diff --git a/Engine/Plugins/Animation/ControlRig/Source/ControlRig/Public/Units/RigUnitContext.h b/Engine/Plugins/Animation/ControlRig/Source/ControlRig/Public/Units/RigUnitContext.h
--- a/Engine/Plugins/Animation/ControlRig/Source/ControlRig/Public/Units/RigUnitContext.h
+++ b/Engine/Plugins/Animation/ControlRig/Source/ControlRig/Public/Units/RigUnitContext.h
@@ -100,6 +100,8 @@ struct FControlRigExecuteContext : public FRigVMExecuteContext
 	FControlRigExecuteContext()
 		: FRigVMExecuteContext()
 		, Hierarchy(nullptr)
+		, ModuleInstanceName(NAME_None)
+		, ModuleInstanceNameSpace()
 	{
 	}
 
@@ -135,6 +137,18 @@ struct FControlRigExecuteContext : public FRigVMExecuteContext
 		return nullptr;
 	}
 
+	/**
+	 * Add the namespace from a given name
+	 */
+	FName AddRigModuleNameSpace(const FName& InName) const;
+	FString AddRigModuleNameSpace(const FString& InName) const;
+
+	/**
+	 * Remove the namespace from a given name
+	 */
+	FName RemoveRigModuleNameSpace(const FName& InName) const;
+	FString RemoveRigModuleNameSpace(const FString& InName) const;
+
 	/** The list of available asset user data object */
 	TArray<const UAssetUserData*> AssetUserData;
 
@@ -146,6 +160,8 @@ struct FControlRigExecuteContext : public FRigVMExecuteContext
 	
 	FRigUnitContext UnitContext;
 	URigHierarchy* Hierarchy;
+	FName ModuleInstanceName;
+	FString ModuleInstanceNameSpace;
 };
 
 #if WITH_EDITOR
diff --git a/Engine/Plugins/Animation/ControlRig/Source/ControlRigDeveloper/Private/ControlRigBlueprint.cpp b/Engine/Plugins/Animation/ControlRig/Source/ControlRigDeveloper/Private/ControlRigBlueprint.cpp
--- a/Engine/Plugins/Animation/ControlRig/Source/ControlRigDeveloper/Private/ControlRigBlueprint.cpp
+++ b/Engine/Plugins/Animation/ControlRig/Source/ControlRigDeveloper/Private/ControlRigBlueprint.cpp
@@ -24,6 +24,8 @@
 #include "RigVMModel/Nodes/RigVMAggregateNode.h"
 #include "Rigs/RigControlHierarchy.h"
 #include "Settings/ControlRigSettings.h"
+#include "Units/Execution/RigUnit_PrepareForExecution.h"
+#include "Units/Execution/RigUnit_DynamicHierarchy.h"
 
 #include UE_INLINE_GENERATED_CPP_BY_NAME(ControlRigBlueprint)
 
@@ -34,6 +36,7 @@
 #include "Editor/UnrealEdEngine.h"
 #include "Editor/Transactor.h"
 #include "CookOnTheSide/CookOnTheFlyServer.h"
+#include "ScopedTransaction.h"
 #endif//WITH_EDITOR
 
 #define LOCTEXT_NAMESPACE "ControlRigBlueprint"
@@ -138,11 +141,447 @@ USkeletalMesh* UControlRigBlueprint::GetPreviewMesh() const
 
 bool UControlRigBlueprint::IsControlRigModule() const
 {
+	return RigModuleSettings.Identifier.IsValid();
+}
+
+#if WITH_EDITORONLY_DATA
+
+bool UControlRigBlueprint::CanTurnIntoControlRigModule(bool InAutoConvertHierarchy, FString* OutErrorMessage) const
+{
+	if(IsControlRigModule())
+	{
+		if(OutErrorMessage)
+		{
+			static const FString Message = TEXT("This asset is already a Control Rig Module.");
+			*OutErrorMessage = Message;
+		}
+		return false;
+	}
+
+	if(Hierarchy == nullptr)
+	{
+		if(OutErrorMessage)
+		{
+			static const FString Message = TEXT("This asset contains no hierarchy.");
+			*OutErrorMessage = Message;
+		}
+		return false;
+	}
+
+	const TArray<FRigElementKey> Keys = Hierarchy->GetAllKeys(true);
+	for(const FRigElementKey& Key : Keys)
+	{
+		if(!InAutoConvertHierarchy)
+		{
+			if(Key.Type != ERigElementType::Bone &&
+				Key.Type != ERigElementType::Curve &&
+				Key.Type != ERigElementType::Connector)
+			{
+				if(OutErrorMessage)
+				{
+					static constexpr TCHAR Format[] = TEXT("The hierarchy contains elements other than bones (for example '%s'). Modules only allow imported bones and user authored connectors.");
+					*OutErrorMessage = FString::Printf(Format, *Key.ToString());
+				}
+				return false;
+			}
+
+			if(Key.Type == ERigElementType::Bone)
+			{
+				if(Hierarchy->FindChecked<FRigBoneElement>(Key)->BoneType != ERigBoneType::Imported)
+				{
+					if(OutErrorMessage)
+					{
+						static constexpr TCHAR Format[] = TEXT("The hierarchy contains a user defined bone ('%s') - only imported bones are allowed.");
+						*OutErrorMessage = FString::Printf(Format, *Key.ToString());
+					}
+					return false;
+				}
+			}
+		}
+	}
+	
+	return true;
+}
+
+bool UControlRigBlueprint::TurnIntoControlRigModule(bool InAutoConvertHierarchy, FString* OutErrorMessage)
+{
+	if(!CanTurnIntoControlRigModule(InAutoConvertHierarchy, OutErrorMessage))
+	{
+		return false;
+	}
+
+	FScopedTransaction Transaction(LOCTEXT("TurnIntoControlRigModule", "Turn Rig into Module"));
+
+	Modify();
+	RigModuleSettings.Identifier = FRigModuleIdentifier();
+	RigModuleSettings.Identifier.Name = GetName();
+
 	if(Hierarchy)
 	{
-		return Hierarchy->Num(ERigElementType::Connector) > 0;
+		URigHierarchyController* Controller = Hierarchy->GetController(true);
+
+		// create a copy of this hierarchy
+		URigHierarchy* CopyOfHierarchy = NewObject<URigHierarchy>(GetTransientPackage());
+		CopyOfHierarchy->CopyHierarchy(Hierarchy);
+
+		// also create a hierarchy based on the preview mesh
+		URigHierarchy* PreviewMeshHierarchy = NewObject<URigHierarchy>(GetTransientPackage());
+		if(PreviewSkeletalMesh)
+		{
+			PreviewMeshHierarchy->GetController(true)->ImportBones(PreviewSkeletalMesh->GetSkeleton());
+		}
+
+		// disable compilation
+		{
+			FRigVMBlueprintCompileScope CompileScope(this);
+
+			// remove everything from the hierarchy
+			Hierarchy->Reset();
+
+			const TArray<FRigElementKey> AllKeys = CopyOfHierarchy->GetAllKeys(true);
+			TArray<FRigElementKey> KeysToSpawn;
+
+			for(const FRigElementKey& Key : AllKeys)
+			{
+				if(Key.Type == ERigElementType::Curve)
+				{
+					continue;
+				}
+				if(Key.Type == ERigElementType::Bone)
+				{
+					if(PreviewMeshHierarchy->Contains(Key))
+					{
+						continue;
+					}
+				}
+				KeysToSpawn.Add(Key);
+			}
+
+			(void)ConvertHierarchyElementsToSpawnerNodes(CopyOfHierarchy, KeysToSpawn, false);
+
+			if(Hierarchy->Num(ERigElementType::Connector) == 0)
+			{
+				static const FName RootName = TEXT("Root");
+				Controller->AddConnector(RootName, FTransform::Identity);
+			}
+		}
 	}
-	return false;
+
+	OnRigTypeChangedDelegate.Broadcast(this);
+	return true;
+}
+
+bool UControlRigBlueprint::CanTurnIntoStandaloneRig(FString* OutErrorMessage) const
+{
+	return IsControlRigModule();
+}
+
+bool UControlRigBlueprint::TurnIntoStandaloneRig(FString* OutErrorMessage)
+{
+	if(!CanTurnIntoStandaloneRig(OutErrorMessage))
+	{
+		return false;
+	}
+
+	FScopedTransaction Transaction(LOCTEXT("TurnIntoControlRigModule", "Turn Rig into Module"));
+
+	Modify();
+	RigModuleSettings = FRigModuleSettings();
+
+	if(Hierarchy)
+	{
+		Hierarchy->Modify();
+		Hierarchy->Reset();
+		if(PreviewSkeletalMesh)
+		{
+			Hierarchy->GetController(true)->ImportBones(PreviewSkeletalMesh->GetSkeleton());
+		}
+	}
+
+	OnRigTypeChangedDelegate.Broadcast(this);
+	return true;
+}
+
+TArray<URigVMNode*> UControlRigBlueprint::ConvertHierarchyElementsToSpawnerNodes(URigHierarchy* InHierarchy, TArray<FRigElementKey> InKeys, bool bRemoveElements)
+{
+	TArray<URigVMNode*> SpawnerNodes;
+
+	// find the construction event 
+	const URigVMNode* EventNode = nullptr;
+	for(const URigVMGraph* Graph : GetRigVMClient()->GetAllModels(false, false))
+	{
+		for(const URigVMNode* Node : Graph->GetNodes())
+		{
+			if(Node->IsEvent() && Node->GetEventName() == FRigUnit_PrepareForExecution::EventName)
+			{
+				EventNode = Node;
+				break;
+			}
+		}
+		if(EventNode)
+		{
+			break;
+		}
+	}
+
+	FVector2D NodePosition = FVector2D::ZeroVector;
+	const FVector2D NodePositionIncrement = FVector2D(400, 0);
+	
+	// if we didn't find the construction event yet, create it
+	if(EventNode == nullptr)
+	{
+		const URigVMGraph* ConstructionGraph = GetRigVMClient()->AddModel(TEXT("ConstructionGraph"), true);
+		URigVMController* GraphController = GetRigVMClient()->GetOrCreateController(ConstructionGraph);
+		EventNode = GraphController->AddUnitNode(FRigUnit_PrepareForExecution::StaticStruct(), FRigUnit::GetMethodName(), NodePosition);
+		NodePosition += NodePositionIncrement;
+	}
+
+	const URigVMPin* LastPin = EventNode->FindPin(GET_MEMBER_NAME_STRING_CHECKED(FRigUnit_PrepareForExecution, ExecuteContext));
+	if(LastPin)
+	{
+		// follow the node's execution links to find the last one
+		bool bCarryOn = true;
+		while(bCarryOn)
+		{
+			static const TArray<FString> ExecutePinPaths = {
+				FRigVMStruct::ControlFlowCompletedName.ToString(),
+				FRigVMStruct::ExecuteContextName.ToString()
+			};
+
+			for(const FString& ExecutePinPath : ExecutePinPaths)
+			{
+				if(const URigVMPin* ExecutePin = LastPin->GetNode()->FindPin(ExecutePinPath))
+				{
+					const TArray<URigVMPin*> TargetPins = ExecutePin->GetLinkedTargetPins();
+					if(TargetPins.IsEmpty())
+					{
+						bCarryOn = false;
+						break;
+					}
+					LastPin = TargetPins[0];
+					NodePosition = LastPin->GetNode()->GetPosition() + NodePositionIncrement;
+				}
+			}
+		}
+	}
+
+	const URigVMGraph* ConstructionGraph = EventNode->GetGraph();
+	URigVMController* GraphController = GetRigVMClient()->GetOrCreateController(ConstructionGraph);
+
+	auto GetParentAndTransformDefaults = [InHierarchy](const FRigElementKey& InKey, FString& OutParentDefault, FString& OutTransformDefault)
+	{
+		const FRigElementKey Parent = InHierarchy->GetFirstParent(InKey);
+		OutParentDefault.Reset();
+		FRigElementKey::StaticStruct()->ExportText(OutParentDefault, &Parent, nullptr, nullptr, PPF_None, nullptr);
+
+		const FTransform Transform = InHierarchy->GetInitialLocalTransform(InKey);
+		OutTransformDefault.Reset();
+		TBaseStructure<FTransform>::Get()->ExportText(OutTransformDefault, &Transform, nullptr, nullptr, PPF_None, nullptr);
+	};
+
+	TMap<FRigElementKey, const URigVMPin*> ParentItemPinMap;
+	auto AddParentItemLink = [GraphController, InHierarchy, &SpawnerNodes, &ParentItemPinMap]
+		(const FRigElementKey& Key, URigVMNode* Node)
+		{
+			SpawnerNodes.Add(Node);
+			ParentItemPinMap.Add(Key, Node->FindPin(GET_MEMBER_NAME_STRING_CHECKED(FRigUnit_HierarchyAddElement, Item)));
+
+			if(const URigVMPin** SourcePin = ParentItemPinMap.Find(InHierarchy->GetFirstParent(Key)))
+			{
+				if(const URigVMPin* TargetPin = Node->FindPin(GET_MEMBER_NAME_STRING_CHECKED(FRigUnit_HierarchyAddElement, Parent)))
+				{
+					GraphController->AddLink((*SourcePin)->GetPinPath(), TargetPin->GetPinPath(), true);
+				}
+			}
+		};
+	
+	for(const FRigElementKey& Key : InKeys)
+	{
+		if(Key.Type == ERigElementType::Bone)
+		{
+			FString ParentDefault, TransformDefault;
+			GetParentAndTransformDefaults(Key, ParentDefault, TransformDefault);
+
+			URigVMNode* AddBoneNode = GraphController->AddUnitNode(FRigUnit_HierarchyAddBone::StaticStruct(), FRigUnit::GetMethodName(), NodePosition);
+			NodePosition += NodePositionIncrement;
+			AddParentItemLink(Key, AddBoneNode);
+
+			if(LastPin)
+			{
+				if(const URigVMPin* NextPin = AddBoneNode->FindPin(GET_MEMBER_NAME_STRING_CHECKED(FRigUnit_HierarchyAddBone, ExecuteContext)))
+				{
+					GraphController->AddLink(LastPin->GetPinPath(), NextPin->GetPinPath(), true);
+					LastPin = NextPin;
+				}
+			}
+
+			GraphController->SetPinDefaultValue(AddBoneNode->FindPin(GET_MEMBER_NAME_STRING_CHECKED(FRigUnit_HierarchyAddElement, Name))->GetPinPath(), Key.Name.ToString(), true, true);
+			GraphController->SetPinDefaultValue(AddBoneNode->FindPin(GET_MEMBER_NAME_STRING_CHECKED(FRigUnit_HierarchyAddElement, Parent))->GetPinPath(), ParentDefault, true, true);
+			GraphController->SetPinDefaultValue(AddBoneNode->FindPin(GET_MEMBER_NAME_STRING_CHECKED(FRigUnit_HierarchyAddBone, Space))->GetPinPath(), TEXT("LocalSpace"), true, true);
+			GraphController->SetPinDefaultValue(AddBoneNode->FindPin(GET_MEMBER_NAME_STRING_CHECKED(FRigUnit_HierarchyAddBone, Transform))->GetPinPath(), TransformDefault, true, true);
+		}
+		else if(Key.Type == ERigElementType::Null)
+		{
+			FString ParentDefault, TransformDefault;
+			GetParentAndTransformDefaults(Key, ParentDefault, TransformDefault);
+
+			URigVMNode* AddNullNode = GraphController->AddUnitNode(FRigUnit_HierarchyAddNull::StaticStruct(), FRigUnit::GetMethodName(), NodePosition);
+			NodePosition += NodePositionIncrement;
+			SpawnerNodes.Add(AddNullNode);
+
+			if(LastPin)
+			{
+				if(const URigVMPin* NextPin = AddNullNode->FindPin(GET_MEMBER_NAME_STRING_CHECKED(FRigUnit_HierarchyAddBone, ExecuteContext)))
+				{
+					GraphController->AddLink(LastPin->GetPinPath(), NextPin->GetPinPath(), true);
+					LastPin = NextPin;
+				}
+			}
+
+			GraphController->SetPinDefaultValue(AddNullNode->FindPin(GET_MEMBER_NAME_STRING_CHECKED(FRigUnit_HierarchyAddElement, Name))->GetPinPath(), Key.Name.ToString(), true, true);
+			GraphController->SetPinDefaultValue(AddNullNode->FindPin(GET_MEMBER_NAME_STRING_CHECKED(FRigUnit_HierarchyAddElement, Parent))->GetPinPath(), ParentDefault, true, true);
+			GraphController->SetPinDefaultValue(AddNullNode->FindPin(GET_MEMBER_NAME_STRING_CHECKED(FRigUnit_HierarchyAddNull, Space))->GetPinPath(), TEXT("LocalSpace"), true, true);
+			GraphController->SetPinDefaultValue(AddNullNode->FindPin(GET_MEMBER_NAME_STRING_CHECKED(FRigUnit_HierarchyAddNull, Transform))->GetPinPath(), TransformDefault, true, true);
+		}
+		else if(Key.Type == ERigElementType::Control)
+		{
+			FRigControlElement* ControlElement = InHierarchy->FindChecked<FRigControlElement>(Key);
+			
+			FString ParentDefault, TransformDefault;
+			GetParentAndTransformDefaults(Key, ParentDefault, TransformDefault);
+
+			const FTransform OffsetTransform = InHierarchy->GetControlOffsetTransform(ControlElement, ERigTransformType::InitialLocal);
+			FString OffsetDefault;
+			TBaseStructure<FTransform>::Get()->ExportText(OffsetDefault, &OffsetTransform, nullptr, nullptr, PPF_None, nullptr);
+			
+			if(ControlElement->Settings.AnimationType == ERigControlAnimationType::AnimationChannel)
+			{
+				// todo
+			}
+			else
+			{
+				UScriptStruct* UnitNodeStruct = nullptr;
+				TRigVMTypeIndex TypeIndex = INDEX_NONE;
+				FString InitialValue;
+				switch(ControlElement->Settings.ControlType)
+				{
+					case ERigControlType::Float:
+					case ERigControlType::ScaleFloat:
+					{
+						UnitNodeStruct = FRigUnit_HierarchyAddControlFloat::StaticStruct();
+						TypeIndex = RigVMTypeUtils::TypeIndex::Float;
+						InitialValue = InHierarchy->GetControlValue(Key, ERigControlValueType::Initial).ToString<float>();
+						break;
+					}
+					case ERigControlType::Integer:
+					{
+						UnitNodeStruct = FRigUnit_HierarchyAddControlInteger::StaticStruct();
+						TypeIndex = RigVMTypeUtils::TypeIndex::Int32; 
+						InitialValue = InHierarchy->GetControlValue(Key, ERigControlValueType::Initial).ToString<int32>();
+						break;
+					}
+					case ERigControlType::Vector2D:
+					{
+						UnitNodeStruct = FRigUnit_HierarchyAddControlVector2D::StaticStruct();
+						TypeIndex = FRigVMRegistry::Get().GetTypeIndex<FVector2D>(); 
+						InitialValue = InHierarchy->GetControlValue(Key, ERigControlValueType::Initial).ToString<FVector2D>();
+						break;
+					}
+					case ERigControlType::Position:
+					case ERigControlType::Scale:
+					{
+						UnitNodeStruct = FRigUnit_HierarchyAddControlVector::StaticStruct();
+						TypeIndex = FRigVMRegistry::Get().GetTypeIndex<FVector>(); 
+						InitialValue = InHierarchy->GetControlValue(Key, ERigControlValueType::Initial).ToString<FVector>();
+						break;
+					}
+					case ERigControlType::Rotator:
+					{
+						UnitNodeStruct = FRigUnit_HierarchyAddControlRotator::StaticStruct();
+						TypeIndex = FRigVMRegistry::Get().GetTypeIndex<FRotator>(); 
+						InitialValue = InHierarchy->GetControlValue(Key, ERigControlValueType::Initial).ToString<FRotator>();
+						break;
+					}
+					case ERigControlType::Transform:
+					case ERigControlType::TransformNoScale:
+					case ERigControlType::EulerTransform:
+					{
+						UnitNodeStruct = FRigUnit_HierarchyAddControlTransform::StaticStruct();
+						TypeIndex = FRigVMRegistry::Get().GetTypeIndex<FTransform>(); 
+						const FTransform InitialTransform = InHierarchy->GetInitialLocalTransform(Key);
+						TBaseStructure<FTransform>::Get()->ExportText(InitialValue, &InitialTransform, nullptr, nullptr, PPF_None, nullptr);
+						break;
+					}
+					default:
+					{
+						break;
+					}
+				}
+
+				if(UnitNodeStruct == nullptr)
+				{
+					continue;
+				}
+
+				URigVMNode* AddControlNode = GraphController->AddUnitNode(UnitNodeStruct, FRigUnit::GetMethodName(), NodePosition);
+				NodePosition += NodePositionIncrement;
+				AddParentItemLink(Key, AddControlNode);
+
+				if(LastPin)
+				{
+					if(const URigVMPin* NextPin = AddControlNode->FindPin(GET_MEMBER_NAME_STRING_CHECKED(FRigUnit_HierarchyAddBone, ExecuteContext)))
+					{
+						GraphController->AddLink(LastPin->GetPinPath(), NextPin->GetPinPath(), true);
+						LastPin = NextPin;
+					}
+				}
+
+				GraphController->ResolveWildCardPin(AddControlNode->FindPin(GET_MEMBER_NAME_STRING_CHECKED(FRigUnit_HierarchyAddControlInteger, InitialValue))->GetPinPath(), TypeIndex, true);
+				GraphController->SetPinDefaultValue(AddControlNode->FindPin(GET_MEMBER_NAME_STRING_CHECKED(FRigUnit_HierarchyAddElement, Name))->GetPinPath(), Key.Name.ToString(), true, true);
+				GraphController->SetPinDefaultValue(AddControlNode->FindPin(GET_MEMBER_NAME_STRING_CHECKED(FRigUnit_HierarchyAddElement, Parent))->GetPinPath(), ParentDefault, true, true);
+				GraphController->SetPinDefaultValue(AddControlNode->FindPin(GET_MEMBER_NAME_STRING_CHECKED(FRigUnit_HierarchyAddControlElement, OffsetSpace))->GetPinPath(), TEXT("LocalSpace"), true, true);
+				GraphController->SetPinDefaultValue(AddControlNode->FindPin(GET_MEMBER_NAME_STRING_CHECKED(FRigUnit_HierarchyAddControlElement, OffsetTransform))->GetPinPath(), OffsetDefault, true, true);
+				GraphController->SetPinDefaultValue(AddControlNode->FindPin(GET_MEMBER_NAME_STRING_CHECKED(FRigUnit_HierarchyAddControlInteger, InitialValue))->GetPinPath(), InitialValue, true, true);
+
+				if(const FStructProperty* SettingsProperty = CastField<FStructProperty>(UnitNodeStruct->FindPropertyByName(TEXT("Settings"))))
+				{
+					UScriptStruct* SettingsStruct = CastChecked<UScriptStruct>(SettingsProperty->Struct);
+					FStructOnScope SettingsScope(SettingsStruct);
+					FRigUnit_HierarchyAddControl_Settings* Settings = (FRigUnit_HierarchyAddControl_Settings*)SettingsScope.GetStructMemory();
+					Settings->ConfigureFrom(ControlElement, ControlElement->Settings);
+					FString SettingsDefault;
+					SettingsStruct->ExportText(SettingsDefault, Settings, nullptr, nullptr, PPF_None, nullptr);
+
+					GraphController->SetPinDefaultValue(AddControlNode->FindPin(SettingsProperty->GetName())->GetPinPath(), SettingsDefault, true, true);
+				}
+			}
+		}
+	}
+
+	if(bRemoveElements && InHierarchy)
+	{
+		InHierarchy->Modify();
+		for(const FRigElementKey& Key : InKeys)
+		{
+			InHierarchy->GetController(true)->RemoveElement(Key, true);
+		}
+	}
+
+	return SpawnerNodes;
+}
+
+#endif // WITH_EDITORONLY_DATA
+
+UTexture2D* UControlRigBlueprint::GetRigModuleIcon() const
+{
+	if(IsControlRigModule())
+	{
+		if(UTexture2D* Icon = Cast<UTexture2D>(RigModuleSettings.Icon.TryLoad()))
+		{
+			return Icon;
+		}
+	}
+	return nullptr;
 }
 
 void UControlRigBlueprint::SetPreviewMesh(USkeletalMesh* PreviewMesh, bool bMarkAsDirty/*=true*/)
@@ -234,20 +673,48 @@ void UControlRigBlueprint::PreSave(FObjectPreSaveContext ObjectSaveContext)
 	UpdateExposedModuleConnectors();
 }
 
-void UControlRigBlueprint::UpdateExposedModuleConnectors()
+void UControlRigBlueprint::UpdateExposedModuleConnectors() const
 {
-	ModuleSettings.ExposedConnectors.Reset();
-	Hierarchy->ForEach<FRigConnectorElement>([this](const FRigConnectorElement* ConnectorElement) -> bool
+	UControlRigBlueprint* MutableThis = ((UControlRigBlueprint*)this);
+	MutableThis->RigModuleSettings.ExposedConnectors.Reset();
+	Hierarchy->ForEach<FRigConnectorElement>([MutableThis](const FRigConnectorElement* ConnectorElement) -> bool
 	{
 		FRigModuleConnector ExposedConnector;
 		ExposedConnector.Name = ConnectorElement->GetName();
-		ExposedConnector.bIsRoot = Hierarchy->GetNumberOfParents(ConnectorElement) == 0;
+		ExposedConnector.bIsRoot = MutableThis->Hierarchy->GetNumberOfParents(ConnectorElement) == 0;
 		ExposedConnector.Settings = ConnectorElement->Settings;
-		ModuleSettings.ExposedConnectors.Add(ExposedConnector);
+		MutableThis->RigModuleSettings.ExposedConnectors.Add(ExposedConnector);
 		return true;
 	});
 }
 
+bool UControlRigBlueprint::ResolveConnector(const FRigElementKey& DraggedKey, const FRigElementKey& TargetKey, bool bSetupUndoRedo)
+{
+	FScopedTransaction Transaction(LOCTEXT("ResolveConnector", "Resolve connector"));
+
+	if(bSetupUndoRedo)
+	{
+		Modify();
+	}
+	if(TargetKey.IsValid())
+	{
+		FRigElementKey& ExistingTargetKey = ConnectionMap.FindOrAdd(DraggedKey);
+		if(ExistingTargetKey == TargetKey)
+		{
+			return false;
+		}
+		ExistingTargetKey = TargetKey;
+	}
+	else
+	{
+		ConnectionMap.Remove(DraggedKey);
+	}
+
+	PropagateHierarchyFromBPToInstances();
+
+	return true;
+}
+
 void UControlRigBlueprint::PostLoad()
 {
 	Super::PostLoad();
@@ -509,6 +976,11 @@ void UControlRigBlueprint::PostTransacted(const FTransactionObjectEvent& Transac
 		{
 			PropagateDrawInstructionsFromBPToInstances();
 		}
+
+		if (PropertiesChanged.Contains(GET_MEMBER_NAME_CHECKED(UControlRigBlueprint, ConnectionMap)))
+		{
+			PropagateHierarchyFromBPToInstances();
+		}
 	}
 }
 
@@ -1159,6 +1631,13 @@ void UControlRigBlueprint::PatchVariableNodesOnLoad()
 	Super::PatchVariableNodesOnLoad();
 }
 
+void UControlRigBlueprint::UpdateElementKeyRedirector(UControlRig* InControlRig) const
+{
+	InControlRig->HierarchySettings = HierarchySettings;
+	InControlRig->RigModuleSettings = RigModuleSettings;
+	InControlRig->ElementKeyRedirector = FRigElementKeyRedirector(ConnectionMap, Hierarchy);
+}
+
 void UControlRigBlueprint::PropagatePoseFromInstanceToBP(UControlRig* InControlRig) const
 {
 	check(InControlRig);
@@ -1199,7 +1678,9 @@ void UControlRigBlueprint::PropagateHierarchyFromBPToInstances() const
 		{
 			DefaultObject->PostInitInstanceIfRequired();
 			DefaultObject->GetHierarchy()->CopyHierarchy(Hierarchy);
-			DefaultObject->HierarchySettings = HierarchySettings;
+
+			UpdateElementKeyRedirector(DefaultObject);
+
 			if (!DefaultObject->HasAnyFlags(RF_NeedPostLoad)) // If CDO is loading, skip Init, it will be done later
 			{
 				DefaultObject->Initialize(true);
diff --git a/Engine/Plugins/Animation/ControlRig/Source/ControlRigDeveloper/Public/ControlRigBlueprint.h b/Engine/Plugins/Animation/ControlRig/Source/ControlRigDeveloper/Public/ControlRigBlueprint.h
--- a/Engine/Plugins/Animation/ControlRig/Source/ControlRigDeveloper/Public/ControlRigBlueprint.h
+++ b/Engine/Plugins/Animation/ControlRig/Source/ControlRigDeveloper/Public/ControlRigBlueprint.h
@@ -6,6 +6,7 @@
 #include "ControlRigBlueprintGeneratedClass.h"
 #include "UObject/ObjectMacros.h"
 #include "Engine/Blueprint.h"
+#include "Engine/Texture2D.h"
 #include "ControlRigDefines.h"
 #include "Rigs/RigHierarchyContainer.h"
 #include "Rigs/RigHierarchy.h"
@@ -90,15 +91,55 @@ public:
 	UFUNCTION(BlueprintPure, Category = "Control Rig Blueprint")
 	bool IsControlRigModule() const;
 
+#if WITH_EDITORONLY_DATA
+	
+	bool CanTurnIntoControlRigModule_Blueprint(bool InAutoConvertHierarchy = false) const { return CanTurnIntoControlRigModule(InAutoConvertHierarchy); }
+
+	bool CanTurnIntoControlRigModule(bool InAutoConvertHierarchy, FString* OutErrorMessage = nullptr) const;
+
+	UFUNCTION(BlueprintCallable, meta = (DisplayName = "TurnIntoControlRigModule", ScriptName = "TurnIntoControlRigModule"), Category = "Control Rig Blueprint")
+	bool TurnIntoControlRigModule_Blueprint() { return TurnIntoControlRigModule(); }
+
+	bool TurnIntoControlRigModule(bool InAutoConvertHierarchy = false, FString* OutErrorMessage = nullptr);
+
+	UFUNCTION(BlueprintPure, meta = (DisplayName = "CanTurnIntoStandaloneRig", ScriptName = "CanTurnIntoStandaloneRig"), Category = "Control Rig Blueprint")
+	bool CanTurnIntoStandaloneRig_Blueprint() const { return CanTurnIntoStandaloneRig(); }
+
+	bool CanTurnIntoStandaloneRig(FString* OutErrorMessage = nullptr) const;
+
+	UFUNCTION(BlueprintCallable, meta = (DisplayName = "TurnIntoStandaloneRig", ScriptName = "TurnIntoStandaloneRig"), Category = "Control Rig Blueprint")
+	bool TurnIntoStandaloneRig_Blueprint() { return TurnIntoStandaloneRig(); }
+
+	bool TurnIntoStandaloneRig(FString* OutErrorMessage = nullptr);
+
+	UFUNCTION(BlueprintCallable, Category = "Control Rig Blueprint")
+	TArray<URigVMNode*> ConvertHierarchyElementsToSpawnerNodes(URigHierarchy* InHierarchy, TArray<FRigElementKey> InKeys, bool bRemoveElements = true);
+
+#endif // WITH_EDITORONLY_DATA
+
+	UFUNCTION(BlueprintPure, Category = "Control Rig Blueprint")
+	UTexture2D* GetRigModuleIcon() const;
+
+	DECLARE_EVENT_OneParam(UControlRigBlueprint, FOnRigTypeChanged, UControlRigBlueprint*);
+
+	FOnRigTypeChanged& OnRigTypeChanged() { return OnRigTypeChangedDelegate; }
+
 	UPROPERTY(EditAnywhere, Category = "Hierarchy")
 	FRigHierarchySettings HierarchySettings;
 
 	UPROPERTY(EditAnywhere, Category = "Hierarchy", AssetRegistrySearchable)
-	FRigModuleSettings ModuleSettings;
+	FRigModuleSettings RigModuleSettings;
+
+	UPROPERTY()
+	TMap<FRigElementKey, FRigElementKey> ConnectionMap;
 
 protected:
 
-	void UpdateExposedModuleConnectors();
+	FOnRigTypeChanged OnRigTypeChangedDelegate;
+	
+	void UpdateExposedModuleConnectors() const;
+
+	bool ResolveConnector(const FRigElementKey& DraggedKey, const FRigElementKey& TargetKey, bool bSetupUndoRedo = true);
 
 	/** Asset searchable information about exposed public functions on this rig */
 	UPROPERTY(AssetRegistrySearchable)
@@ -205,6 +246,7 @@ protected:
 	virtual void PatchVariableNodesOnLoad() override;
 
 public:
+	void UpdateElementKeyRedirector(UControlRig* InControlRig) const;
 	void PropagatePoseFromInstanceToBP(UControlRig* InControlRig) const;
 	void PropagatePoseFromBPToInstances() const;
 	void PropagateHierarchyFromBPToInstances() const;
diff --git a/Engine/Plugins/Animation/ControlRig/Source/ControlRigEditor/Private/ControlRigBlueprintDetails.cpp b/Engine/Plugins/Animation/ControlRig/Source/ControlRigEditor/Private/ControlRigBlueprintDetails.cpp
--- a/Engine/Plugins/Animation/ControlRig/Source/ControlRigEditor/Private/ControlRigBlueprintDetails.cpp
+++ b/Engine/Plugins/Animation/ControlRig/Source/ControlRigEditor/Private/ControlRigBlueprintDetails.cpp
@@ -34,6 +34,6 @@ void FControlRigBlueprintDetails::CustomizeDetails(IDetailLayoutBuilder& DetailL
 	
 	if(!bIsValidRigModule)
 	{
-		DetailLayout.HideProperty(GET_MEMBER_NAME_CHECKED(UControlRigBlueprint, ModuleSettings));
+		DetailLayout.HideProperty(GET_MEMBER_NAME_CHECKED(UControlRigBlueprint, RigModuleSettings));
 	}
 }
\ No newline at end of file
diff --git a/Engine/Plugins/Animation/ControlRig/Source/ControlRigEditor/Private/ControlRigBlueprintEditorLibrary.cpp b/Engine/Plugins/Animation/ControlRig/Source/ControlRigEditor/Private/ControlRigBlueprintEditorLibrary.cpp
--- a/Engine/Plugins/Animation/ControlRig/Source/ControlRigEditor/Private/ControlRigBlueprintEditorLibrary.cpp
+++ b/Engine/Plugins/Animation/ControlRig/Source/ControlRigEditor/Private/ControlRigBlueprintEditorLibrary.cpp
@@ -103,7 +103,7 @@ TArray<FRigModuleDescription> UControlRigBlueprintEditorLibrary::GetAvailableRig
 
 	for(const FAssetData& AssetData : AssetDataList)
 	{
-		static const FName ModuleSettingsName = GET_MEMBER_NAME_CHECKED(UControlRigBlueprint, ModuleSettings);
+		static const FName ModuleSettingsName = GET_MEMBER_NAME_CHECKED(UControlRigBlueprint, RigModuleSettings);
 		if(AssetData.FindTag(ModuleSettingsName))
 		{
 			const FString TagValue = AssetData.GetTagValueRef<FString>(ModuleSettingsName);
diff --git a/Engine/Plugins/Animation/ControlRig/Source/ControlRigEditor/Private/ControlRigThumbnailRenderer.cpp b/Engine/Plugins/Animation/ControlRig/Source/ControlRigEditor/Private/ControlRigThumbnailRenderer.cpp
--- a/Engine/Plugins/Animation/ControlRig/Source/ControlRigEditor/Private/ControlRigThumbnailRenderer.cpp
+++ b/Engine/Plugins/Animation/ControlRig/Source/ControlRigEditor/Private/ControlRigThumbnailRenderer.cpp
@@ -5,10 +5,13 @@
 #include "Engine/StaticMesh.h"
 #include "Components/StaticMeshComponent.h"
 #include "Engine/SkeletalMesh.h"
+#include "Engine/Texture2D.h"
+#include "TextureResource.h"
 #include "Animation/SkeletalMeshActor.h"
 #include "Components/SkeletalMeshComponent.h"
 #include "Materials/MaterialInstanceDynamic.h"
 #include "ControlRig.h"
+#include "CanvasTypes.h"
 
 #include UE_INLINE_GENERATED_CPP_BY_NAME(ControlRigThumbnailRenderer)
 
@@ -22,6 +25,11 @@ bool UControlRigThumbnailRenderer::CanVisualizeAsset(UObject* Object)
 {
 	if (UControlRigBlueprint* InRigBlueprint = Cast<UControlRigBlueprint>(Object))
 	{
+		if(InRigBlueprint->GetRigModuleIcon())
+		{
+			return true;
+		}
+		
 		USkeletalMesh* SkeletalMesh = InRigBlueprint->PreviewSkeletalMesh.Get();
 		int32 MissingMeshCount = 0;
 
@@ -55,6 +63,12 @@ void UControlRigThumbnailRenderer::Draw(UObject* Object, int32 X, int32 Y, uint3
 
 	if (UControlRigBlueprint* InRigBlueprint = Cast<UControlRigBlueprint>(Object))
 	{
+		if(UTexture2D* ModuleIcon = InRigBlueprint->GetRigModuleIcon())
+		{
+			Canvas->DrawTile(X, Y, Width, Height, 0, 0, 1, 1, FLinearColor::White, ModuleIcon->GetResource(), 1.f);
+			return;
+		}
+		
 		UObject* ObjectToDraw = InRigBlueprint->PreviewSkeletalMesh.Get();
 		if(ObjectToDraw == nullptr)
 		{
@@ -81,6 +95,11 @@ void UControlRigThumbnailRenderer::AddAdditionalPreviewSceneContent(UObject* Obj
 	TSharedRef<FSkeletalMeshThumbnailScene> ThumbnailScene = ThumbnailSceneCache.EnsureThumbnailScene(Object);
 	if (ThumbnailScene->GetPreviewActor() && RigBlueprint && !RigBlueprint->ShapeLibraries.IsEmpty() && RigBlueprint->GeneratedClass)
 	{
+		if(RigBlueprint->GetRigModuleIcon())
+		{
+			return;
+		}
+
 		UControlRig* ControlRig = nullptr;
 
 		// reuse the current control rig if possible
diff --git a/Engine/Plugins/Animation/ControlRig/Source/ControlRigEditor/Private/Editor/ControlRigEditor.cpp b/Engine/Plugins/Animation/ControlRig/Source/ControlRigEditor/Private/Editor/ControlRigEditor.cpp
--- a/Engine/Plugins/Animation/ControlRig/Source/ControlRigEditor/Private/Editor/ControlRigEditor.cpp
+++ b/Engine/Plugins/Animation/ControlRig/Source/ControlRigEditor/Private/Editor/ControlRigEditor.cpp
@@ -145,6 +145,8 @@ FControlRigEditor::~FControlRigEditor()
 			RigBlueprint->OnHierarchyModified().RemoveAll(EditMode);
 			EditMode->OnEditorClosed();
 		}
+
+		RigBlueprint->OnRigTypeChanged().RemoveAll(this);
 	}
 
 	if (PersonaToolkit.IsValid())
@@ -244,6 +246,8 @@ void FControlRigEditor::InitRigVMEditor(const EToolkitMode::Type Mode, const TSh
 		{
 			ControlRigBlueprint->OnHierarchyModified().AddSP(EditMode, &FControlRigEditMode::OnHierarchyModified_AnyThread);
 		}
+
+		ControlRigBlueprint->OnRigTypeChanged().AddSP(this, &FControlRigEditor::HandleRigTypeChanged);
 	}
 
 	CreateRigHierarchyToGraphDragAndDropMenu();
@@ -502,6 +506,100 @@ void FControlRigEditor::FillToolbar(FToolBarBuilder& ToolbarBuilder, bool bEndSe
 	FRigVMEditor::FillToolbar(ToolbarBuilder, false);
 	
 	{
+		if(CVarControlRigHierarchyEnableModules.GetValueOnAnyThread())
+		{
+			TWeakObjectPtr<UControlRigBlueprint> WeakBlueprint = GetControlRigBlueprint();
+			ToolbarBuilder.AddToolBarButton(
+				FUIAction(
+					FExecuteAction::CreateLambda([WeakBlueprint]()
+					{
+						if(WeakBlueprint.IsValid())
+						{
+							if(WeakBlueprint->IsControlRigModule())
+							{
+								WeakBlueprint->TurnIntoStandaloneRig();
+							}
+							else
+							{
+								if(!WeakBlueprint->CanTurnIntoControlRigModule(false))
+								{
+									static const FText Message(LOCTEXT("TurnIntoControlRigModuleMessage", "This rig requires some changes to the hierarchy to turn it into a module.\n\nWe'll try to recreate the hierarchy by relying on nodes in the construction event instead.\n\nDo you want to continue?"));
+									EAppReturnType::Type Ret = FMessageDialog::Open(EAppMsgType::YesNo, Message);
+									if(Ret == EAppReturnType::No)
+									{
+										return;
+									}
+								}
+								WeakBlueprint->TurnIntoControlRigModule(true);
+							}
+						}
+					}),
+					FCanExecuteAction::CreateLambda([WeakBlueprint]
+					{
+						if(WeakBlueprint.IsValid())
+						{
+							if(WeakBlueprint->IsControlRigModule())
+							{
+								return WeakBlueprint->CanTurnIntoStandaloneRig();
+							}
+							return WeakBlueprint->CanTurnIntoControlRigModule(true);
+						}
+						return false;
+					})
+				),
+				NAME_None,
+				TAttribute<FText>::CreateLambda([WeakBlueprint]()
+				{
+					static const FText StandaloneRig = LOCTEXT("StandaloneRig", "Standalone Rig"); 
+					static const FText RigModule = LOCTEXT("RigModule", "Rig Module");
+					if(WeakBlueprint.IsValid())
+					{
+						if(WeakBlueprint->IsControlRigModule())
+						{
+							return RigModule;
+						}
+					}
+					return StandaloneRig;
+				}),
+				TAttribute<FText>::CreateLambda([WeakBlueprint]()
+				{
+					static const FText StandaloneRigTooltip = LOCTEXT("StandaloneRigTooltip", "A standalone control rig."); 
+					static const FText RigModuleTooltip = LOCTEXT("RigModuleTooltip", "A rig module used to build rigs.");
+					if(WeakBlueprint.IsValid())
+					{
+						if(!WeakBlueprint->IsControlRigModule())
+						{
+							FString FailureReason;
+							if(!WeakBlueprint->CanTurnIntoControlRigModule(true, &FailureReason))
+							{
+								return FText::Format(
+									LOCTEXT("StandaloneRigTooltipFormat", "{0}\n\nThis rig cannot be turned into a module:\n\n{1}"),
+									StandaloneRigTooltip,
+									FText::FromString(FailureReason)
+								);
+							}
+							return StandaloneRigTooltip;
+						}
+					}
+					return RigModuleTooltip;
+				}),
+				TAttribute<FSlateIcon>::CreateLambda([WeakBlueprint]()
+				{
+					static const FSlateIcon ModuleIcon = FSlateIcon(FControlRigEditorStyle::Get().GetStyleSetName(), "ControlRig.Tree.Connector");
+					static const FSlateIcon RigIcon = FSlateIcon(FControlRigEditorStyle::Get().GetStyleSetName(),"ClassIcon.ControlRigBlueprint"); 
+					if(WeakBlueprint.IsValid())
+					{
+						if(WeakBlueprint->IsControlRigModule())
+						{
+							return ModuleIcon;
+						}
+					}
+					return RigIcon;
+				}),
+				EUserInterfaceActionType::Button
+			);
+		}
+
 		if(CVarControlRigShowTestingToolbar.GetValueOnAnyThread())
 		{
 			ToolbarBuilder.AddSeparator();
@@ -804,6 +902,7 @@ void FControlRigEditor::HandleSetObjectBeingDebugged(UObject* InObject)
 			PreviouslyDebuggedControlRig->GetHierarchy()->OnModified().RemoveAll(this);
 			PreviouslyDebuggedControlRig->OnPreForwardsSolve_AnyThread().RemoveAll(this);
 			PreviouslyDebuggedControlRig->OnPreConstructionForUI_AnyThread().RemoveAll(this);
+			PreviouslyDebuggedControlRig->OnPreConstruction_AnyThread().RemoveAll(this);
 			PreviouslyDebuggedControlRig->OnPostConstruction_AnyThread().RemoveAll(this);
 			PreviouslyDebuggedControlRig->ControlModified().RemoveAll(this);
 		}
@@ -859,15 +958,17 @@ void FControlRigEditor::HandleSetObjectBeingDebugged(UObject* InObject)
 		DebuggedControlRig->GetHierarchy()->OnModified().RemoveAll(this);
 		DebuggedControlRig->OnPreForwardsSolve_AnyThread().RemoveAll(this);
 		DebuggedControlRig->OnPreConstructionForUI_AnyThread().RemoveAll(this);
+		DebuggedControlRig->OnPreConstruction_AnyThread().RemoveAll(this);
 		DebuggedControlRig->OnPostConstruction_AnyThread().RemoveAll(this);
 		DebuggedControlRig->ControlModified().RemoveAll(this);
 
 		DebuggedControlRig->GetHierarchy()->OnModified().AddSP(this, &FControlRigEditor::OnHierarchyModified_AnyThread);
 		DebuggedControlRig->OnPreForwardsSolve_AnyThread().AddSP(this, &FControlRigEditor::OnPreForwardsSolve_AnyThread);
-		DebuggedControlRig->OnPreConstructionForUI_AnyThread().AddSP(this, &FControlRigEditor::OnPreConstruction_AnyThread);
+		DebuggedControlRig->OnPreConstructionForUI_AnyThread().AddSP(this, &FControlRigEditor::OnPreConstructionForUI_AnyThread);
+		DebuggedControlRig->OnPreConstruction_AnyThread().AddSP(this, &FControlRigEditor::OnPreConstruction_AnyThread);
 		DebuggedControlRig->OnPostConstruction_AnyThread().AddSP(this, &FControlRigEditor::OnPostConstruction_AnyThread);
 		DebuggedControlRig->ControlModified().AddSP(this, &FControlRigEditor::HandleOnControlModified);
-
+		
 		LastHierarchyHash = INDEX_NONE;
 
 		if(EditorSkelComp)
@@ -1118,6 +1219,14 @@ void FControlRigEditor::HandleVMCompiledEvent(UObject* InCompiledObject, URigVM*
 		RefreshDirectManipulationTextList();
 		bRefreshDirectionManipulationTargetsRequired = false;
 	}
+
+	if(UControlRigBlueprint* ControlRigBlueprint = GetControlRigBlueprint())
+	{
+		if(UControlRig* ControlRig = InVM->GetTypedOuter<UControlRig>())
+		{
+			ControlRigBlueprint->UpdateElementKeyRedirector(ControlRig);
+		}
+	}
 }
 
 void FControlRigEditor::SaveAsset_Execute()
@@ -1128,8 +1237,9 @@ void FControlRigEditor::SaveAsset_Execute()
 	UControlRigBlueprint* RigBlueprint = Cast<UControlRigBlueprint>(GetBlueprintObj());
 	if(const UControlRig* ControlRig = GetControlRig())
 	{
-		const UControlRig* CDO = ControlRig->GetClass()->GetDefaultObject<UControlRig>();
+		UControlRig* CDO = ControlRig->GetClass()->GetDefaultObject<UControlRig>();
 		CDO->DynamicHierarchy->CopyHierarchy(RigBlueprint->Hierarchy);
+		RigBlueprint->UpdateElementKeyRedirector(CDO);
 	}
 
 	FBlueprintActionDatabase& ActionDatabase = FBlueprintActionDatabase::Get();
@@ -1145,8 +1255,9 @@ void FControlRigEditor::SaveAssetAs_Execute()
 	UControlRigBlueprint* RigBlueprint = Cast<UControlRigBlueprint>(GetBlueprintObj());
 	if(const UControlRig* ControlRig = GetControlRig())
 	{
-		const UControlRig* CDO = ControlRig->GetClass()->GetDefaultObject<UControlRig>();
+		UControlRig* CDO = ControlRig->GetClass()->GetDefaultObject<UControlRig>();
 		CDO->DynamicHierarchy->CopyHierarchy(RigBlueprint->Hierarchy);
+		RigBlueprint->UpdateElementKeyRedirector(CDO);
 	}
 
 	FBlueprintActionDatabase& ActionDatabase = FBlueprintActionDatabase::Get();
@@ -1367,6 +1478,7 @@ void FControlRigEditor::HandleVMExecutedEvent(URigVMHost* InHost, const FName& I
 		{
 			if(!DebuggedControlRig->RigUnitManipulationInfos.IsEmpty())
 			{
+				const FRigHierarchyRedirectorGuard RedirectorGuard(DebuggedControlRig->GetHierarchy(), DebuggedControlRig->ElementKeyRedirector);
 				FControlRigExecuteContext& ExecuteContext = DebuggedControlRig->GetExtendedExecuteContext().GetPublicDataSafe<FControlRigExecuteContext>();
 				
 				for(const TSharedPtr<FRigDirectManipulationInfo>& ManipulationInfo : DebuggedControlRig->RigUnitManipulationInfos)
@@ -2296,6 +2408,11 @@ void FControlRigEditor::OnFinishedChangingProperties(const FPropertyChangedEvent
 		{
 			ControlRigBP->PropagateDrawInstructionsFromBPToInstances();
 		}
+
+		else if (PropertyChangedEvent.MemberProperty->GetNameCPP() == GET_MEMBER_NAME_STRING_CHECKED(UControlRigBlueprint, RigModuleSettings))
+		{
+			ControlRigBP->PropagateHierarchyFromBPToInstances();
+		}
 	}
 }
 
@@ -2747,13 +2864,15 @@ void FControlRigEditor::OnHierarchyModified(ERigHierarchyNotification InNotif, U
 			{
 				if(InHierarchy->GetConnectors().Num() == 1)
 				{
-					FNotificationInfo Info(LOCTEXT("FirstConnectorEncountered", "Looks like you have added the first connector. This rig can now be configured as a module in the class settings Hierarchy -> Module Settings."));
+					FNotificationInfo Info(LOCTEXT("FirstConnectorEncountered", "Looks like you have added the first connector. This rig will now be configured as a module, settings can be found in the class settings Hierarchy -> Module Settings."));
 					Info.bFireAndForget = true;
 					Info.FadeOutDuration = 5.0f;
 					Info.ExpireDuration = 5.0f;
 
 					TSharedPtr<SNotificationItem> NotificationPtr = FSlateNotificationManager::Get().AddNotification(Info);
 					NotificationPtr->SetCompletionState(SNotificationItem::CS_Success);
+
+					RigBlueprint->TurnIntoControlRigModule();
 				}
 			}
 			// no break - fall through
@@ -3074,6 +3193,14 @@ void FControlRigEditor::OnHierarchyModified_AnyThread(ERigHierarchyNotification
 	}
 }
 
+void FControlRigEditor::HandleRigTypeChanged(UControlRigBlueprint* InBlueprint)
+{
+	// todo: fire a notification.
+	// todo: reapply the preview mesh and react to it accordingly.
+
+	Compile();
+}
+
 void FControlRigEditor::SynchronizeViewportBoneSelection()
 {
 	UControlRigBlueprint* RigBlueprint = GetControlRigBlueprint();
@@ -4090,6 +4217,7 @@ void FControlRigEditor::HandleOnControlModified(UControlRig* Subject, FRigContro
 			// update the node based on the incoming pose. once that is done we'll need to compare the node instance
 			// with the settings on the node in the graph and update them accordingly.
 			FControlRigExecuteContext& ExecuteContext = DebuggedControlRig->GetExtendedExecuteContext().GetPublicDataSafe<FControlRigExecuteContext>();
+			const FRigHierarchyRedirectorGuard RedirectorGuard(DebuggedControlRig->GetHierarchy(), DebuggedControlRig->ElementKeyRedirector);
 			if(UnitInstance->UpdateDirectManipulationFromHierarchy(UnitNode, NodeInstance, ExecuteContext, ManipulationInfo))
 			{
 				UnitNode->UpdateHostFromStructInstance(DebuggedControlRig, NodeInstance);
@@ -4236,16 +4364,66 @@ void FControlRigEditor::OnPreForwardsSolve_AnyThread(UControlRig* InRig, const F
 	}
 }
 
-void FControlRigEditor::OnPreConstruction_AnyThread(UControlRig* InRig, const FName& InEventName)
+void FControlRigEditor::OnPreConstructionForUI_AnyThread(UControlRig* InRig, const FName& InEventName)
 {
 	bIsConstructionEventRunning = true;
 	const TArrayView<const FRigElementKey> Elements;
 	PreConstructionPose = InRig->GetHierarchy()->GetPose(false, ERigElementType::ToResetAfterConstructionEvent, Elements);
+
+	if(const UControlRigBlueprint* RigBlueprint = GetControlRigBlueprint())
+	{
+		if(RigBlueprint->IsControlRigModule())
+		{
+			ConnectorInfos = RigBlueprint->Hierarchy->GetConnectorInfos();
+		}
+	}
+}
+
+void FControlRigEditor::OnPreConstruction_AnyThread(UControlRig* InRig, const FName& InEventName)
+{
+	if(UControlRigBlueprint* RigBlueprint = GetControlRigBlueprint())
+	{
+		if(RigBlueprint->IsControlRigModule())
+		{
+			if(RigBlueprint->PreviewSkeletalMesh)
+			{
+				if(URigHierarchyController* Controller = InRig->GetHierarchy()->GetController(true))
+				{
+					const TArray<FRigElementKey> Bones = Controller->ImportBones(RigBlueprint->PreviewSkeletalMesh->GetSkeleton());
+
+					if(URigVM* VM = InRig->GetVM())
+					{
+						const int32 EntryIndex = VM->GetByteCode().FindEntryIndex(FRigUnit_PrepareForExecution::EventName);
+						if(EntryIndex != INDEX_NONE)
+						{
+							const int32 InstructionIndex = VM->GetByteCode().GetEntry(EntryIndex).InstructionIndex;
+							for(const FRigElementKey& Bone : Bones)
+							{
+								if(FRigBaseElement* Element = InRig->GetHierarchy()->Find(Bone))
+								{
+									Element->CreatedAtInstructionIndex = InstructionIndex;
+								}
+							}
+						}
+					}
+					//(void)RigBlueprint->UpdateConnectionRedirectorOnCDO();
+				}
+			}
+		}
+	}
 }
 
 void FControlRigEditor::OnPostConstruction_AnyThread(UControlRig* InRig, const FName& InEventName)
 {
 	bIsConstructionEventRunning = false;
+
+	if(UControlRigBlueprint* RigBlueprint = GetControlRigBlueprint())
+	{
+		if(RigBlueprint->IsControlRigModule())
+		{
+			RigBlueprint->Hierarchy->RestoreConnectorsFromInfos(ConnectorInfos);
+		}
+	}
 	
 	const int32 HierarchyHash = InRig->GetHierarchy()->GetTopologyHash(false);
 	if(LastHierarchyHash != HierarchyHash)
diff --git a/Engine/Plugins/Animation/ControlRig/Source/ControlRigEditor/Private/Editor/ControlRigEditor.h b/Engine/Plugins/Animation/ControlRig/Source/ControlRigEditor/Private/Editor/ControlRigEditor.h
--- a/Engine/Plugins/Animation/ControlRig/Source/ControlRigEditor/Private/Editor/ControlRigEditor.h
+++ b/Engine/Plugins/Animation/ControlRig/Source/ControlRigEditor/Private/Editor/ControlRigEditor.h
@@ -139,6 +139,8 @@ public:
 	void OnHierarchyModified(ERigHierarchyNotification InNotif, URigHierarchy* InHierarchy, const FRigBaseElement* InElement);
 	void OnHierarchyModified_AnyThread(ERigHierarchyNotification InNotif, URigHierarchy* InHierarchy, const FRigBaseElement* InElement);
 
+	void HandleRigTypeChanged(UControlRigBlueprint* InBlueprint);
+
 	const FName RigHierarchyToGraphDragAndDropMenuName = TEXT("ControlRigEditor.RigHierarchyToGraphDragAndDropMenu");
 	void CreateRigHierarchyToGraphDragAndDropMenu() const;
 	void OnGraphNodeDropToPerform(TSharedPtr<FGraphNodeDragDropOp> InDragDropOp, UEdGraph* InGraph, const FVector2D& InNodePosition, const FVector2D& InScreenPosition);
@@ -294,9 +296,11 @@ protected:
 	TWeakObjectPtr<AStaticMeshActor> WeakGroundActorPtr;
 
 	void OnPreForwardsSolve_AnyThread(UControlRig* InRig, const FName& InEventName);
+	void OnPreConstructionForUI_AnyThread(UControlRig* InRig, const FName& InEventName);
 	void OnPreConstruction_AnyThread(UControlRig* InRig, const FName& InEventName);
 	void OnPostConstruction_AnyThread(UControlRig* InRig, const FName& InEventName);
 	FRigPose PreConstructionPose;
+	TArray<FRigConnectorInfo> ConnectorInfos;
 
 	bool bIsConstructionEventRunning;
 	uint32 LastHierarchyHash;
diff --git a/Engine/Plugins/Animation/ControlRig/Source/ControlRigEditor/Private/Editor/SRigHierarchy.cpp b/Engine/Plugins/Animation/ControlRig/Source/ControlRigEditor/Private/Editor/SRigHierarchy.cpp
--- a/Engine/Plugins/Animation/ControlRig/Source/ControlRigEditor/Private/Editor/SRigHierarchy.cpp
+++ b/Engine/Plugins/Animation/ControlRig/Source/ControlRigEditor/Private/Editor/SRigHierarchy.cpp
@@ -88,6 +88,15 @@ FString FRigElementHierarchyDragDropOp::GetJoinedElementNames() const
 	return FString::Join(ElementNameStrings, TEXT(","));
 }
 
+bool FRigElementHierarchyDragDropOp::IsDraggingSingleConnector() const
+{
+	if(Elements.Num() == 1)
+	{
+		return Elements[0].Type == ERigElementType::Connector;
+	}
+	return false;
+}
+
 ///////////////////////////////////////////////////////////
 
 const FName SRigHierarchy::ContextMenuName = TEXT("ControlRigEditor.RigHierarchy.ContextMenu");
@@ -134,6 +143,7 @@ void SRigHierarchy::Construct(const FArguments& InArgs, TSharedRef<FControlRigEd
 	Delegates.OnCanAcceptDrop = FOnRigTreeCanAcceptDrop::CreateSP(this, &SRigHierarchy::OnCanAcceptDrop);
 	Delegates.OnAcceptDrop = FOnRigTreeAcceptDrop::CreateSP(this, &SRigHierarchy::OnAcceptDrop);
 	Delegates.OnDragDetected = FOnDragDetected::CreateSP(this, &SRigHierarchy::OnDragDetected);
+	Delegates.OnGetResolvedKey = FOnRigTreeGetResolvedKey::CreateSP(this, &SRigHierarchy::OnGetResolvedKey);
 
 	ChildSlot
 	[
@@ -527,6 +537,18 @@ void SRigHierarchy::OnFilterTextChanged(const FText& SearchText)
 
 void SRigHierarchy::RefreshTreeView(bool bRebuildContent)
 {
+	if(const URigHierarchy* Hierarchy = GetHierarchy())
+	{
+		// is the rig currently running
+		if(Hierarchy->HasExecuteContext())
+		{
+			FFunctionGraphTask::CreateAndDispatchWhenReady([this, bRebuildContent]()
+			{
+				RefreshTreeView(bRebuildContent);
+			}, TStatId(), NULL, ENamedThreads::GameThread);
+		}
+	}
+	
 	bool bDummySuspensionFlag = false;
 	bool* SuspensionFlagPtr = &bDummySuspensionFlag;
 	if (ControlRigEditor.IsValid())
@@ -894,7 +916,11 @@ void SRigHierarchy::OnPostConstruction_AnyThread(UControlRig* InRig, const FName
 
 	bIsConstructionEventRunning = false;
 
-	const int32 HierarchyHash = InRig->GetHierarchy()->GetTopologyHash(false);
+	const URigHierarchy* Hierarchy = InRig->GetHierarchy();
+	const int32 HierarchyHash = HashCombine(
+		Hierarchy->GetTopologyHash(false),
+		InRig->ElementKeyRedirector.GetHash());
+
 	if(LastHierarchyHash != HierarchyHash)
 	{
 		LastHierarchyHash = HierarchyHash;
@@ -1524,6 +1550,18 @@ void SRigHierarchy::CreateImportMenu(FMenuBuilder& MenuBuilder)
 	);
 }
 
+FRigElementKey SRigHierarchy::OnGetResolvedKey(const FRigElementKey& InKey)
+{
+	if (const UControlRigBlueprint* Blueprint = ControlRigEditor.Pin()->GetControlRigBlueprint())
+	{
+		if(const FRigElementKey* ResolvedKey = Blueprint->ConnectionMap.Find(InKey))
+		{
+			return *ResolvedKey;
+		}
+	}
+	return InKey;
+}
+
 void SRigHierarchy::ImportHierarchy(const FAssetData& InAssetData)
 {
 	if (bIsChangingRigHierarchy)
@@ -1893,7 +1931,31 @@ void SRigHierarchy::HandleNewItem(ERigElementType InElementType, bool bIsAnimati
 				}
 				case ERigElementType::Connector:
 				{
-					NewItemKey = Controller->AddConnector(NewElementName, ParentKey, ParentTransform, true, FRigConnectorSettings(), true);
+					FString FailureReason;
+					if(!ControlRigBlueprint->CanTurnIntoControlRigModule(false, &FailureReason))
+					{
+						if(ControlRigBlueprint->Hierarchy->Num(ERigElementType::Connector) == 0)
+						{
+							static constexpr TCHAR Format[] = TEXT("Connector cannot be created: %s");
+							UE_LOG(LogControlRig, Warning, Format, *FailureReason);
+							FNotificationInfo Info(FText::FromString(FString::Printf(Format, *FailureReason)));
+							Info.bUseSuccessFailIcons = true;
+							Info.Image = FAppStyle::GetBrush(TEXT("MessageLog.Warning"));
+							Info.bFireAndForget = true;
+							Info.bUseThrobber = true;
+							Info.FadeOutDuration = 2.f;
+							Info.ExpireDuration = 8.f;;
+							TSharedPtr<SNotificationItem> NotificationPtr = FSlateNotificationManager::Get().AddNotification(Info);
+							if (NotificationPtr)
+							{
+								NotificationPtr->SetCompletionState(SNotificationItem::CS_Fail);
+							}
+							return;
+						}
+					}
+
+					NewItemKey = Controller->AddConnector(NewElementName, ParentTransform, true, FRigConnectorSettings(), true);
+					(void)ResolveConnector(NewItemKey, ParentKey);
 					break;
 				}
 				default:
@@ -2316,7 +2378,7 @@ TOptional<EItemDropZone> SRigHierarchy::OnCanAcceptDrop(const FDragDropEvent& Dr
 
 			for (const FRigElementKey& DraggedKey : RigDragDropOp->GetElements())
 			{
-				if(Hierarchy->IsProcedural(DraggedKey))
+				if(Hierarchy->IsProcedural(DraggedKey) && !RigDragDropOp->IsDraggingSingleConnector())
 				{
 					return InvalidDropZone;
 				}
@@ -2327,7 +2389,7 @@ TOptional<EItemDropZone> SRigHierarchy::OnCanAcceptDrop(const FDragDropEvent& Dr
 					return InvalidDropZone;
 				}
 
-				if(DropZone == EItemDropZone::OntoItem)
+				if(DropZone == EItemDropZone::OntoItem && !RigDragDropOp->IsDraggingSingleConnector())
 				{
 					if(Hierarchy->IsParentedTo(TargetKey, DraggedKey))
 					{
@@ -2338,11 +2400,17 @@ TOptional<EItemDropZone> SRigHierarchy::OnCanAcceptDrop(const FDragDropEvent& Dr
 		}
 
 		// don't allow dragging onto procedural items
-		if(TargetKey.IsValid() && !GetDefaultHierarchy()->Contains(TargetKey))
+		if(TargetKey.IsValid() && !GetDefaultHierarchy()->Contains(TargetKey) && !RigDragDropOp->IsDraggingSingleConnector())
 		{
 			return InvalidDropZone;
 		}
 
+		// resolve connector rules
+		if(RigDragDropOp->IsDraggingSingleConnector())
+		{
+			// todo resolve rules
+		}
+
 		switch (TargetKey.Type)
 		{
 			case ERigElementType::Bone:
@@ -2426,6 +2494,7 @@ FReply SRigHierarchy::OnAcceptDrop(const FDragDropEvent& DragDropEvent, EItemDro
 			FRigElementKey TargetKey;
 			int32 LocalIndex = INDEX_NONE;
 
+
 			if (TargetItem.IsValid())
 			{
 				switch(DropZone)
@@ -2449,6 +2518,11 @@ FReply SRigHierarchy::OnAcceptDrop(const FDragDropEvent& DragDropEvent, EItemDro
 					}
 				}
 			}
+
+			if(RigDragDropOp->IsDraggingSingleConnector())
+			{
+				return ResolveConnector(RigDragDropOp->GetElements()[0], TargetKey);
+			}
 			
 			return ReparentOrMatchTransform(RigDragDropOp->GetElements(), TargetKey, bReparentItems, LocalIndex);			
 		}
@@ -3055,6 +3129,25 @@ FReply SRigHierarchy::ReparentOrMatchTransform(const TArray<FRigElementKey>& Dra
 
 }
 
+FReply SRigHierarchy::ResolveConnector(const FRigElementKey& DraggedKey, const FRigElementKey& TargetKey)
+{
+	if (UControlRigBlueprint* Blueprint = ControlRigEditor.Pin()->GetControlRigBlueprint())
+	{
+		if (const URigHierarchy* DebuggedHierarchy = GetHierarchy())
+		{
+			if(DebuggedHierarchy->Contains(DraggedKey))
+			{
+				if(Blueprint->ResolveConnector(DraggedKey, TargetKey))
+				{
+					RefreshTreeView();
+					return FReply::Handled();
+				}
+			}
+		}
+	}
+	return FReply::Unhandled();
+}
+
 void SRigHierarchy::HandleSetInitialTransformFromClosestBone()
 {
 	if (IsControlOrNullSelected(false))
diff --git a/Engine/Plugins/Animation/ControlRig/Source/ControlRigEditor/Private/Editor/SRigHierarchy.h b/Engine/Plugins/Animation/ControlRig/Source/ControlRigEditor/Private/Editor/SRigHierarchy.h
--- a/Engine/Plugins/Animation/ControlRig/Source/ControlRigEditor/Private/Editor/SRigHierarchy.h
+++ b/Engine/Plugins/Animation/ControlRig/Source/ControlRigEditor/Private/Editor/SRigHierarchy.h
@@ -46,6 +46,8 @@ public:
 
 	FString GetJoinedElementNames() const;
 
+	bool IsDraggingSingleConnector() const;
+
 private:
 
 	/** Data for the property paths this item represents */
@@ -195,6 +197,7 @@ private:
 	URigHierarchy* GetHierarchy() const;
 	URigHierarchy* GetDefaultHierarchy() const;
 	const URigHierarchy* GetHierarchyForTreeView() const { return GetHierarchy(); }
+	FRigElementKey OnGetResolvedKey(const FRigElementKey& InKey);
 	
 	void ImportHierarchy(const FAssetData& InAssetData);
 	void CreateImportMenu(FMenuBuilder& MenuBuilder);
@@ -218,7 +221,8 @@ private:
 	void HandleParent(const FToolMenuContext& Context);
 	void HandleAlign(const FToolMenuContext& Context);
 	FReply ReparentOrMatchTransform(const TArray<FRigElementKey>& DraggedKeys, FRigElementKey TargetKey, bool bReparentItems, int32 LocalIndex = INDEX_NONE);
-
+	FReply ResolveConnector(const FRigElementKey& DraggedKey, const FRigElementKey& TargetKey);
+	
 	FName CreateUniqueName(const FName& InBaseName, ERigElementType InElementType) const;
 
 	void ClearDetailPanel() const;
diff --git a/Engine/Plugins/Animation/ControlRig/Source/ControlRigEditor/Private/Editor/SRigHierarchyTreeView.cpp b/Engine/Plugins/Animation/ControlRig/Source/ControlRigEditor/Private/Editor/SRigHierarchyTreeView.cpp
--- a/Engine/Plugins/Animation/ControlRig/Source/ControlRigEditor/Private/Editor/SRigHierarchyTreeView.cpp
+++ b/Engine/Plugins/Animation/ControlRig/Source/ControlRigEditor/Private/Editor/SRigHierarchyTreeView.cpp
@@ -481,6 +481,14 @@ bool SRigHierarchyTreeView::AddElement(const FRigBaseElement* InElement)
 		if(const URigHierarchy* Hierarchy = Delegates.GetHierarchy())
 		{
 			FRigElementKey ParentKey = Hierarchy->GetFirstParent(InElement->GetKey());
+			if(InElement->GetType() == ERigElementType::Connector)
+			{
+				ParentKey = Delegates.GetResolvedKey(InElement->GetKey());
+				if(ParentKey == InElement->GetKey())
+				{
+					ParentKey.Reset();
+				}
+			}
 
 			TArray<FRigElementWeight> ParentWeights = Hierarchy->GetParentWeightArray(InElement->GetKey());
 			if(ParentWeights.Num() > 0)
@@ -621,12 +629,26 @@ void SRigHierarchyTreeView::RefreshTreeView(bool bRebuildContent)
 		const URigHierarchy* Hierarchy = Delegates.GetHierarchy();
 		if(Hierarchy)
 		{
+			TArray<FRigConnectorElement*> Connectors;
 			Hierarchy->Traverse([&](FRigBaseElement* Element, bool& bContinue)
 			{
-				AddElement(Element);
+				if(FRigConnectorElement* Connector = Cast<FRigConnectorElement>(Element))
+				{
+					Connectors.Add(Connector);
+				}
+				else
+				{
+					AddElement(Element);
+				}
 				bContinue = true;
 			});
 
+			// add all of the connectors. their parent relationship in the tree represents resolval
+			for(FRigConnectorElement* Connector : Connectors)
+			{
+				AddElement(Connector);
+			}
+
 			// expand all elements upon the initial construction of the tree
 			if (ExpansionState.Num() == 0)
 			{
diff --git a/Engine/Plugins/Animation/ControlRig/Source/ControlRigEditor/Private/Editor/SRigHierarchyTreeView.h b/Engine/Plugins/Animation/ControlRig/Source/ControlRigEditor/Private/Editor/SRigHierarchyTreeView.h
--- a/Engine/Plugins/Animation/ControlRig/Source/ControlRigEditor/Private/Editor/SRigHierarchyTreeView.h
+++ b/Engine/Plugins/Animation/ControlRig/Source/ControlRigEditor/Private/Editor/SRigHierarchyTreeView.h
@@ -63,6 +63,7 @@ DECLARE_DELEGATE_RetVal(const FRigTreeDisplaySettings&, FOnGetRigTreeDisplaySett
 DECLARE_DELEGATE_RetVal_TwoParams(FName, FOnRigTreeRenameElement, const FRigElementKey& /*OldKey*/, const FString& /*NewName*/);
 DECLARE_DELEGATE_RetVal_ThreeParams(bool, FOnRigTreeVerifyElementNameChanged, const FRigElementKey& /*OldKey*/, const FString& /*NewName*/, FText& /*OutErrorMessage*/);
 DECLARE_DELEGATE_RetVal_TwoParams(bool, FOnRigTreeCompareKeys, const FRigElementKey& /*A*/, const FRigElementKey& /*B*/);
+DECLARE_DELEGATE_RetVal_OneParam(FRigElementKey, FOnRigTreeGetResolvedKey, const FRigElementKey&);
 
 typedef STableRow<TSharedPtr<FRigTreeElement>>::FOnCanAcceptDrop FOnRigTreeCanAcceptDrop;
 typedef STableRow<TSharedPtr<FRigTreeElement>>::FOnAcceptDrop FOnRigTreeAcceptDrop;
@@ -86,6 +87,7 @@ struct CONTROLRIGEDITOR_API FRigTreeDelegates
 	FOnRigTreeMouseButtonDoubleClick OnMouseButtonDoubleClick;
 	FOnRigTreeSetExpansionRecursive OnSetExpansionRecursive;
 	FOnRigTreeCompareKeys OnCompareKeys;
+	FOnRigTreeGetResolvedKey OnGetResolvedKey;
 
 	FRigTreeDelegates()
 	{
@@ -138,6 +140,15 @@ struct CONTROLRIGEDITOR_API FRigTreeDelegates
 		OnSelectionChanged.ExecuteIfBound(Selection, SelectInfo);
 	}
 
+	FRigElementKey GetResolvedKey(const FRigElementKey& InKey)
+	{
+		if(OnGetResolvedKey.IsBound())
+		{
+			return OnGetResolvedKey.Execute(InKey);
+		}
+		return InKey;
+	}
+
 	static FRigTreeDisplaySettings DefaultDisplaySettings;
 	bool bIsChangingRigHierarchy;
 };
