Control Rig: Use soft pointers for UDataAssetLink and UControlRigShapeLibraryLink so that they can refer to assets across plugins and fail quietly in the absence of the pointed-to assets.

#rnx

diff --git a/Engine/Plugins/Animation/ControlRig/Source/ControlRig/Private/ControlRigAssetUserData.cpp b/Engine/Plugins/Animation/ControlRig/Source/ControlRig/Private/ControlRigAssetUserData.cpp
--- a/Engine/Plugins/Animation/ControlRig/Source/ControlRig/Private/ControlRigAssetUserData.cpp
+++ b/Engine/Plugins/Animation/ControlRig/Source/ControlRig/Private/ControlRigAssetUserData.cpp
@@ -2,12 +2,12 @@
 
 #include "ControlRigAssetUserData.h"
 
-void UControlRigShapeLibraryLink::SetShapeLibrary(UControlRigShapeLibrary* InShapeLibrary)
+void UControlRigShapeLibraryLink::SetShapeLibrary(TSoftObjectPtr<UControlRigShapeLibrary> InShapeLibrary)
 {
 	InvalidateCache();
 	ShapeLibrary = InShapeLibrary;
 	ShapeNames.Reset();
-	if(ShapeLibrary)
+	if(ShapeLibrary.IsValid())
 	{
 		for(const FControlRigShapeDefinition& Shape : ShapeLibrary->Shapes)
 		{
@@ -24,22 +24,22 @@ const UNameSpacedUserData::FUserData* UControlRigShapeLibraryLink::GetUserData(c
 		return ResultFromSuper;
 	}
 
-	if(ShapeLibrary)
+	if(const UControlRigShapeLibrary* LoadedShapeLibrary = ShapeLibrary.LoadSynchronous())
 	{
 		if(InPath.Equals(GET_MEMBER_NAME_STRING_CHECKED(UControlRigShapeLibraryLink, ShapeLibrary), ESearchCase::CaseSensitive))
 		{
 			static const FProperty* ShapeLibraryProperty = FindPropertyByName(StaticClass(), GET_MEMBER_NAME_CHECKED(UControlRigShapeLibraryLink, ShapeLibrary));
-			return StoreCacheForUserData({InPath, ShapeLibraryProperty, (const uint8*)ShapeLibrary});
+			return StoreCacheForUserData({InPath, ShapeLibraryProperty, reinterpret_cast<const uint8*>(LoadedShapeLibrary)});
 		}
 		if(InPath.Equals(DefaultShapePath, ESearchCase::CaseSensitive))
 		{
 			static const FProperty* DefaultShapeProperty = FindPropertyByName(FControlRigShapeDefinition::StaticStruct(), GET_MEMBER_NAME_CHECKED(FControlRigShapeDefinition, ShapeName));
-			return StoreCacheForUserData({InPath, DefaultShapeProperty, (const uint8*)&ShapeLibrary->DefaultShape.ShapeName});
+			return StoreCacheForUserData({InPath, DefaultShapeProperty, reinterpret_cast<const uint8*>(&LoadedShapeLibrary->DefaultShape.ShapeName)});
 		}
 		if(InPath.Equals(ShapeNamesPath, ESearchCase::CaseSensitive))
 		{
 			static const FProperty* ShapeNamesProperty = FindPropertyByName(StaticClass(), GET_MEMBER_NAME_CHECKED(UControlRigShapeLibraryLink, ShapeNames));
-			return StoreCacheForUserData({InPath, ShapeNamesProperty, (const uint8*)&ShapeNames});
+			return StoreCacheForUserData({InPath, ShapeNamesProperty, reinterpret_cast<const uint8*>(&ShapeNames)});
 		}
 		if(OutErrorMessage && OutErrorMessage->IsEmpty())
 		{
@@ -63,7 +63,7 @@ const TArray<const UNameSpacedUserData::FUserData*>& UControlRigShapeLibraryLink
 		return ResultFromSuper;
 	}
 
-	if(ShapeLibrary)
+	if(!ShapeLibrary.IsNull())
 	{
 		// UControlRigShapeLibraryLink doesn't offer any arrays other than the top level 
 		if(InParentPath.IsEmpty())
@@ -71,10 +71,13 @@ const TArray<const UNameSpacedUserData::FUserData*>& UControlRigShapeLibraryLink
 			const FUserData* ShapeLibraryUserData = GetUserData(GET_MEMBER_NAME_STRING_CHECKED(UControlRigShapeLibraryLink, ShapeLibrary), OutErrorMessage);
 			const FUserData* DefaultShapeUserData = GetUserData(DefaultShapePath, OutErrorMessage);
 			const FUserData* ShapeNamesUserData = GetUserData(ShapeNamesPath, OutErrorMessage);
-			check(ShapeLibraryUserData);
-			check(DefaultShapeUserData);
-			check(ShapeNamesUserData);
-			return StoreCacheForUserDataArray(InParentPath, {ShapeLibraryUserData, DefaultShapeUserData, ShapeNamesUserData});
+
+			if (ensure(ShapeLibraryUserData) &&
+				ensure(DefaultShapeUserData) &&
+				ensure(ShapeNamesUserData))
+			{
+				return StoreCacheForUserDataArray(InParentPath, {ShapeLibraryUserData, DefaultShapeUserData, ShapeNamesUserData});
+			}
 		}
 		return EmptyUserDatas;
 	}
@@ -92,11 +95,11 @@ void UControlRigShapeLibraryLink::PostEditChangeProperty(FPropertyChangedEvent&
 {
 	Super::PostEditChangeProperty(PropertyChangedEvent);
 
-	static const FProperty* ShapeLibraryProperty= FindPropertyByName(GetClass(), GET_MEMBER_NAME_CHECKED(UControlRigShapeLibraryLink, ShapeLibrary));
+	static const FProperty* ShapeLibraryProperty = FindPropertyByName(GetClass(), GET_MEMBER_NAME_CHECKED(UControlRigShapeLibraryLink, ShapeLibrary));
 	if(PropertyChangedEvent.Property == ShapeLibraryProperty ||
 		PropertyChangedEvent.MemberProperty == ShapeLibraryProperty)
 	{
-		SetShapeLibrary(ShapeLibrary);
+		SetShapeLibrary(ShapeLibrary.Get());
 	}
 }
 
diff --git a/Engine/Plugins/Animation/ControlRig/Source/ControlRig/Public/ControlRigAssetUserData.h b/Engine/Plugins/Animation/ControlRig/Source/ControlRig/Public/ControlRigAssetUserData.h
--- a/Engine/Plugins/Animation/ControlRig/Source/ControlRig/Public/ControlRigAssetUserData.h
+++ b/Engine/Plugins/Animation/ControlRig/Source/ControlRig/Public/ControlRigAssetUserData.h
@@ -18,13 +18,13 @@ public:
 
 	/** If assigned, the data asset link will provide access to the data asset's content. */
 	UPROPERTY(EditAnywhere, BlueprintReadWrite, Category = General, BlueprintGetter = GetShapeLibrary, BlueprintSetter = SetShapeLibrary, Meta = (DisplayAfter="NameSpace"))
-	TObjectPtr<UControlRigShapeLibrary> ShapeLibrary;
+	TSoftObjectPtr<UControlRigShapeLibrary> ShapeLibrary;
 
 	UFUNCTION(BlueprintGetter)
-	UControlRigShapeLibrary* GetShapeLibrary() const { return ShapeLibrary; }
+	TSoftObjectPtr<UControlRigShapeLibrary> GetShapeLibrary() const { return ShapeLibrary; }
 
 	UFUNCTION(BlueprintSetter)
-	void SetShapeLibrary(UControlRigShapeLibrary* InShapeLibrary);
+	void SetShapeLibrary(TSoftObjectPtr<UControlRigShapeLibrary> InShapeLibrary);
 
 	virtual const FUserData* GetUserData(const FString& InPath, FString* OutErrorMessage = nullptr) const override;
 	virtual const TArray<const FUserData*>& GetUserDataArray(const FString& InParentPath = FString(), FString* OutErrorMessage = nullptr) const override;
diff --git a/Engine/Plugins/Runtime/RigVM/Source/RigVM/Private/RigVMCore/RigVMAssetUserData.cpp b/Engine/Plugins/Runtime/RigVM/Source/RigVM/Private/RigVMCore/RigVMAssetUserData.cpp
--- a/Engine/Plugins/Runtime/RigVM/Source/RigVM/Private/RigVMCore/RigVMAssetUserData.cpp
+++ b/Engine/Plugins/Runtime/RigVM/Source/RigVM/Private/RigVMCore/RigVMAssetUserData.cpp
@@ -372,12 +372,12 @@ bool UNameSpacedUserData::IsPropertySupported(const FProperty* InProperty, const
 	return true;
 }
 
-void UDataAssetLink::SetDataAsset(UDataAsset* InDataAsset)
+void UDataAssetLink::SetDataAsset(TSoftObjectPtr<UDataAsset> InDataAsset)
 {
 	InvalidateCache();
 	DataAsset = InDataAsset;
 
-	if(NameSpace.IsEmpty() && DataAsset)
+	if(NameSpace.IsEmpty() && DataAsset.IsValid())
 	{
 		NameSpace = DataAsset->GetName();
 	}
@@ -391,10 +391,10 @@ const UNameSpacedUserData::FUserData* UDataAssetLink::GetUserData(const FString&
 		return ResultFromSuper;
 	}
 
-	if(DataAsset)
+	if(const UDataAsset* LoadedDataAsset = DataAsset.LoadSynchronous())
 	{
 		// this method caches as well - so the next time around Super::GetUserData should return the cache 
-		return GetUserDataWithinStruct(DataAsset->GetClass(), (const uint8*)DataAsset, InPath, FString(), OutErrorMessage);
+		return GetUserDataWithinStruct(LoadedDataAsset->GetClass(), reinterpret_cast<const uint8*>(LoadedDataAsset), InPath, FString(), OutErrorMessage);
 	}
 
 	if(OutErrorMessage && OutErrorMessage->IsEmpty())
@@ -413,13 +413,13 @@ const TArray<const UNameSpacedUserData::FUserData*>& UDataAssetLink::GetUserData
 		return ResultFromSuper;
 	}
 
-	if(DataAsset)
+	if(const UDataAsset* LoadedDataAsset = DataAsset.LoadSynchronous())
 	{
 		// we should only get here if we haven't cached this user data array before.
 		if(InParentPath.IsEmpty())
 		{
 			// this method caches as well - so the next time around Super::GetUserDataArray should return the cache 
-			return GetUserDataArrayWithinStruct(DataAsset->GetClass(), (const uint8*)DataAsset, InParentPath, OutErrorMessage);
+			return GetUserDataArrayWithinStruct(LoadedDataAsset->GetClass(), reinterpret_cast<const uint8*>(LoadedDataAsset), InParentPath, OutErrorMessage);
 		}
 
 		if(const FUserData* ParentUserData = GetUserData(InParentPath, OutErrorMessage))
diff --git a/Engine/Plugins/Runtime/RigVM/Source/RigVM/Public/RigVMCore/RigVMAssetUserData.h b/Engine/Plugins/Runtime/RigVM/Source/RigVM/Public/RigVMCore/RigVMAssetUserData.h
--- a/Engine/Plugins/Runtime/RigVM/Source/RigVM/Public/RigVMCore/RigVMAssetUserData.h
+++ b/Engine/Plugins/Runtime/RigVM/Source/RigVM/Public/RigVMCore/RigVMAssetUserData.h
@@ -192,13 +192,13 @@ public:
 
 	/** If assigned, the data asset link will provide access to the data asset's content. */
 	UPROPERTY(EditAnywhere, BlueprintReadWrite, Category = General, BlueprintGetter = GetDataAsset, BlueprintSetter = SetDataAsset, Meta = (DisplayAfter="NameSpace"))
-	TObjectPtr<UDataAsset> DataAsset;
+	TSoftObjectPtr<UDataAsset> DataAsset;
 
 	UFUNCTION(BlueprintGetter)
-	UDataAsset* GetDataAsset() const { return DataAsset; }
+	TSoftObjectPtr<UDataAsset> GetDataAsset() const { return DataAsset; }
 
 	UFUNCTION(BlueprintSetter)
-	void SetDataAsset(UDataAsset* InDataAsset);
+	void SetDataAsset(TSoftObjectPtr<UDataAsset> InDataAsset);
 
 	virtual const FUserData* GetUserData(const FString& InPath, FString* OutErrorMessage = nullptr) const override;
 	virtual const TArray<const FUserData*>& GetUserDataArray(const FString& InParentPath = FString(), FString* OutErrorMessage = nullptr) const override;
