Control Rig: Modules - WIP

Adding namespaces for spawned controls in modules
Adding support for capsule based tags in SRigHierarchy

diff --git a/Engine/Plugins/Animation/ControlRig/Source/ControlRig/Private/ControlRig.cpp b/Engine/Plugins/Animation/ControlRig/Source/ControlRig/Private/ControlRig.cpp
--- a/Engine/Plugins/Animation/ControlRig/Source/ControlRig/Private/ControlRig.cpp
+++ b/Engine/Plugins/Animation/ControlRig/Source/ControlRig/Private/ControlRig.cpp
@@ -524,6 +524,64 @@ FText UControlRig::GetToolTipText() const
 }
 #endif
 
+bool UControlRig::AllConnectorsAreResolved(FString* OutFailureReason, FRigElementKey* OutConnector) const
+{
+	if(const URigHierarchy* Hierarchy = GetHierarchy())
+	{
+		// todo: introduce a cache here based on ElementKeyDirector hash and
+		// topology hash of the hierarchy
+		
+		const TArray<FRigConnectorElement*> Connectors = Hierarchy->GetConnectors(false);
+
+		// collect the connection map
+		TMap<FRigElementKey, FRigElementKey> ConnectionMap;
+		for(const FRigConnectorElement* Connector : Connectors)
+		{
+			if(const FCachedRigElement* Cache = ElementKeyRedirector.Find(Connector->GetKey()))
+			{
+				if(const_cast<FCachedRigElement*>(Cache)->UpdateCache(Hierarchy))
+				{
+					ConnectionMap.Add(Connector->GetKey(), Cache->GetKey());
+				}
+				else
+				{
+					if(OutFailureReason)
+					{
+						static constexpr TCHAR Format[] = TEXT("Connector '%s' has invalid target '%s'.");
+						*OutFailureReason = FString::Printf(Format, *Connector->GetName(), *Cache->GetKey().ToString());
+					}
+					if(OutConnector)
+					{
+						*OutConnector = Connector->GetKey();
+					}
+					return false;
+				}
+			}
+			else
+			{
+				if(OutFailureReason)
+				{
+					static constexpr TCHAR Format[] = TEXT("Connector '%s' is not resolved.");
+					*OutFailureReason = FString::Printf(Format, *Connector->GetName());
+				}
+				if(OutConnector)
+				{
+					*OutConnector = Connector->GetKey();
+				}
+				return false;
+			}
+		}
+
+		// now that we have the connection try if we can connect this way
+		const FRigConnectionInfo ConnectionInfo(ConnectionMap, Hierarchy, Hierarchy);
+		if(!Hierarchy->CanConnect(&ConnectionInfo, OutFailureReason, OutConnector))
+		{
+			return false;
+		}
+	}
+	return true;
+}
+
 bool UControlRig::Execute(const FName& InEventName)
 {
 	if(!CanExecute())
@@ -1227,6 +1285,18 @@ bool UControlRig::Execute(const FName& InEventName)
 
 bool UControlRig::Execute_Internal(const FName& InEventName)
 {
+	if(IsRigModule())
+	{
+		FString ConnectorWarning;
+		if(!AllConnectorsAreResolved(&ConnectorWarning))
+		{
+#if WITH_EDITOR
+			LogOnce(EMessageSeverity::Warning, INDEX_NONE, ConnectorWarning);
+#endif
+			return false;
+		}
+	}
+	
 	if (VM)
 	{
 		FRigVMExtendedExecuteContext& Context = GetExtendedExecuteContext();
@@ -1320,7 +1390,7 @@ bool UControlRig::Execute_Internal(const FName& InEventName)
 		FControlRigExecuteContext& PublicContext = Context.GetPublicDataSafe<FControlRigExecuteContext>();
 		TGuardValue<FName> ModuleInstanceNameGuard(PublicContext.ModuleInstanceName, GetModuleInstanceName());
 		TGuardValue<FString> ModuleInstanceNameSpaceGuard(PublicContext.ModuleInstanceNameSpace, GetRigModuleNameSpace());
-		FRigHierarchyRedirectorGuard ElementRedirectorGuard(DynamicHierarchy, GetElementKeyRedirector());
+		FRigHierarchyRedirectorGuard ElementRedirectorGuard(this);
 
 		TArray<TRigVMMemoryStorage*> LocalMemory = VM->GetLocalMemoryArray(Context);
 		const bool bSuccess = VM->Execute(Context, LocalMemory, InEventName) != ERigVMExecuteResult::Failed;
@@ -1584,10 +1654,16 @@ const FString& UControlRig::GetRigModuleNameSpace() const
 		if(const UControlRig* ParentRig = GetParentRig())
 		{
 			const FString& ParentNameSpace = ParentRig->GetRigModuleNameSpace();
-			static constexpr TCHAR JoinFormat[] = TEXT("%s%s::");
+			static constexpr TCHAR JoinFormat[] = TEXT("%s%s:");
 			RigModuleNameSpace = FString::Printf(JoinFormat, *ParentNameSpace, *GetName());
 			return RigModuleNameSpace;
 		}
+
+		// this means we are not an instance - so we are authoring the rig module right now.
+		// for visualization purposes we'll use the name of the rig module chosen in the settings.
+		static constexpr TCHAR NameSpaceFormat[] = TEXT("%s:"); 
+		RigModuleNameSpace = FString::Printf(NameSpaceFormat, *GetRigModuleSettings().Identifier.Name);
+		return RigModuleNameSpace;
 	}
 
 	static const FString EmptyNameSpace;
diff --git a/Engine/Plugins/Animation/ControlRig/Source/ControlRig/Private/Rigs/RigHierarchy.cpp b/Engine/Plugins/Animation/ControlRig/Source/ControlRig/Private/Rigs/RigHierarchy.cpp
--- a/Engine/Plugins/Animation/ControlRig/Source/ControlRig/Private/Rigs/RigHierarchy.cpp
+++ b/Engine/Plugins/Animation/ControlRig/Source/ControlRig/Private/Rigs/RigHierarchy.cpp
@@ -108,7 +108,6 @@ URigHierarchy::URigHierarchy()
 , bIsInteracting(false)
 , LastInteractedKey()
 , bSuspendNotifications(false)
-, bIsRunningGetIndex(false)
 , HierarchyController(nullptr)
 , bIsControllerAvailable(true)
 , ResetPoseHash(INDEX_NONE)
@@ -121,6 +120,7 @@ URigHierarchy::URigHierarchy()
 , bEnableCacheValidityCheck(bEnableValidityCheckbyDefault)
 , HierarchyForCacheValidation()
 , bUpdatePreferedEulerAngleWhenSettingTransform(true)
+, bAllowNameSpaceWhenSanitizingName(false)
 , ExecuteContext(nullptr)
 #if WITH_EDITOR
 , bRecordTransformsAtRuntime(true)
@@ -1094,7 +1094,7 @@ TArray<FRigElementKey> URigHierarchy::GetSelectedKeys(ERigElementType InTypeFilt
 	return Selection;
 }
 
-void URigHierarchy::SanitizeName(FRigName& InOutName)
+void URigHierarchy::SanitizeName(FRigName& InOutName, bool bAllowNameSpaces)
 {
 	// Sanitize the name
 	FString SanitizedNameString = InOutName.GetName();
@@ -1110,6 +1110,11 @@ void URigHierarchy::SanitizeName(FRigName& InOutName)
 
 		if (!bGoodChar)
 		{
+			if(bAllowNameSpaces && C == ':')
+			{
+				continue;
+			}
+			
 			C = '_';
 			bChangedSomething = true;
 		}
@@ -1172,7 +1177,7 @@ bool URigHierarchy::IsNameAvailable(const FRigName& InPotentialNewName, ERigElem
 	}
 
 	FRigName SanitizedName = UnsanitizedName;
-	SanitizeName(SanitizedName);
+	SanitizeName(SanitizedName, bAllowNameSpaceWhenSanitizingName);
 
 	if (SanitizedName != UnsanitizedName)
 	{
@@ -1249,7 +1254,18 @@ bool URigHierarchy::IsDisplayNameAvailable(const FRigElementKey& InParentElement
 FRigName URigHierarchy::GetSafeNewName(const FRigName& InPotentialNewName, ERigElementType InType) const
 {
 	FRigName SanitizedName = InPotentialNewName;
-	SanitizeName(SanitizedName);
+	SanitizeName(SanitizedName, false);
+
+	if(ExecuteContext)
+	{
+		const FControlRigExecuteContext& CRContext = ExecuteContext->GetPublicData<FControlRigExecuteContext>();
+		if(!CRContext.ModuleInstanceNameSpace.IsEmpty())
+		{
+			SanitizedName = CRContext.ModuleInstanceNameSpace + SanitizedName.GetName();
+			bAllowNameSpaceWhenSanitizingName = true;
+		}
+	}
+
 	FRigName Name = SanitizedName;
 
 	int32 Suffix = 1;
@@ -1262,6 +1278,8 @@ FRigName URigHierarchy::GetSafeNewName(const FRigName& InPotentialNewName, ERigE
 		}
 		Name.SetName(FString::Printf(TEXT("%s_%d"), *BaseString, ++Suffix));
 	}
+
+	bAllowNameSpaceWhenSanitizingName = false;
 	return Name;
 }
 
@@ -2352,7 +2370,7 @@ void URigHierarchy::Traverse(TFunction<void(FRigBaseElement*, bool& /* continue
 	}
 }
 
-bool URigHierarchy::CanConnect(const FRigConnectionInfo* InConnectionInfo, FString* OutFailureReason) const
+bool URigHierarchy::CanConnect(const FRigConnectionInfo* InConnectionInfo, FString* OutFailureReason, FRigElementKey* OutConnector) const
 {
 	check(InConnectionInfo);
 	check(InConnectionInfo->IsValid());
@@ -2368,6 +2386,10 @@ bool URigHierarchy::CanConnect(const FRigConnectionInfo* InConnectionInfo, FStri
 				static constexpr TCHAR Format[] = TEXT("Connector key '%s' not provided as part of the connection map.");
 				OutFailureReason->Appendf(Format, *ExpectedConnector->GetKey().ToString());
 			}
+			if(OutConnector)
+			{
+				*OutConnector = ExpectedConnector->GetKey();
+			}
 			return false;
 		}
 	}
@@ -2393,6 +2415,10 @@ bool URigHierarchy::CanConnect(const FRigConnectionInfo* InConnectionInfo, FStri
 				static constexpr TCHAR Format[] = TEXT("Connector key '%s' is not a connector.");
 				OutFailureReason->Appendf(Format, *ConnectorKey.ToString());
 			}
+			if(OutConnector)
+			{
+				*OutConnector = ConnectorKey;
+			}
 			return false;
 		}
 		if(!TargetKey.IsValid())
@@ -2402,6 +2428,10 @@ bool URigHierarchy::CanConnect(const FRigConnectionInfo* InConnectionInfo, FStri
 				static constexpr TCHAR Format[] = TEXT("Target key '%s' is not valid.");
 				OutFailureReason->Appendf(Format, *TargetKey.ToString());
 			}
+			if(OutConnector)
+			{
+				*OutConnector = ConnectorKey;
+			}
 			return false;
 		}
 
@@ -2413,6 +2443,10 @@ bool URigHierarchy::CanConnect(const FRigConnectionInfo* InConnectionInfo, FStri
 				static constexpr TCHAR Format[] = TEXT("Connector element '%s' does not exist.");
 				OutFailureReason->Appendf(Format, *ConnectorKey.ToString());
 			}
+			if(OutConnector)
+			{
+				*OutConnector = Connector->GetKey();
+			}
 			return false;
 		}
 		const FRigBaseElement* Target = Find<FRigBaseElement>(TargetKey);
@@ -2423,6 +2457,10 @@ bool URigHierarchy::CanConnect(const FRigConnectionInfo* InConnectionInfo, FStri
 				static constexpr TCHAR Format[] = TEXT("Target element '%s' does not exist.");
 				OutFailureReason->Appendf(Format, *TargetKey.ToString());
 			}
+			if(OutConnector)
+			{
+				*OutConnector = Connector->GetKey();
+			}
 			return false;
 		}
 		const FRigTransformElement* TargetTransform = Cast<FRigTransformElement>(Target);
@@ -2433,11 +2471,19 @@ bool URigHierarchy::CanConnect(const FRigConnectionInfo* InConnectionInfo, FStri
 				static constexpr TCHAR Format[] = TEXT("Target element '%s' is not a transform.");
 				OutFailureReason->Appendf(Format, *TargetKey.ToString());
 			}
+			if(OutConnector)
+			{
+				*OutConnector = Connector->GetKey();
+			}
 			return false;
 		}
 
 		if(!Connector->CanConnect(InConnectionInfo, OutFailureReason))
 		{
+			if(OutConnector)
+			{
+				*OutConnector = Connector->GetKey();
+			}
 			return false;
 		}
 	}
@@ -6563,3 +6609,8 @@ TArray<FString> URigHierarchy::ConnectorSettingsToPythonCommands(const FRigConne
 }
 
 #endif
+
+FRigHierarchyRedirectorGuard::FRigHierarchyRedirectorGuard(UControlRig* InControlRig)
+: Guard(InControlRig->GetHierarchy()->ElementKeyRedirector, &InControlRig->GetElementKeyRedirector())
+{
+}
\ No newline at end of file
diff --git a/Engine/Plugins/Animation/ControlRig/Source/ControlRig/Private/Rigs/RigHierarchyCache.cpp b/Engine/Plugins/Animation/ControlRig/Source/ControlRig/Private/Rigs/RigHierarchyCache.cpp
--- a/Engine/Plugins/Animation/ControlRig/Source/ControlRig/Private/Rigs/RigHierarchyCache.cpp
+++ b/Engine/Plugins/Animation/ControlRig/Source/ControlRig/Private/Rigs/RigHierarchyCache.cpp
@@ -80,9 +80,7 @@ FRigElementKeyRedirector::FRigElementKeyRedirector(const TMap<FRigElementKey, FR
 	for(const TPair<FRigElementKey, FRigElementKey>& Pair : InMap)
 	{
 		check(Pair.Key.IsValid());
-		InternalKeyToExternalKey.Add(Pair.Key, FCachedRigElement(Pair.Value, InHierarchy, true));
-		ExternalKeys.Add(Pair.Value);
-		Hash = HashCombine(Hash, HashCombine(GetTypeHash(Pair.Key), GetTypeHash(Pair.Value)));
+		Add(Pair.Key, Pair.Value, InHierarchy);
 	}
 }
 
@@ -96,9 +94,31 @@ FRigElementKeyRedirector::FRigElementKeyRedirector(const FRigElementKeyRedirecto
 	for(const TPair<FRigElementKey, FCachedRigElement>& Pair : InOther.InternalKeyToExternalKey)
 	{
 		check(Pair.Key.IsValid());
-		InternalKeyToExternalKey.Add(Pair.Key, FCachedRigElement(Pair.Value.GetKey(), InHierarchy, true));
-		ExternalKeys.Add(Pair.Value.GetKey());
-		Hash = HashCombine(Hash, HashCombine(GetTypeHash(Pair.Key), GetTypeHash(Pair.Value.GetKey())));
+		Add(Pair.Key, Pair.Value.GetKey(), InHierarchy);
 	}
 }
 
+const FRigElementKey* FRigElementKeyRedirector::FindReverse(const FRigElementKey& InKey) const
+{
+	for(const TPair<FRigElementKey, FCachedRigElement>& Pair : InternalKeyToExternalKey)
+	{
+		if(Pair.Value.GetKey() == InKey)
+		{
+			return &Pair.Key;
+		}
+	}
+	return nullptr;
+}
+
+void FRigElementKeyRedirector::Add(const FRigElementKey& InSource, const FRigElementKey& InTarget, const URigHierarchy* InHierarchy)
+{
+	if(!InSource.IsValid() || !InTarget.IsValid() || InSource == InTarget)
+	{
+		return;
+	}
+
+	InternalKeyToExternalKey.Add(InSource, FCachedRigElement(InTarget, InHierarchy, true));
+	ExternalKeys.Add(InTarget);
+	Hash = HashCombine(Hash, HashCombine(GetTypeHash(InSource), GetTypeHash(InTarget)));
+}
+
diff --git a/Engine/Plugins/Animation/ControlRig/Source/ControlRig/Private/Rigs/RigHierarchyController.cpp b/Engine/Plugins/Animation/ControlRig/Source/ControlRig/Private/Rigs/RigHierarchyController.cpp
--- a/Engine/Plugins/Animation/ControlRig/Source/ControlRig/Private/Rigs/RigHierarchyController.cpp
+++ b/Engine/Plugins/Animation/ControlRig/Source/ControlRig/Private/Rigs/RigHierarchyController.cpp
@@ -1,6 +1,7 @@
 // Copyright Epic Games, Inc. All Rights Reserved.
 
 #include "Rigs/RigHierarchyController.h"
+#include "ControlRig.h"
 #include "AnimationCoreLibrary.h"
 #include "UObject/Package.h"
 
@@ -88,13 +89,25 @@ bool URigHierarchyController::SelectElement(FRigElementKey InKey, bool bSelect,
 		}
 	}
 
-	FRigBaseElement* Element = Hierarchy->Find(InKey);
+	FRigElementKey Key = InKey;
+	if(Hierarchy->ElementKeyRedirector)
+	{
+		if(const FCachedRigElement* Cache = Hierarchy->ElementKeyRedirector->Find(Key))
+		{
+			if(const_cast<FCachedRigElement*>(Cache)->UpdateCache(Hierarchy.Get()))
+			{
+				Key = Cache->GetKey();
+			}
+		}
+	}
+
+	FRigBaseElement* Element = Hierarchy->Find(Key);
 	if(Element == nullptr)
 	{
 		return false;
 	}
 
-	const bool bSelectionState = Hierarchy->OrderedSelection.Contains(InKey);
+	const bool bSelectionState = Hierarchy->OrderedSelection.Contains(Key);
 	ensure(bSelectionState == Element->bSelected);
 	if(Element->bSelected == bSelect)
 	{
@@ -104,11 +117,11 @@ bool URigHierarchyController::SelectElement(FRigElementKey InKey, bool bSelect,
 	Element->bSelected = bSelect;
 	if(bSelect)
 	{
-		Hierarchy->OrderedSelection.Add(InKey);
+		Hierarchy->OrderedSelection.Add(Key);
 	}
 	else
 	{
-		Hierarchy->OrderedSelection.Remove(InKey);
+		Hierarchy->OrderedSelection.Remove(Key);
 	}
 
 	if(Element->bSelected)
@@ -214,7 +227,7 @@ FRigElementKey URigHierarchyController::AddBone(FName InName, FRigElementKey InP
 		NewElement->Key.Type = ERigElementType::Bone;
 		NewElement->Key.Name = Hierarchy->GetSafeNewName(InName, NewElement->Key.Type);
 		NewElement->BoneType = InBoneType;
-		AddElement(NewElement, Hierarchy->Get(Hierarchy->GetIndex(InParent)), true);
+		AddElement(NewElement, Hierarchy->Get(Hierarchy->GetIndex(InParent)), true, InName);
 
 		if(bTransformInGlobal)
 		{
@@ -275,7 +288,7 @@ FRigElementKey URigHierarchyController::AddNull(FName InName, FRigElementKey InP
 		TGuardValue<bool> DisableCacheValidityChecks(Hierarchy->bEnableCacheValidityCheck, false);		
 		NewElement->Key.Type = ERigElementType::Null;
 		NewElement->Key.Name = Hierarchy->GetSafeNewName(InName, NewElement->Key.Type);
-		AddElement(NewElement, Hierarchy->Get(Hierarchy->GetIndex(InParent)), false);
+		AddElement(NewElement, Hierarchy->Get(Hierarchy->GetIndex(InParent)), false, InName);
 
 		if(bTransformInGlobal)
 		{
@@ -350,7 +363,16 @@ FRigElementKey URigHierarchyController::AddControl(
 		{
 			NewElement->Settings.DisplayName = Hierarchy->GetSafeNewDisplayName(InParent, NewElement->Settings.DisplayName); 
 		}
-		AddElement(NewElement, Hierarchy->Get(Hierarchy->GetIndex(InParent)), false);
+		else if(Hierarchy->HasExecuteContext())
+		{
+			const FControlRigExecuteContext& CRContext = Hierarchy->ExecuteContext->GetPublicData<FControlRigExecuteContext>();
+			if(!CRContext.ModuleInstanceNameSpace.IsEmpty())
+			{
+				NewElement->Settings.DisplayName = Hierarchy->GetSafeNewDisplayName(InParent, InName);
+			}
+		}
+		
+		AddElement(NewElement, Hierarchy->Get(Hierarchy->GetIndex(InParent)), false, InName);
 		
 		NewElement->Offset.Set(ERigTransformType::InitialLocal, InOffsetTransform);  
 		NewElement->Shape.Set(ERigTransformType::InitialLocal, InShapeTransform);  
@@ -434,7 +456,7 @@ FRigElementKey URigHierarchyController::AddCurve(FName InName, float InValue, bo
 		NewElement->Key.Type = ERigElementType::Curve;
 		NewElement->Key.Name = Hierarchy->GetSafeNewName(InName, NewElement->Key.Type);
 		NewElement->Value = InValue;
-		AddElement(NewElement, nullptr, false);
+		AddElement(NewElement, nullptr, false, InName);
 	}
 
 #if WITH_EDITOR
@@ -483,7 +505,7 @@ FRigElementKey URigHierarchyController::AddRigidBody(FName InName, FRigElementKe
 		NewElement->Key.Type = ERigElementType::RigidBody;
 		NewElement->Key.Name = Hierarchy->GetSafeNewName(InName, NewElement->Key.Type);
 		NewElement->Settings = InSettings;
-		AddElement(NewElement, Hierarchy->Get(Hierarchy->GetIndex(InParent)), true);
+		AddElement(NewElement, Hierarchy->Get(Hierarchy->GetIndex(InParent)), true, InName);
 
 		Hierarchy->SetTransform(NewElement, InLocalTransform, ERigTransformType::InitialLocal, true, false);
 		NewElement->Pose.Current = NewElement->Pose.Initial;
@@ -535,7 +557,7 @@ FRigElementKey URigHierarchyController::AddReference(FName InName, FRigElementKe
 		NewElement->Key.Type = ERigElementType::Reference;
 		NewElement->Key.Name = Hierarchy->GetSafeNewName(InName, NewElement->Key.Type);
 		NewElement->GetWorldTransformDelegate = InDelegate;
-		AddElement(NewElement, Hierarchy->Get(Hierarchy->GetIndex(InParent)), true);
+		AddElement(NewElement, Hierarchy->Get(Hierarchy->GetIndex(InParent)), true, InName);
 
 		Hierarchy->SetTransform(NewElement, FTransform::Identity, ERigTransformType::InitialLocal, true, false);
 		NewElement->Pose.Current = NewElement->Pose.Initial;
@@ -573,7 +595,7 @@ FRigElementKey URigHierarchyController::AddConnector(FName InName, FTransform In
 		NewElement->Key.Type = ERigElementType::Connector;
 		NewElement->Key.Name = Hierarchy->GetSafeNewName(InName, NewElement->Key.Type);
 		NewElement->Settings = InSettings;
-		AddElement(NewElement, nullptr, true);
+		AddElement(NewElement, nullptr, true, InName);
 
 		if(bTransformInGlobal)
 		{
@@ -1279,9 +1301,10 @@ TArray<FRigElementKey> URigHierarchyController::ImportFromText(FString InContent
 				continue;
 			}
 		}
-		
-		NewElement->Key.Name = Hierarchy->GetSafeNewName(NewElement->Key.Name, NewElement->Key.Type);
-		AddElement(NewElement, nullptr, true);
+
+		const FName DesiredName = NewElement->Key.Name;
+		NewElement->Key.Name = Hierarchy->GetSafeNewName(DesiredName, NewElement->Key.Type);
+		AddElement(NewElement, nullptr, true, DesiredName);
 
 		KeyMap.FindOrAdd(PerElementData.Key) = NewElement->Key;
 
@@ -1760,7 +1783,7 @@ bool URigHierarchyController::IsValid() const
 	return Hierarchy.IsValid(bPendingKillAcceptable);
 }
 
-int32 URigHierarchyController::AddElement(FRigBaseElement* InElementToAdd, FRigBaseElement* InFirstParent, bool bMaintainGlobalTransform)
+int32 URigHierarchyController::AddElement(FRigBaseElement* InElementToAdd, FRigBaseElement* InFirstParent, bool bMaintainGlobalTransform, const FName& InDesiredName)
 {
 	ensure(IsValid());
 
@@ -1772,6 +1795,54 @@ int32 URigHierarchyController::AddElement(FRigBaseElement* InElementToAdd, FRigB
 	Hierarchy->IndexLookup.Add(InElementToAdd->Key, InElementToAdd->Index);
 	Hierarchy->IncrementTopologyVersion();
 
+	FRigName DesiredName = InDesiredName;
+	URigHierarchy::SanitizeName(DesiredName);
+
+	int32 NameSpaceTokenIndex = INDEX_NONE; 
+	if(InElementToAdd->GetName().FindLastChar(':', NameSpaceTokenIndex))
+	{
+		static const FName ShortNameMetaName = TEXT("ShortName");
+		const FString ShortName = InElementToAdd->GetName().RightChop(NameSpaceTokenIndex + 1);
+		if(!ShortName.IsEmpty())
+		{
+			Hierarchy->SetNameMetadata(InElementToAdd->Key, ShortNameMetaName, *ShortName);
+		}
+	}
+
+	if(!InDesiredName.IsNone() &&
+		!InElementToAdd->GetFName().IsEqual(DesiredName.GetFName(), ENameCase::CaseSensitive))
+	{
+		static const FName DesiredNameMetaName = TEXT("DesiredName");
+		static const FName DesiredKeyMetaName = TEXT("DesiredKey");
+		Hierarchy->SetNameMetadata(InElementToAdd->Key, DesiredNameMetaName, DesiredName.GetFName());
+		Hierarchy->SetRigElementKeyMetadata(InElementToAdd->Key, DesiredKeyMetaName, FRigElementKey(DesiredName.GetFName(), InElementToAdd->Key.Type));
+	}
+	
+	if(Hierarchy->HasExecuteContext())
+	{
+		const FControlRigExecuteContext& CRContext = Hierarchy->ExecuteContext->GetPublicData<FControlRigExecuteContext>();
+
+		if(!CRContext.ModuleInstanceName.IsNone())
+		{
+			static const FName ModuleMetaName = TEXT("Module");
+			Hierarchy->SetNameMetadata(InElementToAdd->Key, ModuleMetaName, CRContext.ModuleInstanceName);
+		}
+		
+		if(!CRContext.ModuleInstanceNameSpace.IsEmpty())
+		{
+			if(InElementToAdd->GetName().StartsWith(CRContext.ModuleInstanceNameSpace, ESearchCase::CaseSensitive))
+			{
+				static const FName NameSpaceMetaName = TEXT("NameSpace");
+				Hierarchy->SetNameMetadata(InElementToAdd->Key, NameSpaceMetaName, *CRContext.ModuleInstanceNameSpace);
+
+				if(Hierarchy->ElementKeyRedirector)
+				{
+					Hierarchy->ElementKeyRedirector->Add(FRigElementKey(DesiredName.GetFName(), InElementToAdd->Key.Type), InElementToAdd->Key, Hierarchy.Get());
+				}
+			}
+		}
+	}
+	
 	{
 		const TGuardValue<bool> Guard(bSuspendAllNotifications, true);
 		SetParent(InElementToAdd, InFirstParent, bMaintainGlobalTransform);
diff --git a/Engine/Plugins/Animation/ControlRig/Source/ControlRig/Private/Rigs/RigHierarchyElements.cpp b/Engine/Plugins/Animation/ControlRig/Source/ControlRig/Private/Rigs/RigHierarchyElements.cpp
--- a/Engine/Plugins/Animation/ControlRig/Source/ControlRig/Private/Rigs/RigHierarchyElements.cpp
+++ b/Engine/Plugins/Animation/ControlRig/Source/ControlRig/Private/Rigs/RigHierarchyElements.cpp
@@ -1480,6 +1480,23 @@ void FRigReferenceElement::CopyPose(FRigBaseElement* InOther, bool bCurrent, boo
 	}
 }
 
+////////////////////////////////////////////////////////////////////////////////
+// FRigConnectionInfo
+////////////////////////////////////////////////////////////////////////////////
+
+FRigConnectionInfo::FRigConnectionInfo(const FRigElementKeyRedirector* InRedirector, const URigHierarchy* InHierarchy)
+	: SourceHierarchy(InHierarchy)
+	, TargetHierarchy(InHierarchy)
+{
+	for(const TPair<FRigElementKey, FCachedRigElement>& Pair : InRedirector->InternalKeyToExternalKey)
+	{
+		if(Pair.Key.Type == ERigElementType::Connector)
+		{
+			ConnectionMap.Add(Pair.Key, Pair.Value.GetKey());
+		}
+	}
+}
+
 ////////////////////////////////////////////////////////////////////////////////
 // FRigConnectorSettings
 ////////////////////////////////////////////////////////////////////////////////
@@ -1597,6 +1614,11 @@ bool FRigConnectorElement::CanConnect(const FRigConnectionInfo* InConnectionInfo
 
 	for(const FRigConnectionRuleStash& Stash : Settings.Rules)
 	{
+		if(!Stash.IsValid())
+		{
+			continue;
+		}
+		
 		TSharedPtr<FStructOnScope> RuleScope;
 		const FRigConnectionRule* Rule = Stash.Get(RuleScope);
 		if(Rule == nullptr)
diff --git a/Engine/Plugins/Animation/ControlRig/Source/ControlRig/Private/Units/Execution/RigUnit_Hierarchy.cpp b/Engine/Plugins/Animation/ControlRig/Source/ControlRig/Private/Units/Execution/RigUnit_Hierarchy.cpp
--- a/Engine/Plugins/Animation/ControlRig/Source/ControlRig/Private/Units/Execution/RigUnit_Hierarchy.cpp
+++ b/Engine/Plugins/Animation/ControlRig/Source/ControlRig/Private/Units/Execution/RigUnit_Hierarchy.cpp
@@ -166,6 +166,64 @@ FRigUnit_HierarchyGetSiblingsItemArray_Execute()
 	Siblings = CachedSiblings.Keys;
 }
 
+FRigUnit_HierarchyGetChainItemArray_Execute()
+{
+	DECLARE_SCOPE_HIERARCHICAL_COUNTER_RIGUNIT()
+
+	if(!CachedStart.IsIdentical(Start, ExecuteContext.Hierarchy) ||
+		!CachedEnd.IsIdentical(End, ExecuteContext.Hierarchy))
+	{
+		CachedChain.Reset();
+
+		if(CachedStart.UpdateCache(Start, ExecuteContext.Hierarchy) &&
+			CachedEnd.UpdateCache(End, ExecuteContext.Hierarchy))
+		{
+			TArray<FRigElementKey> Keys;
+
+			const FRigTransformElement* StartElement = Cast<FRigTransformElement>(CachedStart.GetElement());
+			const FRigTransformElement* EndElement = Cast<FRigTransformElement>(CachedEnd.GetElement());
+			if(StartElement == nullptr || EndElement == nullptr)
+			{
+				Keys.Reset();
+			}
+			else
+			{
+				if(bIncludeEnd)
+				{
+					Keys.Add(EndElement->GetKey());
+				}
+
+				const FRigTransformElement* Parent = Cast<FRigTransformElement>(ExecuteContext.Hierarchy->GetFirstParent(EndElement));
+				while(Parent && Parent != StartElement)
+				{
+					Keys.Add(Parent->GetKey());
+					Parent = Cast<FRigTransformElement>(ExecuteContext.Hierarchy->GetFirstParent(Parent));
+				}
+
+				if(Parent != StartElement)
+				{
+					Keys.Reset();
+#if WITH_EDITOR
+					if(ExecuteContext.GetLog())
+					{
+						ExecuteContext.GetLog()->Report(EMessageSeverity::Info, ExecuteContext.GetFunctionName(), ExecuteContext.GetInstructionIndex(), TEXT("Start and End are not part of the same chain."));
+					}
+#endif
+				}
+				else if(bIncludeStart)
+				{
+					Keys.Add(StartElement->GetKey());
+				}
+
+				CachedChain = FRigElementKeyCollection(Keys);
+				CachedChain = FRigElementKeyCollection::MakeReversed(CachedChain);
+			}
+		}
+	}
+
+	Chain = CachedChain.Keys;
+}
+
 FRigUnit_HierarchyGetPose_Execute()
 {
 	FRigUnit_HierarchyGetPoseItemArray::StaticExecute(ExecuteContext, Initial, ElementType, ItemsToGet.Keys, Pose);
diff --git a/Engine/Plugins/Animation/ControlRig/Source/ControlRig/Private/Units/Execution/RigUnit_Hierarchy.h b/Engine/Plugins/Animation/ControlRig/Source/ControlRig/Private/Units/Execution/RigUnit_Hierarchy.h
--- a/Engine/Plugins/Animation/ControlRig/Source/ControlRig/Private/Units/Execution/RigUnit_Hierarchy.h
+++ b/Engine/Plugins/Animation/ControlRig/Source/ControlRig/Private/Units/Execution/RigUnit_Hierarchy.h
@@ -257,6 +257,53 @@ struct CONTROLRIG_API FRigUnit_HierarchyGetSiblingsItemArray : public FRigUnit_H
 	FRigElementKeyCollection CachedSiblings;
 };
 
+/**
+ * Returns a chain between two items
+ */
+USTRUCT(meta=(DisplayName="Get Chain", Keywords="Chain,Siblings,Hierarchy", Varying))
+struct CONTROLRIG_API FRigUnit_HierarchyGetChainItemArray : public FRigUnit_HierarchyBase
+{
+	GENERATED_BODY()
+
+	FRigUnit_HierarchyGetChainItemArray()
+	{
+		Start = End = FRigElementKey(NAME_None, ERigElementType::Bone);
+		CachedStart = CachedEnd = FCachedRigElement();
+		CachedChain = FRigElementKeyCollection();
+		bIncludeStart = bIncludeEnd = true;
+	}
+
+	RIGVM_METHOD()
+	virtual void Execute() override;
+
+	UPROPERTY(meta = (Input, ExpandByDefault))
+	FRigElementKey Start;
+
+	UPROPERTY(meta = (Input, ExpandByDefault))
+	FRigElementKey End;
+
+	UPROPERTY(meta = (Input))
+	bool bIncludeStart;
+
+	UPROPERTY(meta = (Input))
+	bool bIncludeEnd;
+
+	UPROPERTY(meta = (Output))
+	TArray<FRigElementKey> Chain;
+
+	// Used to cache the internally used item
+	UPROPERTY()
+	FCachedRigElement CachedStart;
+
+	// Used to cache the internally used item
+	UPROPERTY()
+	FCachedRigElement CachedEnd;
+
+	// Used to cache the internally used siblings
+	UPROPERTY()
+	FRigElementKeyCollection CachedChain;
+};
+
 /**
  * Returns the hierarchy's pose
  */
diff --git a/Engine/Plugins/Animation/ControlRig/Source/ControlRig/Public/ControlRig.h b/Engine/Plugins/Animation/ControlRig/Source/ControlRig/Public/ControlRig.h
--- a/Engine/Plugins/Animation/ControlRig/Source/ControlRig/Public/ControlRig.h
+++ b/Engine/Plugins/Animation/ControlRig/Source/ControlRig/Public/ControlRig.h
@@ -188,6 +188,8 @@ public:
 	virtual bool Execute_Internal(const FName& InEventName) override;
 	virtual void RequestInit() override;
 
+	bool AllConnectorsAreResolved(FString* OutFailureReason = nullptr, FRigElementKey* OutConnector = nullptr) const;
+
 	/** Requests to perform construction during the next execution */
 	UFUNCTION(BlueprintCallable, Category = "Control Rig")
 	void RequestConstruction();
diff --git a/Engine/Plugins/Animation/ControlRig/Source/ControlRig/Public/Rigs/RigHierarchy.h b/Engine/Plugins/Animation/ControlRig/Source/ControlRig/Public/Rigs/RigHierarchy.h
--- a/Engine/Plugins/Animation/ControlRig/Source/ControlRig/Public/Rigs/RigHierarchy.h
+++ b/Engine/Plugins/Animation/ControlRig/Source/ControlRig/Public/Rigs/RigHierarchy.h
@@ -18,6 +18,7 @@
 #include "Containers/Queue.h"
 #include "RigHierarchy.generated.h"
 
+class UControlRig;
 class URigHierarchy;
 class URigHierarchyController;
 
@@ -406,17 +407,12 @@ public:
 		{
 			if(FCachedRigElement* CachedRigElement = ElementKeyRedirector->Find(InKey))
 			{
-				TGuardValue<bool> GuardRecursion(bIsRunningGetIndex, true);
 				if(CachedRigElement->UpdateCache(this))
 				{
 					return CachedRigElement->GetIndex();
 				}
 				return INDEX_NONE;
 			}
-			if(!bIsRunningGetIndex && ElementKeyRedirector->ContainsExternalKey(InKey))
-			{
-				return INDEX_NONE;
-			}
 		}
 		
 		if(const int32* Index = IndexLookup.Find(InKey))
@@ -1635,7 +1631,7 @@ public:
 	 * Sanitizes a name by removing invalid characters.
 	 * @param InOutName The name to sanitize in place.
 	 */
-	static void SanitizeName(FRigName& InOutName);
+	static void SanitizeName(FRigName& InOutName, bool bAllowNameSpaces = false);
 
 	/**
 	 * Sanitizes a name by removing invalid characters.
@@ -3297,7 +3293,7 @@ public:
 	{
 		return CanConnect(&InConnectionInfo);
 	}
-	bool CanConnect(const FRigConnectionInfo* InConnectionInfo, FString* OutFailureReason = nullptr) const;
+	bool CanConnect(const FRigConnectionInfo* InConnectionInfo, FString* OutFailureReason = nullptr, FRigElementKey* OutConnector = nullptr) const;
 
 	/**
 	 * Returns the currently resolved target for given connector key
@@ -4201,11 +4197,6 @@ private:
 	 */
 	bool bSuspendNotifications;
 
-	/**
-	 * If set to true the hierarchy is currently running GetIndex - flag to avoid infinite recursion
-	 */
-	mutable bool bIsRunningGetIndex;
-
 	/**
 	 * The event fired during undo / redo
 	 */
@@ -4479,13 +4470,13 @@ private:
 #endif
 
 	template<typename T>
-	const T& GetMetadata(const FRigElementKey& InItem, ERigMetadataType InType, const FName& InMetadataName, const T& DefaultValue) const
+	const T& GetMetadata(const FRigElementKey& InItem, ERigMetadataType InType, const FRigName& InMetadataName, const T& DefaultValue) const
 	{
 		return GetMetadata<T>(Find(InItem), InType, InMetadataName, DefaultValue);
 	}
 
 	template<typename T>
-	const T& GetMetadata(const FRigBaseElement* InElement, ERigMetadataType InType, const FName& InMetadataName, const T& DefaultValue) const
+	const T& GetMetadata(const FRigBaseElement* InElement, ERigMetadataType InType, const FRigName& InMetadataName, const T& DefaultValue) const
 	{
 		if(InElement)
 		{
@@ -4498,26 +4489,26 @@ private:
 	}
 
 	template<typename T>
-	const TArray<T>& GetArrayMetadata(const FRigElementKey& InItem, ERigMetadataType InType, const FName& InMetadataName) const
+	const TArray<T>& GetArrayMetadata(const FRigElementKey& InItem, ERigMetadataType InType, const FRigName& InMetadataName) const
 	{
 		return GetArrayMetadata<T>(Find(InItem), InType, InMetadataName);
 	}
 
 	template<typename T>
-	const TArray<T>& GetArrayMetadata(const FRigBaseElement* InElement, ERigMetadataType InType, const FName& InMetadataName) const
+	const TArray<T>& GetArrayMetadata(const FRigBaseElement* InElement, ERigMetadataType InType, const FRigName& InMetadataName) const
 	{
 		static const TArray<T> EmptyArray;
 		return GetMetadata<TArray<T>>(InElement, InType, InMetadataName, EmptyArray);
 	}
 
 	template<typename T>
-	bool SetMetadata(const FRigElementKey& InItem, ERigMetadataType InType, const FName& InMetadataName, const T& InValue)
+	bool SetMetadata(const FRigElementKey& InItem, ERigMetadataType InType, const FRigName& InMetadataName, const T& InValue)
 	{
 		return SetMetadata<T>(Find(InItem), InType, InMetadataName, InValue);
 	}
 
 	template<typename T>
-	bool SetMetadata(FRigBaseElement* InElement, ERigMetadataType InType, const FName& InMetadataName, const T& InValue)
+	bool SetMetadata(FRigBaseElement* InElement, ERigMetadataType InType, const FRigName& InMetadataName, const T& InValue)
 	{
 		if(InElement)
 		{
@@ -4527,13 +4518,13 @@ private:
 	}
 
 	template<typename T>
-	bool SetArrayMetadata(const FRigElementKey& InItem, ERigMetadataType InType, const FName& InMetadataName, const TArray<T>& InValue)
+	bool SetArrayMetadata(const FRigElementKey& InItem, ERigMetadataType InType, const FRigName& InMetadataName, const TArray<T>& InValue)
 	{
 		return SetMetadata<TArray<T>>(Find(InItem), InType, InMetadataName, InValue);
 	}
 
 	template<typename T>
-	bool SetArrayMetadata(FRigBaseElement* InElement, ERigMetadataType InType, const FName& InMetadataName, const TArray<T>& InValue)
+	bool SetArrayMetadata(FRigBaseElement* InElement, ERigMetadataType InType, const FRigName& InMetadataName, const TArray<T>& InValue)
 	{
 		return SetMetadata<TArray<T>>(InElement, InType, InMetadataName, InValue);
 	}
@@ -4566,6 +4557,7 @@ protected:
 	mutable TMap<FRigElementKey, FRigElementKey> DefaultParentPerElement;
 
 	bool bUpdatePreferedEulerAngleWhenSettingTransform;
+	mutable bool bAllowNameSpaceWhenSanitizingName;
 
 private:
 	
@@ -4764,6 +4756,8 @@ public:
 	{
 	}
 
+	FRigHierarchyRedirectorGuard(UControlRig* InControlRig);
+
 private:
 	TGuardValue<FRigElementKeyRedirector*> Guard;
 };
diff --git a/Engine/Plugins/Animation/ControlRig/Source/ControlRig/Public/Rigs/RigHierarchyCache.h b/Engine/Plugins/Animation/ControlRig/Source/ControlRig/Public/Rigs/RigHierarchyCache.h
--- a/Engine/Plugins/Animation/ControlRig/Source/ControlRig/Public/Rigs/RigHierarchyCache.h
+++ b/Engine/Plugins/Animation/ControlRig/Source/ControlRig/Public/Rigs/RigHierarchyCache.h
@@ -171,11 +171,18 @@ public:
 	bool ContainsExternalKey(const FRigElementKey& InKey) const { return ExternalKeys.Contains(InKey); }
 	const FCachedRigElement* Find(const FRigElementKey& InKey) const { return InternalKeyToExternalKey.Find(InKey); }
 	FCachedRigElement* Find(const FRigElementKey& InKey) { return InternalKeyToExternalKey.Find(InKey); }
+	const FRigElementKey* FindReverse(const FRigElementKey& InKey) const;
 	uint32 GetHash() const { return Hash; }
 	
 private:
 
+	void Add(const FRigElementKey& InSource, const FRigElementKey& InTarget, const URigHierarchy* InHierarchy);
+
 	TMap<FRigElementKey, FCachedRigElement> InternalKeyToExternalKey;
 	TSet<FRigElementKey> ExternalKeys;
 	uint32 Hash;
+
+	friend class URigHierarchy;
+	friend class URigHierarchyController;
+	friend struct FRigConnectionInfo;
 };
\ No newline at end of file
diff --git a/Engine/Plugins/Animation/ControlRig/Source/ControlRig/Public/Rigs/RigHierarchyController.h b/Engine/Plugins/Animation/ControlRig/Source/ControlRig/Public/Rigs/RigHierarchyController.h
--- a/Engine/Plugins/Animation/ControlRig/Source/ControlRig/Public/Rigs/RigHierarchyController.h
+++ b/Engine/Plugins/Animation/ControlRig/Source/ControlRig/Public/Rigs/RigHierarchyController.h
@@ -620,9 +620,10 @@ private:
 	 * @param InElementToAdd The new element to add to the hierarchy 
 	 * @param InFirstParent The (optional) parent of the new bone. If you don't need a parent, pass nullptr
 	 * @param bMaintainGlobalTransform If set to true the child will stay in the same place spatially, otherwise it will maintain it's local transform (and potential move).
+	 * @param InDesiredName The original desired name
 	 * @return The index of the newly added element
 	 */
-	int32 AddElement(FRigBaseElement* InElementToAdd, FRigBaseElement* InFirstParent, bool bMaintainGlobalTransform);
+	int32 AddElement(FRigBaseElement* InElementToAdd, FRigBaseElement* InFirstParent, bool bMaintainGlobalTransform, const FName& InDesiredName = NAME_None);
 
 	/**
 	 * Removes an existing element from the hierarchy.
diff --git a/Engine/Plugins/Animation/ControlRig/Source/ControlRig/Public/Rigs/RigHierarchyElements.h b/Engine/Plugins/Animation/ControlRig/Source/ControlRig/Public/Rigs/RigHierarchyElements.h
--- a/Engine/Plugins/Animation/ControlRig/Source/ControlRig/Public/Rigs/RigHierarchyElements.h
+++ b/Engine/Plugins/Animation/ControlRig/Source/ControlRig/Public/Rigs/RigHierarchyElements.h
@@ -12,6 +12,7 @@ struct FRigVMExecuteContext;
 struct FRigBaseElement;
 struct FRigControlElement;
 class URigHierarchy;
+class FRigElementKeyRedirector;
 
 DECLARE_DELEGATE_RetVal_ThreeParams(FTransform, FRigReferenceGetWorldTransformDelegate, const FRigVMExecuteContext*, const FRigElementKey& /* Key */, bool /* bInitial */);
 DECLARE_DELEGATE_TwoParams(FRigElementMetadataChangedDelegate, const FRigElementKey& /* Key */, const FName& /* Name */);
@@ -1619,12 +1620,14 @@ public:
 	, TargetHierarchy(nullptr)
 	{}
 
-	FRigConnectionInfo(const TMap<FRigElementKey, FRigElementKey>& InMap, URigHierarchy* InSourceHierarchy, URigHierarchy* InTargetHierarchy)
+	FRigConnectionInfo(const TMap<FRigElementKey, FRigElementKey>& InMap, const URigHierarchy* InSourceHierarchy, const URigHierarchy* InTargetHierarchy)
 	: ConnectionMap(InMap)
 	, SourceHierarchy(InSourceHierarchy)
 	, TargetHierarchy(InTargetHierarchy)
 	{}
 
+	FRigConnectionInfo(const FRigElementKeyRedirector* InRedirector, const URigHierarchy* InHierarchy);
+
 	bool IsValid() const
 	{
 		return (!ConnectionMap.IsEmpty()) && (SourceHierarchy != nullptr) && (TargetHierarchy != nullptr);
@@ -1637,11 +1640,11 @@ public:
 
 	// The hierarchy owning the connectors
 	UPROPERTY(BlueprintReadWrite, EditAnywhere, Category = Connections)
-	URigHierarchy* SourceHierarchy; 
+	const URigHierarchy* SourceHierarchy; 
 
 	// The hierarchy to be linked into
 	UPROPERTY(BlueprintReadWrite, EditAnywhere, Category = Connections)
-	URigHierarchy* TargetHierarchy; 
+	const URigHierarchy* TargetHierarchy; 
 };
 
 USTRUCT(BlueprintType)
diff --git a/Engine/Plugins/Animation/ControlRig/Source/ControlRigDeveloper/Private/ControlRigBlueprint.cpp b/Engine/Plugins/Animation/ControlRig/Source/ControlRigDeveloper/Private/ControlRigBlueprint.cpp
--- a/Engine/Plugins/Animation/ControlRig/Source/ControlRigDeveloper/Private/ControlRigBlueprint.cpp
+++ b/Engine/Plugins/Animation/ControlRig/Source/ControlRigDeveloper/Private/ControlRigBlueprint.cpp
@@ -93,6 +93,33 @@ bool UControlRigBlueprint::RequiresForceLoadMembers(UObject* InObject) const
 	return Super::RequiresForceLoadMembers(InObject);
 }
 
+void UControlRigBlueprint::PostEditChangeProperty(FPropertyChangedEvent& PropertyChangedEvent)
+{
+	Super::PostEditChangeProperty(PropertyChangedEvent);
+
+	// if this is any of our external variables we need to request construction so that the rig rebuilds itself
+	if(NewVariables.ContainsByPredicate([&PropertyChangedEvent](const FBPVariableDescription& Variable)
+	{
+		return Variable.VarName == PropertyChangedEvent.GetMemberPropertyName();
+	}))
+	{
+		if(UControlRig* DebuggedControlRig = Cast<UControlRig>(GetObjectBeingDebugged()))
+		{
+			if(const FProperty* PropertyOnRig = DebuggedControlRig->GetClass()->FindPropertyByName(PropertyChangedEvent.MemberProperty->GetFName()))
+			{
+				if(PropertyOnRig->SameType(PropertyChangedEvent.MemberProperty))
+				{
+					UControlRig* CDO = DebuggedControlRig->GetClass()->GetDefaultObject<UControlRig>();
+					const uint8* SourceMemory = PropertyOnRig->ContainerPtrToValuePtr<uint8>(CDO);
+					uint8* TargetMemory = PropertyOnRig->ContainerPtrToValuePtr<uint8>(DebuggedControlRig);
+					PropertyOnRig->CopyCompleteValue(TargetMemory, SourceMemory);
+				}
+			}
+			DebuggedControlRig->RequestConstruction();
+		}
+	}
+}
+
 void UControlRigBlueprint::PostEditChangeChainProperty(FPropertyChangedChainEvent& PropertyChangedEvent)
 {
 	Super::PostEditChangeChainProperty(PropertyChangedEvent);
@@ -711,6 +738,7 @@ bool UControlRigBlueprint::ResolveConnector(const FRigElementKey& DraggedKey, co
 	}
 
 	PropagateHierarchyFromBPToInstances();
+	RequestAutoVMRecompilation();
 
 	return true;
 }
diff --git a/Engine/Plugins/Animation/ControlRig/Source/ControlRigDeveloper/Private/Graph/ControlRigGraph.cpp b/Engine/Plugins/Animation/ControlRig/Source/ControlRigDeveloper/Private/Graph/ControlRigGraph.cpp
--- a/Engine/Plugins/Animation/ControlRig/Source/ControlRigDeveloper/Private/Graph/ControlRigGraph.cpp
+++ b/Engine/Plugins/Animation/ControlRig/Source/ControlRigDeveloper/Private/Graph/ControlRigGraph.cpp
@@ -75,6 +75,8 @@ void UControlRigGraph::CacheNameLists(URigHierarchy* InHierarchy, const FRigVMDr
 
 	if(LastHierarchyTopologyVersion != InHierarchy->GetTopologyVersion())
 	{
+		UControlRig* ControlRig = InHierarchy->GetTypedOuter<UControlRig>();
+
 		ElementNameLists.FindOrAdd(ERigElementType::All);
 		ElementNameLists.FindOrAdd(ERigElementType::Bone);
 		ElementNameLists.FindOrAdd(ERigElementType::Null);
@@ -93,15 +95,15 @@ void UControlRigGraph::CacheNameLists(URigHierarchy* InHierarchy, const FRigVMDr
 		TArray<TSharedPtr<FString>>& ReferenceNameList = ElementNameLists.FindChecked(ERigElementType::Reference);
 		TArray<TSharedPtr<FString>>& ConnectorNameList = ElementNameLists.FindChecked(ERigElementType::Connector);
 		
-		CacheNameListForHierarchy<FRigBaseElement>(InHierarchy, AllNameList, false);
-		CacheNameListForHierarchy<FRigBoneElement>(InHierarchy, BoneNameList, false);
-		CacheNameListForHierarchy<FRigNullElement>(InHierarchy, NullNameList, false);
-		CacheNameListForHierarchy<FRigControlElement>(InHierarchy, ControlNameList, false);
-		CacheNameListForHierarchy<FRigControlElement>(InHierarchy, ControlNameListWithoutAnimationChannels, true);
-		CacheNameListForHierarchy<FRigCurveElement>(InHierarchy, CurveNameList, false);
-		CacheNameListForHierarchy<FRigRigidBodyElement>(InHierarchy, RigidBodyNameList, false);
-		CacheNameListForHierarchy<FRigReferenceElement>(InHierarchy, ReferenceNameList, false);
-		CacheNameListForHierarchy<FRigConnectorElement>(InHierarchy, ConnectorNameList, false);
+		CacheNameListForHierarchy<FRigBaseElement>(ControlRig, InHierarchy, AllNameList, false);
+		CacheNameListForHierarchy<FRigBoneElement>(ControlRig, InHierarchy, BoneNameList, false);
+		CacheNameListForHierarchy<FRigNullElement>(ControlRig, InHierarchy, NullNameList, false);
+		CacheNameListForHierarchy<FRigControlElement>(ControlRig, InHierarchy, ControlNameList, false);
+		CacheNameListForHierarchy<FRigControlElement>(ControlRig, InHierarchy, ControlNameListWithoutAnimationChannels, true);
+		CacheNameListForHierarchy<FRigCurveElement>(ControlRig, InHierarchy, CurveNameList, false);
+		CacheNameListForHierarchy<FRigRigidBodyElement>(ControlRig, InHierarchy, RigidBodyNameList, false);
+		CacheNameListForHierarchy<FRigReferenceElement>(ControlRig, InHierarchy, ReferenceNameList, false);
+		CacheNameListForHierarchy<FRigConnectorElement>(ControlRig, InHierarchy, ConnectorNameList, false);
 
 		LastHierarchyTopologyVersion = InHierarchy->GetTopologyVersion();
 	}
diff --git a/Engine/Plugins/Animation/ControlRig/Source/ControlRigDeveloper/Public/ControlRigBlueprint.h b/Engine/Plugins/Animation/ControlRig/Source/ControlRigDeveloper/Public/ControlRigBlueprint.h
--- a/Engine/Plugins/Animation/ControlRig/Source/ControlRigDeveloper/Public/ControlRigBlueprint.h
+++ b/Engine/Plugins/Animation/ControlRig/Source/ControlRigDeveloper/Public/ControlRigBlueprint.h
@@ -71,6 +71,7 @@ public:
 	virtual bool RequiresForceLoadMembers(UObject* InObject) const override;
 
 	// UObject interface
+	virtual void PostEditChangeProperty(FPropertyChangedEvent& PropertyChangedEvent) override;
 	virtual void PostEditChangeChainProperty(struct FPropertyChangedChainEvent& PropertyChangedEvent) override;
 
 #endif	// #if WITH_EDITOR
diff --git a/Engine/Plugins/Animation/ControlRig/Source/ControlRigDeveloper/Public/Graph/ControlRigGraph.h b/Engine/Plugins/Animation/ControlRig/Source/ControlRigDeveloper/Public/Graph/ControlRigGraph.h
--- a/Engine/Plugins/Animation/ControlRig/Source/ControlRigDeveloper/Public/Graph/ControlRigGraph.h
+++ b/Engine/Plugins/Animation/ControlRig/Source/ControlRigDeveloper/Public/Graph/ControlRigGraph.h
@@ -79,7 +79,7 @@ private:
 	}
 
 	template<class T>
-	void CacheNameListForHierarchy(URigHierarchy* InHierarchy, TArray<TSharedPtr<FString>>& OutNameList, bool bFilter = true)
+	void CacheNameListForHierarchy(UControlRig* InControlRig, URigHierarchy* InHierarchy, TArray<TSharedPtr<FString>>& OutNameList, bool bFilter = true)
 	{
         TArray<FString> Names;
 		for (auto Element : *InHierarchy)
@@ -88,6 +88,14 @@ private:
 			{
 				if(!bFilter || IncludeElementInNameList<T>(Cast<T>(Element)))
 				{
+					if(InControlRig)
+					{
+						if(const FRigElementKey* SourceKey = InControlRig->GetElementKeyRedirector().FindReverse(Element->GetKey()))
+						{
+							Names.Add(SourceKey->Name.ToString());
+							continue;
+						}
+					}
 					Names.Add(Element->GetName());
 				}
 			}
diff --git a/Engine/Plugins/Animation/ControlRig/Source/ControlRigEditor/Private/ControlRigElementDetails.cpp b/Engine/Plugins/Animation/ControlRig/Source/ControlRigEditor/Private/ControlRigElementDetails.cpp
--- a/Engine/Plugins/Animation/ControlRig/Source/ControlRigEditor/Private/ControlRigElementDetails.cpp
+++ b/Engine/Plugins/Animation/ControlRig/Source/ControlRigEditor/Private/ControlRigElementDetails.cpp
@@ -757,7 +757,7 @@ void FRigBaseElementDetails::CustomizeDetails(IDetailLayoutBuilder& DetailBuilde
 
 	DetailBuilder.HideCategory(TEXT("RigElement"));
 
-	if(!bAllAnimationChannels)
+	if(!bAllControls || !bAllAnimationChannels)
 	{
 		GeneralCategory.AddCustomRow(FText::FromString(TEXT("Name")))
 		.NameContent()
diff --git a/Engine/Plugins/Animation/ControlRig/Source/ControlRigEditor/Private/Editor/ControlRigEditor.cpp b/Engine/Plugins/Animation/ControlRig/Source/ControlRigEditor/Private/Editor/ControlRigEditor.cpp
--- a/Engine/Plugins/Animation/ControlRig/Source/ControlRigEditor/Private/Editor/ControlRigEditor.cpp
+++ b/Engine/Plugins/Animation/ControlRig/Source/ControlRigEditor/Private/Editor/ControlRigEditor.cpp
@@ -986,6 +986,12 @@ void FControlRigEditor::HandleSetObjectBeingDebugged(UObject* InObject)
 }
 
 void FControlRigEditor::SetDetailViewForRigElements()
+{
+	URigHierarchy* HierarchyBeingDebugged = GetHierarchyBeingDebugged();
+	SetDetailViewForRigElements(HierarchyBeingDebugged->GetSelectedKeys());
+}
+
+void FControlRigEditor::SetDetailViewForRigElements(const TArray<FRigElementKey>& InKeys)
 {
 	if(IsDetailsPanelRefreshSuspended())
 	{
@@ -998,10 +1004,9 @@ void FControlRigEditor::SetDetailViewForRigElements()
 	URigHierarchy* HierarchyBeingDebugged = GetHierarchyBeingDebugged();
 	TArray<UObject*> Objects;
 
-	TArray<FRigElementKey> CurrentSelection = HierarchyBeingDebugged->GetSelectedKeys();
-	for(const FRigElementKey& SelectedKey : CurrentSelection)
+	for(const FRigElementKey& Key : InKeys)
 	{
-		FRigBaseElement* Element = HierarchyBeingDebugged->Find(SelectedKey);
+		FRigBaseElement* Element = HierarchyBeingDebugged->Find(Key);
 		if (Element == nullptr)
 		{
 			continue;
@@ -1478,7 +1483,7 @@ void FControlRigEditor::HandleVMExecutedEvent(URigVMHost* InHost, const FName& I
 		{
 			if(!DebuggedControlRig->RigUnitManipulationInfos.IsEmpty())
 			{
-				const FRigHierarchyRedirectorGuard RedirectorGuard(DebuggedControlRig->GetHierarchy(), DebuggedControlRig->ElementKeyRedirector);
+				const FRigHierarchyRedirectorGuard RedirectorGuard(DebuggedControlRig);
 				FControlRigExecuteContext& ExecuteContext = DebuggedControlRig->GetExtendedExecuteContext().GetPublicDataSafe<FControlRigExecuteContext>();
 				
 				for(const TSharedPtr<FRigDirectManipulationInfo>& ManipulationInfo : DebuggedControlRig->RigUnitManipulationInfos)
@@ -1738,7 +1743,74 @@ void FControlRigEditor::HandleViewportCreated(const TSharedRef<class IPersonaVie
 			],
 			DirectManipulationNotificationOptions
 		);
+	}
+
+	{
+		InViewport->AddNotification(
+			EMessageSeverity::Warning,
+			false,
+			SNew(SHorizontalBox)
+			.Visibility(this, &FControlRigEditor::GetConnectorWarningVisibility)
+			.ToolTipText(LOCTEXT("ConnectorWarningTooltip", "This rig has unresolved connectors."))
+			+SHorizontalBox::Slot()
+			.FillWidth(1.0f)
+			.Padding(4.0f, 4.0f)
+			[
+				SNew(SHorizontalBox)
+				+SHorizontalBox::Slot()
+				.AutoWidth()
+				.VAlign(VAlign_Center)
+				.Padding(0.0f, 0.0f, 4.0f, 0.0f)
+				[
+					SNew(STextBlock)
+					.TextStyle(FAppStyle::Get(), "AnimViewport.MessageText")
+					.Font(FAppStyle::Get().GetFontStyle("FontAwesome.9"))
+					.Text(FEditorFontGlyphs::Exclamation_Triangle)
+				]
+				+SHorizontalBox::Slot()
+				.VAlign(VAlign_Center)
+				.AutoWidth()
+				[
+					SNew(STextBlock)
+					.TextStyle(FAppStyle::Get(), "AnimViewport.MessageText")
+					.Text(this, &FControlRigEditor::GetConnectorWarningText)
+				]
+				+SHorizontalBox::Slot()
+				.AutoWidth()
+				.Padding(2.0f, 0.0f)
+				[
+					SNew(SButton)
+					.ForegroundColor(FSlateColor::UseForeground())
+					.ButtonStyle(FAppStyle::Get(), "FlatButton.Primary")
+					.ToolTipText(LOCTEXT("ConnectorWarningNavigateTooltip", "Navigate to the first unresolved connector in the hierarchy"))
+					.OnClicked(this, &FControlRigEditor::OnNavigateToConnectorWarning)
+					[
+						SNew(SHorizontalBox)
+						+SHorizontalBox::Slot()
+						.AutoWidth()
+						.VAlign(VAlign_Center)
+						.Padding(0.0f, 0.0f, 4.0f, 0.0f)
+						[
+							SNew(STextBlock)
+							.TextStyle(FAppStyle::Get(), "AnimViewport.MessageText")
+							.Font(FAppStyle::Get().GetFontStyle("FontAwesome.9"))
+							.Text(FEditorFontGlyphs::Cog)
+						]
+						+SHorizontalBox::Slot()
+						.VAlign(VAlign_Center)
+						.AutoWidth()
+						[
+							SNew(STextBlock)
+							.TextStyle(FAppStyle::Get(), "AnimViewport.MessageText")
+							.Text(LOCTEXT("ConnectorWarningNavigateButtonLabel", "Discover"))
+						]
+					]
+				]
+			],
+			FPersonaViewportNotificationOptions(TAttribute<EVisibility>::CreateSP(this, &FControlRigEditor::GetConnectorWarningVisibility))
+		);
 	}	
+
 	InViewport->AddNotification(MakeAttributeLambda(GetErrorSeverity),
 		false,
 		SNew(SHorizontalBox)
@@ -2773,6 +2845,45 @@ void FControlRigEditor::RefreshDirectManipulationTextList()
 	}
 }
 
+EVisibility FControlRigEditor::GetConnectorWarningVisibility() const
+{
+	if(GetConnectorWarningText().IsEmpty())
+	{
+		return EVisibility::Hidden;
+	}
+	return EVisibility::Visible;
+}
+
+FText FControlRigEditor::GetConnectorWarningText() const
+{
+	if (const UControlRigBlueprint* Blueprint = GetControlRigBlueprint())
+	{
+		if(Blueprint->IsControlRigModule())
+		{
+			if(UControlRig* ControlRig = GetControlRig())
+			{
+				FString FailureReason;
+				if(!ControlRig->AllConnectorsAreResolved(&FailureReason))
+				{
+					if(FailureReason.IsEmpty())
+					{
+						static const FText ConnectorWarningDefault = LOCTEXT("ConnectorWarningDefault", "This rig has unresolved connectors.");
+						return ConnectorWarningDefault;
+					}
+					return FText::FromString(FailureReason);
+				}
+			}
+		}
+	}
+	return FText();
+}
+
+FReply FControlRigEditor::OnNavigateToConnectorWarning() const
+{
+	RequestNavigateToConnectorWarningDelegate.Broadcast();
+	return FReply::Handled();
+}
+
 void FControlRigEditor::BindCommands()
 {
 	FRigVMEditor::BindCommands();
@@ -4217,7 +4328,7 @@ void FControlRigEditor::HandleOnControlModified(UControlRig* Subject, FRigContro
 			// update the node based on the incoming pose. once that is done we'll need to compare the node instance
 			// with the settings on the node in the graph and update them accordingly.
 			FControlRigExecuteContext& ExecuteContext = DebuggedControlRig->GetExtendedExecuteContext().GetPublicDataSafe<FControlRigExecuteContext>();
-			const FRigHierarchyRedirectorGuard RedirectorGuard(DebuggedControlRig->GetHierarchy(), DebuggedControlRig->ElementKeyRedirector);
+			const FRigHierarchyRedirectorGuard RedirectorGuard(DebuggedControlRig);
 			if(UnitInstance->UpdateDirectManipulationFromHierarchy(UnitNode, NodeInstance, ExecuteContext, ManipulationInfo))
 			{
 				UnitNode->UpdateHostFromStructInstance(DebuggedControlRig, NodeInstance);
@@ -4387,26 +4498,42 @@ void FControlRigEditor::OnPreConstruction_AnyThread(UControlRig* InRig, const FN
 		{
 			if(RigBlueprint->PreviewSkeletalMesh)
 			{
-				if(URigHierarchyController* Controller = InRig->GetHierarchy()->GetController(true))
+				if(URigHierarchy* Hierarchy = InRig->GetHierarchy())
 				{
-					const TArray<FRigElementKey> Bones = Controller->ImportBones(RigBlueprint->PreviewSkeletalMesh->GetSkeleton());
-
-					if(URigVM* VM = InRig->GetVM())
+					if(URigHierarchyController* Controller = Hierarchy->GetController(true))
 					{
-						const int32 EntryIndex = VM->GetByteCode().FindEntryIndex(FRigUnit_PrepareForExecution::EventName);
-						if(EntryIndex != INDEX_NONE)
+						// find the instruction index for the construction event
+						int32 InstructionIndex = INDEX_NONE;
+						if(URigVM* VM = InRig->GetVM())
 						{
-							const int32 InstructionIndex = VM->GetByteCode().GetEntry(EntryIndex).InstructionIndex;
-							for(const FRigElementKey& Bone : Bones)
+							const int32 EntryIndex = VM->GetByteCode().FindEntryIndex(FRigUnit_PrepareForExecution::EventName);
+							if(EntryIndex != INDEX_NONE)
 							{
-								if(FRigBaseElement* Element = InRig->GetHierarchy()->Find(Bone))
-								{
-									Element->CreatedAtInstructionIndex = InstructionIndex;
-								}
+								InstructionIndex = VM->GetByteCode().GetEntry(EntryIndex).InstructionIndex;
+							}
+						}
+
+						// import the bones for the preview hierarchy
+						const TArray<FRigElementKey> Bones = Controller->ImportBones(RigBlueprint->PreviewSkeletalMesh->GetSkeleton());
+						for(const FRigElementKey& Bone : Bones)
+						{
+							if(FRigBaseElement* Element = InRig->GetHierarchy()->Find(Bone))
+							{
+								Element->CreatedAtInstructionIndex = InstructionIndex;
+							}
+						}
+
+						// create a null to store controls under
+						static const FRigElementKey ControlParentKey(TEXT("Controls"), ERigElementType::Null);
+						if(!Hierarchy->Contains(ControlParentKey))
+						{
+							const FRigElementKey Null = Controller->AddNull(ControlParentKey.Name, FRigElementKey(), FTransform::Identity, true, false, false);
+							if(FRigBaseElement* Element = InRig->GetHierarchy()->Find(Null))
+							{
+								Element->CreatedAtInstructionIndex = InstructionIndex;
 							}
 						}
 					}
-					//(void)RigBlueprint->UpdateConnectionRedirectorOnCDO();
 				}
 			}
 		}
diff --git a/Engine/Plugins/Animation/ControlRig/Source/ControlRigEditor/Private/Editor/ControlRigEditor.h b/Engine/Plugins/Animation/ControlRig/Source/ControlRigEditor/Private/Editor/ControlRigEditor.h
--- a/Engine/Plugins/Animation/ControlRig/Source/ControlRigEditor/Private/Editor/ControlRigEditor.h
+++ b/Engine/Plugins/Animation/ControlRig/Source/ControlRigEditor/Private/Editor/ControlRigEditor.h
@@ -117,6 +117,7 @@ public:
 	URigHierarchy* GetHierarchyBeingDebugged() const;
 
 	void SetDetailViewForRigElements();
+	void SetDetailViewForRigElements(const TArray<FRigElementKey>& InKeys);
 	bool DetailViewShowsAnyRigElement() const;
 	bool DetailViewShowsRigElement(FRigElementKey InKey) const;
 	virtual void RefreshDetailView() override;
@@ -165,6 +166,12 @@ public:
 	bool ClearDirectManipulationSubject() { return SetDirectionManipulationSubject(nullptr); }
 	void RefreshDirectManipulationTextList();
 
+	// Rig connector functionality
+	EVisibility GetConnectorWarningVisibility() const;
+	FText GetConnectorWarningText() const;
+	FReply OnNavigateToConnectorWarning() const;
+	FSimpleMulticastDelegate& OnRequestNavigateToConnectorWarning() { return RequestNavigateToConnectorWarningDelegate; }
+	
 protected:
 
 	virtual void BindCommands() override;
@@ -311,7 +318,8 @@ protected:
 	mutable TArray<TSharedPtr<FString>> DirectManipulationTextList;
 	TSharedPtr<SComboBox<TSharedPtr<FString>>> DirectManipulationCombo;
 	bool bRefreshDirectionManipulationTargetsRequired;
-	
+	FSimpleMulticastDelegate RequestNavigateToConnectorWarningDelegate;
+
 	static const TArray<FName> ForwardsSolveEventQueue;
 	static const TArray<FName> BackwardsSolveEventQueue;
 	static const TArray<FName> ConstructionEventQueue;
diff --git a/Engine/Plugins/Animation/ControlRig/Source/ControlRigEditor/Private/Editor/SRigHierarchy.cpp b/Engine/Plugins/Animation/ControlRig/Source/ControlRigEditor/Private/Editor/SRigHierarchy.cpp
--- a/Engine/Plugins/Animation/ControlRig/Source/ControlRigEditor/Private/Editor/SRigHierarchy.cpp
+++ b/Engine/Plugins/Animation/ControlRig/Source/ControlRigEditor/Private/Editor/SRigHierarchy.cpp
@@ -144,6 +144,7 @@ void SRigHierarchy::Construct(const FArguments& InArgs, TSharedRef<FControlRigEd
 	Delegates.OnAcceptDrop = FOnRigTreeAcceptDrop::CreateSP(this, &SRigHierarchy::OnAcceptDrop);
 	Delegates.OnDragDetected = FOnDragDetected::CreateSP(this, &SRigHierarchy::OnDragDetected);
 	Delegates.OnGetResolvedKey = FOnRigTreeGetResolvedKey::CreateSP(this, &SRigHierarchy::OnGetResolvedKey);
+	Delegates.OnRequestDetailsInspection = FOnRigTreeRequestDetailsInspection::CreateSP(this, &SRigHierarchy::OnRequestDetailsInspection);
 
 	ChildSlot
 	[
@@ -278,6 +279,7 @@ void SRigHierarchy::Construct(const FArguments& InArgs, TSharedRef<FControlRigEd
 		ControlRigEditor.Pin()->OnGetViewportContextMenu().BindSP(this, &SRigHierarchy::GetContextMenu);
 		ControlRigEditor.Pin()->OnViewportContextMenuCommands().BindSP(this, &SRigHierarchy::GetContextMenuCommands);
 		ControlRigEditor.Pin()->OnEditorClosed().AddSP(this, &SRigHierarchy::OnEditorClose);
+		ControlRigEditor.Pin()->OnRequestNavigateToConnectorWarning().AddSP(this, &SRigHierarchy::OnNavigateToFirstConnectorWarning);
 	}
 	
 	CreateContextMenu();
@@ -292,6 +294,7 @@ void SRigHierarchy::OnEditorClose(const FRigVMEditor* InEditor, URigVMBlueprint*
 		Editor->GetKeyDownDelegate().Unbind();
 		Editor->OnGetViewportContextMenu().Unbind();
 		Editor->OnViewportContextMenuCommands().Unbind();
+		Editor->OnRequestNavigateToConnectorWarning().RemoveAll(this);
 	}
 
 	if (UControlRigBlueprint* BP = Cast<UControlRigBlueprint>(InBlueprint))
@@ -537,7 +540,8 @@ void SRigHierarchy::OnFilterTextChanged(const FText& SearchText)
 
 void SRigHierarchy::RefreshTreeView(bool bRebuildContent)
 {
-	if(const URigHierarchy* Hierarchy = GetHierarchy())
+	const URigHierarchy* Hierarchy = GetHierarchy();
+	if(Hierarchy)
 	{
 		// is the rig currently running
 		if(Hierarchy->HasExecuteContext())
@@ -962,6 +966,34 @@ void SRigHierarchy::OnPostConstruction_AnyThread(UControlRig* InRig, const FName
 	}
 }
 
+void SRigHierarchy::OnNavigateToFirstConnectorWarning()
+{
+	if(ControlRigEditor.IsValid())
+	{
+		if(UControlRig* ControlRig = ControlRigEditor.Pin()->GetControlRig())
+		{
+			FRigElementKey ConnectorKey;
+			if(!ControlRig->AllConnectorsAreResolved(nullptr, &ConnectorKey))
+			{
+				if(ConnectorKey.IsValid())
+				{
+					if(URigHierarchy* Hierarchy = ControlRig->GetHierarchy())
+					{
+						if(URigHierarchyController* HierarchyController = Hierarchy->GetController())
+						{
+							{
+								const FRigHierarchyRedirectorGuard RedirectorGuard(ControlRig);
+								HierarchyController->SetSelection({ConnectorKey}, false);
+							}
+							HandleFrameSelection();
+						}
+					}
+				}
+			}
+		}
+	}
+}
+
 void SRigHierarchy::ClearDetailPanel() const
 {
 	if(ControlRigEditor.IsValid())
@@ -1562,6 +1594,15 @@ FRigElementKey SRigHierarchy::OnGetResolvedKey(const FRigElementKey& InKey)
 	return InKey;
 }
 
+void SRigHierarchy::OnRequestDetailsInspection(const FRigElementKey& InKey)
+{
+	if(!ControlRigEditor.IsValid())
+	{
+		return;
+	}
+	ControlRigEditor.Pin()->SetDetailViewForRigElements({InKey});
+}
+
 void SRigHierarchy::ImportHierarchy(const FAssetData& InAssetData)
 {
 	if (bIsChangingRigHierarchy)
@@ -2459,6 +2500,31 @@ TOptional<EItemDropZone> SRigHierarchy::OnCanAcceptDrop(const FDragDropEvent& Dr
 		}
 	}
 
+	const TSharedPtr<FRigHierarchyTagDragDropOp> TagDragDropOp = DragDropEvent.GetOperationAs<FRigHierarchyTagDragDropOp>();
+	if(TagDragDropOp.IsValid())
+	{
+		if(DropZone != EItemDropZone::OntoItem)
+		{
+			return InvalidDropZone;
+		}
+
+		if (const URigHierarchy* Hierarchy = GetHierarchy())
+		{
+			FRigElementKey DraggedKey;
+			FRigElementKey::StaticStruct()->ImportText(*TagDragDropOp->GetIdentifier(), &DraggedKey, nullptr, EPropertyPortFlags::PPF_None, nullptr, FRigElementKey::StaticStruct()->GetName(), true);
+
+			if(Hierarchy->Contains(DraggedKey))
+			{
+				// todo: apply rules
+				ReturnDropZone = DropZone;
+			}
+			else if(!TargetItem.IsValid())
+			{
+				ReturnDropZone = DropZone;
+			}
+		}
+	}
+
 	return ReturnDropZone;
 }
 
@@ -2529,6 +2595,21 @@ FReply SRigHierarchy::OnAcceptDrop(const FDragDropEvent& DragDropEvent, EItemDro
 
 	}
 
+	const TSharedPtr<FRigHierarchyTagDragDropOp> TagDragDropOp = DragDropEvent.GetOperationAs<FRigHierarchyTagDragDropOp>();
+	if(TagDragDropOp.IsValid())
+	{
+		if (const URigHierarchy* Hierarchy = GetHierarchy())
+		{
+			FRigElementKey DraggedKey;
+			FRigElementKey::StaticStruct()->ImportText(*TagDragDropOp->GetIdentifier(), &DraggedKey, nullptr, EPropertyPortFlags::PPF_None, nullptr, FRigElementKey::StaticStruct()->GetName(), true);
+			if(TargetItem.IsValid())
+			{
+				return ResolveConnector(DraggedKey, TargetItem->Key);
+			}
+			return ResolveConnector(DraggedKey, FRigElementKey());
+		}
+	}
+
 	return FReply::Unhandled();
 }
 
diff --git a/Engine/Plugins/Animation/ControlRig/Source/ControlRigEditor/Private/Editor/SRigHierarchy.h b/Engine/Plugins/Animation/ControlRig/Source/ControlRigEditor/Private/Editor/SRigHierarchy.h
--- a/Engine/Plugins/Animation/ControlRig/Source/ControlRigEditor/Private/Editor/SRigHierarchy.h
+++ b/Engine/Plugins/Animation/ControlRig/Source/ControlRigEditor/Private/Editor/SRigHierarchy.h
@@ -198,6 +198,7 @@ private:
 	URigHierarchy* GetDefaultHierarchy() const;
 	const URigHierarchy* GetHierarchyForTreeView() const { return GetHierarchy(); }
 	FRigElementKey OnGetResolvedKey(const FRigElementKey& InKey);
+	void OnRequestDetailsInspection(const FRigElementKey& InKey);
 	
 	void ImportHierarchy(const FAssetData& InAssetData);
 	void CreateImportMenu(FMenuBuilder& MenuBuilder);
@@ -234,6 +235,7 @@ private:
 	void HandleSetObjectBeingDebugged(UObject* InObject);
 	void OnPreConstruction_AnyThread(UControlRig* InRig, const FName& InEventName);
 	void OnPostConstruction_AnyThread(UControlRig* InRig, const FName& InEventName);
+	void OnNavigateToFirstConnectorWarning();
 
 	bool bIsConstructionEventRunning;
 	uint32 LastHierarchyHash;
@@ -249,4 +251,5 @@ public:
 	friend class FRigTreeElement;
 	friend class SRigHierarchyItem;
 	friend class UControlRigBlueprintEditorLibrary;
-};
\ No newline at end of file
+};
+
diff --git a/Engine/Plugins/Animation/ControlRig/Source/ControlRigEditor/Private/Editor/SRigHierarchyTagWidget.cpp b/Engine/Plugins/Animation/ControlRig/Source/ControlRigEditor/Private/Editor/SRigHierarchyTagWidget.cpp
new file mode 100644
--- /dev/null
+++ b/Engine/Plugins/Animation/ControlRig/Source/ControlRigEditor/Private/Editor/SRigHierarchyTagWidget.cpp
@@ -0,0 +1,141 @@
+// Copyright Epic Games, Inc. All Rights Reserved.
+
+#include "Editor/SRigHierarchyTagWidget.h"
+#include "Widgets/SBoxPanel.h"
+#include "Widgets/Images/SImage.h"
+#include "Widgets/Text/STextBlock.h"
+#include "Brushes/SlateRoundedBoxBrush.h"
+#include "DetailLayoutBuilder.h"
+
+#define LOCTEXT_NAMESPACE "SRigHierarchyTagWidget"
+
+//////////////////////////////////////////////////////////////
+/// SRigHierarchyTagWidget
+///////////////////////////////////////////////////////////
+
+TSharedRef<FRigHierarchyTagDragDropOp> FRigHierarchyTagDragDropOp::New(TSharedPtr<SRigHierarchyTagWidget> InTagWidget)
+{
+	TSharedRef<FRigHierarchyTagDragDropOp> Operation = MakeShared<FRigHierarchyTagDragDropOp>();
+	Operation->Text = InTagWidget->Text.Get();
+	Operation->Identifier = InTagWidget->Identifier.Get();
+	Operation->Construct();
+	return Operation;
+}
+
+TSharedPtr<SWidget> FRigHierarchyTagDragDropOp::GetDefaultDecorator() const
+{
+	return SNew(SBorder)
+		.Visibility(EVisibility::Visible)
+		.BorderImage(FAppStyle::GetBrush("Menu.Background"))
+		[
+			SNew(STextBlock)
+			.Text(Text)
+		];
+}
+
+void SRigHierarchyTagWidget::Construct(const FArguments& InArgs)
+{
+	Text = InArgs._Text;
+	Icon = InArgs._Icon;
+	Color = InArgs._Color;
+	Radius = InArgs._Radius;
+	Padding = InArgs._Padding;
+	ContentPadding = InArgs._ContentPadding;
+	Identifier = InArgs._Identifier;
+	bAllowDragDrop = InArgs._AllowDragDrop;
+	OnClicked = InArgs._OnClicked;
+
+	const FMargin CombinedPadding = ContentPadding + Padding;
+
+	SetVisibility(InArgs._Visibility);
+
+	ChildSlot
+	[
+		SNew(SHorizontalBox)
+		.ToolTipText(InArgs._TooltipText)
+
+		+ SHorizontalBox::Slot()
+		.AutoWidth()
+		.VAlign(VAlign_Center)
+		.Padding(CombinedPadding.Left, CombinedPadding.Top, 0, CombinedPadding.Bottom)
+		[
+			SNew(SImage)
+			.Visibility_Lambda([this]()
+			{
+				return Icon.Get() ? EVisibility::Visible : EVisibility::Collapsed;
+			})
+			.Image(Icon)
+		]
+
+		+ SHorizontalBox::Slot()
+		.AutoWidth()
+		.VAlign(VAlign_Center)
+		.Padding(TAttribute<FMargin>::CreateSP(this, &SRigHierarchyTagWidget::GetTextPadding))
+		[
+			SNew(STextBlock)
+			.Text(Text)
+			.ColorAndOpacity(InArgs._TextColor)
+			.Font(IDetailLayoutBuilder::GetDetailFont())
+		]
+	];
+}
+
+int32 SRigHierarchyTagWidget::OnPaint(const FPaintArgs& Args, const FGeometry& AllottedGeometry,
+	const FSlateRect& MyCullingRect, FSlateWindowElementList& OutDrawElements, int32 LayerId,
+	const FWidgetStyle& InWidgetStyle, bool bParentEnabled) const
+{
+	const FSlateRoundedBoxBrush RoundedBoxBrush(FLinearColor::White, Radius);
+	FSlateDrawElement::MakeBox(
+		OutDrawElements,
+		LayerId++,
+		AllottedGeometry.ToPaintGeometry(
+			AllottedGeometry.GetLocalSize() - Padding.GetDesiredSize(),
+			FSlateLayoutTransform(FVector2d(Padding.Left, Padding.Top))
+		),
+		&RoundedBoxBrush,
+		ESlateDrawEffect::NoPixelSnapping,
+		Color.Get()
+	);
+	return SCompoundWidget::OnPaint(Args, AllottedGeometry, MyCullingRect, OutDrawElements, LayerId, InWidgetStyle, bParentEnabled);
+}
+
+FReply SRigHierarchyTagWidget::OnMouseButtonDown(const FGeometry& MyGeometry, const FPointerEvent& MouseEvent)
+{
+	if (MouseEvent.GetEffectingButton() == EKeys::LeftMouseButton)
+	{
+		if(OnClicked.IsBound())
+		{
+			OnClicked.Execute();
+		}
+		return FReply::Handled().DetectDrag(SharedThis(this), EKeys::LeftMouseButton);
+	}
+	return FReply::Unhandled();
+}
+
+FReply SRigHierarchyTagWidget::OnDragDetected(const FGeometry& MyGeometry, const FPointerEvent& MouseEvent)
+{
+	if(bAllowDragDrop && (Identifier.IsSet() || Identifier.IsBound())) 
+	{
+		if (MouseEvent.IsMouseButtonDown(EKeys::LeftMouseButton) && !Identifier.Get().IsEmpty())
+		{
+			TSharedRef<FRigHierarchyTagDragDropOp> DragDropOp = FRigHierarchyTagDragDropOp::New(SharedThis(this));
+			return FReply::Handled().BeginDragDrop(DragDropOp);
+		}
+	}
+	return FReply::Unhandled();
+}
+
+FMargin SRigHierarchyTagWidget::GetTextPadding() const
+{
+	const FMargin CombinedPadding = ContentPadding + Padding;
+	if(Icon.IsBound() || Icon.IsSet())
+	{
+		if(Icon.Get())
+		{
+			return FMargin(ContentPadding.Left * 2, ContentPadding.Top, ContentPadding.Right, ContentPadding.Bottom);
+		}
+	}
+	return CombinedPadding;
+}
+
+#undef LOCTEXT_NAMESPACE
diff --git a/Engine/Plugins/Animation/ControlRig/Source/ControlRigEditor/Private/Editor/SRigHierarchyTagWidget.h b/Engine/Plugins/Animation/ControlRig/Source/ControlRigEditor/Private/Editor/SRigHierarchyTagWidget.h
new file mode 100644
--- /dev/null
+++ b/Engine/Plugins/Animation/ControlRig/Source/ControlRigEditor/Private/Editor/SRigHierarchyTagWidget.h
@@ -0,0 +1,80 @@
+// Copyright Epic Games, Inc. All Rights Reserved.
+
+#pragma once
+
+#include "CoreMinimal.h"
+#include "Widgets/SCompoundWidget.h"
+#include "DragAndDrop/DecoratedDragDropOp.h"
+
+class SRigHierarchyTagWidget;
+
+class FRigHierarchyTagDragDropOp : public FDecoratedDragDropOp
+{
+public:
+	DRAG_DROP_OPERATOR_TYPE(FRigHierarchyTagDragDropOp, FDecoratedDragDropOp)
+
+	static TSharedRef<FRigHierarchyTagDragDropOp> New(TSharedPtr<SRigHierarchyTagWidget> InTagWidget);
+
+	virtual TSharedPtr<SWidget> GetDefaultDecorator() const override;
+
+	/** @return The identifier being dragged */
+	const FString& GetIdentifier() const
+	{
+		return Identifier;
+	}
+
+private:
+
+	FText Text;
+	FString Identifier;
+};
+
+class SRigHierarchyTagWidget : public SCompoundWidget
+{
+public:
+	SLATE_BEGIN_ARGS(SRigHierarchyTagWidget)
+		: _Text()
+		, _TooltipText()
+		, _Icon(nullptr)
+		, _Color(FLinearColor(.2, .2, .2))
+		, _TextColor(FLinearColor(.8, .8, .8))
+		, _Radius(5)
+		, _Padding(10, 0, 0, 0)
+		, _ContentPadding(2, 2, 2, 2)
+		, _Identifier()
+		, _AllowDragDrop(false)
+		, _OnClicked()
+	{}
+		SLATE_ATTRIBUTE(FText, Text)
+		SLATE_ATTRIBUTE(FText, TooltipText)
+		SLATE_ATTRIBUTE(const FSlateBrush*, Icon)
+		SLATE_ATTRIBUTE(FLinearColor, Color)
+		SLATE_ATTRIBUTE(FSlateColor, TextColor)
+		SLATE_ARGUMENT(float, Radius)
+		SLATE_ARGUMENT(FMargin, Padding)
+		SLATE_ARGUMENT(FMargin, ContentPadding)
+		SLATE_ATTRIBUTE(FString, Identifier)
+		SLATE_ARGUMENT(bool, AllowDragDrop)
+		SLATE_EVENT(FSimpleDelegate, OnClicked)
+	SLATE_END_ARGS()
+	
+	void Construct(const FArguments& InArgs);
+	virtual int32 OnPaint(const FPaintArgs& Args, const FGeometry& AllottedGeometry, const FSlateRect& MyCullingRect, FSlateWindowElementList& OutDrawElements, int32 LayerId, const FWidgetStyle& InWidgetStyle, bool bParentEnabled) const override;
+	virtual FReply OnMouseButtonDown(const FGeometry& MyGeometry, const FPointerEvent& MouseEvent) override;
+	virtual FReply OnDragDetected(const FGeometry& MyGeometry, const FPointerEvent& MouseEvent) override;
+private:
+
+	FMargin GetTextPadding() const;
+
+	TAttribute<FText> Text;
+	TAttribute<const FSlateBrush*> Icon;
+	TAttribute<FLinearColor> Color;
+	float Radius = 5.f;
+	FMargin Padding;
+	FMargin ContentPadding;
+	TAttribute<FString> Identifier;
+	bool bAllowDragDrop = false;
+	FSimpleDelegate OnClicked;
+
+	friend class FRigHierarchyTagDragDropOp;
+};
diff --git a/Engine/Plugins/Animation/ControlRig/Source/ControlRigEditor/Private/Editor/SRigHierarchyTreeView.cpp b/Engine/Plugins/Animation/ControlRig/Source/ControlRigEditor/Private/Editor/SRigHierarchyTreeView.cpp
--- a/Engine/Plugins/Animation/ControlRig/Source/ControlRigEditor/Private/Editor/SRigHierarchyTreeView.cpp
+++ b/Engine/Plugins/Animation/ControlRig/Source/ControlRigEditor/Private/Editor/SRigHierarchyTreeView.cpp
@@ -17,6 +17,8 @@
 #include "Editor/SRigHierarchy.h"
 #include "Settings/ControlRigSettings.h"
 #include "Graph/ControlRigGraphSchema.h"
+#include "Rigs/AdditiveControlRig.h"
+#include "Rigs/RigHierarchyController.h"
 #include "Styling/AppStyle.h"
 
 #define LOCTEXT_NAMESPACE "SRigHierarchyTreeView"
@@ -126,6 +128,7 @@ void SRigHierarchyItem::Construct(const FArguments& InArgs, const TSharedRef<STa
 	}
 
 	TSharedPtr< SInlineEditableTextBlock > InlineWidget;
+	TSharedPtr< SHorizontalBox > HorizontalBox;
 
 	STableRow<TSharedPtr<FRigTreeElement>>::Construct(
 		STableRow<TSharedPtr<FRigTreeElement>>::FArguments()
@@ -135,7 +138,7 @@ void SRigHierarchyItem::Construct(const FArguments& InArgs, const TSharedRef<STa
 		.ShowWires(true)
 		.Content()
 		[
-			SNew(SHorizontalBox)
+			SAssignNew(HorizontalBox, SHorizontalBox)
 			+ SHorizontalBox::Slot()
 			.MaxWidth(18)
 			.FillWidth(1.0)
@@ -181,6 +184,15 @@ void SRigHierarchyItem::Construct(const FArguments& InArgs, const TSharedRef<STa
 			]
 		], OwnerTable);
 
+	for(const SRigHierarchyTagWidget::FArguments& TagArguments : InRigTreeElement->Tags)
+	{
+		HorizontalBox->AddSlot()
+		.AutoWidth()
+		[
+			SArgumentNew(TagArguments, SRigHierarchyTagWidget)
+		];
+	}
+
 	InRigTreeElement->OnRenameRequested.BindSP(InlineWidget.Get(), &SInlineEditableTextBlock::EnterEditingMode);
 }
 
@@ -412,6 +424,7 @@ bool SRigHierarchyTreeView::AddElement(const FRigBaseElement* InElement)
 	}
 
 	const FRigTreeDisplaySettings& Settings = Delegates.GetDisplaySettings();
+	const URigHierarchy* Hierarchy = Delegates.GetHierarchy();
 
 	switch(InElement->GetType())
 	{
@@ -465,6 +478,47 @@ bool SRigHierarchyTreeView::AddElement(const FRigBaseElement* InElement)
 		{
 			return false;
 		}
+		case ERigElementType::Connector:
+		{
+			if(Hierarchy)
+			{
+				// add the connector as a tag rather than its own element in the tree
+				if(UControlRig* ControlRig = Hierarchy->GetTypedOuter<UControlRig>())
+				{
+					FRigElementKeyRedirector& Redirector = ControlRig->GetElementKeyRedirector();
+					if(const FCachedRigElement* Cache = Redirector.Find(InElement->GetKey()))
+					{
+						if(const_cast<FCachedRigElement*>(Cache)->UpdateCache(Hierarchy))
+						{
+							if(const TSharedPtr<FRigTreeElement>* TargetElementPtr = ElementMap.Find(Cache->GetKey()))
+							{
+								const FRigElementKey ConnectorKey = InElement->GetKey();
+
+								SRigHierarchyTagWidget::FArguments TagArguments;
+								TagArguments.Text(FText::FromString(InElement->GetName()));
+								TagArguments.TextColor(FSlateColor(FLinearColor::White));
+								TagArguments.Color(FLinearColor(0.0, 112.f/255.f, 224.f/255.f));
+								TagArguments.AllowDragDrop(true);
+								FString Identifier;
+								FRigElementKey::StaticStruct()->ExportText(Identifier, &ConnectorKey, nullptr, nullptr, PPF_None, nullptr);
+								TagArguments.Identifier(Identifier);
+
+								TagArguments.OnClicked_Lambda([ConnectorKey, this]()
+								{
+									Delegates.RequestDetailsInspection(ConnectorKey);
+								});
+
+								TargetElementPtr->Get()->Tags.Add(TagArguments);
+
+								AddConnectorResolveWarningTag(*TargetElementPtr, InElement, Hierarchy);
+								return true;
+							}
+						}
+					}
+				}
+			}
+			break;
+		}
 		default:
 		{
 			break;
@@ -478,8 +532,13 @@ bool SRigHierarchyTreeView::AddElement(const FRigBaseElement* InElement)
 
 	if (ElementMap.Contains(InElement->GetKey()))
 	{
-		if(const URigHierarchy* Hierarchy = Delegates.GetHierarchy())
+		if(Hierarchy)
 		{
+			if(InElement->GetType() == ERigElementType::Connector)
+			{
+				AddConnectorResolveWarningTag(ElementMap.FindChecked(InElement->GetKey()), InElement, Hierarchy);
+			}
+			
 			FRigElementKey ParentKey = Hierarchy->GetFirstParent(InElement->GetKey());
 			if(InElement->GetType() == ERigElementType::Connector)
 			{
@@ -798,6 +857,72 @@ const TSharedPtr<FRigTreeElement>* SRigHierarchyTreeView::FindItemAtPosition(FVe
 	return nullptr;
 }
 
+void SRigHierarchyTreeView::AddConnectorResolveWarningTag(TSharedPtr<FRigTreeElement> InTreeElement,
+	const FRigBaseElement* InRigElement, const URigHierarchy* InHierarchy)
+{
+	check(InTreeElement.IsValid());
+	check(InRigElement);
+	check(InRigElement->GetType() == ERigElementType::Connector);
+
+	if(UControlRig* ControlRig = InHierarchy->GetTypedOuter<UControlRig>())
+	{
+		TWeakObjectPtr<UControlRig> ControlRigPtr(ControlRig);
+		const FRigElementKey ConnectorKey = InRigElement->GetKey();
+		
+		TAttribute<FText> GetTooltipText = TAttribute<FText>::CreateSP(this,
+			&SRigHierarchyTreeView::GetConnectorWarningMessage, InTreeElement, ControlRigPtr, ConnectorKey);
+
+		SRigHierarchyTagWidget::FArguments TagArguments;
+		TagArguments.Visibility_Lambda([GetTooltipText]() -> EVisibility
+		{
+			return GetTooltipText.Get().IsEmpty() ? EVisibility::Collapsed : EVisibility::Visible;
+		});
+		TagArguments.Text(LOCTEXT("ConnectorWarningTagLabel", "Warning"));
+		TagArguments.ToolTipText(GetTooltipText);
+		TagArguments.TextColor(FSlateColor(FLinearColor::Black));
+		TagArguments.Color(FLinearColor(0.728f, 0.364f, 0.003f));
+		TagArguments.Icon(FAppStyle::GetBrush("AnimSlotManager.Warning"));
+		InTreeElement->Tags.Add(TagArguments);
+	}
+}
+
+FText SRigHierarchyTreeView::GetConnectorWarningMessage(TSharedPtr<FRigTreeElement> InTreeElement,
+	TWeakObjectPtr<UControlRig> InControlRigPtr, const FRigElementKey InConnectorKey) const
+{
+	if(UControlRig* ControlRig = InControlRigPtr.Get())
+	{
+		URigHierarchy* Hierarchy = ControlRig->GetHierarchy();
+				
+		static const FText NotResolvedWarning = LOCTEXT("ConnectorWarningConnectorNotResolved", "Connector is not resolved.");
+
+		FRigElementKeyRedirector& Redirector = ControlRig->GetElementKeyRedirector();
+		if(const FCachedRigElement* Cache = Redirector.Find(InConnectorKey))
+		{
+			if(const_cast<FCachedRigElement*>(Cache)->UpdateCache(Hierarchy))
+			{
+				if(const FRigConnectorElement* ConnectorElement = Hierarchy->Find<FRigConnectorElement>(InConnectorKey))
+				{
+					const FRigConnectionInfo ConnectionInfo(&ControlRig->GetElementKeyRedirector(), Hierarchy);
+					FString FailureReason;
+					if(!ConnectorElement->CanConnect(&ConnectionInfo, &FailureReason))
+					{
+						return FText::FromString(FailureReason);
+					}
+				}
+			}
+			else
+			{
+				return NotResolvedWarning;
+			}
+		}
+		else
+		{
+			return NotResolvedWarning;
+		}
+	}
+	return FText();
+}
+
 bool SRigHierarchyItem::OnVerifyNameChanged(const FText& InText, FText& OutErrorMessage)
 {
 	const FString NewName = InText.ToString();
diff --git a/Engine/Plugins/Animation/ControlRig/Source/ControlRigEditor/Private/Editor/SRigHierarchyTreeView.h b/Engine/Plugins/Animation/ControlRig/Source/ControlRigEditor/Private/Editor/SRigHierarchyTreeView.h
--- a/Engine/Plugins/Animation/ControlRig/Source/ControlRigEditor/Private/Editor/SRigHierarchyTreeView.h
+++ b/Engine/Plugins/Animation/ControlRig/Source/ControlRigEditor/Private/Editor/SRigHierarchyTreeView.h
@@ -5,6 +5,7 @@
 #include "CoreMinimal.h"
 #include "Widgets/Views/STreeView.h"
 #include "Rigs/RigHierarchy.h"
+#include "SRigHierarchyTagWidget.h"
 
 class SSearchBox;
 class SRigHierarchyTreeView;
@@ -64,6 +65,7 @@ DECLARE_DELEGATE_RetVal_TwoParams(FName, FOnRigTreeRenameElement, const FRigElem
 DECLARE_DELEGATE_RetVal_ThreeParams(bool, FOnRigTreeVerifyElementNameChanged, const FRigElementKey& /*OldKey*/, const FString& /*NewName*/, FText& /*OutErrorMessage*/);
 DECLARE_DELEGATE_RetVal_TwoParams(bool, FOnRigTreeCompareKeys, const FRigElementKey& /*A*/, const FRigElementKey& /*B*/);
 DECLARE_DELEGATE_RetVal_OneParam(FRigElementKey, FOnRigTreeGetResolvedKey, const FRigElementKey&);
+DECLARE_DELEGATE_OneParam(FOnRigTreeRequestDetailsInspection, const FRigElementKey&);
 
 typedef STableRow<TSharedPtr<FRigTreeElement>>::FOnCanAcceptDrop FOnRigTreeCanAcceptDrop;
 typedef STableRow<TSharedPtr<FRigTreeElement>>::FOnAcceptDrop FOnRigTreeAcceptDrop;
@@ -88,6 +90,7 @@ struct CONTROLRIGEDITOR_API FRigTreeDelegates
 	FOnRigTreeSetExpansionRecursive OnSetExpansionRecursive;
 	FOnRigTreeCompareKeys OnCompareKeys;
 	FOnRigTreeGetResolvedKey OnGetResolvedKey;
+	FOnRigTreeRequestDetailsInspection OnRequestDetailsInspection;
 
 	FRigTreeDelegates()
 	{
@@ -149,6 +152,14 @@ struct CONTROLRIGEDITOR_API FRigTreeDelegates
 		return InKey;
 	}
 
+	void RequestDetailsInspection(const FRigElementKey& InKey)
+	{
+		if(OnRequestDetailsInspection.IsBound())
+		{
+			return OnRequestDetailsInspection.Execute(InKey);
+		}
+	}
+
 	static FRigTreeDisplaySettings DefaultDisplaySettings;
 	bool bIsChangingRigHierarchy;
 };
@@ -206,6 +217,9 @@ public:
 
 	/** The color to use when rendering the label text */
 	FSlateColor TextColor;
+
+	/** The tag arguments for this element */
+	TArray<SRigHierarchyTagWidget::FArguments> Tags;
 };
 
 class SRigHierarchyItem : public STableRow<TSharedPtr<FRigTreeElement>>
@@ -300,6 +314,9 @@ public:
 
 private:
 
+	void AddConnectorResolveWarningTag(TSharedPtr<FRigTreeElement> InTreeElement, const FRigBaseElement* InRigElement, const URigHierarchy* InHierarchy);
+	FText GetConnectorWarningMessage(TSharedPtr<FRigTreeElement> InTreeElement, TWeakObjectPtr<UControlRig> InControlRigPtr, const FRigElementKey InConnectorKey) const;
+
 	/** A temporary snapshot of the SparseItemInfos in STreeView, used during RefreshTreeView() */
 	TSparseItemMap OldSparseItemInfos;
 
