ControlRig: Fix static analysis in RigHierarchy.cpp

diff --git a/Engine/Plugins/Animation/ControlRig/Source/ControlRig/Private/Rigs/RigHierarchy.cpp b/Engine/Plugins/Animation/ControlRig/Source/ControlRig/Private/Rigs/RigHierarchy.cpp
--- a/Engine/Plugins/Animation/ControlRig/Source/ControlRig/Private/Rigs/RigHierarchy.cpp
+++ b/Engine/Plugins/Animation/ControlRig/Source/ControlRig/Private/Rigs/RigHierarchy.cpp
@@ -2443,10 +2443,6 @@ bool URigHierarchy::CanConnect(const FRigConnectionInfo* InConnectionInfo, FStri
 				static constexpr TCHAR Format[] = TEXT("Connector element '%s' does not exist.");
 				OutFailureReason->Appendf(Format, *ConnectorKey.ToString());
 			}
-			if(OutConnector)
-			{
-				*OutConnector = Connector->GetKey();
-			}
 			return false;
 		}
 		const FRigBaseElement* Target = Find<FRigBaseElement>(TargetKey);
@@ -2457,7 +2453,7 @@ bool URigHierarchy::CanConnect(const FRigConnectionInfo* InConnectionInfo, FStri
 				static constexpr TCHAR Format[] = TEXT("Target element '%s' does not exist.");
 				OutFailureReason->Appendf(Format, *TargetKey.ToString());
 			}
-			if(OutConnector)
+			if(OutConnector && Connector)
 			{
 				*OutConnector = Connector->GetKey();
 			}
@@ -2471,14 +2467,14 @@ bool URigHierarchy::CanConnect(const FRigConnectionInfo* InConnectionInfo, FStri
 				static constexpr TCHAR Format[] = TEXT("Target element '%s' is not a transform.");
 				OutFailureReason->Appendf(Format, *TargetKey.ToString());
 			}
-			if(OutConnector)
+			if(OutConnector && Connector)
 			{
 				*OutConnector = Connector->GetKey();
 			}
 			return false;
 		}
 
-		if(!Connector->CanConnect(InConnectionInfo, OutFailureReason))
+		if(Connector && !Connector->CanConnect(InConnectionInfo, OutFailureReason))
 		{
 			if(OutConnector)
 			{
