Control Rig: avoid cleaning up connector related delegate on editor close

diff --git a/Engine/Plugins/Animation/ControlRig/Source/ControlRigEditor/Private/Editor/SRigHierarchy.cpp b/Engine/Plugins/Animation/ControlRig/Source/ControlRigEditor/Private/Editor/SRigHierarchy.cpp
--- a/Engine/Plugins/Animation/ControlRig/Source/ControlRigEditor/Private/Editor/SRigHierarchy.cpp
+++ b/Engine/Plugins/Animation/ControlRig/Source/ControlRigEditor/Private/Editor/SRigHierarchy.cpp
@@ -294,7 +294,6 @@ void SRigHierarchy::OnEditorClose(const FRigVMEditor* InEditor, URigVMBlueprint*
 		Editor->GetKeyDownDelegate().Unbind();
 		Editor->OnGetViewportContextMenu().Unbind();
 		Editor->OnViewportContextMenuCommands().Unbind();
-		Editor->OnRequestNavigateToConnectorWarning().RemoveAll(this);
 	}
 
 	if (UControlRigBlueprint* BP = Cast<UControlRigBlueprint>(InBlueprint))
