Control Rig: Create Transient temporary data link objects in the transient packlage, to avoid them to keep objects alive

diff --git a/Engine/Plugins/Animation/ControlRig/Source/ControlRig/Private/ControlRig.cpp b/Engine/Plugins/Animation/ControlRig/Source/ControlRig/Private/ControlRig.cpp
--- a/Engine/Plugins/Animation/ControlRig/Source/ControlRig/Private/ControlRig.cpp
+++ b/Engine/Plugins/Animation/ControlRig/Source/ControlRig/Private/ControlRig.cpp
@@ -2804,7 +2804,7 @@ const TArray<UAssetUserData*>* UControlRig::GetAssetUserDataArray() const
 						{
 							ExternalVariableDataAssetLinks.Add(
 								ExternalVariable.Name,
-								NewObject<UDataAssetLink>((UObject*)this));
+								NewObject<UDataAssetLink>(GetTransientPackage(), UDataAssetLink::StaticClass(), NAME_None, RF_Transient));
 						}
 
 						TObjectPtr<UDataAssetLink>& DataAssetLink = ExternalVariableDataAssetLinks.FindChecked(ExternalVariable.Name);
