Control RIg: Invert pose into additive control rig

diff --git a/Engine/Plugins/Animation/ControlRig/Source/ControlRig/Private/ControlRig.cpp b/Engine/Plugins/Animation/ControlRig/Source/ControlRig/Private/ControlRig.cpp
--- a/Engine/Plugins/Animation/ControlRig/Source/ControlRig/Private/ControlRig.cpp
+++ b/Engine/Plugins/Animation/ControlRig/Source/ControlRig/Private/ControlRig.cpp
@@ -100,6 +100,7 @@ void UControlRig::BeginDestroy()
 	PostConstructionEvent.Clear();
 	PreForwardsSolveEvent.Clear();
 	PostForwardsSolveEvent.Clear();
+	PreAdditiveValuesApplicationEvent.Clear();
 
 #if WITH_EDITOR
 	if (!HasAnyFlags(RF_ClassDefaultObject))
@@ -362,6 +363,12 @@ void UControlRig::Evaluate_AnyThread()
 		// Store control pose after backwards solve to figure out additive local transforms based on animation
 		ControlsAfterBackwardsSolve = Hierarchy->GetPose(false, ERigElementType::Control, TArrayView<const FRigElementKey>());
 
+		if (PreAdditiveValuesApplicationEvent.IsBound())
+		{
+			FControlRigBracketScope BracketScope(PreAdditiveValuesApplicationBracket);
+			PreAdditiveValuesApplicationEvent.Broadcast(this, TEXT("Additive"));
+		}
+		
 		// Apply additive controls
 		for (TPair<FRigElementKey, FRigSetControlValueInfo>& Value : ControlValues)
 		{
@@ -446,6 +453,20 @@ void UControlRig::ClearPoseBeforeBackwardsSolve()
 	PoseBeforeBackwardsSolve.Reset();
 }
 
+void UControlRig::InvertInputPose(EControlRigSetKey InSetKey)
+{
+	for (const FRigPoseElement& PoseElement : ControlsAfterBackwardsSolve)
+	{
+		if (FRigControlElement* ControlElement = Cast<FRigControlElement>(DynamicHierarchy->Get(PoseElement.Index)))
+		{
+			FRigControlValue Value;
+			Value.SetFromTransform(PoseElement.LocalTransform.Inverse(), ControlElement->Settings.ControlType, ControlElement->Settings.PrimaryAxis);
+			FRigSetControlValueInfo Info = {Value, true, FRigControlModifiedContext(InSetKey), false, false, false};
+			ControlValues.Add(ControlElement->GetKey(), Info);
+		}
+	}
+}
+
 void UControlRig::InitializeFromCDO()
 {
 	Super::InitializeFromCDO();
diff --git a/Engine/Plugins/Animation/ControlRig/Source/ControlRig/Public/ControlRig.h b/Engine/Plugins/Animation/ControlRig/Source/ControlRig/Public/ControlRig.h
--- a/Engine/Plugins/Animation/ControlRig/Source/ControlRig/Public/ControlRig.h
+++ b/Engine/Plugins/Animation/ControlRig/Source/ControlRig/Public/ControlRig.h
@@ -140,6 +140,9 @@ public:
 	 * This usually indicates a new pose should be stored. */
 	void ClearPoseBeforeBackwardsSolve();
 
+	/* For additive rigs, will set control values by inverting the pose found after the backwards solve */
+	void InvertInputPose(EControlRigSetKey InSetKey);
+
 	/** Setup bindings to a runtime object (or clear by passing in nullptr). */
 	void SetObjectBinding(TSharedPtr<IControlRigObjectBinding> InObjectBinding)
 	{
@@ -319,6 +322,7 @@ public:
 	FControlRigExecuteEvent& OnPostConstruction_AnyThread() { return PostConstructionEvent; }
 	FControlRigExecuteEvent& OnPreForwardsSolve_AnyThread() { return PreForwardsSolveEvent; }
 	FControlRigExecuteEvent& OnPostForwardsSolve_AnyThread() { return PostForwardsSolveEvent; }
+	FControlRigExecuteEvent& OnPreAdditiveValuesApplication_AnyThread() { return PreAdditiveValuesApplicationEvent; }
 	FRigEventDelegate& OnRigEvent_AnyThread() { return RigEventDelegate; }
 
 	// Setup the initial transform / ref pose of the bones based upon an anim instance
@@ -439,6 +443,9 @@ private:
 	/** Broadcasts a notification after a forward solve has been initiated */
 	FControlRigExecuteEvent PostForwardsSolveEvent;
 
+	/** Broadcasts a notification before additive controls have been applied */
+	FControlRigExecuteEvent PreAdditiveValuesApplicationEvent;
+
 	/** Handle changes within the hierarchy */
 	void HandleHierarchyModified(ERigHierarchyNotification InNotification, URigHierarchy* InHierarchy, const FRigBaseElement* InElement);
 
@@ -574,6 +581,7 @@ protected:
 	int32 PostConstructionBracket;
 	int32 PreForwardsSolveBracket;
 	int32 PostForwardsSolveBracket;
+	int32 PreAdditiveValuesApplicationBracket;
 	int32 InteractionBracket;
 	int32 InterRigSyncBracket;
 	int32 ControlUndoBracketIndex;
diff --git a/Engine/Plugins/Animation/ControlRig/Source/ControlRigEditor/Private/EditMode/ControlRigEditMode.cpp b/Engine/Plugins/Animation/ControlRig/Source/ControlRigEditor/Private/EditMode/ControlRigEditMode.cpp
--- a/Engine/Plugins/Animation/ControlRig/Source/ControlRigEditor/Private/EditMode/ControlRigEditMode.cpp
+++ b/Engine/Plugins/Animation/ControlRig/Source/ControlRigEditor/Private/EditMode/ControlRigEditMode.cpp
@@ -2358,6 +2358,9 @@ void FControlRigEditMode::BindCommands()
 	CommandBindings->MapAction(
 		Commands.ResetAllTransforms,
 		FExecuteAction::CreateRaw(this, &FControlRigEditMode::ResetTransforms, false));
+	CommandBindings->MapAction(
+		Commands.InvertInputPose,
+		FExecuteAction::CreateRaw(this, &FControlRigEditMode::InvertInputPose));
 	CommandBindings->MapAction(
 		Commands.ClearSelection,
 		FExecuteAction::CreateRaw(this, &FControlRigEditMode::ClearSelection));
@@ -3043,6 +3046,22 @@ void FControlRigEditMode::ResetTransforms(bool bSelectionOnly)
 	}
 }
 
+void FControlRigEditMode::InvertInputPose()
+{
+	FScopedTransaction Transaction(LOCTEXT("HierarchyInvertInputPose", "Invert Input Pose"));
+	for (TWeakObjectPtr<UControlRig>& RuntimeRigPtr : RuntimeControlRigs)
+	{
+		if (RuntimeRigPtr.IsValid())
+		{
+			if (RuntimeRigPtr->IsAdditive())
+			{
+				RuntimeRigPtr->InvertInputPose(EControlRigSetKey::Never);
+				RuntimeRigPtr->Evaluate_AnyThread();
+			}
+		}
+	}
+}
+
 bool FControlRigEditMode::MouseMove(FEditorViewportClient* ViewportClient, FViewport* Viewport, int32 x, int32 y)
 {
 	// Inform units of hover state
diff --git a/Engine/Plugins/Animation/ControlRig/Source/ControlRigEditor/Private/EditMode/ControlRigEditModeCommands.cpp b/Engine/Plugins/Animation/ControlRig/Source/ControlRigEditor/Private/EditMode/ControlRigEditModeCommands.cpp
--- a/Engine/Plugins/Animation/ControlRig/Source/ControlRigEditor/Private/EditMode/ControlRigEditModeCommands.cpp
+++ b/Engine/Plugins/Animation/ControlRig/Source/ControlRigEditor/Private/EditMode/ControlRigEditModeCommands.cpp
@@ -8,6 +8,7 @@ void FControlRigEditModeCommands::RegisterCommands()
 {
 	UI_COMMAND(ResetTransforms, "Reset Transform", "Reset the Controls Transforms", EUserInterfaceActionType::Button, FInputChord(EKeys::G, EModifierKey::Control));
 	UI_COMMAND(ResetAllTransforms, "Reset All Transform", "Reset all of the Controls Transforms", EUserInterfaceActionType::Button, FInputChord(EKeys::G, EModifierKey::Control | EModifierKey::Shift));
+	UI_COMMAND(InvertInputPose, "Invert Input Pose", "Invert Input Pose for additive control rigs", EUserInterfaceActionType::Button, FInputChord(EKeys::I, EModifierKey::Control));
 	UI_COMMAND(ToggleManipulators, "Toggle Shapes", "Toggles visibility of active control rig shapes in the viewport", EUserInterfaceActionType::Button, FInputChord(EKeys::T));
 	UI_COMMAND(ToggleAllManipulators, "Toggle All Shapes", "Toggles visibility of all control rig shapes in the viewport", EUserInterfaceActionType::Button, FInputChord(EKeys::T, EModifierKey::Control | EModifierKey::Alt));
 	UI_COMMAND(FrameSelection, "Frame Selection", "Focus the viewport on the current selection", EUserInterfaceActionType::Button, FInputChord(EKeys::F));
diff --git a/Engine/Plugins/Animation/ControlRig/Source/ControlRigEditor/Private/EditMode/ControlRigEditModeCommands.h b/Engine/Plugins/Animation/ControlRig/Source/ControlRigEditor/Private/EditMode/ControlRigEditModeCommands.h
--- a/Engine/Plugins/Animation/ControlRig/Source/ControlRigEditor/Private/EditMode/ControlRigEditModeCommands.h
+++ b/Engine/Plugins/Animation/ControlRig/Source/ControlRigEditor/Private/EditMode/ControlRigEditModeCommands.h
@@ -38,6 +38,9 @@ public:
 	/** Reset Transforms for Controls */
 	TSharedPtr< FUICommandInfo > ResetAllTransforms;
 
+	/** Invert input pose for additive control rigs */
+	TSharedPtr< FUICommandInfo > InvertInputPose;
+
 	/** Clear Selection*/
 	TSharedPtr< FUICommandInfo > ClearSelection;
 
diff --git a/Engine/Plugins/Animation/ControlRig/Source/ControlRigEditor/Private/Sequencer/ControlRigParameterTrackEditor.cpp b/Engine/Plugins/Animation/ControlRig/Source/ControlRigEditor/Private/Sequencer/ControlRigParameterTrackEditor.cpp
--- a/Engine/Plugins/Animation/ControlRig/Source/ControlRigEditor/Private/Sequencer/ControlRigParameterTrackEditor.cpp
+++ b/Engine/Plugins/Animation/ControlRig/Source/ControlRigEditor/Private/Sequencer/ControlRigParameterTrackEditor.cpp
@@ -88,6 +88,7 @@
 #include "Constraints/TransformConstraintChannelInterface.h"
 #include "BakingAnimationKeySettings.h"
 #include "FrameNumberDetailsCustomization.h"
+#include "Editor/UnrealEd/Private/FbxExporter.h"
 #include "Sequencer/ControlRigSequencerHelpers.h"
 #include "Widgets/Layout/SSpacer.h"
 
@@ -1230,6 +1231,70 @@ void FControlRigParameterTrackEditor::BakeToControlRig(UClass* InClass, FGuid Ob
 	}
 }
 
+void FControlRigParameterTrackEditor::BakeInvertedPose(UControlRig* InControlRig, UMovieSceneControlRigParameterTrack* Track)
+{
+	if (!InControlRig->IsAdditive())
+	{
+		return;
+	}
+
+	UMovieSceneControlRigParameterSection* Section = Cast<UMovieSceneControlRigParameterSection>(Track->GetSectionToKey());
+	USkeletalMeshComponent* SkelMeshComp = Cast<USkeletalMeshComponent>(InControlRig->GetObjectBinding()->GetBoundObject());
+	UMovieSceneSequence* MovieSceneSequence = GetSequencer()->GetFocusedMovieSceneSequence();
+	UMovieScene* MovieScene = MovieSceneSequence->GetMovieScene();
+	UAnimSeqExportOption* ExportOptions = NewObject<UAnimSeqExportOption>(GetTransientPackage(), NAME_None);
+	UBakeToControlRigSettings* BakeSettings = GetMutableDefault<UBakeToControlRigSettings>();
+	const TSharedPtr<ISequencer> ParentSequencer = GetSequencer();
+	FMovieSceneSequenceIDRef Template = ParentSequencer->GetFocusedTemplateID();
+	FMovieSceneSequenceTransform RootToLocalTransform = ParentSequencer->GetFocusedMovieSceneSequenceTransform();
+
+	if (ExportOptions == nullptr || MovieScene == nullptr || SkelMeshComp == nullptr)
+	{
+		UE_LOG(LogMovieScene, Error, TEXT("FControlRigParameterTrackEditor::BakeInvertedPose All parameters must be valid."));
+		return;
+	}
+	
+	const FScopedTransaction Transaction(LOCTEXT("BakeInvertedPose_Transaction", "Bake Inverted Pose"));
+
+	UnFbx::FLevelSequenceAnimTrackAdapter AnimTrackAdapter(ParentSequencer.Get(), MovieScene, RootToLocalTransform);
+	int32 AnimationLength = AnimTrackAdapter.GetLength();
+	FScopedSlowTask Progress(AnimationLength, LOCTEXT("BakingToControlRig_SlowTask", "Baking To Control Rig..."));
+	Progress.MakeDialog(true);
+
+	auto DelegateHandle = InControlRig->OnPreAdditiveValuesApplication_AnyThread().AddLambda([](UControlRig* InControlRig, const FName& InEventName)
+	{
+		InControlRig->InvertInputPose(EControlRigSetKey::Never);
+	});
+
+	auto KeyFrame = [this, ParentSequencer, InControlRig, SkelMeshComp](const FFrameNumber FrameNumber)
+	{
+		const FFrameNumber NewTime = ConvertFrameTime(FrameNumber, ParentSequencer->GetFocusedDisplayRate(), ParentSequencer->GetFocusedTickResolution()).FrameNumber;
+		float LocalTime = ParentSequencer->GetFocusedTickResolution().AsSeconds(FFrameTime(NewTime));
+
+		AddControlKeys(SkelMeshComp, InControlRig, InControlRig->GetFName(), NAME_None, EControlRigContextChannelToKey::AllTransform, 
+				ESequencerKeyMode::ManualKeyForced, LocalTime);
+	};
+
+	FInitAnimationCB InitCallback = FInitAnimationCB::CreateLambda([]{});
+	FStartAnimationCB StartCallback = FStartAnimationCB::CreateLambda([AnimTrackAdapter, KeyFrame]
+	{
+		KeyFrame(AnimTrackAdapter.GetLocalStartFrame());
+	});
+	FTickAnimationCB TickCallback = FTickAnimationCB::CreateLambda([KeyFrame, &Progress](float DeltaTime, FFrameNumber FrameNumber)
+	{
+		KeyFrame(FrameNumber);
+		Progress.EnterProgressFrame(1);
+	});
+	FEndAnimationCB EndCallback = FEndAnimationCB::CreateLambda([]{});
+
+	MovieSceneToolHelpers::BakeToSkelMeshToCallbacks(MovieScene,ParentSequencer.Get(),
+		SkelMeshComp, Template, RootToLocalTransform, ExportOptions,
+		InitCallback, StartCallback, TickCallback, EndCallback);
+
+	InControlRig->OnPreAdditiveValuesApplication_AnyThread().Remove(DelegateHandle);
+	ParentSequencer->NotifyMovieSceneDataChanged(EMovieSceneDataChangeType::TrackValueChanged);
+}
+
 void FControlRigParameterTrackEditor::BuildObjectBindingTrackMenu(FMenuBuilder& MenuBuilder, const TArray<FGuid>& ObjectBindings, const UClass* ObjectClass)
 {
 	if(!ObjectClass)
@@ -3362,7 +3427,7 @@ void FControlRigParameterTrackEditor::GetControlRigKeys(
 
 
 			bool bMaskKeyOut = (ControlIndex >= ControlsMask.Num() || ControlsMask[ControlIndex] == false);
-			bool bSetKey = ControlElement->GetFName() == ParameterName && !bMaskKeyOut;
+			bool bSetKey = ParameterName.IsNone() || (ControlElement->GetFName() == ParameterName && !bMaskKeyOut);
 
 			FRigControlValue ControlValue = InControlRig->GetControlValue(ControlElement, ERigControlValueType::Current);
 
@@ -3981,6 +4046,23 @@ void FControlRigParameterTrackEditor::BuildTrackContextMenu(FMenuBuilder& MenuBu
 
 		MenuBuilder.AddMenuSeparator();
 	}
+	else if (UControlRig* AdditiveRig = Cast<UControlRig>(Track->GetControlRig()))
+	{
+		if (AdditiveRig->IsAdditive())
+		{
+			MenuBuilder.BeginSection("Additive Control Rig", LOCTEXT("AdditiveControlRig", "Additive Control Rig"));
+			{
+				MenuBuilder.AddMenuEntry(
+					LOCTEXT("Bake Inverted Pose", "Bake Inverted Pose"),
+					LOCTEXT("BakeInvertedPoseToolTip", "Bake inversion of the input pose into the rig"),
+					FSlateIcon(),
+					FUIAction(
+						FExecuteAction::CreateRaw(this, &FControlRigParameterTrackEditor::BakeInvertedPose, AdditiveRig, Track)));
+			}
+			MenuBuilder.EndSection();
+			MenuBuilder.AddMenuSeparator();
+		}
+	}
 
 }
 
diff --git a/Engine/Plugins/Animation/ControlRig/Source/ControlRigEditor/Private/Sequencer/ControlRigParameterTrackEditor.h b/Engine/Plugins/Animation/ControlRig/Source/ControlRigEditor/Private/Sequencer/ControlRigParameterTrackEditor.h
--- a/Engine/Plugins/Animation/ControlRig/Source/ControlRigEditor/Private/Sequencer/ControlRigParameterTrackEditor.h
+++ b/Engine/Plugins/Animation/ControlRig/Source/ControlRigEditor/Private/Sequencer/ControlRigParameterTrackEditor.h
@@ -196,6 +196,9 @@ private:
 	/** Bake To Control Rig Sub Menu*/
 	void BakeToControlRig(UClass* InClass, FGuid ObjectBinding,UObject* BoundObject, USkeletalMeshComponent* SkelMeshComp, USkeleton* Skeleton);
 
+	/** Bake Inersion of Additive Control Rig to Rest Pose*/
+	void BakeInvertedPose(UControlRig* InControlRig, UMovieSceneControlRigParameterTrack* Track);
+
 	/** Set Up EditMode for Specified Control Rig*/
 	void SetUpEditModeIfNeeded(UControlRig* ControlRig);
 
diff --git a/Engine/Plugins/Animation/ControlRig/Source/ControlRigEditor/Public/EditMode/ControlRigEditMode.h b/Engine/Plugins/Animation/ControlRig/Source/ControlRigEditor/Public/EditMode/ControlRigEditMode.h
--- a/Engine/Plugins/Animation/ControlRig/Source/ControlRigEditor/Public/EditMode/ControlRigEditMode.h
+++ b/Engine/Plugins/Animation/ControlRig/Source/ControlRigEditor/Public/EditMode/ControlRigEditMode.h
@@ -248,6 +248,9 @@ private:
 	/** Reset Transforms */
 	void ResetTransforms(bool bSelectionOnly);
 
+	/** Invert Input Pose */
+	void InvertInputPose();
+
 	/** Increase Shape Size */
 	void IncreaseShapeSize();
 
diff --git a/Engine/Source/Editor/MovieSceneTools/Private/MovieSceneToolHelpers.cpp b/Engine/Source/Editor/MovieSceneTools/Private/MovieSceneToolHelpers.cpp
--- a/Engine/Source/Editor/MovieSceneTools/Private/MovieSceneToolHelpers.cpp
+++ b/Engine/Source/Editor/MovieSceneTools/Private/MovieSceneToolHelpers.cpp
@@ -3876,7 +3876,7 @@ bool MovieSceneToolHelpers::BakeToSkelMeshToCallbacks(UMovieScene* MovieScene, I
 		int32 LocalIndex = LocalStartFrame + FrameCount;
 		FFrameNumber LocalFrame(LocalIndex);
 		TickFrame(LocalFrame, DeltaTime, MovieScene, AnimTrackAdapter, BakeHelpers, SkelMeshComps, LiveLinkClient, SourceAndMode);
-		TickCallback.ExecuteIfBound(DeltaTime);
+		TickCallback.ExecuteIfBound(DeltaTime, LocalFrame);
 	}
 
 	for (IMovieSceneToolsAnimationBakeHelper* BakeHelper : BakeHelpers)
@@ -3954,7 +3954,7 @@ bool MovieSceneToolHelpers::ExportToAnimSequence(UAnimSequence* AnimSequence, UA
 		SkelMeshComp->UpdateLODStatus();
 	});
 
-	FTickAnimationCB TickCallback = FTickAnimationCB::CreateLambda([&AnimationRecorder](float DeltaTime)
+	FTickAnimationCB TickCallback = FTickAnimationCB::CreateLambda([&AnimationRecorder](float DeltaTime, FFrameNumber FrameNumber)
 	{
 		AnimationRecorder.Update(DeltaTime);
 
diff --git a/Engine/Source/Editor/MovieSceneTools/Public/MovieSceneToolHelpers.h b/Engine/Source/Editor/MovieSceneTools/Public/MovieSceneToolHelpers.h
--- a/Engine/Source/Editor/MovieSceneTools/Public/MovieSceneToolHelpers.h
+++ b/Engine/Source/Editor/MovieSceneTools/Public/MovieSceneToolHelpers.h
@@ -69,7 +69,7 @@ struct FFBXInOutParameters
 
 DECLARE_DELEGATE(FInitAnimationCB);
 DECLARE_DELEGATE(FStartAnimationCB);
-DECLARE_DELEGATE_OneParam(FTickAnimationCB, float);
+DECLARE_DELEGATE_TwoParams(FTickAnimationCB, float, FFrameNumber);
 DECLARE_DELEGATE(FEndAnimationCB);
 
 //Skel Mesh Recorder to set up and restore various parameters on the skelmesh
