Expose FControlRigClassFilter for use outside of ControlRigEditor module/plugin

diff --git a/Engine/Plugins/Animation/ControlRig/Source/ControlRigEditor/Private/ControlRigEditorModule.cpp b/Engine/Plugins/Animation/ControlRig/Source/ControlRigEditor/Private/ControlRigEditorModule.cpp
--- a/Engine/Plugins/Animation/ControlRig/Source/ControlRigEditor/Private/ControlRigEditorModule.cpp
+++ b/Engine/Plugins/Animation/ControlRig/Source/ControlRigEditor/Private/ControlRigEditorModule.cpp
@@ -62,7 +62,6 @@
 #include "ISequencer.h"
 #include "ILevelSequenceEditorToolkit.h"
 #include "ClassViewerModule.h"
-#include "ClassViewerFilter.h"
 #include "Editor.h"
 #include "LevelEditorViewport.h"
 #include "Animation/SkeletalMeshActor.h"
@@ -94,6 +93,7 @@
 #include "Constraints/ControlRigTransformableHandle.h"
 #include "Constraints/TransformConstraintChannelInterface.h"
 #include "Async/TaskGraphInterfaces.h"
+#include "AssetRegistry/IAssetRegistry.h"
 
 #define LOCTEXT_NAMESPACE "ControlRigEditorModule"
 
@@ -888,126 +888,6 @@ TSharedRef< SWidget > FControlRigEditorModule::GenerateAnimationMenu(TWeakPtr<IA
 						LOCTEXT("BakeToControlRig", "Bake To Control Rig"), NSLOCTEXT("AnimationModeToolkit", "BakeToControlRigTooltip", "This Control Rig will Drive This Animation."),
 						FNewMenuDelegate::CreateLambda([this, AnimSequence, SkeletalMesh, Skeleton](FMenuBuilder& InSubMenuBuilder)
 							{
-								//todo move to .h for ue5
-								class FControlRigClassFilter : public IClassViewerFilter
-								{
-								public:
-									FControlRigClassFilter(bool bInCheckSkeleton, bool bInCheckAnimatable, bool bInCheckInversion, USkeleton* InSkeleton) :
-										bFilterAssetBySkeleton(bInCheckSkeleton),
-										bFilterExposesAnimatableControls(bInCheckAnimatable),
-										bFilterInversion(bInCheckInversion),
-										AssetRegistry(FModuleManager::GetModuleChecked<FAssetRegistryModule>(TEXT("AssetRegistry")).Get())
-									{
-										if (InSkeleton)
-										{
-											SkeletonName = FAssetData(InSkeleton).GetExportTextName();
-										}
-									}
-									bool bFilterAssetBySkeleton;
-									bool bFilterExposesAnimatableControls;
-									bool bFilterInversion;
-
-									FString SkeletonName;
-									const IAssetRegistry& AssetRegistry;
-
-									bool MatchesFilter(const FAssetData& AssetData)
-									{
-										bool bExposesAnimatableControls = AssetData.GetTagValueRef<bool>(TEXT("bExposesAnimatableControls"));
-										if (bFilterExposesAnimatableControls == true && bExposesAnimatableControls == false)
-										{
-											return false;
-										}
-										if (bFilterInversion)
-										{
-											bool bHasInversion = false;
-											FAssetDataTagMapSharedView::FFindTagResult Tag = AssetData.TagsAndValues.FindTag(TEXT("SupportedEventNames"));
-											if (Tag.IsSet())
-											{
-												FString EventString = FRigUnit_InverseExecution::EventName.ToString();
-												FString OldEventString = FString(TEXT("Inverse"));
-												TArray<FString> SupportedEventNames;
-												Tag.GetValue().ParseIntoArray(SupportedEventNames, TEXT(","), true);
-
-												for (const FString& Name : SupportedEventNames)
-												{
-													if (Name.Contains(EventString) || Name.Contains(OldEventString))
-													{
-														bHasInversion = true;
-														break;
-													}
-												}
-												if (bHasInversion == false)
-												{
-													return false;
-												}
-											}
-										}
-										if (bFilterAssetBySkeleton)
-										{
-											FString PreviewSkeletalMesh = AssetData.GetTagValueRef<FString>(TEXT("PreviewSkeletalMesh"));
-											if (PreviewSkeletalMesh.Len() > 0)
-											{
-												FAssetData SkelMeshData = AssetRegistry.GetAssetByObjectPath(FSoftObjectPath(PreviewSkeletalMesh));
-												FString PreviewSkeleton = SkelMeshData.GetTagValueRef<FString>(TEXT("Skeleton"));
-												if (PreviewSkeleton == SkeletonName)
-												{
-													return true;
-												}
-											}
-											FString PreviewSkeleton = AssetData.GetTagValueRef<FString>(TEXT("PreviewSkeleton"));
-											if (PreviewSkeleton == SkeletonName)
-											{
-												return true;
-											}
-											FString SourceHierarchyImport = AssetData.GetTagValueRef<FString>(TEXT("SourceHierarchyImport"));
-											if (SourceHierarchyImport == SkeletonName)
-											{
-												return true;
-											}
-											FString SourceCurveImport = AssetData.GetTagValueRef<FString>(TEXT("SourceCurveImport"));
-											if (SourceCurveImport == SkeletonName)
-											{
-												return true;
-											}
-											return false;
-										}
-										return true;
-
-									}
-									bool IsClassAllowed(const FClassViewerInitializationOptions& InInitOptions, const UClass* InClass, TSharedRef< FClassViewerFilterFuncs > InFilterFuncs) override
-									{
-										if (InClass)
-										{
-											const bool bChildOfObjectClass = InClass->IsChildOf(UControlRig::StaticClass());
-											const bool bMatchesFlags = !InClass->HasAnyClassFlags(CLASS_Hidden | CLASS_HideDropDown | CLASS_Deprecated | CLASS_Abstract);
-											const bool bNotNative = !InClass->IsNative();
-
-											if (bChildOfObjectClass && bMatchesFlags && bNotNative)
-											{
-												FAssetData AssetData(InClass);
-												return MatchesFilter(AssetData);
-											}
-										}
-										return false;
-									}
-
-									virtual bool IsUnloadedClassAllowed(const FClassViewerInitializationOptions& InInitOptions, const TSharedRef< const IUnloadedBlueprintData > InUnloadedClassData, TSharedRef< FClassViewerFilterFuncs > InFilterFuncs) override
-									{
-										const bool bChildOfObjectClass = InUnloadedClassData->IsChildOf(UControlRig::StaticClass());
-										const bool bMatchesFlags = !InUnloadedClassData->HasAnyClassFlags(CLASS_Hidden | CLASS_HideDropDown | CLASS_Deprecated | CLASS_Abstract);
-										if (bChildOfObjectClass && bMatchesFlags)
-										{
-											FString GeneratedClassPathString = InUnloadedClassData->GetClassPathName().ToString();
-											FString BlueprintPath = GeneratedClassPathString.LeftChop(2); // Chop off _C
-											FAssetData AssetData = AssetRegistry.GetAssetByObjectPath(FSoftObjectPath(BlueprintPath));
-											return MatchesFilter(AssetData);
-
-										}
-										return false;
-									}
-
-								};
-
 								FClassViewerInitializationOptions Options;
 								Options.bShowUnloadedBlueprints = true;
 								Options.NameTypeToDisplay = EClassViewerNameTypeToDisplay::DisplayName;
@@ -1432,6 +1312,149 @@ TSharedRef<IControlRigEditor> FControlRigEditorModule::CreateControlRigEditor(co
 	return NewControlRigEditor;
 }
 
+FControlRigClassFilter::FControlRigClassFilter(bool bInCheckSkeleton, bool bInCheckAnimatable, bool bInCheckInversion, USkeleton* InSkeleton) :	bFilterAssetBySkeleton(bInCheckSkeleton),
+	bFilterExposesAnimatableControls(bInCheckAnimatable),
+	bFilterInversion(bInCheckInversion),
+	Skeleton(InSkeleton),
+	AssetRegistry(FModuleManager::GetModuleChecked<FAssetRegistryModule>(TEXT("AssetRegistry")).Get())
+{
+}
+
+bool FControlRigClassFilter::MatchesFilter(const FAssetData& AssetData)
+{
+	const bool bExposesAnimatableControls = AssetData.GetTagValueRef<bool>(TEXT("bExposesAnimatableControls"));
+	if (bFilterExposesAnimatableControls == true && bExposesAnimatableControls == false)
+	{
+		return false;
+	}
+	if (bFilterInversion)
+	{		
+		FAssetDataTagMapSharedView::FFindTagResult Tag = AssetData.TagsAndValues.FindTag(TEXT("SupportedEventNames"));
+		if (Tag.IsSet())
+		{
+			bool bHasInversion = false;
+			FString EventString = FRigUnit_InverseExecution::EventName.ToString();
+			FString OldEventString = FString(TEXT("Inverse"));
+			TArray<FString> SupportedEventNames;
+			Tag.GetValue().ParseIntoArray(SupportedEventNames, TEXT(","), true);
+
+			for (const FString& Name : SupportedEventNames)
+			{
+				if (Name.Contains(EventString) || Name.Contains(OldEventString))
+				{
+					bHasInversion = true;
+					break;
+				}
+			}
+			if (bHasInversion == false)
+			{
+				return false;
+			}
+		}
+	}
+	if (bFilterAssetBySkeleton)
+	{
+		FString SkeletonName;
+		if (Skeleton)
+		{
+			SkeletonName = FAssetData(Skeleton).GetExportTextName();
+		}
+		FString PreviewSkeletalMesh = AssetData.GetTagValueRef<FString>(TEXT("PreviewSkeletalMesh"));
+		if (PreviewSkeletalMesh.Len() > 0)
+		{
+			FAssetData SkelMeshData = AssetRegistry.GetAssetByObjectPath(FSoftObjectPath(PreviewSkeletalMesh));
+			FString PreviewSkeleton = SkelMeshData.GetTagValueRef<FString>(TEXT("Skeleton"));
+			if (PreviewSkeleton == SkeletonName)
+			{
+				return true;
+			}
+			else if(Skeleton)
+			{
+				if (Skeleton->IsCompatibleForEditor(PreviewSkeleton))
+				{
+					return true;
+				}
+			}
+		}
+		FString PreviewSkeleton = AssetData.GetTagValueRef<FString>(TEXT("PreviewSkeleton"));
+		if (PreviewSkeleton == SkeletonName)
+		{
+			return true;
+		}
+		else if (Skeleton)
+		{
+			if (Skeleton->IsCompatibleForEditor(PreviewSkeleton))
+			{
+				return true;
+			}
+		}
+		FString SourceHierarchyImport = AssetData.GetTagValueRef<FString>(TEXT("SourceHierarchyImport"));
+		if (SourceHierarchyImport == SkeletonName)
+		{
+			return true;
+		}
+		else if (Skeleton)
+		{
+			if (Skeleton->IsCompatibleForEditor(SourceHierarchyImport))
+			{
+				return true;
+			}
+		}
+		FString SourceCurveImport = AssetData.GetTagValueRef<FString>(TEXT("SourceCurveImport"));
+		if (SourceCurveImport == SkeletonName)
+		{
+			return true;
+		}
+		else if (Skeleton)
+		{
+			if (Skeleton->IsCompatibleForEditor(SourceCurveImport))
+			{
+				return true;
+			}
+		}
+		return false;
+	}
+	return true;
+}
+
+bool FControlRigClassFilter::IsClassAllowed(const FClassViewerInitializationOptions& InInitOptions, const UClass* InClass, TSharedRef< FClassViewerFilterFuncs > InFilterFuncs) 
+{
+	if(InClass)
+	{
+		const bool bChildOfObjectClass = InClass->IsChildOf(UControlRig::StaticClass());
+		const bool bMatchesFlags = !InClass->HasAnyClassFlags(CLASS_Hidden | CLASS_HideDropDown | CLASS_Deprecated | CLASS_Abstract);
+		const bool bNotNative = !InClass->IsNative();
+
+		// Allow any class contained in the extra picker common classes array
+		if (InInitOptions.ExtraPickerCommonClasses.Contains(InClass))
+		{
+			return true;
+		}
+			
+		if (bChildOfObjectClass && bMatchesFlags && bNotNative)
+		{
+			const FAssetData AssetData(InClass);
+			return MatchesFilter(AssetData);
+		}
+	}
+	return false;
+}
+
+bool FControlRigClassFilter::IsUnloadedClassAllowed(const FClassViewerInitializationOptions& InInitOptions, const TSharedRef< const IUnloadedBlueprintData > InUnloadedClassData, TSharedRef< FClassViewerFilterFuncs > InFilterFuncs)
+{
+	const bool bChildOfObjectClass = InUnloadedClassData->IsChildOf(UControlRig::StaticClass());
+	const bool bMatchesFlags = !InUnloadedClassData->HasAnyClassFlags(CLASS_Hidden | CLASS_HideDropDown | CLASS_Deprecated | CLASS_Abstract);
+	if (bChildOfObjectClass && bMatchesFlags)
+	{
+		const FString GeneratedClassPathString = InUnloadedClassData->GetClassPathName().ToString();
+		const FString BlueprintPath = GeneratedClassPathString.LeftChop(2); // Chop off _C
+		const FAssetData AssetData = AssetRegistry.GetAssetByObjectPath(FSoftObjectPath(BlueprintPath));
+		return MatchesFilter(AssetData);
+
+	}
+	return false;
+}
+
 IMPLEMENT_MODULE(FControlRigEditorModule, ControlRigEditor)
 
 #undef LOCTEXT_NAMESPACE
diff --git a/Engine/Plugins/Animation/ControlRig/Source/ControlRigEditor/Private/Sequencer/ControlRigParameterTrackEditor.cpp b/Engine/Plugins/Animation/ControlRig/Source/ControlRigEditor/Private/Sequencer/ControlRigParameterTrackEditor.cpp
--- a/Engine/Plugins/Animation/ControlRig/Source/ControlRigEditor/Private/Sequencer/ControlRigParameterTrackEditor.cpp
+++ b/Engine/Plugins/Animation/ControlRig/Source/ControlRigEditor/Private/Sequencer/ControlRigParameterTrackEditor.cpp
@@ -650,161 +650,6 @@ void FControlRigParameterTrackEditor::BuildObjectBindingContextMenu(FMenuBuilder
 	}
 }
 
-class FControlRigClassFilter : public IClassViewerFilter
-{
-public:
-	FControlRigClassFilter(bool bInCheckSkeleton, bool bInCheckAnimatable, bool bInCheckInversion, USkeleton* InSkeleton) :
-		bFilterAssetBySkeleton(bInCheckSkeleton),
-		bFilterExposesAnimatableControls(bInCheckAnimatable),
-		bFilterInversion(bInCheckInversion),
-		AssetRegistry(FModuleManager::GetModuleChecked<FAssetRegistryModule>(TEXT("AssetRegistry")).Get())
-	{
-		Skeleton = InSkeleton;
-	}
-	bool bFilterAssetBySkeleton;
-	bool bFilterExposesAnimatableControls;
-	bool bFilterInversion;
-
-	USkeleton* Skeleton;
-	const IAssetRegistry& AssetRegistry;
-
-	bool MatchesFilter(const FAssetData& AssetData)
-	{
-		bool bExposesAnimatableControls = AssetData.GetTagValueRef<bool>(TEXT("bExposesAnimatableControls"));
-		if (bFilterExposesAnimatableControls == true && bExposesAnimatableControls == false)
-		{
-			return false;
-		}
-		if (bFilterInversion)
-		{
-			bool bHasInversion = false;
-			FAssetDataTagMapSharedView::FFindTagResult Tag = AssetData.TagsAndValues.FindTag(TEXT("SupportedEventNames"));
-			if (Tag.IsSet())
-			{
-				FString EventString = FRigUnit_InverseExecution::EventName.ToString();
-				FString OldEventString = FString(TEXT("Inverse"));
-				TArray<FString> SupportedEventNames;
-				Tag.GetValue().ParseIntoArray(SupportedEventNames, TEXT(","), true);
-
-				for (const FString& Name : SupportedEventNames)
-				{
-					if (Name.Contains(EventString) || Name.Contains(OldEventString))
-					{
-						bHasInversion = true;
-						break;
-					}
-				}
-				if (bHasInversion == false)
-				{
-					return false;
-				}
-			}
-		}
-		if (bFilterAssetBySkeleton)
-		{
-			FString SkeletonName;
-			if (Skeleton)
-			{
-				SkeletonName = FAssetData(Skeleton).GetExportTextName();
-			}
-			FString PreviewSkeletalMesh = AssetData.GetTagValueRef<FString>(TEXT("PreviewSkeletalMesh"));
-			if (PreviewSkeletalMesh.Len() > 0)
-			{
-				FAssetData SkelMeshData = AssetRegistry.GetAssetByObjectPath(FSoftObjectPath(PreviewSkeletalMesh));
-				FString PreviewSkeleton = SkelMeshData.GetTagValueRef<FString>(TEXT("Skeleton"));
-				if (PreviewSkeleton == SkeletonName)
-				{
-					return true;
-				}
-				else if(Skeleton)
-				{
-					if (Skeleton->IsCompatibleForEditor(PreviewSkeleton))
-					{
-						return true;
-					}
-				}
-			}
-			FString PreviewSkeleton = AssetData.GetTagValueRef<FString>(TEXT("PreviewSkeleton"));
-			if (PreviewSkeleton == SkeletonName)
-			{
-				return true;
-			}
-			else if (Skeleton)
-			{
-				if (Skeleton->IsCompatibleForEditor(PreviewSkeleton))
-				{
-					return true;
-				}
-			}
-			FString SourceHierarchyImport = AssetData.GetTagValueRef<FString>(TEXT("SourceHierarchyImport"));
-			if (SourceHierarchyImport == SkeletonName)
-			{
-				return true;
-			}
-			else if (Skeleton)
-			{
-				if (Skeleton->IsCompatibleForEditor(SourceHierarchyImport))
-				{
-					return true;
-				}
-			}
-			FString SourceCurveImport = AssetData.GetTagValueRef<FString>(TEXT("SourceCurveImport"));
-			if (SourceCurveImport == SkeletonName)
-			{
-				return true;
-			}
-			else if (Skeleton)
-			{
-				if (Skeleton->IsCompatibleForEditor(SourceCurveImport))
-				{
-					return true;
-				}
-			}
-			return false;
-		}
-		return true;
-
-	}
-	bool IsClassAllowed(const FClassViewerInitializationOptions& InInitOptions, const UClass* InClass, TSharedRef< FClassViewerFilterFuncs > InFilterFuncs) override
-	{
-		if(InClass)
-		{
-			const bool bChildOfObjectClass = InClass->IsChildOf(UControlRig::StaticClass());
-			const bool bMatchesFlags = !InClass->HasAnyClassFlags(CLASS_Hidden | CLASS_HideDropDown | CLASS_Deprecated | CLASS_Abstract);
-			const bool bNotNative = !InClass->IsNative();
-
-			// ALlow any class contained in the extra picker common classes array
-			if (InInitOptions.ExtraPickerCommonClasses.Contains(InClass))
-			{
-				return true;
-			}
-			
-			if (bChildOfObjectClass && bMatchesFlags && bNotNative)
-			{
-				FAssetData AssetData(InClass);
-				return MatchesFilter(AssetData);
-			}
-		}
-		return false;
-	}
-
-	virtual bool IsUnloadedClassAllowed(const FClassViewerInitializationOptions& InInitOptions, const TSharedRef< const IUnloadedBlueprintData > InUnloadedClassData, TSharedRef< FClassViewerFilterFuncs > InFilterFuncs) override
-	{
-		const bool bChildOfObjectClass = InUnloadedClassData->IsChildOf(UControlRig::StaticClass());
-		const bool bMatchesFlags = !InUnloadedClassData->HasAnyClassFlags(CLASS_Hidden | CLASS_HideDropDown | CLASS_Deprecated | CLASS_Abstract);
-		if (bChildOfObjectClass && bMatchesFlags)
-		{
-			FString GeneratedClassPathString = InUnloadedClassData->GetClassPathName().ToString();
-			FString BlueprintPath = GeneratedClassPathString.LeftChop(2); // Chop off _C
-			FAssetData AssetData = AssetRegistry.GetAssetByObjectPath(FSoftObjectPath(BlueprintPath));
-			return MatchesFilter(AssetData);
-
-		}
-		return false;
-	}
-
-};
-
 void FControlRigParameterTrackEditor::BakeToControlRigSubMenu(FMenuBuilder& MenuBuilder, FGuid ObjectBinding, UObject* BoundObject, USkeletalMeshComponent* SkelMeshComp, USkeleton* Skeleton)
 {
 	const TSharedPtr<ISequencer> ParentSequencer = GetSequencer();
@@ -814,19 +659,18 @@ void FControlRigParameterTrackEditor::BakeToControlRigSubMenu(FMenuBuilder& Menu
 		FClassViewerInitializationOptions Options;
 		Options.bShowUnloadedBlueprints = true;
 		Options.NameTypeToDisplay = EClassViewerNameTypeToDisplay::DisplayName;
-		TSharedPtr<FControlRigClassFilter> ClassFilter = MakeShareable(new FControlRigClassFilter(bFilterAssetBySkeleton, false, true, Skeleton));
+		const TSharedPtr<FControlRigClassFilter> ClassFilter = MakeShareable(new FControlRigClassFilter(bFilterAssetBySkeleton, false, true, Skeleton));
 		Options.ClassFilters.Add(ClassFilter.ToSharedRef());
 		Options.bShowNoneOption = false;
 		Options.ExtraPickerCommonClasses.Add(UFKControlRig::StaticClass());
 
 		FClassViewerModule& ClassViewerModule = FModuleManager::LoadModuleChecked<FClassViewerModule>("ClassViewer");
 
-		TSharedRef<SWidget> ClassViewer = ClassViewerModule.CreateClassViewer(Options, FOnClassPicked::CreateRaw(this, &FControlRigParameterTrackEditor::BakeToControlRig, ObjectBinding, BoundObject, SkelMeshComp, Skeleton));
+		const TSharedRef<SWidget> ClassViewer = ClassViewerModule.CreateClassViewer(Options, FOnClassPicked::CreateRaw(this, &FControlRigParameterTrackEditor::BakeToControlRig, ObjectBinding, BoundObject, SkelMeshComp, Skeleton));
 		MenuBuilder.AddWidget(ClassViewer, FText::GetEmpty(), true);
 	}
 }
 
-
 class SBakeToAnimAndControlRigOptionsWindow : public SCompoundWidget
 {
 public:
diff --git a/Engine/Plugins/Animation/ControlRig/Source/ControlRigEditor/Public/IControlRigEditorModule.h b/Engine/Plugins/Animation/ControlRig/Source/ControlRigEditor/Public/IControlRigEditorModule.h
--- a/Engine/Plugins/Animation/ControlRig/Source/ControlRigEditor/Public/IControlRigEditorModule.h
+++ b/Engine/Plugins/Animation/ControlRig/Source/ControlRigEditor/Public/IControlRigEditorModule.h
@@ -5,6 +5,10 @@
 #include "CoreMinimal.h"
 #include "IControlRigEditor.h"
 #include "RigVMEditorModule.h"
+#include "Kismet2/StructureEditorUtils.h"
+#include "AssetRegistry/IAssetRegistry.h"
+#include "ClassViewerModule.h"
+#include "ClassViewerFilter.h"
 
 DECLARE_LOG_CATEGORY_EXTERN(LogControlRigEditor, Log, All);
 
@@ -14,6 +18,25 @@ class UControlRigGraphNode;
 class UControlRigGraphSchema;
 class FConnectionDrawingPolicy;
 
+class CONTROLRIGEDITOR_API FControlRigClassFilter : public IClassViewerFilter
+{
+public:
+	FControlRigClassFilter(bool bInCheckSkeleton, bool bInCheckAnimatable, bool bInCheckInversion, USkeleton* InSkeleton);
+	virtual bool IsClassAllowed(const FClassViewerInitializationOptions& InInitOptions, const UClass* InClass, TSharedRef< FClassViewerFilterFuncs > InFilterFuncs) override;
+	virtual bool IsUnloadedClassAllowed(const FClassViewerInitializationOptions& InInitOptions, const TSharedRef< const IUnloadedBlueprintData > InUnloadedClassData, TSharedRef< FClassViewerFilterFuncs > InFilterFuncs) override;	
+
+private:
+	bool MatchesFilter(const FAssetData& AssetData);
+
+private:
+	bool bFilterAssetBySkeleton;
+	bool bFilterExposesAnimatableControls;
+	bool bFilterInversion;
+
+	USkeleton* Skeleton;
+	const IAssetRegistry& AssetRegistry;
+};
+
 class IControlRigEditorModule : public FRigVMEditorModule
 {
 public:
