Control Rig: Fix crash when loading a Control Rig which had a saved RigHierarchy
#rnx

diff --git a/Engine/Plugins/Animation/ControlRig/Source/ControlRig/Private/ControlRig.cpp b/Engine/Plugins/Animation/ControlRig/Source/ControlRig/Private/ControlRig.cpp
--- a/Engine/Plugins/Animation/ControlRig/Source/ControlRig/Private/ControlRig.cpp
+++ b/Engine/Plugins/Animation/ControlRig/Source/ControlRig/Private/ControlRig.cpp
@@ -3100,7 +3100,22 @@ void UControlRig::PostInitInstance(URigVMHost* InCDO)
 	UpdateVMSettings();
 
 	// set up the hierarchy
-	DynamicHierarchy = NewObject<URigHierarchy>(this, TEXT("DynamicHierarchy"), SubObjectFlags);
+	{
+		// If this is not a CDO, it should have never saved its hieararchy. However, we have found that some rigs in the past
+		// did save their hiearchies. If that's the case, let's mark them as garbage and rename them before creating our own hierarchy.
+		if (!HasAnyFlags(RF_ClassDefaultObject))
+		{
+			UObject* ObjectFound = StaticFindObjectFast(URigHierarchy::StaticClass(), this, TEXT("DynamicHierarchy"), true, RF_DefaultSubObject);
+			if (ObjectFound)
+			{
+				FName NewName = MakeUniqueObjectName(GetTransientPackage(), URigHierarchy::StaticClass(), TEXT("DynamicHierarchy_Deleted"));
+				ObjectFound->Rename(*NewName.ToString(), GetTransientPackage(), REN_ForceNoResetLoaders | REN_DoNotDirty | REN_DontCreateRedirectors | REN_NonTransactional);
+				ObjectFound->MarkAsGarbage();
+			}
+		}
+
+		DynamicHierarchy = NewObject<URigHierarchy>(this, TEXT("DynamicHierarchy"), SubObjectFlags);
+	}
 
 #if WITH_EDITOR
 		const TWeakObjectPtr<UControlRig> WeakThis = this;
