- Generalize Animating-in-Engine tools for use outside of the level editor
     * EditorMode is now retrieved from widget/toolkit rather than forcefully retrieving it from the global LevelEditor
     * Tween tool can now also operate on  ControlRigParameterTrack(s) without a binding
- Prep ControlRigEditMode to be used outisde of ContorlRigEditor/LevelEditor contexts

diff --git a/Engine/Plugins/Animation/ControlRig/Source/ControlRigEditor/Private/ControlRigSequencerEditorLibrary.cpp b/Engine/Plugins/Animation/ControlRig/Source/ControlRigEditor/Private/ControlRigSequencerEditorLibrary.cpp
--- a/Engine/Plugins/Animation/ControlRig/Source/ControlRigEditor/Private/ControlRigSequencerEditorLibrary.cpp
+++ b/Engine/Plugins/Animation/ControlRig/Source/ControlRigEditor/Private/ControlRigSequencerEditorLibrary.cpp
@@ -48,6 +48,7 @@
 #include "Constraints/MovieSceneConstraintChannelHelper.h"
 #include "Sections/MovieSceneConstrainedSection.h"
 #include "BakingAnimationKeySettings.h"
+#include "EditMode/ControlRigEditMode.h"
 
 #include UE_INLINE_GENERATED_CPP_BY_NAME(ControlRigSequencerEditorLibrary)
 
@@ -413,30 +414,35 @@ bool UControlRigSequencerEditorLibrary::BlendValuesOnSelected(ULevelSequence* Le
 	if (WeakSequencer.IsValid() && WeakSequencer.Pin()->GetFocusedMovieSceneSequence() == LevelSequence
 		&& LevelSequence->GetMovieScene())
 	{
-		FControlsToTween ControlsToTween;
-		LevelSequence->GetMovieScene()->Modify();
-		switch(BlendOperation)
-		{ 
+		if (FControlRigEditMode* EditMode = static_cast<FControlRigEditMode*>(GLevelEditorModeTools().GetActiveMode(FControlRigEditMode::ModeName)))
+		{
+			TWeakPtr<FControlRigEditMode> WeakMode = StaticCastSharedRef<FControlRigEditMode, FEdMode>(EditMode->AsShared()).ToWeakPtr();
+			
+			FControlsToTween ControlsToTween;
+			LevelSequence->GetMovieScene()->Modify();
+			switch(BlendOperation)
+			{ 
 			case EAnimToolBlendOperation::Tween:
-			{
-				FControlsToTween BlendTool;
-				BlendTool.Setup(WeakSequencer);
-				BlendTool.Blend(WeakSequencer, BlendValue);
-				return true;
-			}
+				{
+					FControlsToTween BlendTool;
+					BlendTool.Setup(WeakSequencer, WeakMode);
+					BlendTool.Blend(WeakSequencer, BlendValue);
+					return true;
+				}
 			case EAnimToolBlendOperation::BlendToNeighbor:
-			{
-				FBlendNeighborSlider BlendTool;
-				BlendTool.Setup(WeakSequencer);
-				BlendTool.Blend(WeakSequencer, BlendValue);
-				return true;
-			}
+				{
+					FBlendNeighborSlider BlendTool;
+					BlendTool.Setup(WeakSequencer, WeakMode);
+					BlendTool.Blend(WeakSequencer, BlendValue);
+					return true;
+				}
 			case EAnimToolBlendOperation::PushPull:
-			{
-				FPushPullSlider BlendTool;
-				BlendTool.Setup(WeakSequencer);
-				BlendTool.Blend(WeakSequencer, BlendValue);
-				return true;
+				{
+					FPushPullSlider BlendTool;
+					BlendTool.Setup(WeakSequencer, WeakMode);
+					BlendTool.Blend(WeakSequencer, BlendValue);
+					return true;
+				}
 			}
 		}
 	}
diff --git a/Engine/Plugins/Animation/ControlRig/Source/ControlRigEditor/Private/EditMode/ControlRigEditMode.cpp b/Engine/Plugins/Animation/ControlRig/Source/ControlRigEditor/Private/EditMode/ControlRigEditMode.cpp
--- a/Engine/Plugins/Animation/ControlRig/Source/ControlRigEditor/Private/EditMode/ControlRigEditMode.cpp
+++ b/Engine/Plugins/Animation/ControlRig/Source/ControlRigEditor/Private/EditMode/ControlRigEditMode.cpp
@@ -1,6 +1,8 @@
 // Copyright Epic Games, Inc. All Rights Reserved.
 
 #include "EditMode/ControlRigEditMode.h"
+
+#include "AnimationEditorPreviewActor.h"
 #include "Components/StaticMeshComponent.h"
 #include "EditMode/ControlRigEditModeToolkit.h"
 #include "Toolkits/ToolkitManager.h"
@@ -63,6 +65,9 @@
 #include "Constraints/MovieSceneConstraintChannelHelper.h"
 #include "Editor/EditorPerProjectUserSettings.h"
 #include "TransformConstraint.h"
+#include "PersonaModule.h"
+#include "Animation/DebugSkelMeshComponent.h"
+#include "AnimationEditorPreviewActor.h"
 
 #include UE_INLINE_GENERATED_CPP_BY_NAME(ControlRigEditMode)
 
@@ -156,9 +161,17 @@ FControlRigEditMode::FControlRigEditMode()
 	ControlProxy = NewObject<UControlRigDetailPanelControlProxies>(GetTransientPackage(), NAME_None);
 	ControlProxy->SetFlags(RF_Transactional);
 
-	const UControlRigEditModeSettings* Settings = GetDefault<UControlRigEditModeSettings>();
+	UControlRigEditModeSettings* Settings = GetMutableDefault<UControlRigEditModeSettings>();
 	bShowControlsAsOverlay = Settings->bShowControlsAsOverlay;
 
+	Settings->GizmoScaleDelegate.AddLambda([this](float GizmoScale)
+	{
+		if (FEditorModeTools* ModeTools = GetModeManager())
+		{
+			ModeTools->SetWidgetScale(GizmoScale);
+		}
+	});
+
 	CommandBindings = MakeShareable(new FUICommandList);
 	BindCommands();
 
@@ -289,9 +302,10 @@ bool FControlRigEditMode::IsInLevelEditor() const
 {
 	return GetModeManager() == &GLevelEditorModeTools();
 }
+
 void FControlRigEditMode::SetUpDetailPanel()
 {
-	if (IsInLevelEditor() && Toolkit)
+	if (!AreEditingControlRigDirectly() && Toolkit)
 	{
 		StaticCastSharedPtr<SControlRigEditModeTools>(Toolkit->GetInlineContent())->SetSequencer(WeakSequencer.Pin());
 		StaticCastSharedPtr<SControlRigEditModeTools>(Toolkit->GetInlineContent())->SetSettingsDetailsObject(GetMutableDefault<UControlRigEditModeSettings>());	
@@ -512,7 +526,7 @@ void FControlRigEditMode::Tick(FEditorViewportClient* ViewportClient, float Delt
 		HandleSelectionChanged();
 		bSelectionChanged = false;
 	}
-	if (IsInLevelEditor() == false)
+	if (!AreEditingControlRigDirectly() == false)
 	{
 		ViewportClient->Invalidate();
 	}
@@ -552,7 +566,7 @@ void FControlRigEditMode::Tick(FEditorViewportClient* ViewportClient, float Delt
 							ShapeActor->SetSelected(true);
 						}
 
-						if (IsInLevelEditor())
+						if (!AreEditingControlRigDirectly())
 						{
 							FRigControlElement* ControlElement = ControlRig->FindControl(SelectedKey.Name);
 							if (ControlElement)
@@ -578,8 +592,8 @@ void FControlRigEditMode::Tick(FEditorViewportClient* ViewportClient, float Delt
 		// We need to tick here since changing a bone for example
 		// might have changed the transform of the Control
 		PostPoseUpdate();
-		
-		if (IsInLevelEditor() == false) //only do this check if not in level editor
+
+		if (!AreEditingControlRigDirectly() == false) //only do this check if not in level editor
 		{
 			for (TWeakObjectPtr<UControlRig>& RuntimeRigPtr : RuntimeControlRigs)
 			{
@@ -740,7 +754,7 @@ TSet<FName> FControlRigEditMode::GetActiveControlsFromSequencer(UControlRig* Con
 void FControlRigEditMode::Render(const FSceneView* View, FViewport* Viewport, FPrimitiveDrawInterface* PDI)
 {	
 	const UControlRigEditModeSettings* Settings = GetDefault<UControlRigEditModeSettings>();
-	const bool bIsInGameView = IsInLevelEditor() ? (ViewportToGameView.Find(Viewport) && ViewportToGameView[Viewport]) : false;
+	const bool bIsInGameView = !AreEditingControlRigDirectly() ? (ViewportToGameView.Find(Viewport) && ViewportToGameView[Viewport]) : false;
 	bool bRender = !Settings->bHideControlShapes;
 	for (TWeakObjectPtr<UControlRig>& ControlRigPtr : RuntimeControlRigs)
 	{
@@ -749,7 +763,7 @@ void FControlRigEditMode::Render(const FSceneView* View, FViewport* Viewport, FP
 		if (bRender && ControlRig && ControlRig->GetControlsVisible())
 		{
 			FTransform ComponentTransform = FTransform::Identity;
-			if (IsInLevelEditor())
+			if (!AreEditingControlRigDirectly())
 			{
 				ComponentTransform = GetHostingSceneComponentTransform(ControlRig);
 			}
@@ -1006,7 +1020,7 @@ bool FControlRigEditMode::EndTracking(FEditorViewportClient* InViewportClient, F
 		}
 		InteractionScopes.Reset();
 
-		if (bWasInteracting && IsInLevelEditor())
+		if (bWasInteracting && !AreEditingControlRigDirectly())
 		{
 			// We invalidate the hit proxies when in level editor to ensure that the gizmo's hit proxy is up to date.
 			// The invalidation is called here to avoid useless viewport update in the FControlRigEditMode::Tick
@@ -1045,33 +1059,35 @@ bool FControlRigEditMode::StartTracking(FEditorViewportClient* InViewportClient,
 	
 	if (InteractionScopes.Num() == 0)
 	{
-		bool bShouldModify = IsInLevelEditor();
-		if (!bShouldModify)
+		const bool bShouldModify = [this]() -> bool
 		{
-			for (TWeakObjectPtr<UControlRig>& RuntimeRigPtr : RuntimeControlRigs)
+			if (AreEditingControlRigDirectly())
 			{
-				if (UControlRig* ControlRig = RuntimeRigPtr.Get())
+				for (TWeakObjectPtr<UControlRig>& RuntimeRigPtr : RuntimeControlRigs)
 				{
-					TArray<FRigElementKey> SelectedRigElements = GetSelectedRigElements(ControlRig);
-					for (const FRigElementKey& Key : SelectedRigElements)
+					if (UControlRig* ControlRig = RuntimeRigPtr.Get())
 					{
-						if (Key.Type != ERigElementType::Control)
+						TArray<FRigElementKey> SelectedRigElements = GetSelectedRigElements(ControlRig);
+						for (const FRigElementKey& Key : SelectedRigElements)
 						{
-							bShouldModify = true;
-							break;
+							if (Key.Type != ERigElementType::Control)
+							{
+								return true;
+							}
 						}
 					}
 				}
 			}
-		}
+			
+			return !AreEditingControlRigDirectly();
+		}();
 
-		if (!IsInLevelEditor())
+		if (AreEditingControlRigDirectly())
 		{
 			for (TWeakObjectPtr<UControlRig>& RuntimeRigPtr : RuntimeControlRigs)
 			{
 				if (UControlRig* ControlRig = RuntimeRigPtr.Get())
 				{
-
 					UObject* Blueprint = ControlRig->GetClass()->ClassGeneratedBy;
 					if (Blueprint)
 					{
@@ -1094,7 +1110,7 @@ bool FControlRigEditMode::StartTracking(FEditorViewportClient* InViewportClient,
 
 	//in level editor only transact if we have at least one control selected, in editor we only select CR stuff so always transact
 
-	if (IsInLevelEditor())
+	if (!AreEditingControlRigDirectly())
 	{
 		for (TWeakObjectPtr<UControlRig>& RuntimeRigPtr : RuntimeControlRigs)
 		{
@@ -1252,7 +1268,7 @@ bool FControlRigEditMode::HandleClick(FEditorViewportClient* InViewportClient, H
 				AControlRigShapeActor* ShapeActor = CastChecked<AControlRigShapeActor>(ActorHitProxy->Actor);
 				if (ShapeActor->IsSelectable() && ShapeActor->ControlRig.IsValid())
 				{
-					FScopedTransaction ScopedTransaction(LOCTEXT("SelectControlTransaction", "Select Control"), IsInLevelEditor() && !GIsTransacting);
+					FScopedTransaction ScopedTransaction(LOCTEXT("SelectControlTransaction", "Select Control"), !AreEditingControlRigDirectly() && !GIsTransacting);
 
 					// temporarily disable the interaction scope
 					FControlRigInteractionScope** InteractionScope = InteractionScopes.Find(ShapeActor->ControlRig.Get());
@@ -1279,7 +1295,7 @@ bool FControlRigEditMode::HandleClick(FEditorViewportClient* InViewportClient, H
 					else
 					{
 						//also need to clear actor selection. Sequencer will handle this automatically if done in Sequencder UI but not if done by clicking
-						if (IsInLevelEditor())
+						if (!AreEditingControlRigDirectly())
 						{
 							if (GEditor && GEditor->GetSelectedActorCount())
 							{
@@ -1337,7 +1353,7 @@ bool FControlRigEditMode::HandleClick(FEditorViewportClient* InViewportClient, H
 										FName ControlName(*(Result.BoneName.ToString() + TEXT("_CONTROL")));
 										if (ControlRig->FindControl(ControlName))
 										{
-											FScopedTransaction ScopedTransaction(LOCTEXT("SelectControlTransaction", "Select Control"), IsInLevelEditor() && !GIsTransacting);
+											FScopedTransaction ScopedTransaction(LOCTEXT("SelectControlTransaction", "Select Control"), !AreEditingControlRigDirectly() && !GIsTransacting);
 
 											if (Click.IsShiftDown()) //guess we just select
 											{
@@ -1369,7 +1385,7 @@ bool FControlRigEditMode::HandleClick(FEditorViewportClient* InViewportClient, H
 		FName ControlName(*(FKBoneProxy->BoneName.ToString() + TEXT("_CONTROL")));
 		if (FKBoneProxy->ControlRig->FindControl(ControlName))
 		{
-			FScopedTransaction ScopedTransaction(LOCTEXT("SelectControlTransaction", "Select Control"), IsInLevelEditor() && !GIsTransacting);
+			FScopedTransaction ScopedTransaction(LOCTEXT("SelectControlTransaction", "Select Control"), !AreEditingControlRigDirectly() && !GIsTransacting);
 
 			if (Click.IsShiftDown()) //guess we just select
 			{
@@ -1446,7 +1462,7 @@ bool FControlRigEditMode::HandleClick(FEditorViewportClient* InViewportClient, H
 	// clear selected controls
 	if (Click.IsShiftDown() ==false && Click.IsControlDown() == false)
 	{
-		FScopedTransaction ScopedTransaction(LOCTEXT("SelectControlTransaction", "Select Control"), IsInLevelEditor() && !GIsTransacting);
+		FScopedTransaction ScopedTransaction(LOCTEXT("SelectControlTransaction", "Select Control"), !AreEditingControlRigDirectly() && !GIsTransacting);
 		ClearRigElementSelection(FRigElementTypeHelper::ToMask(ERigElementType::All));
 	}
 
@@ -1616,7 +1632,7 @@ bool FControlRigEditMode::BoxSelect(FBox& InBox, bool InSelect)
 	}
 	const bool bStrictDragSelection = GetDefault<ULevelEditorViewportSettings>()->bStrictBoxSelection;
 
-	FScopedTransaction ScopedTransaction(LOCTEXT("SelectControlTransaction", "Select Control"), IsInLevelEditor() && !GIsTransacting);
+	FScopedTransaction ScopedTransaction(LOCTEXT("SelectControlTransaction", "Select Control"), !AreEditingControlRigDirectly() && !GIsTransacting);
 	const bool bShiftDown = LevelViewportClient->Viewport->KeyState(EKeys::LeftShift) || LevelViewportClient->Viewport->KeyState(EKeys::RightShift);
 	if (!bShiftDown)
 	{
@@ -1688,7 +1704,7 @@ bool FControlRigEditMode::FrustumSelect(const FConvexVolume& InFrustum, FEditorV
 		return FEdMode::FrustumSelect(InFrustum, InViewportClient, InSelect);
 	}
 
-	FScopedTransaction ScopedTransaction(LOCTEXT("SelectControlTransaction", "Select Control"), IsInLevelEditor() && !GIsTransacting);
+	FScopedTransaction ScopedTransaction(LOCTEXT("SelectControlTransaction", "Select Control"), !AreEditingControlRigDirectly() && !GIsTransacting);
 	bool bSomethingSelected(false);
 	const bool bShiftDown = InViewportClient->Viewport->KeyState(EKeys::LeftShift) || InViewportClient->Viewport->KeyState(EKeys::RightShift);
 	if (!bShiftDown)
@@ -1973,7 +1989,7 @@ void FControlRigEditMode::ClearRigElementSelection(uint32 InTypes)
 	{
 		if (UControlRig* ControlRig = RuntimeRigPtr.Get())
 		{
-			if (IsInLevelEditor())
+			if (!AreEditingControlRigDirectly())
 			{
 				if (URigHierarchyController* Controller = ControlRig->GetHierarchy()->GetController())
 				{
@@ -2335,7 +2351,7 @@ void FControlRigEditMode::HandleSelectionChanged()
 	// update the pivot transform of our selected objects (they could be animating)
 	RecalcPivotTransform();
 	//need to force the redraw also
-	if (IsInLevelEditor())
+	if (!AreEditingControlRigDirectly())
 	{
 		GEditor->RedrawLevelEditingViewports(true);
 	}
@@ -2690,9 +2706,9 @@ void FControlRigEditMode::OpenSpacePickerWidget()
 	.Hierarchy(Hierarchy)
 	.Controls(SelectedControls)
 	.Title(LOCTEXT("PickSpace", "Pick Space"))
-	.AllowDelete(!IsInLevelEditor())
-	.AllowReorder(!IsInLevelEditor())
-	.AllowAdd(!IsInLevelEditor())
+	.AllowDelete(AreEditingControlRigDirectly())
+	.AllowReorder(AreEditingControlRigDirectly())
+	.AllowAdd(AreEditingControlRigDirectly())
 	.GetControlCustomization_Lambda([this, RuntimeRig](URigHierarchy*, const FRigElementKey& InControlKey)
 	{
 		return RuntimeRig->GetControlCustomization(InControlKey);
@@ -2700,7 +2716,7 @@ void FControlRigEditMode::OpenSpacePickerWidget()
 	.OnActiveSpaceChanged_Lambda([this, SelectedControls, RuntimeRig](URigHierarchy* InHierarchy, const FRigElementKey& InControlKey, const FRigElementKey& InSpaceKey)
 	{
 		check(SelectedControls.Contains(InControlKey));
-		if (IsInLevelEditor())
+		if (!AreEditingControlRigDirectly())
 		{
 			if (WeakSequencer.IsValid())
 			{
@@ -2736,7 +2752,7 @@ void FControlRigEditMode::OpenSpacePickerWidget()
 		check(SelectedControls.Contains(InControlKey));
 
 		// check if we are in the control rig editor or in the level
-		if(!IsInLevelEditor())
+		if(AreEditingControlRigDirectly())
 		{
 			if (UControlRigBlueprint* Blueprint = Cast<UControlRigBlueprint>(RuntimeRig->GetClass()->ClassGeneratedBy))
 			{
@@ -2805,7 +2821,7 @@ FText FControlRigEditMode::GetToggleControlShapeTransformEditHotKey() const
 
 void FControlRigEditMode::ToggleManipulators()
 {
-	if (IsInLevelEditor())
+	if (!AreEditingControlRigDirectly())
 	{
 		TMap<UControlRig*, TArray<FRigElementKey>> SelectedControls;
 		GetAllSelectedControls(SelectedControls);
@@ -2834,7 +2850,7 @@ void FControlRigEditMode::ToggleAllManipulators()
 	Settings->bHideControlShapes = !Settings->bHideControlShapes;
 
 	//turn on all if in level editor in case any where off
-	if (IsInLevelEditor() && Settings->bHideControlShapes)
+	if (!AreEditingControlRigDirectly() && Settings->bHideControlShapes)
 	{
 		for (TWeakObjectPtr<UControlRig>& RuntimeRigPtr : RuntimeControlRigs)
 		{
@@ -3215,7 +3231,7 @@ void FControlRigEditMode::PostUndo()
 	}
 
 	//normal actor undo will force the redraw, so we need to do the same for our transients/controls.
-	if (IsInLevelEditor() && (bInvalidateViewport || UsesTransformWidget()))
+	if (!AreEditingControlRigDirectly() && (bInvalidateViewport || UsesTransformWidget()))
 	{
 		GEditor->GetTimerManager()->SetTimerForNextTick([this]()
 		{
@@ -3439,7 +3455,7 @@ void FControlRigEditMode::CreateShapeActors(UControlRig* ControlRig)
 			}
 		}
 	}
-	if (IsInLevelEditor())
+	if (!AreEditingControlRigDirectly())
 	{
 
 		if (ControlProxy)
@@ -3467,14 +3483,14 @@ FControlRigEditMode* FControlRigEditMode::GetEditModeFromWorldContext(UWorld* In
 bool FControlRigEditMode::ShapeSelectionOverride(const UPrimitiveComponent* InComponent) const
 {
     //Think we only want to do this in regular editor, in the level editor we are driving selection
-	if (!IsInLevelEditor())
+	if (AreEditingControlRigDirectly())
 	{
-	AControlRigShapeActor* OwnerActor = Cast<AControlRigShapeActor>(InComponent->GetOwner());
-	if (OwnerActor)
-	{
-		// See if the actor is in a selected unit proxy
-		return OwnerActor->IsSelected();
-	}
+	    AControlRigShapeActor* OwnerActor = Cast<AControlRigShapeActor>(InComponent->GetOwner());
+	    if (OwnerActor)
+	    {
+		    // See if the actor is in a selected unit proxy
+		    return OwnerActor->IsSelected();
+	    }
 	}
 
 	return false;
@@ -3553,7 +3569,7 @@ bool FControlRigEditMode::AreRigElementSelectedAndMovable(UControlRig* ControlRi
 	}
 
 	//when in sequencer/level we don't have that delegate so don't check.
-	if (!IsInLevelEditor())
+	if (AreEditingControlRigDirectly())
 	{
 		if (!IsTransformDelegateAvailable())
 		{
@@ -3562,10 +3578,10 @@ bool FControlRigEditMode::AreRigElementSelectedAndMovable(UControlRig* ControlRi
 	}
 	else //do check for the binding though
 	{
-		if (GetHostingSceneComponent(ControlRig) == nullptr)
-		{
-			return false;
-		}
+		// if (GetHostingSceneComponent(ControlRig) == nullptr)
+		// {
+		// 	return false;
+		// }
 	}
 
 	return true;
@@ -3584,11 +3600,12 @@ void FControlRigEditMode::ReplaceControlRig(UControlRig* OldControlRig, UControl
 }
 void FControlRigEditMode::OnHierarchyModified(ERigHierarchyNotification InNotif, URigHierarchy* InHierarchy, const FRigBaseElement* InElement)
 {
-	if(bSuspendHierarchyNotifs)
+	if(bSuspendHierarchyNotifs || InElement == nullptr)
 	{
 		return;
 	}
-	
+
+	check(InElement);
 	switch(InNotif)
 	{
 		case ERigHierarchyNotification::ElementAdded:
@@ -3615,7 +3632,7 @@ void FControlRigEditMode::OnHierarchyModified(ERigHierarchyNotification InNotif,
 					{
 						// try to lazily apply the changes to the actor
 						const UControlRigEditModeSettings* Settings = GetDefault<UControlRigEditModeSettings>();
-						if (ShapeActor->UpdateControlSettings(InNotif, ControlRig, ControlElement, Settings->bHideControlShapes, IsInLevelEditor()))
+						if (ShapeActor->UpdateControlSettings(InNotif, ControlRig, ControlElement, Settings->bHideControlShapes, !AreEditingControlRigDirectly()))
 						{
 							break;
 						}
@@ -3632,7 +3649,7 @@ void FControlRigEditMode::OnHierarchyModified(ERigHierarchyNotification InNotif,
 		}
 		case ERigHierarchyNotification::ControlDrivenListChanged:
 		{
-			if (IsInLevelEditor())
+			if (!AreEditingControlRigDirectly())
 			{
 				// to synchronize the selection between the viewport / editmode and the details panel / sequencer
 				// we re-select the control. during deselection we recover the previously set driven list
@@ -3698,7 +3715,8 @@ void FControlRigEditMode::OnHierarchyModified(ERigHierarchyNotification InNotif,
 					// if it's control
 					if (Key.Type == ERigElementType::Control)
 					{
-						if (IsInLevelEditor())
+						FScopedTransaction ScopedTransaction(LOCTEXT("SelectControlTransaction", "Select Control"), !AreEditingControlRigDirectly() && !GIsTransacting);
+						if (!AreEditingControlRigDirectly())
 						{
 							ControlProxy->Modify();
 						}
@@ -3709,7 +3727,7 @@ void FControlRigEditMode::OnHierarchyModified(ERigHierarchyNotification InNotif,
 							ShapeActor->SetSelected(bSelected);
 
 						}
-						if (IsInLevelEditor())
+						if (!AreEditingControlRigDirectly())
 						{
 							if (const FRigControlElement* ControlElement = ControlRig->GetHierarchy()->Find<FRigControlElement>(Key))
 							{
@@ -3932,7 +3950,7 @@ void FControlRigEditMode::OnCoordSystemChanged(ECoordSystem InCoordSystem)
 
 bool FControlRigEditMode::CanChangeControlShapeTransform()
 {
-	if (!IsInLevelEditor())
+	if (AreEditingControlRigDirectly())
 	{
 		for (TWeakObjectPtr<UControlRig> RuntimeRigPtr : RuntimeControlRigs)
 		{
@@ -3982,7 +4000,7 @@ void FControlRigEditMode::SetControlShapeTransform(
 	}
 
 	static constexpr bool bNotify = true, bFixEuler = true, bUndo = true;
-	if (!IsInLevelEditor())
+	if (!AreEditingControlRigDirectly())
 	{
 		// assumes it's attached to actor
 		ControlRig->SetControlGlobalTransform(
@@ -4636,7 +4654,6 @@ void FControlRigEditMode::DestroyShapesActors(UControlRig* ControlRig)
 
 USceneComponent* FControlRigEditMode::GetHostingSceneComponent(const UControlRig* ControlRig) const
 {
-	
 	if (ControlRig == nullptr && GetControlRigs().Num() > 0)
 	{
 		ControlRig = GetControlRigs()[0].Get();
@@ -4646,9 +4663,32 @@ USceneComponent* FControlRigEditMode::GetHostingSceneComponent(const UControlRig
 		TSharedPtr<IControlRigObjectBinding> ObjectBinding = ControlRig->GetObjectBinding();
 		if (ObjectBinding.IsValid())
 		{
-			return Cast<USceneComponent>(ObjectBinding->GetBoundObject());
+			if (USceneComponent* BoundSceneComponent = Cast<USceneComponent>(ObjectBinding->GetBoundObject()))
+			{
+				return BoundSceneComponent;
+			}
+			else if (USkeleton* BoundSkeleton = Cast<USkeleton>(ObjectBinding->GetBoundObject()))
+			{
+				// Bound to a Skeleton means we are previewing an Animation Sequence
+				if (WorldPtr)
+				{
+					TObjectPtr<AActor>* PreviewActor = WorldPtr->PersistentLevel->Actors.FindByPredicate([](TObjectPtr<AActor> Actor)
+					{
+						return Actor && Actor->GetClass() == AAnimationEditorPreviewActor::StaticClass();
+					});
+
+					if(PreviewActor)
+					{
+						if (UDebugSkelMeshComponent* DebugComponent = (*PreviewActor)->FindComponentByClass<UDebugSkelMeshComponent>())
+						{
+							return DebugComponent;
+						}
+					}
+				}
+			}			
 		}
-	}
+		
+	}	
 
 	return nullptr;
 }
@@ -4673,7 +4713,7 @@ void FControlRigEditMode::PostPoseUpdate()
 	for (auto& ShapeActors : ControlRigShapeActors)
 	{
 		FTransform ComponentTransform = FTransform::Identity;
-		if (IsInLevelEditor())
+		if (!AreEditingControlRigDirectly())
 		{
 			ComponentTransform = GetHostingSceneComponentTransform(ShapeActors.Key);
 		}
diff --git a/Engine/Plugins/Animation/ControlRig/Source/ControlRigEditor/Private/EditMode/ControlRigEditModeSettings.cpp b/Engine/Plugins/Animation/ControlRig/Source/ControlRigEditor/Private/EditMode/ControlRigEditModeSettings.cpp
--- a/Engine/Plugins/Animation/ControlRig/Source/ControlRigEditor/Private/EditMode/ControlRigEditModeSettings.cpp
+++ b/Engine/Plugins/Animation/ControlRig/Source/ControlRigEditor/Private/EditMode/ControlRigEditModeSettings.cpp
@@ -17,8 +17,7 @@ void UControlRigEditModeSettings::PostEditChangeProperty(struct FPropertyChanged
 #if WITH_EDITOR
 	if (PropertyChangedEvent.Property && PropertyChangedEvent.Property->GetFName() == GET_MEMBER_NAME_CHECKED(UControlRigEditModeSettings, GizmoScale))
 	{
-		FEditorModeTools& Tools = GLevelEditorModeTools();
-		Tools.SetWidgetScale(GizmoScale);
+		GizmoScaleDelegate.Broadcast(GizmoScale);
 	}
 #endif
 	if (PropertyChangedEvent.ChangeType != EPropertyChangeType::Interactive)
@@ -31,8 +30,7 @@ void UControlRigEditModeSettings::PostEditChangeProperty(struct FPropertyChanged
 #if WITH_EDITOR
 void UControlRigEditModeSettings::PostEditUndo()
 {
-	FEditorModeTools& Tools = GLevelEditorModeTools();
-	Tools.SetWidgetScale(GizmoScale);
+	GizmoScaleDelegate.Broadcast(GizmoScale);
 }
 #endif
 
diff --git a/Engine/Plugins/Animation/ControlRig/Source/ControlRigEditor/Private/EditMode/ControlRigEditModeSettings.h b/Engine/Plugins/Animation/ControlRig/Source/ControlRigEditor/Private/EditMode/ControlRigEditModeSettings.h
--- a/Engine/Plugins/Animation/ControlRig/Source/ControlRigEditor/Private/EditMode/ControlRigEditModeSettings.h
+++ b/Engine/Plugins/Animation/ControlRig/Source/ControlRigEditor/Private/EditMode/ControlRigEditModeSettings.h
@@ -7,6 +7,8 @@
 #include "UObject/LazyObjectPtr.h"
 #include "ControlRigEditModeSettings.generated.h"
 
+DECLARE_MULTICAST_DELEGATE_OneParam(FGizmoScaleSet, float /* GizmoScale */);
+
 /** Settings object used to show useful information in the details panel */
 UCLASS(config=EditorPerProjectUserSettings, MinimalAPI)
 class UControlRigEditModeSettings : public UObject
@@ -88,4 +90,7 @@ public:
 
 	UPROPERTY(config)
 	FVector2D LastInViewportTweenWidgetLocation;
+
+	/** Delegate broadcasted whenever GizmoScale is modified */
+	FGizmoScaleSet GizmoScaleDelegate;
 };
\ No newline at end of file
diff --git a/Engine/Plugins/Animation/ControlRig/Source/ControlRigEditor/Private/EditMode/ControlRigEditModeToolkit.cpp b/Engine/Plugins/Animation/ControlRig/Source/ControlRigEditor/Private/EditMode/ControlRigEditModeToolkit.cpp
--- a/Engine/Plugins/Animation/ControlRig/Source/ControlRigEditor/Private/EditMode/ControlRigEditModeToolkit.cpp
+++ b/Engine/Plugins/Animation/ControlRig/Source/ControlRigEditor/Private/EditMode/ControlRigEditModeToolkit.cpp
@@ -156,11 +156,12 @@ FText FControlRigEditModeToolkit::GetActiveToolMessage() const
 	return ModeTools->GetActiveToolMessage();
 }
 
-TSharedRef<SDockTab> SpawnPoseTab(const FSpawnTabArgs& Args)
+TSharedRef<SDockTab> SpawnPoseTab(const FSpawnTabArgs& Args, TWeakPtr<FControlRigEditModeToolkit> SharedToolkit)
 {
 	return SNew(SDockTab)
 		[
 			SNew(SControlRigBaseListWidget)
+			.InSharedToolkit(SharedToolkit)
 		];
 }
 
@@ -213,9 +214,9 @@ void FControlRigEditModeToolkit::CreateAndShowTweenOverlay()
 	{
 		const FVector2D ActiveViewportSize = GetToolkitHost()->GetActiveViewportSize();
 		NewTweenWidgetLocation.X = ActiveViewportSize.X / 2.0f;
-		NewTweenWidgetLocation.Y = ActiveViewportSize.Y - 100.0f;
-		
+		NewTweenWidgetLocation.Y = FMath::Max(ActiveViewportSize.Y - 100.0f, 0);
 	}
+	
 	UpdateTweenWidgetLocation(NewTweenWidgetLocation);
 
 	SAssignNew(TweenWidgetParent, SHorizontalBox)
@@ -228,6 +229,7 @@ void FControlRigEditModeToolkit::CreateAndShowTweenOverlay()
 		[
 			SAssignNew(TweenWidget, SControlRigTweenWidget)
 			.InOwningToolkit(SharedThis(this))
+			.InOwningEditMode(SharedThis(&EditMode))
 		];
 
 	TryShowTweenOverlay();
@@ -323,7 +325,7 @@ void FControlRigEditModeToolkit::UpdateTweenWidgetLocation(const FVector2D InLoc
 	{
 		// reset the location if it was placed out of bounds
 		ScreenPos.X = ActiveViewportSize.X / 2.0f;
-		ScreenPos.Y = ActiveViewportSize.Y - 100.0f;
+		ScreenPos.Y = FMath::Max(ActiveViewportSize.Y - 100.0f, 0);
 	}
 	InViewportTweenWidgetLocation = ScreenPos;
 	UControlRigEditModeSettings* ControlRigEditModeSettings = GetMutableDefault<UControlRigEditModeSettings>();
@@ -372,7 +374,12 @@ void FControlRigEditModeToolkit::RequestModeUITabs()
 		ModeUILayerPtr->GetTabManager()->RegisterDefaultTabWindowSize(SnapperTabName, FVector2D(300, 325));
 
 		ModeUILayerPtr->GetTabManager()->UnregisterTabSpawner(PoseTabName);
-		ModeUILayerPtr->GetTabManager()->RegisterTabSpawner(PoseTabName, FOnSpawnTab::CreateStatic(&SpawnPoseTab))
+
+		TWeakPtr<FControlRigEditModeToolkit> WeakToolkit = SharedThis(this);
+		ModeUILayerPtr->GetTabManager()->RegisterTabSpawner(PoseTabName, FOnSpawnTab::CreateLambda([WeakToolkit](const FSpawnTabArgs& Args)
+		{
+			return SpawnPoseTab(Args, WeakToolkit);
+		}))
 			.SetDisplayName(LOCTEXT("ControlRigPoseTab", "Control Rig Pose"))
 			.SetTooltipText(LOCTEXT("ControlRigPoseTabTooltip", "Show Poses."))
 			.SetGroup(MenuGroup)
diff --git a/Engine/Plugins/Animation/ControlRig/Source/ControlRigEditor/Private/EditMode/SControlRigBaseListWidget.cpp b/Engine/Plugins/Animation/ControlRig/Source/ControlRigEditor/Private/EditMode/SControlRigBaseListWidget.cpp
--- a/Engine/Plugins/Animation/ControlRig/Source/ControlRigEditor/Private/EditMode/SControlRigBaseListWidget.cpp
+++ b/Engine/Plugins/Animation/ControlRig/Source/ControlRigEditor/Private/EditMode/SControlRigBaseListWidget.cpp
@@ -45,6 +45,7 @@
 #include "LevelSequenceEditorBlueprintLibrary.h"
 #include "Misc/PackageName.h"
 #include "AssetToolsModule.h"
+#include "ControlRigEditModeToolkit.h"
 #include "PropertyEditorModule.h"
 
 
@@ -416,7 +417,7 @@ void SControlRigPoseAnimSelectionToolbar::Construct(const FArguments& InArgs)
 
 void SControlRigPoseAnimSelectionToolbar::MakeControlRigAssetDialog(FControlRigAssetType Type, bool bSelectAll)
 {
-	FControlRigEditMode* ControlRigEditMode = static_cast<FControlRigEditMode*>(GLevelEditorModeTools().GetActiveMode(FControlRigEditMode::ModeName));
+	FControlRigEditMode* ControlRigEditMode = OwningControlRigWidget->GetEditMode();
 	if (!ControlRigEditMode)
 	{
 		return;
@@ -451,7 +452,7 @@ void SControlRigPoseAnimSelectionToolbar::MakeControlRigAssetDialog(FControlRigA
 		{
 			FString Path = OwningControlRigWidget->GetCurrentlySelectedPath();
 
-			FControlRigEditMode* ControlRigEditMode = static_cast<FControlRigEditMode*>(GLevelEditorModeTools().GetActiveMode(FControlRigEditMode::ModeName));
+			FControlRigEditMode* ControlRigEditMode =  OwningControlRigWidget->GetEditMode();
 			if (ControlRigEditMode )
 			{
 				TMap<UControlRig*, TArray<FRigElementKey>> AllSelectedControls;
@@ -490,7 +491,7 @@ void SControlRigPoseAnimSelectionToolbar::MakeControlRigAssetDialog(FControlRigA
 
 bool SControlRigPoseAnimSelectionToolbar::CanExecuteMakeControlRigAsset()
 {
-	FControlRigEditMode* ControlRigEditMode = static_cast<FControlRigEditMode*>(GLevelEditorModeTools().GetActiveMode(FControlRigEditMode::ModeName));
+	FControlRigEditMode* ControlRigEditMode =  OwningControlRigWidget->GetEditMode();
 	if (!ControlRigEditMode)
 	{
 		return false;
@@ -697,6 +698,7 @@ FString SPathDialogWithAllowList::GetAssetPath()
 
 void SControlRigBaseListWidget::Construct(const FArguments& InArgs)
 {
+	WeakToolkit = InArgs._InSharedToolkit;
 	FControlRigEditMode* EditMode = GetEditMode();
 	BindCommands();
 
@@ -865,7 +867,7 @@ void SControlRigBaseListWidget::NotifyUser(FNotificationInfo& NotificationInfo)
 
 FControlRigEditMode* SControlRigBaseListWidget::GetEditMode()
 {
-	FControlRigEditMode* ControlRigEditMode = static_cast<FControlRigEditMode*>(GLevelEditorModeTools().GetActiveMode(FControlRigEditMode::ModeName));
+	FControlRigEditMode* ControlRigEditMode = static_cast<FControlRigEditMode*>(WeakToolkit.Pin()->GetEditorMode());
 	return ControlRigEditMode;
 }
 
@@ -985,7 +987,7 @@ void SControlRigBaseListWidget::OnAssetsActivated(const TArray<FAssetData>& Sele
 				// If alt is down, select controls
 				if (FSlateApplication::Get().GetModifierKeys().IsAltDown())
 				{
-					FControlRigEditMode* ControlRigEditMode = static_cast<FControlRigEditMode*>(GLevelEditorModeTools().GetActiveMode(FControlRigEditMode::ModeName));
+					FControlRigEditMode* ControlRigEditMode = GetEditMode();
 					if (ControlRigEditMode)
 					{
 						TArray<UControlRig*> ControlRigs = ControlRigEditMode->GetControlRigsArray(true /*bIsVisible*/);
@@ -1363,7 +1365,7 @@ void SControlRigBaseListWidget::ExecutePastePose(UControlRigPoseAsset* PoseAsset
 {
 	if (PoseAsset)
 	{
-		FControlRigEditMode* ControlRigEditMode = static_cast<FControlRigEditMode*>(GLevelEditorModeTools().GetActiveMode(FControlRigEditMode::ModeName));
+		FControlRigEditMode* ControlRigEditMode = GetEditMode();
 		if (ControlRigEditMode)
 		{
 			TArray<UControlRig*> ControlRigs = ControlRigEditMode->GetControlRigsArray(true /*bIsVisible*/);
@@ -1392,7 +1394,7 @@ void SControlRigBaseListWidget::ExecuteUpdatePose(UControlRigPoseAsset* PoseAsse
 
 void SControlRigBaseListWidget::ExecuteSelectControls(UControlRigPoseAsset* PoseAsset)
 {
-	FControlRigEditMode* ControlRigEditMode = static_cast<FControlRigEditMode*>(GLevelEditorModeTools().GetActiveMode(FControlRigEditMode::ModeName));
+	FControlRigEditMode* ControlRigEditMode = GetEditMode();
 	if (ControlRigEditMode)
 	{
 		TArray<UControlRig*> ControlRigs = ControlRigEditMode->GetControlRigsArray(true /*bIsVisible*/);
@@ -1429,7 +1431,7 @@ void SControlRigBaseListWidget::ExecutePasteMirrorPose(UControlRigPoseAsset* Pos
 {
 	if (PoseAsset)
 	{
-		FControlRigEditMode* ControlRigEditMode = static_cast<FControlRigEditMode*>(GLevelEditorModeTools().GetActiveMode(FControlRigEditMode::ModeName));
+		FControlRigEditMode* ControlRigEditMode = GetEditMode();
 		if (ControlRigEditMode)
 		{
 			TArray<UControlRig*> ControlRigs = ControlRigEditMode->GetControlRigsArray(true /*bIsVisible*/);
@@ -1499,8 +1501,9 @@ void SControlRigBaseListWidget::CreateCurrentView(UObject* Asset)
 TSharedRef<class SControlRigPoseView> SControlRigBaseListWidget::CreatePoseView(UObject* InObject)
 {
 	UControlRigPoseAsset* PoseAsset = Cast<UControlRigPoseAsset>(InObject);
-	return SNew(SControlRigPoseView).
-		PoseAsset(PoseAsset);
+	return SNew(SControlRigPoseView)
+		.PoseAsset(PoseAsset)
+		.OwningWidget(SharedThis(this));
 }
 
 TSharedRef<class SControlRigPoseView> SControlRigBaseListWidget::CreateAnimationView(UObject* InObject)
diff --git a/Engine/Plugins/Animation/ControlRig/Source/ControlRigEditor/Private/EditMode/SControlRigBaseListWidget.h b/Engine/Plugins/Animation/ControlRig/Source/ControlRigEditor/Private/EditMode/SControlRigBaseListWidget.h
--- a/Engine/Plugins/Animation/ControlRig/Source/ControlRigEditor/Private/EditMode/SControlRigBaseListWidget.h
+++ b/Engine/Plugins/Animation/ControlRig/Source/ControlRigEditor/Private/EditMode/SControlRigBaseListWidget.h
@@ -23,6 +23,7 @@
 class UControlRig;
 class FControlRigEditMode;
 class UControlRigPoseAsset;
+class FControlRigEditModeToolkit;
 
 namespace ESelectedControlAsset
 {
@@ -43,7 +44,7 @@ public:
 	SLATE_BEGIN_ARGS(SControlRigBaseListWidget)
 
 	{}
-
+	SLATE_ARGUMENT(TWeakPtr<FControlRigEditModeToolkit>, InSharedToolkit)
 	SLATE_END_ARGS()
 public:
 	void Construct(const FArguments& InArgs);
@@ -133,4 +134,5 @@ private:
 	/** Utility function to display notifications to the user */
 	void NotifyUser(FNotificationInfo& NotificationInfo);
 
+	TWeakPtr<FControlRigEditModeToolkit> WeakToolkit;
 };
diff --git a/Engine/Plugins/Animation/ControlRig/Source/ControlRigEditor/Private/EditMode/SControlRigControlViews.cpp b/Engine/Plugins/Animation/ControlRig/Source/ControlRigEditor/Private/EditMode/SControlRigControlViews.cpp
--- a/Engine/Plugins/Animation/ControlRig/Source/ControlRigEditor/Private/EditMode/SControlRigControlViews.cpp
+++ b/Engine/Plugins/Animation/ControlRig/Source/ControlRigEditor/Private/EditMode/SControlRigControlViews.cpp
@@ -23,6 +23,7 @@
 #include "LevelSequence.h"
 #include "LevelSequenceEditorBlueprintLibrary.h"
 #include "PropertyEditorModule.h"
+#include "SControlRigBaseListWidget.h"
 #include "UnrealClient.h"
 
 #define LOCTEXT_NAMESPACE "ControlRigBaseListWidget"
@@ -190,6 +191,7 @@ bool SControlRigPoseView::bIsMirror = false;
 void SControlRigPoseView::Construct(const FArguments& InArgs)
 {
 	PoseAsset = InArgs._PoseAsset;
+	OwningWidget = InArgs._OwningWidget;
 
 	PoseBlendValue = 0.0f;
 	bIsBlending = false;
@@ -441,7 +443,7 @@ void SControlRigPoseView::Construct(const FArguments& InArgs)
 			]
 	];
 
-	if (FControlRigEditMode* EditMode = static_cast<FControlRigEditMode*>(GLevelEditorModeTools().GetActiveMode(FControlRigEditMode::ModeName)))
+	if (FControlRigEditMode* EditMode = OwningWidget->GetEditMode())
 	{
 		EditMode->OnControlRigAddedOrRemoved().AddRaw(this, &SControlRigPoseView::HandleControlAdded);
 		TArray<UControlRig*> ControlRigs = GetControlRigs();
@@ -454,7 +456,7 @@ void SControlRigPoseView::Construct(const FArguments& InArgs)
 
 SControlRigPoseView::~SControlRigPoseView()
 {
-	if (FControlRigEditMode* EditMode = static_cast<FControlRigEditMode*>(GLevelEditorModeTools().GetActiveMode(FControlRigEditMode::ModeName)))
+	if (FControlRigEditMode* EditMode = OwningWidget->GetEditMode())
 	{
 		EditMode->OnControlRigAddedOrRemoved().RemoveAll(this);
 		TArray<UControlRig*> EditModeRigs = EditMode->GetControlRigsArray(false /*bIsVisible*/);
@@ -662,7 +664,7 @@ TSharedRef<SWidget> SControlRigPoseView::GetThumbnailWidget()
 
 TArray<UControlRig*> SControlRigPoseView::GetControlRigs()
 {
-	FControlRigEditMode* EditMode = static_cast<FControlRigEditMode*>(GLevelEditorModeTools().GetActiveMode(FControlRigEditMode::ModeName));
+	FControlRigEditMode* EditMode = OwningWidget->GetEditMode();
 	TArray<UControlRig*> NewControlRigs;
 	if (EditMode)
 	{
diff --git a/Engine/Plugins/Animation/ControlRig/Source/ControlRigEditor/Private/EditMode/SControlRigControlViews.h b/Engine/Plugins/Animation/ControlRig/Source/ControlRigEditor/Private/EditMode/SControlRigControlViews.h
--- a/Engine/Plugins/Animation/ControlRig/Source/ControlRigEditor/Private/EditMode/SControlRigControlViews.h
+++ b/Engine/Plugins/Animation/ControlRig/Source/ControlRigEditor/Private/EditMode/SControlRigControlViews.h
@@ -20,6 +20,7 @@
 class FAssetThumbnail;
 class IDetailsView;
 class UControlRig;
+class SControlRigBaseListWidget;
 
 //Class to Hold Statics that are shared and externally callable
 class FControlRigView
@@ -32,6 +33,7 @@ class SControlRigPoseView : public SCompoundWidget, public FNotifyHook
 {
 	SLATE_BEGIN_ARGS(SControlRigPoseView) {}
 	SLATE_ARGUMENT(UControlRigPoseAsset*, PoseAsset)
+	SLATE_ARGUMENT(TSharedPtr<SControlRigBaseListWidget>, OwningWidget)
 
 	SLATE_END_ARGS()
 
@@ -79,6 +81,7 @@ private:
 	TSharedPtr<STextBlock> TextStatusBlock1;
 	TSharedPtr<STextBlock> TextStatusBlock2;
 
+	TSharedPtr<SControlRigBaseListWidget> OwningWidget;
 
 	/* Mirroring*/
 	TSharedPtr<IDetailsView> MirrorDetailsView;
diff --git a/Engine/Plugins/Animation/ControlRig/Source/ControlRigEditor/Private/EditMode/SControlRigEditModeTools.cpp b/Engine/Plugins/Animation/ControlRig/Source/ControlRigEditor/Private/EditMode/SControlRigEditModeTools.cpp
--- a/Engine/Plugins/Animation/ControlRig/Source/ControlRigEditor/Private/EditMode/SControlRigEditModeTools.cpp
+++ b/Engine/Plugins/Animation/ControlRig/Source/ControlRigEditor/Private/EditMode/SControlRigEditModeTools.cpp
@@ -1204,14 +1204,12 @@ void SControlRigEditModeTools::CustomizeToolBarPalette(FToolBarBuilder& ToolBarB
 		FUIAction(
 		FExecuteAction::CreateSP(this, &SControlRigEditModeTools::ToggleEditPivotMode),
 		FCanExecuteAction(),
-		FIsActionChecked::CreateLambda([] {
-				if (FLevelEditorModule* LevelEditorModule = FModuleManager::GetModulePtr<FLevelEditorModule>(TEXT("LevelEditor")))
+		FIsActionChecked::CreateLambda([this] {
+				if (TSharedPtr<IToolkit> Toolkit = OwningToolkit.Pin())
 				{
-					TSharedPtr<ILevelEditor> LevelEditorPtr = LevelEditorModule->GetLevelEditorInstance().Pin();
-
-					if (LevelEditorPtr.IsValid())
+					if (Toolkit.IsValid())
 					{
-						FString ActiveToolName = LevelEditorPtr->GetEditorModeManager().GetInteractiveToolsContext()->ToolManager->GetActiveToolName(EToolSide::Left);
+						const FString ActiveToolName = Toolkit->GetToolkitHost()->GetEditorModeManager().GetInteractiveToolsContext()->ToolManager->GetActiveToolName(EToolSide::Left);
 						if (ActiveToolName == TEXT("SequencerPivotTool"))
 						{
 							return true;
@@ -1233,25 +1231,20 @@ void SControlRigEditModeTools::CustomizeToolBarPalette(FToolBarBuilder& ToolBarB
 void SControlRigEditModeTools::ToggleEditPivotMode()
 {
 	FEditorModeID ModeID = TEXT("SequencerToolsEditMode");
-	if (GLevelEditorModeTools().IsModeActive(ModeID))
+	if (const TSharedPtr<IToolkit> Toolkit = OwningToolkit.Pin())
 	{
-		if (FLevelEditorModule* LevelEditorModule = FModuleManager::GetModulePtr<FLevelEditorModule>(TEXT("LevelEditor")))
+		if (Toolkit.IsValid())
 		{
-			TSharedPtr<ILevelEditor> LevelEditorPtr = LevelEditorModule->GetLevelEditorInstance().Pin();
-
-			if (LevelEditorPtr.IsValid())
+			const FEditorModeTools& Tools =Toolkit->GetToolkitHost()->GetEditorModeManager();
+			const FString ActiveToolName = Tools.GetInteractiveToolsContext()->ToolManager->GetActiveToolName(EToolSide::Left);
+			if (ActiveToolName == TEXT("SequencerPivotTool"))
 			{
-				FString ActiveToolName = LevelEditorPtr->GetEditorModeManager().GetInteractiveToolsContext()->ToolManager->GetActiveToolName(EToolSide::Left);
-				if (ActiveToolName == TEXT("SequencerPivotTool"))
-				{
-					LevelEditorPtr->GetEditorModeManager().GetInteractiveToolsContext()->ToolManager->DeactivateTool(EToolSide::Left, EToolShutdownType::Completed);
-				}
-				else
-				{
-					LevelEditorPtr->GetEditorModeManager().GetInteractiveToolsContext()->ToolManager->SelectActiveToolType(EToolSide::Left, TEXT("SequencerPivotTool"));
-					LevelEditorPtr->GetEditorModeManager().GetInteractiveToolsContext()->ToolManager->ActivateTool(EToolSide::Left);
-
-				}
+				Tools.GetInteractiveToolsContext()->ToolManager->DeactivateTool(EToolSide::Left, EToolShutdownType::Completed);
+			}
+			else
+			{
+				Tools.GetInteractiveToolsContext()->ToolManager->SelectActiveToolType(EToolSide::Left, TEXT("SequencerPivotTool"));
+				Tools.GetInteractiveToolsContext()->ToolManager->ActivateTool(EToolSide::Left);
 			}
 		}
 	}
@@ -1299,7 +1292,7 @@ void SControlRigEditModeTools::MakeSelectionSetDialog()
 */
 FText SControlRigEditModeTools::GetActiveToolName() const
 {
-	return  FText();
+	return FText::FromString(TEXT("Control Rig Editing"));
 }
 
 FText SControlRigEditModeTools::GetActiveToolMessage() const
diff --git a/Engine/Plugins/Animation/ControlRig/Source/ControlRigEditor/Private/EditMode/SControlRigTweenWidget.cpp b/Engine/Plugins/Animation/ControlRig/Source/ControlRigEditor/Private/EditMode/SControlRigTweenWidget.cpp
--- a/Engine/Plugins/Animation/ControlRig/Source/ControlRigEditor/Private/EditMode/SControlRigTweenWidget.cpp
+++ b/Engine/Plugins/Animation/ControlRig/Source/ControlRigEditor/Private/EditMode/SControlRigTweenWidget.cpp
@@ -26,11 +26,13 @@
 
 void SControlRigTweenSlider::Construct(const FArguments& InArgs)
 {
-
 	PoseBlendValue = 0.0f;
 	bIsBlending = false;
 	bSliderStartedTransaction = false;
 	AnimSlider = InArgs._InAnimSlider;
+	WeakSequencer = InArgs._InSequencer;
+	WeakEditMode = InArgs._InWeakEditMode;
+	
 	ChildSlot
 	[
 		SNew(SVerticalBox)
@@ -132,12 +134,7 @@ void SControlRigTweenSlider::OnBeginSliderMovement()
 bool SControlRigTweenSlider::Setup()
 {
 	//if getting sequencer from level sequence need to use the current(leader), not the focused
-	ULevelSequence* LevelSequence = ULevelSequenceEditorBlueprintLibrary::GetCurrentLevelSequence();
-	IAssetEditorInstance* AssetEditor = GEditor->GetEditorSubsystem<UAssetEditorSubsystem>()->FindEditorForAsset(LevelSequence, false);
-	ILevelSequenceEditorToolkit* LevelSequenceEditor = static_cast<ILevelSequenceEditorToolkit*>(AssetEditor);
-	WeakSequencer = LevelSequenceEditor ? LevelSequenceEditor->GetSequencer() : nullptr;
-
-	return AnimSlider->Setup(WeakSequencer);
+	return AnimSlider->Setup(WeakSequencer, WeakEditMode);
 	
 }
 
@@ -158,7 +155,6 @@ void SControlRigTweenSlider::OnEndSliderMovement(double NewValue)
 			User.SetFocus(ThisRef->GetParentWidget().ToSharedRef());
 		}
 	});
-	WeakSequencer = nullptr;
 }
 
 
@@ -252,10 +248,11 @@ void SControlRigTweenWidget::Construct(const FArguments& InArgs)
 	TSharedPtr<FBaseAnimSlider> PushPullPtr = MakeShareable(new FPushPullSlider());
 	AnimBlendTools.RegisterAnimSlider(PushPullPtr);
 	OwningToolkit = InArgs._InOwningToolkit;
+	OwningEditMode = InArgs._InOwningEditMode;
+	
 	TSharedPtr<FBaseAnimSlider> TweenPtr = MakeShareable(new FControlsToTween());
 	AnimBlendTools.RegisterAnimSlider(TweenPtr);
 
-
 	// Combo Button to swap sliders 
 	TSharedRef<SComboButton> SliderComoboBtn = SNew(SComboButton)
 		.OnGetMenuContent_Lambda([this]()
@@ -326,7 +323,10 @@ void SControlRigTweenWidget::Construct(const FArguments& InArgs)
 		.VAlign(VAlign_Center)
 		.Padding(2.0f, 2.0f, 2.0f, 2.0f)
 		[
-			SAssignNew(SliderWidget,SControlRigTweenSlider).InAnimSlider(SliderPtr)
+			SAssignNew(SliderWidget,SControlRigTweenSlider)
+			.InAnimSlider(SliderPtr)
+			.InSequencer(OwningEditMode.Pin()->GetWeakSequencer().Pin())
+			.InWeakEditMode(OwningEditMode)
 		];
 	
 	ChildSlot
diff --git a/Engine/Plugins/Animation/ControlRig/Source/ControlRigEditor/Private/EditMode/SControlRigTweenWidget.h b/Engine/Plugins/Animation/ControlRig/Source/ControlRigEditor/Private/EditMode/SControlRigTweenWidget.h
--- a/Engine/Plugins/Animation/ControlRig/Source/ControlRigEditor/Private/EditMode/SControlRigTweenWidget.h
+++ b/Engine/Plugins/Animation/ControlRig/Source/ControlRigEditor/Private/EditMode/SControlRigTweenWidget.h
@@ -25,6 +25,8 @@ class SControlRigTweenSlider : public SCompoundWidget
 {
 	SLATE_BEGIN_ARGS(SControlRigTweenSlider) {}
 	SLATE_ARGUMENT(TSharedPtr<FBaseAnimSlider>, InAnimSlider)
+	SLATE_ARGUMENT(TSharedPtr<ISequencer>, InSequencer)
+	SLATE_ARGUMENT(TWeakPtr<FControlRigEditMode>, InWeakEditMode)
 	SLATE_END_ARGS()
 	~SControlRigTweenSlider()
 	{
@@ -51,8 +53,8 @@ private:
 	double PoseBlendValue;
 	bool bIsBlending;
 	bool bSliderStartedTransaction;
-
 	
+	TWeakPtr<FControlRigEditMode> WeakEditMode;
 	TWeakPtr<ISequencer> WeakSequencer;
 	TSharedPtr<FBaseAnimSlider> AnimSlider;
 	TSharedPtr<SSpinBox<double>> SpinBox;
@@ -64,6 +66,7 @@ class SControlRigTweenWidget : public SCompoundWidget
 {
 	SLATE_BEGIN_ARGS(SControlRigTweenWidget) {}
 	SLATE_ARGUMENT(TSharedPtr<FControlRigEditModeToolkit>, InOwningToolkit)
+	SLATE_ARGUMENT(TSharedPtr<FControlRigEditMode>, InOwningEditMode)
 	SLATE_END_ARGS()
 		~SControlRigTweenWidget()
 	{
@@ -92,6 +95,8 @@ private:
 	TWeakPtr<FControlRigEditModeToolkit> OwningToolkit;
 	FAnimBlendTooLManager  AnimBlendTools;
 
+	TWeakPtr<FControlRigEditMode> OwningEditMode;
+
 	TSharedPtr<SControlRigTweenSlider> SliderWidget;
 	static int32 ActiveSlider;
 
diff --git a/Engine/Plugins/Animation/ControlRig/Source/ControlRigEditor/Private/Editor/ControlRigEditorEditMode.h b/Engine/Plugins/Animation/ControlRig/Source/ControlRigEditor/Private/Editor/ControlRigEditorEditMode.h
--- a/Engine/Plugins/Animation/ControlRig/Source/ControlRigEditor/Private/Editor/ControlRigEditorEditMode.h
+++ b/Engine/Plugins/Animation/ControlRig/Source/ControlRigEditor/Private/Editor/ControlRigEditorEditMode.h
@@ -10,6 +10,9 @@ class FControlRigEditorEditMode : public FControlRigEditMode
 {
 public:
 	static FName ModeName;
+	
+	virtual bool IsInLevelEditor() const override { return false; }
+	virtual bool AreEditingControlRigDirectly() const override { return true; }
 
 	// FEdMode interface
 	virtual void Render(const FSceneView* View, FViewport* Viewport, FPrimitiveDrawInterface* PDI) override;
diff --git a/Engine/Plugins/Animation/ControlRig/Source/ControlRigEditor/Private/Tools/ControlRigTweener.cpp b/Engine/Plugins/Animation/ControlRig/Source/ControlRigEditor/Private/Tools/ControlRigTweener.cpp
--- a/Engine/Plugins/Animation/ControlRig/Source/ControlRigEditor/Private/Tools/ControlRigTweener.cpp
+++ b/Engine/Plugins/Animation/ControlRig/Source/ControlRigEditor/Private/Tools/ControlRigTweener.cpp
@@ -67,7 +67,7 @@ void FAnimSliderKeySelection::FContiguousKeysArray::Add(const TArray<int32>& InC
 
 static int32 GetIndex(const TArray<FKeyPosition>& Keys, double Time)
 {
-	int32 Index = Algo::LowerBoundBy(Keys, Time, [](const FKeyPosition& Value) { return Value.InputValue; });
+	const int32 Index = Algo::LowerBoundBy(Keys, Time, [](const FKeyPosition& Value) { return Value.InputValue; });
 
 	// don't trust precision issues so will double check to make sure the index is correct
 	if (Index != INDEX_NONE)
@@ -88,15 +88,15 @@ static int32 GetIndex(const TArray<FKeyPosition>& Keys, double Time)
 	return Index;
 }
 
-bool FAnimSliderKeySelection::Setup(TWeakPtr<ISequencer>& InSequencer)
+bool FAnimSliderKeySelection::Setup(const TWeakPtr<ISequencer>& InSequencer)
 {
 	using namespace UE::Sequencer;
 
 	KeyMap.Reset();
 	if (InSequencer.IsValid())
 	{
-		TSharedPtr<FSequencerEditorViewModel> SequencerViewModel = InSequencer.Pin()->GetViewModel();
-		FCurveEditorExtension* CurveEditorExtension = SequencerViewModel->CastDynamic<FCurveEditorExtension>();
+		const TSharedPtr<FSequencerEditorViewModel> SequencerViewModel = InSequencer.Pin()->GetViewModel();
+		const FCurveEditorExtension* CurveEditorExtension = SequencerViewModel->CastDynamic<FCurveEditorExtension>();
 		check(CurveEditorExtension);
 		CurveEditor = CurveEditorExtension->GetCurveEditor();
 		GetMapOfContiguousKeys();
@@ -104,10 +104,10 @@ bool FAnimSliderKeySelection::Setup(TWeakPtr<ISequencer>& InSequencer)
 	return (KeyMap.Num() > 0);
 }
 
-bool FAnimSliderObjectSelection::Setup(TWeakPtr<ISequencer>& InSequencer)
+bool FAnimSliderObjectSelection::Setup(TWeakPtr<ISequencer>& InSequencer, TWeakPtr<FControlRigEditMode>& InEditMode)
 {
 	ChannelsArray.Reset();
-	TArray<UControlRig*> ControlRigs = GetControlRigs();
+	const TArray<UControlRig*> ControlRigs = GetControlRigs(InEditMode);
 	return Setup(ControlRigs, InSequencer);
 }
 
@@ -125,9 +125,7 @@ void FAnimSliderKeySelection::GetMapOfContiguousKeys()
 
 	for (const TPair<FCurveModelID, FKeyHandleSet>& Pair : SelectionKeyMap)
 	{
-		FCurveModel* Curve = CurveEditor->FindCurve(Pair.Key);
-
-		if (Curve)
+		if (const FCurveModel* Curve = CurveEditor->FindCurve(Pair.Key))
 		{
 			AllKeyHandles.Reset();
 			Curve->GetKeys(*CurveEditor, TNumericLimits<double>::Lowest(), TNumericLimits<double>::Max(), TNumericLimits<double>::Lowest(), TNumericLimits<double>::Max(), AllKeyHandles);
@@ -185,10 +183,10 @@ void FAnimSliderKeySelection::GetMapOfContiguousKeys()
 * FAnimSliderObjectSelection
 *
 */
-TArray<UControlRig*> FAnimSliderObjectSelection::GetControlRigs()
+TArray<UControlRig*> FAnimSliderObjectSelection::GetControlRigs(TWeakPtr<FControlRigEditMode>& InEditMode)
 {
 	TArray<UControlRig*> ControlRigs;
-	if (FControlRigEditMode* EditMode = static_cast<FControlRigEditMode*>(GLevelEditorModeTools().GetActiveMode(FControlRigEditMode::ModeName)))
+	if (const FControlRigEditMode* EditMode = InEditMode.Pin().Get())
 	{
 		TMap<UControlRig*, TArray<FRigElementKey>> SelectedControls;
 		EditMode->GetAllSelectedControls(SelectedControls);
@@ -241,7 +239,7 @@ void FAnimSliderObjectSelection::SetupChannel(FFrameNumber CurrentFrame, TArray<
 		}
 		for (int32 Index = 0; Index < KeyTimes.Num(); Index++)
 		{
-			FFrameNumber FrameNumber = KeyTimes[Index];
+			const FFrameNumber FrameNumber = KeyTimes[Index];
 			if (FrameNumber < CurrentFrame || (FrameNumber == CurrentFrame && KeyBounds.PreviousIndex == INDEX_NONE))
 			{
 				KeyBounds.PreviousIndex = Index;
@@ -284,21 +282,111 @@ bool FAnimSliderObjectSelection::Setup(const TArray<UControlRig*>& SelectedContr
 	}
 	ISequencer* Sequencer = InSequencer.Pin().Get();
 
-	FQualifiedFrameTime CurrentTime = Sequencer->GetLocalTime();
-
-	UMovieScene* MovieScene = Sequencer->GetFocusedMovieSceneSequence()->GetMovieScene();
+	const FQualifiedFrameTime CurrentTime = Sequencer->GetLocalTime();
+	const UMovieScene* MovieScene = Sequencer->GetFocusedMovieSceneSequence()->GetMovieScene();
 
 	//get selected controls and objects
 	TArray<FGuid> SelectedGuids;
 	Sequencer->GetSelectedObjects(SelectedGuids);
 
+	TArray<UMovieSceneControlRigParameterSection*> HandledSections;
+	auto SetupControlRigTrackChannels = [this, &CurrentTime, &HandledSections](const UMovieSceneControlRigParameterTrack* Track)
+	{
+		check(Track);		
+		for (UMovieSceneSection* MovieSection : Track->GetAllSections())
+		{
+			UMovieSceneControlRigParameterSection* Section = Cast<UMovieSceneControlRigParameterSection>(MovieSection);
+			if (Section && Section->IsActive() && Section->GetRange().Contains(CurrentTime.Time.GetFrame()) && !HandledSections.Contains(Section))
+			{
+				HandledSections.Add(Section);				
+				UControlRig* ControlRig = Track->GetControlRig();
+				TArray<FRigControlElement*> CurrentControls;
+				ControlRig->GetControlsInOrder(CurrentControls);
+
+				FMovieSceneChannelProxy& ChannelProxy = Section->GetChannelProxy();
+				TArrayView<FMovieSceneFloatChannel*> Channels = ChannelProxy.GetChannels<FMovieSceneFloatChannel>();
+				//reuse these arrays
+				TArray<FFrameNumber> KeyTimes;
+				TArray<FKeyHandle> Handles;
+				for (const FRigControlElement* ControlElement : CurrentControls)
+				{
+					if (ControlRig->GetHierarchy()->IsAnimatable(ControlElement) &&  ControlRig->IsControlSelected(ControlElement->GetFName()))
+					{
+						FAnimSliderObjectSelection::FObjectChannels ObjectChannels;
+						ObjectChannels.Section = Section;
+						KeyTimes.SetNum(0);
+						Handles.SetNum(0);
+						const FChannelMapInfo* pChannelIndex = Section->ControlChannelMap.Find(ControlElement->GetFName());
+						if (pChannelIndex == nullptr)
+						{
+							continue;
+						}
+
+						const int32 NumChannels = [&ControlElement]()
+						{
+							switch (ControlElement->Settings.ControlType)
+							{
+							case ERigControlType::Float:
+							case ERigControlType::ScaleFloat:
+								{
+									return 1;
+								}
+							case ERigControlType::Vector2D:
+								{
+									return 2;
+								}
+							case ERigControlType::Position:
+							case ERigControlType::Scale:
+							case ERigControlType::Rotator:
+								{
+									return 3;
+								}
+							case ERigControlType::TransformNoScale:
+								{
+									return 6;
+								}
+							case ERigControlType::Transform:
+							case ERigControlType::EulerTransform:
+								{
+									return 9;
+								}
+							default:
+									return 0;
+							}
+						}();
+					
+						int32 BoundIndex = 0;
+						int32 NumValidChannels = 0;
+						ObjectChannels.KeyBounds.SetNum(NumChannels);
+						for (int32 ChannelIdx = pChannelIndex->ChannelIndex; ChannelIdx < (pChannelIndex->ChannelIndex + NumChannels); ++ChannelIdx)
+						{
+							FMovieSceneFloatChannel* Channel = Channels[ChannelIdx];
+							SetupChannel(CurrentTime.Time.GetFrame(), KeyTimes, Handles, Channel, nullptr,
+								ObjectChannels.KeyBounds[BoundIndex]);
+							if (ObjectChannels.KeyBounds[BoundIndex].bValid)
+							{
+								++NumValidChannels;
+							}
+							++BoundIndex;
+						}
+						if (NumValidChannels > 0)
+						{
+							ChannelsArray.Add(ObjectChannels);
+						}
+					}
+				}
+			}
+		}		
+	};
+	
+	// Handle MovieScene bindings
 	const TArray<FMovieSceneBinding>& Bindings = MovieScene->GetBindings();
 	for (const FMovieSceneBinding& Binding : Bindings)
 	{
 		if (SelectedGuids.Num() > 0 && SelectedGuids.Contains(Binding.GetObjectGuid()))
 		{
 			const TArray<UMovieSceneTrack*>& Tracks = Binding.GetTracks();
-			for (UMovieSceneTrack* Track : Tracks)
+			for (const UMovieSceneTrack* Track : Tracks)
 			{
 				if (Track && Track->IsA<UMovieSceneControlRigParameterTrack>() == false)
 				{
@@ -307,10 +395,6 @@ bool FAnimSliderObjectSelection::Setup(const TArray<UControlRig*>& SelectedContr
 					{
 						if (Section && Section->IsActive())
 						{
-							FMovieSceneChannelProxy& Proxy = Section->GetChannelProxy();
-
-							const FName FloatChannelTypeName = FMovieSceneFloatChannel::StaticStruct()->GetFName();
-
 							FMovieSceneChannelProxy& ChannelProxy = Section->GetChannelProxy();
 							TArrayView<FMovieSceneFloatChannel*> FloatChannels = ChannelProxy.GetChannels<FMovieSceneFloatChannel>();
 							TArrayView<FMovieSceneDoubleChannel*> DoubleChannels = ChannelProxy.GetChannels<FMovieSceneDoubleChannel>();
@@ -327,8 +411,7 @@ bool FAnimSliderObjectSelection::Setup(const TArray<UControlRig*>& SelectedContr
 
 							ObjectChannels.KeyBounds.SetNum(NumFloatChannels + NumDoubleChannels);
 							int32 NumValidChannels = 0;
-							int ChannelIdx = 0;
-							for (ChannelIdx = 0;ChannelIdx < NumFloatChannels; ++ChannelIdx)
+							for (int32 ChannelIdx = 0;ChannelIdx < NumFloatChannels; ++ChannelIdx)
 							{
 								FMovieSceneFloatChannel* FloatChannel = FloatChannels[ChannelIdx];
 								SetupChannel(CurrentTime.Time.GetFrame(), KeyTimes, Handles, FloatChannel, nullptr,
@@ -338,7 +421,7 @@ bool FAnimSliderObjectSelection::Setup(const TArray<UControlRig*>& SelectedContr
 									++NumValidChannels;
 								}
 							}
-							for (ChannelIdx = 0; ChannelIdx < NumDoubleChannels; ++ChannelIdx)
+							for (int32 ChannelIdx = 0; ChannelIdx < NumDoubleChannels; ++ChannelIdx)
 							{
 								FMovieSceneDoubleChannel* DoubleChannel = DoubleChannels[ChannelIdx];
 								SetupChannel(CurrentTime.Time.GetFrame(), KeyTimes, Handles, nullptr, DoubleChannel,
@@ -357,96 +440,25 @@ bool FAnimSliderObjectSelection::Setup(const TArray<UControlRig*>& SelectedContr
 				}
 			}
 		}
-		UMovieSceneControlRigParameterTrack* Track = Cast<UMovieSceneControlRigParameterTrack>(MovieScene->FindTrack(UMovieSceneControlRigParameterTrack::StaticClass(), Binding.GetObjectGuid(), NAME_None));
+		const UMovieSceneControlRigParameterTrack* Track = Cast<UMovieSceneControlRigParameterTrack>(MovieScene->FindTrack(UMovieSceneControlRigParameterTrack::StaticClass(), Binding.GetObjectGuid(), NAME_None));
 		if (Track && Track->GetControlRig() && SelectedControlRigs.Contains(Track->GetControlRig()))
 		{
-			for (UMovieSceneSection* MovieSection : Track->GetAllSections())
+			SetupControlRigTrackChannels(Track);
+		}
+	}
+	
+	// Handle movie tracks in general (for non-binding, USkeleton, ControlRig tracks)
+	for (const UMovieSceneTrack* MovieTrack : MovieScene->GetTracks())
+	{
+		if (const UMovieSceneControlRigParameterTrack* Track = Cast<UMovieSceneControlRigParameterTrack>(MovieTrack))
+		{
+			if (Track && Track->GetControlRig() && SelectedControlRigs.Contains(Track->GetControlRig()))
 			{
-				UMovieSceneControlRigParameterSection* Section = Cast<UMovieSceneControlRigParameterSection>(MovieSection);
-				if (Section && Section->IsActive() && Section->GetRange().Contains(CurrentTime.Time.GetFrame()))
-				{
-					UControlRig* ControlRig = Track->GetControlRig();
-					TArray<FRigControlElement*> CurrentControls;
-					ControlRig->GetControlsInOrder(CurrentControls);
-					URigHierarchy* Hierarchy = ControlRig->GetHierarchy();
-					FMovieSceneChannelProxy& Proxy = Section->GetChannelProxy();
-
-					const FName FloatChannelTypeName = FMovieSceneFloatChannel::StaticStruct()->GetFName();
-
-					FMovieSceneChannelProxy& ChannelProxy = Section->GetChannelProxy();
-					TArrayView<FMovieSceneFloatChannel*> Channels = ChannelProxy.GetChannels<FMovieSceneFloatChannel>();
-					//reuse these arrays
-					TArray<FFrameNumber> KeyTimes;
-					TArray<FKeyHandle> Handles;
-					for (FRigControlElement* ControlElement : CurrentControls)
-					{
-						if (ControlRig->GetHierarchy()->IsAnimatable(ControlElement) &&  ControlRig->IsControlSelected(ControlElement->GetFName()))
-						{
-							FAnimSliderObjectSelection::FObjectChannels ObjectChannels;
-							ObjectChannels.Section = Section;
-							KeyTimes.SetNum(0);
-							Handles.SetNum(0);
-							FChannelMapInfo* pChannelIndex = Section->ControlChannelMap.Find(ControlElement->GetFName());
-							if (pChannelIndex == nullptr)
-							{
-								continue;
-							}
-							int NumChannels = 0;
-							switch (ControlElement->Settings.ControlType)
-							{
-								case ERigControlType::Float:
-								case ERigControlType::ScaleFloat:
-								{
-									NumChannels = 1;
-								}
-								case ERigControlType::Vector2D:
-								{
-									NumChannels = 2;
-									break;
-								}
-								case ERigControlType::Position:
-								case ERigControlType::Scale:
-								case ERigControlType::Rotator:
-								{
-									NumChannels = 3;
-									break;
-								}
-								case ERigControlType::TransformNoScale:
-								{
-									NumChannels = 6;
-									break;
-								}
-								case ERigControlType::Transform:
-								case ERigControlType::EulerTransform:
-								{
-									NumChannels = 9;
-									break;
-								}
-							}
-							int32 BoundIndex = 0;
-							int32 NumValidChannels = 0;
-							ObjectChannels.KeyBounds.SetNum(NumChannels);
-							for (int ChannelIdx = pChannelIndex->ChannelIndex; ChannelIdx < (pChannelIndex->ChannelIndex + NumChannels); ++ChannelIdx)
-							{
-								FMovieSceneFloatChannel* Channel = Channels[ChannelIdx];
-								SetupChannel(CurrentTime.Time.GetFrame(), KeyTimes, Handles, Channel, nullptr,
-									ObjectChannels.KeyBounds[BoundIndex]);
-								if (ObjectChannels.KeyBounds[BoundIndex].bValid)
-								{
-									++NumValidChannels;
-								}
-								++BoundIndex;
-							}
-							if (NumValidChannels > 0)
-							{
-								ChannelsArray.Add(ObjectChannels);
-							}
-						}
-					}
-				}
+				SetupControlRigTrackChannels(Track);
 			}
 		}
 	}
+	
 	return ChannelsArray.Num() > 0;
 }
 
@@ -456,19 +468,16 @@ bool FAnimSliderObjectSelection::Setup(const TArray<UControlRig*>& SelectedContr
 * 
 */
 
-bool FBaseAnimSlider::Setup(TWeakPtr<ISequencer>& InSequencer)
+bool FBaseAnimSlider::Setup(TWeakPtr<ISequencer>& InSequencer, TWeakPtr<FControlRigEditMode>& InEditMode)
 {
 	KeySelection.KeyMap.Reset();
 	ObjectSelection.ChannelsArray.Reset();
 	if (KeySelection.Setup(InSequencer) == false)
 	{
-		return ObjectSelection.Setup(InSequencer);
-	}
-	else
-	{
-		return true;
+		return ObjectSelection.Setup(InSequencer, InEditMode);
 	}
-	return false;
+
+	return true;
 }
 
 /*
@@ -483,7 +492,7 @@ bool FBasicBlendSlider::Blend(TWeakPtr<ISequencer>& InSequencer, const double Bl
 	{
 		return false;
 	}
-	ISequencer* Sequencer = InSequencer.Pin().Get();
+	const ISequencer* Sequencer = InSequencer.Pin().Get();
 
 	bool bDidBlend = false;
 
@@ -493,31 +502,32 @@ bool FBasicBlendSlider::Blend(TWeakPtr<ISequencer>& InSequencer, const double Bl
 		TArray<FKeyPosition> KeyPositions;
 		for (TPair<FCurveModelID, FAnimSliderKeySelection::FContiguousKeysArray>& KeysArray : KeySelection.KeyMap)
 		{
-			FCurveModel* Curve = KeySelection.CurveEditor->FindCurve(KeysArray.Key);
-			if (Curve)
+			if (FCurveModel* Curve = KeySelection.CurveEditor->FindCurve(KeysArray.Key))
 			{
 				Curve->Modify();
 				for (FAnimSliderKeySelection::FContiguousKeys& Keys : KeysArray.Value.KeysArray)
 				{
 					const int32 PreviousIndex = Keys.PreviousIndex != INDEX_NONE ? Keys.PreviousIndex : Keys.Indices[0];
 					const double PreviousTime = KeysArray.Value.AllKeyPositions[PreviousIndex].InputValue;
-					const double  PreviousValue = KeysArray.Value.AllKeyPositions[PreviousIndex].OutputValue;
+					const double PreviousValue = KeysArray.Value.AllKeyPositions[PreviousIndex].OutputValue;
 					const int32 NextIndex = Keys.NextIndex != INDEX_NONE ? Keys.NextIndex : Keys.Indices[Keys.Indices.Num() - 1];
 					const double NextTime = KeysArray.Value.AllKeyPositions[NextIndex].InputValue;
-					const double  NextValue = KeysArray.Value.AllKeyPositions[NextIndex].OutputValue;
+					const double NextValue = KeysArray.Value.AllKeyPositions[NextIndex].OutputValue;
 
-					KeyHandles.Reset();
-					KeyPositions.Reset();
-					for (int32 Index : Keys.Indices)
+					const int32 NumIndices = Keys.Indices.Num();
+					KeyHandles.Reset(NumIndices);
+					KeyPositions.Reset(NumIndices);
+					for (int32 Index = 0; Index < NumIndices; ++Index)
 					{
-						KeyHandles.Add(KeysArray.Value.AllKeyHandles[Index]);
-						const double CurrentValue = KeysArray.Value.AllKeyPositions[Index].OutputValue;
-						const double CurrentTime = KeysArray.Value.AllKeyPositions[Index].InputValue;
+						const int32& KeyIndex = Keys.Indices[Index];
+						KeyHandles[Index] = KeysArray.Value.AllKeyHandles[KeyIndex];
+						const double CurrentValue = KeysArray.Value.AllKeyPositions[KeyIndex].OutputValue;
+						const double CurrentTime = KeysArray.Value.AllKeyPositions[KeyIndex].InputValue;
 
-						double NewValue = DoBlend(PreviousTime,PreviousValue, CurrentTime, CurrentValue,
-							NextTime,  NextValue, BlendValue);
+						const double NewValue = DoBlend(PreviousTime,PreviousValue, CurrentTime, CurrentValue,
+							NextTime, NextValue, BlendValue);
 
-						 KeyPositions.Add(FKeyPosition(KeysArray.Value.AllKeyPositions[Index].InputValue, NewValue));
+						KeyPositions[Index] = FKeyPosition(KeysArray.Value.AllKeyPositions[KeyIndex].InputValue, NewValue);
 					}
 					Curve->SetKeyPositions(KeyHandles, KeyPositions);
 					bDidBlend = true;
@@ -527,38 +537,37 @@ bool FBasicBlendSlider::Blend(TWeakPtr<ISequencer>& InSequencer, const double Bl
 		return bDidBlend;
 	}
 
-	FFrameTime  FrameTime = Sequencer->GetLocalTime().Time;
-	FFrameRate TickResoultion = Sequencer->GetFocusedTickResolution();
-
 	if(ObjectSelection.ChannelsArray.Num() > 0)
-	{ 
+	{
+		const FFrameTime FrameTime = Sequencer->GetLocalTime().Time;
+		const FFrameRate TickResolution = Sequencer->GetFocusedTickResolution();
 		for (const FAnimSliderObjectSelection::FObjectChannels& ObjectChannels : ObjectSelection.ChannelsArray)
 		{
 			if (ObjectChannels.Section)
 			{
 				ObjectChannels.Section->Modify();
 			}
-			for (int Index = 0; Index < ObjectChannels.KeyBounds.Num(); ++Index)
+			for (int32 Index = 0; Index < ObjectChannels.KeyBounds.Num(); ++Index)
 			{
 				if (ObjectChannels.KeyBounds[Index].bValid)
 				{
-					double PreviousValue = ObjectChannels.KeyBounds[Index].PreviousValue;
-					double PreviousTime = TickResoultion.AsSeconds(FFrameTime(ObjectChannels.KeyBounds[Index].PreviousFrame));
-					double NextValue = ObjectChannels.KeyBounds[Index].NextValue;
-					double NextTime = TickResoultion.AsSeconds(FFrameTime(ObjectChannels.KeyBounds[Index].NextFrame));
-					double CurrentValue = ObjectChannels.KeyBounds[Index].CurrentValue;
-					double CurrentTime = TickResoultion.AsSeconds(FFrameTime(ObjectChannels.KeyBounds[Index].CurrentFrame));
-					double NewValue = DoBlend(PreviousTime, PreviousValue, CurrentTime, CurrentValue,
+					const double PreviousValue = ObjectChannels.KeyBounds[Index].PreviousValue;
+					const double PreviousTime = TickResolution.AsSeconds(FFrameTime(ObjectChannels.KeyBounds[Index].PreviousFrame));
+					const double NextValue = ObjectChannels.KeyBounds[Index].NextValue;
+					const double NextTime = TickResolution.AsSeconds(FFrameTime(ObjectChannels.KeyBounds[Index].NextFrame));
+					const double CurrentValue = ObjectChannels.KeyBounds[Index].CurrentValue;
+					const double CurrentTime = TickResolution.AsSeconds(FFrameTime(ObjectChannels.KeyBounds[Index].CurrentFrame));
+					const double NewValue = DoBlend(PreviousTime, PreviousValue, CurrentTime, CurrentValue,
 						NextTime, NextValue, BlendValue);
 					using namespace UE::MovieScene;
 					if (ObjectChannels.KeyBounds[Index].FloatChannel)
 					{
-						EMovieSceneKeyInterpolation KeyInterpolation = GetInterpolationMode(ObjectChannels.KeyBounds[Index].FloatChannel, FrameTime.GetFrame(), Sequencer->GetKeyInterpolation());
+						const EMovieSceneKeyInterpolation KeyInterpolation = GetInterpolationMode(ObjectChannels.KeyBounds[Index].FloatChannel, FrameTime.GetFrame(), Sequencer->GetKeyInterpolation());
 						AddKeyToChannel(ObjectChannels.KeyBounds[Index].FloatChannel, FrameTime.GetFrame(), (float)NewValue, KeyInterpolation);
 					}
 					else if (ObjectChannels.KeyBounds[Index].DoubleChannel)
 					{
-						EMovieSceneKeyInterpolation KeyInterpolation = GetInterpolationMode(ObjectChannels.KeyBounds[Index].DoubleChannel, FrameTime.GetFrame(), Sequencer->GetKeyInterpolation());
+						const EMovieSceneKeyInterpolation KeyInterpolation = GetInterpolationMode(ObjectChannels.KeyBounds[Index].DoubleChannel, FrameTime.GetFrame(), Sequencer->GetKeyInterpolation());
 						AddKeyToChannel(ObjectChannels.KeyBounds[Index].DoubleChannel, FrameTime.GetFrame(), NewValue, KeyInterpolation);
 					}
 					bDidBlend = true;
@@ -579,8 +588,8 @@ double FControlsToTween::DoBlend(const double PreviousTime, const double Previou
 	const double NextTime, const double NextValue, const double BlendValue)
 {
 	//clasic tween will move all to same location, not based on current time at all just blend and values
-	double NormalizedBlendValue = (BlendValue + 1.0f) * 0.5f;
-	double Value = PreviousValue + (NextValue - PreviousValue) * (NormalizedBlendValue);
+	const double NormalizedBlendValue = (BlendValue + 1.0f) * 0.5f;
+	const double Value = PreviousValue + (NextValue - PreviousValue) * (NormalizedBlendValue);
 	return Value;
 }
 
@@ -599,9 +608,9 @@ bool FControlsToTween::Setup(const TArray<UControlRig*>& SelectedControlRigs, TW
 	return ObjectSelection.Setup(SelectedControlRigs, InSequencer);
 }
 
-bool FControlsToTween::Setup(TWeakPtr<ISequencer>& InSequencer)
+bool FControlsToTween::Setup(TWeakPtr<ISequencer>& InSequencer, TWeakPtr<FControlRigEditMode>& InEditMode)
 {
-	return FBaseAnimSlider::Setup(InSequencer);
+	return FBaseAnimSlider::Setup(InSequencer, InEditMode);
 }
 
 /*
@@ -616,7 +625,7 @@ double FPushPullSlider::DoBlend(const double PreviousTime, const double Previous
 {
 	const double T = (CurrentTime - PreviousTime) / (NextTime - PreviousTime);
 	const double ValueAtT = PreviousValue + T * (NextValue - PreviousValue);
-	double NewValue = 0.0;
+	double NewValue;
 	if (BlendValue < 0.0)
 	{
 		NewValue = CurrentValue + (-1.0 * BlendValue) * (ValueAtT - CurrentValue);
diff --git a/Engine/Plugins/Animation/ControlRig/Source/ControlRigEditor/Public/EditMode/ControlRigEditMode.h b/Engine/Plugins/Animation/ControlRig/Source/ControlRigEditor/Public/EditMode/ControlRigEditMode.h
--- a/Engine/Plugins/Animation/ControlRig/Source/ControlRigEditor/Public/EditMode/ControlRigEditMode.h
+++ b/Engine/Plugins/Animation/ControlRig/Source/ControlRigEditor/Public/EditMode/ControlRigEditMode.h
@@ -102,8 +102,10 @@ public:
 	/*Replace old Control Rig with the New Control Rig, perhaps from a recompile in the level editor*/
 	void ReplaceControlRig(UControlRig* OldControlRig, UControlRig* NewControlRig);
 
-	/** This edit mode is re-used between the level editor and the control rig editor. Calling this indicates which context we are in */
-	bool IsInLevelEditor() const;
+	/** This edit mode is re-used between the level editor and the asset editors (control rig editor etc.). Calling this indicates which context we are in */
+	virtual bool IsInLevelEditor() const;
+	/** This is used to differentiate between the control rig editor and any other (asset/level) editors in which this edit mode is used */
+	virtual bool AreEditingControlRigDirectly() const { return false; }
 
 	// FEdMode interface
 	virtual bool UsesToolkits() const override;
@@ -122,14 +124,14 @@ public:
 	virtual bool InputDelta(FEditorViewportClient* InViewportClient, FViewport* InViewport, FVector& InDrag, FRotator& InRot, FVector& InScale) override;
 	virtual bool UsesTransformWidget() const override;
 	virtual bool GetPivotForOrbit(FVector& OutPivot) const override;
-	virtual bool UsesTransformWidget(UE::Widget::EWidgetMode CheckMode) const;
+	virtual bool UsesTransformWidget(UE::Widget::EWidgetMode CheckMode) const override;
 	virtual FVector GetWidgetLocation() const override;
 	virtual bool GetCustomDrawingCoordinateSystem(FMatrix& OutMatrix, void* InData) override;
 	virtual bool GetCustomInputCoordinateSystem(FMatrix& OutMatrix, void* InData) override;
 	virtual bool ShouldDrawWidget() const override;
 	virtual bool IsCompatibleWith(FEditorModeID OtherModeID) const override;
 	virtual bool MouseMove(FEditorViewportClient* ViewportClient, FViewport* Viewport, int32 x, int32 y) override;
-	virtual bool MouseLeave(FEditorViewportClient* ViewportClient, FViewport* Viewport);
+	virtual bool MouseLeave(FEditorViewportClient* ViewportClient, FViewport* Viewport) override;
 	virtual void PostUndo() override;
 
 	/* IPersonaEditMode interface */
diff --git a/Engine/Plugins/Animation/ControlRig/Source/ControlRigEditor/Public/Tools/ControlRigTweener.h b/Engine/Plugins/Animation/ControlRig/Source/ControlRigEditor/Public/Tools/ControlRigTweener.h
--- a/Engine/Plugins/Animation/ControlRig/Source/ControlRigEditor/Public/Tools/ControlRigTweener.h
+++ b/Engine/Plugins/Animation/ControlRig/Source/ControlRigEditor/Public/Tools/ControlRigTweener.h
@@ -42,7 +42,7 @@ struct FAnimSliderKeySelection
 		void Add(const TArray<int32>& ContiguousKeyIndices);
 	};
 
-	bool Setup(TWeakPtr<ISequencer>& InSequencer);
+	bool Setup(const TWeakPtr<ISequencer>& InSequencer);
 
 	TSharedPtr<FCurveEditor> CurveEditor;
 	TMap<FCurveModelID, FContiguousKeysArray> KeyMap;
@@ -84,7 +84,7 @@ struct FAnimSliderObjectSelection
 		UMovieSceneSection* Section;
 	};
 
-	bool Setup(TWeakPtr<ISequencer>& InSequencer);
+	bool Setup(TWeakPtr<ISequencer>& InSequencer, TWeakPtr<class FControlRigEditMode>& InEditMode);
 	bool Setup(const TArray<UControlRig*>& SelectedControlRigs, TWeakPtr<ISequencer>& InSequencer);
 
 	TArray<FObjectChannels>  ChannelsArray;
@@ -95,7 +95,7 @@ private:
 	void SetupChannel(FFrameNumber CurrentFrame, TArray<FFrameNumber>& KeyTimes, TArray<FKeyHandle>& Handles, FMovieSceneFloatChannel* FloatChannel,
 		FMovieSceneDoubleChannel* DoubleChannel, FChannelKeyBounds& KeyBounds);
 
-	TArray<UControlRig*> GetControlRigs();
+	TArray<UControlRig*> GetControlRigs(TWeakPtr<class FControlRigEditMode>& InEditMode);
 };
 
 /*
@@ -110,7 +110,7 @@ struct FBaseAnimSlider
 	* @param InSequencer Sequencer to get keys, controls, etc. from
 	* @param returns true if we have something that we can blend
 	*/
-	virtual bool Setup(TWeakPtr<ISequencer>& InSequencer);
+	virtual bool Setup(TWeakPtr<ISequencer>& InSequencer, TWeakPtr<FControlRigEditMode>& InEditMode);
 
 	/**
 	* @param InSequencer Sequencer to blend at current time
@@ -182,7 +182,7 @@ struct FControlsToTween :public FBasicBlendSlider
 	//BaseAnimSlider overrides
 	virtual FText GetText() override;
 	virtual FText GetTooltipText() override;
-	virtual bool Setup(TWeakPtr<ISequencer>& InSequencer) override;
+	virtual bool Setup(TWeakPtr<ISequencer>& InSequencer, TWeakPtr<FControlRigEditMode>& InEditMode) override;
 
 	//BasicBlendSlider overrides
 	virtual double DoBlend(const double PreviousTime, const double PreviousValue, const double CurrentTime, const double CurrentValue,
