Control Rig: FIx curves not being transferred correctly for control rig components

diff --git a/Engine/Plugins/Animation/ControlRig/Source/ControlRig/Private/ControlRigComponent.cpp b/Engine/Plugins/Animation/ControlRig/Source/ControlRig/Private/ControlRigComponent.cpp
--- a/Engine/Plugins/Animation/ControlRig/Source/ControlRig/Private/ControlRigComponent.cpp
+++ b/Engine/Plugins/Animation/ControlRig/Source/ControlRig/Private/ControlRigComponent.cpp
@@ -1375,10 +1375,13 @@ void UControlRigComponent::ValidateMappingData()
 					}
 
 					// if we didn't find the bone, disable this mapped element
-					if (MappedElement.SubIndex == INDEX_NONE)
+					if (MappedElement.ElementType != ERigElementType::Curve)
 					{
-						MappedElement.ElementIndex = INDEX_NONE;
-						continue;
+						if (MappedElement.SubIndex == INDEX_NONE)
+						{
+							MappedElement.ElementIndex = INDEX_NONE;
+							continue;
+						}
 					}
 				}
 
@@ -1640,7 +1643,7 @@ void UControlRigComponent::TransferOutputs()
 				}
 				else if (MappedElement.ElementType == ERigElementType::Curve)
 				{
-					if (MappedElement.SubIndex >= 0)
+					if (MappedElement.ElementIndex >= 0)
 					{
 						if (LastComponent != MappedElement.SceneComponent || Proxy == nullptr)
 						{
