Control Rig: Fix python user workflows not receiving the correct information
#rnx

diff --git a/Engine/Plugins/Animation/ControlRig/Source/ControlRigDeveloper/Private/ControlRigBlueprint.cpp b/Engine/Plugins/Animation/ControlRig/Source/ControlRigDeveloper/Private/ControlRigBlueprint.cpp
--- a/Engine/Plugins/Animation/ControlRig/Source/ControlRigDeveloper/Private/ControlRigBlueprint.cpp
+++ b/Engine/Plugins/Animation/ControlRig/Source/ControlRigDeveloper/Private/ControlRigBlueprint.cpp
@@ -24,6 +24,7 @@
 #include "RigVMModel/Nodes/RigVMAggregateNode.h"
 #include "Rigs/RigControlHierarchy.h"
 #include "Settings/ControlRigSettings.h"
+#include "Units/ControlRigNodeWorkflow.h"
 #include "Units/Execution/RigUnit_PrepareForExecution.h"
 #include "Units/Execution/RigUnit_DynamicHierarchy.h"
 
@@ -893,6 +894,30 @@ void UControlRigBlueprint::HandlePackageDone()
 	Super::HandlePackageDone();
 }
 
+void UControlRigBlueprint::HandleConfigureRigVMController(const FRigVMClient* InClient, URigVMController* InControllerToConfigure)
+{
+	Super::HandleConfigureRigVMController(InClient, InControllerToConfigure);
+
+	TWeakObjectPtr<URigVMBlueprint> WeakThis(this);
+	InControllerToConfigure->ConfigureWorkflowOptionsDelegate.BindLambda([WeakThis](URigVMUserWorkflowOptions* Options)
+	{
+		if(UControlRigWorkflowOptions* ControlRigNodeWorkflowOptions = Cast<UControlRigWorkflowOptions>(Options))
+		{
+			ControlRigNodeWorkflowOptions->Hierarchy = nullptr;
+			ControlRigNodeWorkflowOptions->Selection.Reset();
+			
+			if(const URigVMBlueprint* StrongThis = WeakThis.Get())
+			{
+				if(UControlRig* ControlRig = Cast<UControlRig>(StrongThis->GetObjectBeingDebugged()))
+				{
+					ControlRigNodeWorkflowOptions->Hierarchy = ControlRig->GetHierarchy();
+					ControlRigNodeWorkflowOptions->Selection = ControlRig->GetHierarchy()->GetSelectedKeys();
+				}
+			}
+		}
+	});
+}
+
 #endif
 
 UClass* UControlRigBlueprint::GetRigVMEdGraphNodeClass() const
diff --git a/Engine/Plugins/Animation/ControlRig/Source/ControlRigDeveloper/Public/ControlRigBlueprint.h b/Engine/Plugins/Animation/ControlRig/Source/ControlRigDeveloper/Public/ControlRigBlueprint.h
--- a/Engine/Plugins/Animation/ControlRig/Source/ControlRigDeveloper/Public/ControlRigBlueprint.h
+++ b/Engine/Plugins/Animation/ControlRig/Source/ControlRigDeveloper/Public/ControlRigBlueprint.h
@@ -273,6 +273,7 @@ private:
 
 #if WITH_EDITOR
 	virtual void HandlePackageDone() override;
+	virtual void HandleConfigureRigVMController(const FRigVMClient* InClient, URigVMController* InControllerToConfigure) override;
 #endif
 
 	// Class used to temporarily cache all 
diff --git a/Engine/Plugins/Runtime/RigVM/Source/RigVMDeveloper/Private/RigVMBlueprint.cpp b/Engine/Plugins/Runtime/RigVM/Source/RigVMDeveloper/Private/RigVMBlueprint.cpp
--- a/Engine/Plugins/Runtime/RigVM/Source/RigVMDeveloper/Private/RigVMBlueprint.cpp
+++ b/Engine/Plugins/Runtime/RigVM/Source/RigVMDeveloper/Private/RigVMBlueprint.cpp
@@ -627,27 +627,6 @@ void URigVMBlueprint::HandleConfigureRigVMController(const FRigVMClient* InClien
 		}
 	});
 
-	/*
-	 * todooooo
-	InControllerToConfigure->ConfigureWorkflowOptionsDelegate.BindLambda([WeakThis](URigVMUserWorkflowOptions* Options)
-	{
-		if(UControlRigWorkflowOptions* ControlRigNodeWorkflowOptions = Cast<UControlRigWorkflowOptions>(Options))
-		{
-			ControlRigNodeWorkflowOptions->Hierarchy = nullptr;
-			ControlRigNodeWorkflowOptions->Selection.Reset();
-			
-			if(const URigVMBlueprint* StrongThis = WeakThis.Get())
-			{
-				if(UControlRig* ControlRig = Cast<UControlRig>(StrongThis->GetObjectBeingDebugged()))
-				{
-					ControlRigNodeWorkflowOptions->Hierarchy = ControlRig->GetHierarchy();
-				}
-				ControlRigNodeWorkflowOptions->Selection = StrongThis->Hierarchy->GetSelectedKeys();
-			}
-		}
-	});
-	*/
-
 #endif
 }
 
