Fixing GC assumption verification crash caused by creating DataSourceRegistry for ControlRig's CDO

diff --git a/Engine/Plugins/Animation/ControlRig/Source/ControlRig/Private/ControlRig.cpp b/Engine/Plugins/Animation/ControlRig/Source/ControlRig/Private/ControlRig.cpp
--- a/Engine/Plugins/Animation/ControlRig/Source/ControlRig/Private/ControlRig.cpp
+++ b/Engine/Plugins/Animation/ControlRig/Source/ControlRig/Private/ControlRig.cpp
@@ -1508,6 +1508,11 @@ UAnimationDataSourceRegistry* UControlRig::GetDataSourceRegistry()
 	if (DataSourceRegistry == nullptr)
 	{
 		DataSourceRegistry = NewObject<UAnimationDataSourceRegistry>(this, NAME_None, RF_Transient);
+
+		if (HasAnyFlags(RF_ClassDefaultObject) && GetClass()->IsNative())
+		{
+			DataSourceRegistry->AddToRoot();
+		}
 	}
 
 	return DataSourceRegistry;
