Control Rig: Fixed FControlRigParameterTrackEditor missing unbind of ControlUndoBracket event
#rnx

diff --git a/Engine/Plugins/Animation/ControlRig/Source/ControlRigEditor/Private/Sequencer/ControlRigParameterTrackEditor.cpp b/Engine/Plugins/Animation/ControlRig/Source/ControlRigEditor/Private/Sequencer/ControlRigParameterTrackEditor.cpp
--- a/Engine/Plugins/Animation/ControlRig/Source/ControlRigEditor/Private/Sequencer/ControlRigParameterTrackEditor.cpp
+++ b/Engine/Plugins/Animation/ControlRig/Source/ControlRigEditor/Private/Sequencer/ControlRigParameterTrackEditor.cpp
@@ -482,6 +482,7 @@ void FControlRigParameterTrackEditor::UnbindControlRig(UControlRig* ControlRig)
 		{
 			Binding->OnControlRigBind().RemoveAll(this);
 		}
+		ControlRig->ControlUndoBracket().RemoveAll(this);
 		ControlRig->ControlRigBound().RemoveAll(this);
 		
 		BoundControlRigs.Remove(ControlRig);
