Control Rig: Detach when destroying shape actors

Fix an ensure that fires if the parent is deleted while the child is attached, it continues to try to detach.

diff --git a/Engine/Plugins/Animation/ControlRig/Source/ControlRigEditor/Private/EditMode/ControlRigEditMode.cpp b/Engine/Plugins/Animation/ControlRig/Source/ControlRigEditor/Private/EditMode/ControlRigEditMode.cpp
--- a/Engine/Plugins/Animation/ControlRig/Source/ControlRigEditor/Private/EditMode/ControlRigEditMode.cpp
+++ b/Engine/Plugins/Animation/ControlRig/Source/ControlRigEditor/Private/EditMode/ControlRigEditMode.cpp
@@ -4651,6 +4651,10 @@ void FControlRigEditMode::DestroyShapesActors(UControlRig* ControlRig)
 				UWorld* World = ShapeActor->GetWorld();
 				if (World)
 				{
+					if (ShapeActor->GetAttachParentActor())
+					{
+						ShapeActor->DetachFromActor(FDetachmentTransformRules::KeepRelativeTransform);
+					}
 					World->EditorDestroyActor(ShapeActor, true);
 				}
 			}
@@ -4673,6 +4677,10 @@ void FControlRigEditMode::DestroyShapesActors(UControlRig* ControlRig)
 				UWorld* World = ShapeActor->GetWorld();
 				if (World)
 				{
+					if (ShapeActor->GetAttachParentActor())
+					{
+						ShapeActor->DetachFromActor(FDetachmentTransformRules::KeepRelativeTransform);
+					}
 					World->EditorDestroyActor(ShapeActor,true);
 				}
 			}
