RigVM: Fix ensure when iterating over asset user data arrays
#rnx

diff --git a/Engine/Plugins/Animation/ControlRig/Source/ControlRig/Private/ControlRig.cpp b/Engine/Plugins/Animation/ControlRig/Source/ControlRig/Private/ControlRig.cpp
--- a/Engine/Plugins/Animation/ControlRig/Source/ControlRig/Private/ControlRig.cpp
+++ b/Engine/Plugins/Animation/ControlRig/Source/ControlRig/Private/ControlRig.cpp
@@ -2857,7 +2857,8 @@ const TArray<UAssetUserData*>* UControlRig::GetAssetUserDataArray() const
 				{
 					if(const TArray<UAssetUserData*>* SkeletonUserDataArray = Skeleton->GetAssetUserDataArray())
 					{
-						for(UAssetUserData* SkeletonUserData : *SkeletonUserDataArray)
+						const TArray<UAssetUserData*> LocalSkeletonUserDataArray = *SkeletonUserDataArray;
+						for(UAssetUserData* SkeletonUserData : LocalSkeletonUserDataArray)
 						{
 							CombinedAssetUserData.Add(SkeletonUserData);
 						}
@@ -2865,7 +2866,8 @@ const TArray<UAssetUserData*>* UControlRig::GetAssetUserDataArray() const
 				}
 				if(const TArray<UAssetUserData*>* SkeletalMeshUserDataArray = SkeletalMesh->GetAssetUserDataArray())
 				{
-					for(UAssetUserData* SkeletalMeshUserData : *SkeletalMeshUserDataArray)
+					const TArray<UAssetUserData*> LocalSkeletalMeshUserDataArray = *SkeletalMeshUserDataArray;
+					for(UAssetUserData* SkeletalMeshUserData : LocalSkeletalMeshUserDataArray)
 					{
 						CombinedAssetUserData.Add(SkeletalMeshUserData);
 					}
@@ -2873,7 +2875,8 @@ const TArray<UAssetUserData*>* UControlRig::GetAssetUserDataArray() const
 			}
 			if (const TArray<UAssetUserData*>* ActorComponentUserDataArray = SkeletalMeshComponent->GetAssetUserDataArray())
 			{
-				for(UAssetUserData* ActorComponentUserData : *ActorComponentUserDataArray)
+				const TArray<UAssetUserData*> LocalActorComponentUserDataArray = *ActorComponentUserDataArray;
+				for(UAssetUserData* ActorComponentUserData : LocalActorComponentUserDataArray)
 				{
 					CombinedAssetUserData.Add(ActorComponentUserData);
 				}
