Sequencer: Control Rig: Add priority order for Control Rig External Nodes.  Allow for multiple Control Rigs on same binding.

diff --git a/Engine/Plugins/Animation/ControlRig/Source/ControlRig/Private/Sequencer/ControlRigLayerInstanceProxy.cpp b/Engine/Plugins/Animation/ControlRig/Source/ControlRig/Private/Sequencer/ControlRigLayerInstanceProxy.cpp
--- a/Engine/Plugins/Animation/ControlRig/Source/ControlRig/Private/Sequencer/ControlRigLayerInstanceProxy.cpp
+++ b/Engine/Plugins/Animation/ControlRig/Source/ControlRig/Private/Sequencer/ControlRigLayerInstanceProxy.cpp
@@ -5,6 +5,7 @@
 #include "Sequencer/ControlRigLayerInstance.h"
 #include "AnimSequencerInstance.h"
 #include "ControlRig.h"
+#include "Sequencer/MovieSceneControlRigParameterTrack.h"
 
 #include UE_INLINE_GENERATED_CPP_BY_NAME(ControlRigLayerInstanceProxy)
 
@@ -39,8 +40,42 @@ void FControlRigLayerInstanceProxy::PreEvaluateAnimation(UAnimInstance* InAnimIn
 	}
 }
 
+void FControlRigLayerInstanceProxy::SortControlRigNodes()
+{
+	auto SortPredicate = [](TSharedPtr<FAnimNode_ControlRig_ExternalSource>& A, TSharedPtr<FAnimNode_ControlRig_ExternalSource>& B)
+	{
+		FAnimNode_ControlRig_ExternalSource* APtr = A.Get();
+		FAnimNode_ControlRig_ExternalSource* BPtr = B.Get();
+		if (APtr && BPtr && APtr->GetControlRig() && B->GetControlRig())
+		{
+			const bool AIsAdditive = A->GetControlRig()->IsAdditive();
+			const bool BIsAdditive = B->GetControlRig()->IsAdditive();
+			if (AIsAdditive != BIsAdditive)
+			{
+				return AIsAdditive; // if additive then first(first goes last) so true if AIsAdditive;
+			}
+			UMovieSceneControlRigParameterTrack* TrackA = APtr->GetControlRig()->GetTypedOuter<UMovieSceneControlRigParameterTrack>();
+			UMovieSceneControlRigParameterTrack* TrackB = BPtr->GetControlRig()->GetTypedOuter<UMovieSceneControlRigParameterTrack>();
+			if (TrackA && TrackB)
+			{
+				return TrackA->GetPriorityOrder() < TrackB->GetPriorityOrder();
+			}
+		}
+		else if (APtr && APtr->GetControlRig())
+		{
+			return false;
+		}
+		return true;
+	};
+
+	Algo::Sort(ControlRigNodes, SortPredicate);
+
+	ConstructNodes(); //we need to sort since prority may change at any time 
+}
+
 bool FControlRigLayerInstanceProxy::Evaluate(FPoseContext& Output)
 {
+	SortControlRigNodes();  //mz todo once we move over to ECS see if we can avoid this and trigger it as needed.
 	check(CurrentRoot);
 	CurrentRoot->Evaluate_AnyThread(Output);
 	return true;
@@ -144,35 +179,120 @@ void FControlRigLayerInstanceProxy::AddControlRigTrack(int32 ControlRigID, UCont
 
 	if(!Node)
 	{
-		FAnimNode_ControlRig_ExternalSource* PreviousNode = nullptr;
-		FAnimNode_ControlRig_ExternalSource* NextNode = nullptr;
-		if (InControlRig->IsAdditive())
-		{
-			PreviousNode = (ControlRigNodes.Num() > 0)? ControlRigNodes.Last().Get() : nullptr;
-			Node = ControlRigNodes.Add_GetRef(MakeShared<FAnimNode_ControlRig_ExternalSource>()).Get();
-		}
-		else
-		{
-			NextNode = (ControlRigNodes.Num() > 0)? ControlRigNodes[0].Get() : nullptr;
-			Node = ControlRigNodes.Insert_GetRef(MakeShared<FAnimNode_ControlRig_ExternalSource>(), 0).Get();
-		}
+		UMovieSceneControlRigParameterTrack* Track = InControlRig->GetTypedOuter<UMovieSceneControlRigParameterTrack>();
 
-		if (PreviousNode)
+		if (ControlRigNodes.Num() > 0)
 		{
-			Node->Source.SetLinkNode(PreviousNode);
+			int32 PriorityOrder = Track->GetPriorityOrder();
+			if (PriorityOrder == INDEX_NONE) //track has no order so just add to end, will happen on creation
+			{
+				if (InControlRig->IsAdditive()) //additive added to end of all
+				{
+					Node = ControlRigNodes[ControlRigNodes.Num() - 1].Get();
+					if (UMovieSceneControlRigParameterTrack* OtherTrack = Node->GetControlRig()->GetTypedOuter<UMovieSceneControlRigParameterTrack>())
+					{
+						Track->SetPriorityOrder(OtherTrack->GetPriorityOrder() + 1);
+					}
+					else
+					{
+						Track->SetPriorityOrder(ControlRigNodes.Num() -1);
+					}
+					Node = ControlRigNodes.Add_GetRef(MakeShared<FAnimNode_ControlRig_ExternalSource>()).Get();
+				}
+				else //if non-additive find first additive then insert there..
+				{
+					for (int32 Index = 0; Index < ControlRigNodes.Num(); ++Index)
+					{
+						Node = ControlRigNodes[Index].Get();
+						if (!Node || !Node->GetControlRig() || !Node->GetControlRig()->IsAdditive())
+						{
+							Node = nullptr;
+							continue; // not additive
+						}
+						//okay first additive so insert, but set priority first
+						if (UMovieSceneControlRigParameterTrack* OtherTrack = Node->GetControlRig()->GetTypedOuter<UMovieSceneControlRigParameterTrack>())
+						{
+							Track->SetPriorityOrder(OtherTrack->GetPriorityOrder() + 1);
+						}
+						else
+						{
+							Track->SetPriorityOrder(Index);
+						}
+						Node = ControlRigNodes.Insert_GetRef(MakeShared<FAnimNode_ControlRig_ExternalSource>(), Index).Get();
+						break;
+					}
+
+					if (Node == nullptr) //add to end
+					{
+						//set priority
+						Node = ControlRigNodes[ControlRigNodes.Num() - 1].Get();
+						if (UMovieSceneControlRigParameterTrack* OtherTrack = Node->GetControlRig()->GetTypedOuter<UMovieSceneControlRigParameterTrack>())
+						{
+							Track->SetPriorityOrder(OtherTrack->GetPriorityOrder() + 1);
+						}
+						else
+						{
+							Track->SetPriorityOrder(ControlRigNodes.Num() - 1);
+						}
+						Node = ControlRigNodes.Add_GetRef(MakeShared<FAnimNode_ControlRig_ExternalSource>()).Get();
+					}
+				}
+			}
+			else
+			{ 
+				for (int32 Index = 0; Index < ControlRigNodes.Num(); ++Index)
+				{
+					Node = ControlRigNodes[Index].Get();
+	
+					if (!Node || !Node->GetControlRig())
+					{
+						continue; // shouldn't happen
+					}
+					//additive but current isn't so keep going
+					if (InControlRig->IsAdditive() && !Node->GetControlRig()->IsAdditive())
+					{
+						continue;
+					}
+					//not additive but this one is so we add it here
+					else if (!InControlRig->IsAdditive() && Node->GetControlRig()->IsAdditive())
+					{
+						Node = ControlRigNodes.Insert_GetRef(MakeShared<FAnimNode_ControlRig_ExternalSource>(), Index).Get();
+						break;
+					}
+					if (UMovieSceneControlRigParameterTrack* OtherTrack = Node->GetControlRig()->GetTypedOuter<UMovieSceneControlRigParameterTrack>())
+					{
+						if (PriorityOrder >= OtherTrack->GetPriorityOrder())
+						{
+							//if last one then add it
+							if (Index == ControlRigNodes.Num() - 1)
+							{
+								Node = ControlRigNodes.Add_GetRef(MakeShared<FAnimNode_ControlRig_ExternalSource>()).Get();
+								break;
+							}
+						}
+						else
+						{
+							Node = ControlRigNodes.Insert_GetRef(MakeShared<FAnimNode_ControlRig_ExternalSource>(), Index).Get();
+							break;
+						}
+					}
+				}
+			}
 		}
 		else
 		{
-			// first node
-			Node->Source.SetLinkNode(&InputPose);
+			Node = ControlRigNodes.Add_GetRef(MakeShared<FAnimNode_ControlRig_ExternalSource>()).Get();
+			int32 PriorityOrder = Track->GetPriorityOrder();
+			if (PriorityOrder == INDEX_NONE) //track has no order so just add to end, will happen on creation
+			{
+				Track->SetPriorityOrder(ControlRigNodes.Num() - 1);
+			}
 		}
+	
+		check(Node);
 		
-		if (NextNode)
-		{
-			NextNode->Source.SetLinkNode(Node);
-		}
-
-		CurrentRoot = ControlRigNodes.Last().Get();
+		//this will set up the link nodes
+		ConstructNodes();
 		SequencerToControlRigNodeMap.FindOrAdd(ControlRigID) = Node;
 	}
 
diff --git a/Engine/Plugins/Animation/ControlRig/Source/ControlRig/Private/Sequencer/ControlRigLayerInstanceProxy.h b/Engine/Plugins/Animation/ControlRig/Source/ControlRig/Private/Sequencer/ControlRigLayerInstanceProxy.h
--- a/Engine/Plugins/Animation/ControlRig/Source/ControlRig/Private/Sequencer/ControlRigLayerInstanceProxy.h
+++ b/Engine/Plugins/Animation/ControlRig/Source/ControlRig/Private/Sequencer/ControlRigLayerInstanceProxy.h
@@ -109,6 +109,9 @@ public:
 
 	friend struct FAnimNode_ControlRigInputPose;
 protected:
+	/** Sort Control Rig node*/
+	void SortControlRigNodes();
+
 	/** Find ControlRig node of the */
 	FAnimNode_ControlRig_ExternalSource* FindControlRigNode(int32 ControlRigID) const;
 
diff --git a/Engine/Plugins/Animation/ControlRig/Source/ControlRig/Private/Sequencer/MovieSceneControlRigParameterTemplate.cpp b/Engine/Plugins/Animation/ControlRig/Source/ControlRig/Private/Sequencer/MovieSceneControlRigParameterTemplate.cpp
--- a/Engine/Plugins/Animation/ControlRig/Source/ControlRig/Private/Sequencer/MovieSceneControlRigParameterTemplate.cpp
+++ b/Engine/Plugins/Animation/ControlRig/Source/ControlRig/Private/Sequencer/MovieSceneControlRigParameterTemplate.cpp
@@ -1222,7 +1222,13 @@ struct FControlRigParameterExecutionToken : IMovieSceneExecutionToken
 					{
 						if (UControlRigLayerInstance* AnimInstance = Cast<UControlRigLayerInstance>(SkeletalMeshComponent->GetAnimInstance()))
 						{
-							float Weight = 1.0f;
+							float Weight = Section->EvaluateEasing(Context.GetTime());
+							if (EnumHasAllFlags(Section->TransformMask.GetChannels(), EMovieSceneTransformChannel::Weight))
+							{
+								float ManualWeight = 1.f;
+								Section->Weight.Evaluate(Context.GetTime(), ManualWeight);
+								Weight *= ManualWeight;
+							}
 							FControlRigIOSettings InputSettings;
 							InputSettings.bUpdateCurves = true;
 							InputSettings.bUpdatePose = true;
diff --git a/Engine/Plugins/Animation/ControlRig/Source/ControlRig/Private/Sequencer/MovieSceneControlRigParameterTrack.cpp b/Engine/Plugins/Animation/ControlRig/Source/ControlRig/Private/Sequencer/MovieSceneControlRigParameterTrack.cpp
--- a/Engine/Plugins/Animation/ControlRig/Source/ControlRig/Private/Sequencer/MovieSceneControlRigParameterTrack.cpp
+++ b/Engine/Plugins/Animation/ControlRig/Source/ControlRig/Private/Sequencer/MovieSceneControlRigParameterTrack.cpp
@@ -31,6 +31,7 @@ FControlRotationOrder::FControlRotationOrder()
 UMovieSceneControlRigParameterTrack::UMovieSceneControlRigParameterTrack(const FObjectInitializer& ObjectInitializer)
 	: Super(ObjectInitializer)
 	, ControlRig(nullptr)
+	, PriorityOrder(INDEX_NONE)
 {
 #if WITH_EDITORONLY_DATA
 	TrackTint = AbsoluteRigTrackColor;
@@ -790,6 +791,23 @@ void UMovieSceneControlRigParameterTrack::GetSelectedNodes(TArray<FName>& Select
 	}
 }
 
+int32 UMovieSceneControlRigParameterTrack::GetPriorityOrder() const
+{
+	return PriorityOrder;
+}
+
+void UMovieSceneControlRigParameterTrack::SetPriorityOrder(int32 InPriorityIndex)
+{
+	if (InPriorityIndex >= 0)
+	{
+		PriorityOrder = InPriorityIndex;
+	}
+	else
+	{
+		PriorityOrder = 0;
+	}
+}
+
 #if WITH_EDITOR
 bool UMovieSceneControlRigParameterTrack::GetFbxCurveDataFromChannelMetadata(const FMovieSceneChannelMetaData& MetaData, FControlRigFbxCurveData& OutCurveData)
 {
diff --git a/Engine/Plugins/Animation/ControlRig/Source/ControlRig/Public/Sequencer/MovieSceneControlRigParameterTrack.h b/Engine/Plugins/Animation/ControlRig/Source/ControlRig/Public/Sequencer/MovieSceneControlRigParameterTrack.h
--- a/Engine/Plugins/Animation/ControlRig/Source/ControlRig/Public/Sequencer/MovieSceneControlRigParameterTrack.h
+++ b/Engine/Plugins/Animation/ControlRig/Source/ControlRig/Public/Sequencer/MovieSceneControlRigParameterTrack.h
@@ -156,6 +156,9 @@ public:
 	// if order is not set then it uses the default FRotator conversions
 	CONTROLRIG_API void ChangeControlRotationOrder(const FName& InControlName, const TOptional<EEulerRotationOrder>& NewOrder,
 		EMovieSceneKeyInterpolation Interpolation = EMovieSceneKeyInterpolation::SmartAuto);
+
+	CONTROLRIG_API int32 GetPriorityOrder() const;
+	CONTROLRIG_API void SetPriorityOrder(int32 InPriorityIndex);
 private:
 	//get the rotation orderm will not be set if it's default FRotator order. If bCurrent is true, it uses what's set
 	//if false it uses the default setting from the current control rig.
@@ -173,8 +176,6 @@ private:
 		IMovieSceneConstrainedSection* InSection,
 		FMovieSceneConstraintChannel* InChannel) const;
 	IMovieSceneConstrainedSection::FConstraintChannelAddedEvent OnConstraintChannelAdded;
-
-
 	void ReconstructControlRig();
 
 private:
@@ -199,6 +200,9 @@ private:
 	UPROPERTY()
 	TMap<FName,FControlRotationOrder> ControlsRotationOrder;
 
+	UPROPERTY()
+	int32 PriorityOrder;
+
 public:
 	static CONTROLRIG_API FColor AbsoluteRigTrackColor;
 	static CONTROLRIG_API FColor AdditiveRigTrackColor;
diff --git a/Engine/Plugins/Animation/ControlRig/Source/ControlRigEditor/Private/Sequencer/ControlRigParameterTrackEditor.cpp b/Engine/Plugins/Animation/ControlRig/Source/ControlRigEditor/Private/Sequencer/ControlRigParameterTrackEditor.cpp
--- a/Engine/Plugins/Animation/ControlRig/Source/ControlRigEditor/Private/Sequencer/ControlRigParameterTrackEditor.cpp
+++ b/Engine/Plugins/Animation/ControlRig/Source/ControlRigEditor/Private/Sequencer/ControlRigParameterTrackEditor.cpp
@@ -9,6 +9,7 @@
 #include "Framework/Commands/Commands.h"
 #include "Rendering/DrawElements.h"
 #include "Widgets/SBoxPanel.h"
+#include "Widgets/Input/SSpinBox.h"
 #include "Framework/MultiBox/MultiBoxBuilder.h"
 #include "GameFramework/Actor.h"
 #include "AssetRegistry/AssetData.h"
@@ -1174,26 +1175,8 @@ void FControlRigParameterTrackEditor::BuildObjectBindingTrackMenu(FMenuBuilder&
 
 		if (Skeleton)
 		{
-			//if there are any other absolute control rigs we don't allow it for now..
-			//mz todo will allow later
-			UMovieScene* MovieScene = GetSequencer()->GetFocusedMovieSceneSequence()->GetMovieScene();
-			TArray<UMovieSceneTrack*> ExistingTracks = MovieScene->FindTracks(UMovieSceneControlRigParameterTrack::StaticClass(), ObjectBindings[0], NAME_None);
-			ExistingTracks = ExistingTracks.FilterByPredicate([](UMovieSceneTrack* Track)
-			{
-				if (UMovieSceneControlRigParameterTrack* CRTrack = Cast<UMovieSceneControlRigParameterTrack>(Track))
-				{
-					if (UControlRig* ControlRig = CRTrack->GetControlRig())
-					{
-						return !ControlRig->IsAdditive();
-					}
-				}
-				return false;
-			});
-			if (ExistingTracks.IsEmpty())
-			{
-				UMovieSceneTrack* Track = nullptr;
-				MenuBuilder.AddSubMenu(LOCTEXT("ControlRigText", "Control Rig"), FText(), FNewMenuDelegate::CreateSP(this, &FControlRigParameterTrackEditor::HandleAddTrackSubMenu, ObjectBindings, Track));
-			}
+			UMovieSceneTrack* Track = nullptr;
+			MenuBuilder.AddSubMenu(LOCTEXT("ControlRigText", "Control Rig"), FText(), FNewMenuDelegate::CreateSP(this, &FControlRigParameterTrackEditor::HandleAddTrackSubMenu, ObjectBindings, Track));
 		}
 	}
 }
@@ -3861,6 +3844,22 @@ void FControlRigParameterTrackEditor::BuildTrackContextMenu(FMenuBuilder& MenuBu
 	}
 	MenuBuilder.EndSection();
 
+	MenuBuilder.BeginSection("Control Rig", LOCTEXT("ControlRig", "Control Rig"));
+	{
+		MenuBuilder.AddWidget(
+			SNew(SSpinBox<int32>)
+			.MinValue(0)
+			.Font(FAppStyle::GetFontStyle(TEXT("MenuItem.Font")))
+			.ToolTipText(LOCTEXT("OrderTooltip", "Order for this Control Rig to evaluate compared to others on the same binding"))
+			.Value_Lambda([Track]() { return Track->GetPriorityOrder(); })
+			.OnValueChanged_Lambda([Track](int32 InValue) { Track->SetPriorityOrder(InValue); })
+			,
+			LOCTEXT("Order", "Order")
+		);
+
+	}
+	MenuBuilder.EndSection();
+
 	MenuBuilder.AddMenuSeparator();
 
 	if (UFKControlRig* AutoRig = Cast<UFKControlRig>(Track->GetControlRig()))
