Sequencer: Tween need to call SetNum(Num) instead of Reset(Num), not sure when or how this changed.

diff --git a/Engine/Plugins/Animation/ControlRig/Source/ControlRigEditor/Private/Tools/ControlRigTweener.cpp b/Engine/Plugins/Animation/ControlRig/Source/ControlRigEditor/Private/Tools/ControlRigTweener.cpp
--- a/Engine/Plugins/Animation/ControlRig/Source/ControlRigEditor/Private/Tools/ControlRigTweener.cpp
+++ b/Engine/Plugins/Animation/ControlRig/Source/ControlRigEditor/Private/Tools/ControlRigTweener.cpp
@@ -515,8 +515,8 @@ bool FBasicBlendSlider::Blend(TWeakPtr<ISequencer>& InSequencer, const double Bl
 					const double NextValue = KeysArray.Value.AllKeyPositions[NextIndex].OutputValue;
 
 					const int32 NumIndices = Keys.Indices.Num();
-					KeyHandles.Reset(NumIndices);
-					KeyPositions.Reset(NumIndices);
+					KeyHandles.SetNum(NumIndices);
+					KeyPositions.SetNum(NumIndices);
 					for (int32 Index = 0; Index < NumIndices; ++Index)
 					{
 						const int32& KeyIndex = Keys.Indices[Index];
