Control Rig: Avoid crash caused by null track when setting priority order
#rnx

diff --git a/Engine/Plugins/Animation/ControlRig/Source/ControlRig/Private/Sequencer/ControlRigLayerInstanceProxy.cpp b/Engine/Plugins/Animation/ControlRig/Source/ControlRig/Private/Sequencer/ControlRigLayerInstanceProxy.cpp
--- a/Engine/Plugins/Animation/ControlRig/Source/ControlRig/Private/Sequencer/ControlRigLayerInstanceProxy.cpp
+++ b/Engine/Plugins/Animation/ControlRig/Source/ControlRig/Private/Sequencer/ControlRigLayerInstanceProxy.cpp
@@ -189,13 +189,16 @@ void FControlRigLayerInstanceProxy::AddControlRigTrack(int32 ControlRigID, UCont
 				if (InControlRig->IsAdditive()) //additive added to end of all
 				{
 					Node = ControlRigNodes[ControlRigNodes.Num() - 1].Get();
-					if (UMovieSceneControlRigParameterTrack* OtherTrack = Node->GetControlRig()->GetTypedOuter<UMovieSceneControlRigParameterTrack>())
-					{
-						Track->SetPriorityOrder(OtherTrack->GetPriorityOrder() + 1);
-					}
-					else
+					if (Track)
 					{
-						Track->SetPriorityOrder(ControlRigNodes.Num() -1);
+						if (UMovieSceneControlRigParameterTrack* OtherTrack = Node->GetControlRig()->GetTypedOuter<UMovieSceneControlRigParameterTrack>())
+						{
+							Track->SetPriorityOrder(OtherTrack->GetPriorityOrder() + 1);
+						}
+						else
+						{
+							Track->SetPriorityOrder(ControlRigNodes.Num() -1);
+						}
 					}
 					Node = ControlRigNodes.Add_GetRef(MakeShared<FAnimNode_ControlRig_ExternalSource>()).Get();
 				}
@@ -210,13 +213,16 @@ void FControlRigLayerInstanceProxy::AddControlRigTrack(int32 ControlRigID, UCont
 							continue; // not additive
 						}
 						//okay first additive so insert, but set priority first
-						if (UMovieSceneControlRigParameterTrack* OtherTrack = Node->GetControlRig()->GetTypedOuter<UMovieSceneControlRigParameterTrack>())
-						{
-							Track->SetPriorityOrder(OtherTrack->GetPriorityOrder() + 1);
-						}
-						else
+						if (Track)
 						{
-							Track->SetPriorityOrder(Index);
+							if (UMovieSceneControlRigParameterTrack* OtherTrack = Node->GetControlRig()->GetTypedOuter<UMovieSceneControlRigParameterTrack>())
+							{
+								Track->SetPriorityOrder(OtherTrack->GetPriorityOrder() + 1);
+							}
+							else
+							{
+								Track->SetPriorityOrder(Index);
+							}
 						}
 						Node = ControlRigNodes.Insert_GetRef(MakeShared<FAnimNode_ControlRig_ExternalSource>(), Index).Get();
 						break;
@@ -226,13 +232,16 @@ void FControlRigLayerInstanceProxy::AddControlRigTrack(int32 ControlRigID, UCont
 					{
 						//set priority
 						Node = ControlRigNodes[ControlRigNodes.Num() - 1].Get();
-						if (UMovieSceneControlRigParameterTrack* OtherTrack = Node->GetControlRig()->GetTypedOuter<UMovieSceneControlRigParameterTrack>())
-						{
-							Track->SetPriorityOrder(OtherTrack->GetPriorityOrder() + 1);
-						}
-						else
+						if (Track)
 						{
-							Track->SetPriorityOrder(ControlRigNodes.Num() - 1);
+							if (UMovieSceneControlRigParameterTrack* OtherTrack = Node->GetControlRig()->GetTypedOuter<UMovieSceneControlRigParameterTrack>())
+							{
+								Track->SetPriorityOrder(OtherTrack->GetPriorityOrder() + 1);
+							}
+							else
+							{
+								Track->SetPriorityOrder(ControlRigNodes.Num() - 1);
+							}
 						}
 						Node = ControlRigNodes.Add_GetRef(MakeShared<FAnimNode_ControlRig_ExternalSource>()).Get();
 					}
@@ -282,10 +291,13 @@ void FControlRigLayerInstanceProxy::AddControlRigTrack(int32 ControlRigID, UCont
 		else
 		{
 			Node = ControlRigNodes.Add_GetRef(MakeShared<FAnimNode_ControlRig_ExternalSource>()).Get();
-			int32 PriorityOrder = Track->GetPriorityOrder();
-			if (PriorityOrder == INDEX_NONE) //track has no order so just add to end, will happen on creation
+			if (Track)
 			{
-				Track->SetPriorityOrder(ControlRigNodes.Num() - 1);
+				int32 PriorityOrder = Track->GetPriorityOrder();
+				if (PriorityOrder == INDEX_NONE) //track has no order so just add to end, will happen on creation
+				{
+					Track->SetPriorityOrder(ControlRigNodes.Num() - 1);
+				}
 			}
 		}
 	
