Control Rig: Change hotkeys so that Ctrl+G resets to rest pose, and Ctrl+I zeros controls

diff --git a/Engine/Plugins/Animation/ControlRig/Source/ControlRig/Private/ControlRig.cpp b/Engine/Plugins/Animation/ControlRig/Source/ControlRig/Private/ControlRig.cpp
--- a/Engine/Plugins/Animation/ControlRig/Source/ControlRig/Private/ControlRig.cpp
+++ b/Engine/Plugins/Animation/ControlRig/Source/ControlRig/Private/ControlRig.cpp
@@ -453,16 +453,19 @@ void UControlRig::ClearPoseBeforeBackwardsSolve()
 	PoseBeforeBackwardsSolve.Reset();
 }
 
-void UControlRig::InvertInputPose(EControlRigSetKey InSetKey)
+void UControlRig::InvertInputPose(const TArray<FRigElementKey>& InElements, EControlRigSetKey InSetKey)
 {
 	for (const FRigPoseElement& PoseElement : ControlsAfterBackwardsSolve)
 	{
 		if (FRigControlElement* ControlElement = Cast<FRigControlElement>(DynamicHierarchy->Get(PoseElement.Index)))
 		{
-			FRigControlValue Value;
-			Value.SetFromTransform(PoseElement.LocalTransform.Inverse(), ControlElement->Settings.ControlType, ControlElement->Settings.PrimaryAxis);
-			FRigSetControlValueInfo Info = {Value, true, FRigControlModifiedContext(InSetKey), false, false, false};
-			ControlValues.Add(ControlElement->GetKey(), Info);
+			if (InElements.IsEmpty() || InElements.Contains(ControlElement->GetKey()))
+			{
+				FRigControlValue Value;
+				Value.SetFromTransform(PoseElement.LocalTransform.Inverse(), ControlElement->Settings.ControlType, ControlElement->Settings.PrimaryAxis);
+				FRigSetControlValueInfo Info = {Value, true, FRigControlModifiedContext(InSetKey), false, false, false};
+				ControlValues.Add(ControlElement->GetKey(), Info);
+			}
 		}
 	}
 }
diff --git a/Engine/Plugins/Animation/ControlRig/Source/ControlRig/Public/ControlRig.h b/Engine/Plugins/Animation/ControlRig/Source/ControlRig/Public/ControlRig.h
--- a/Engine/Plugins/Animation/ControlRig/Source/ControlRig/Public/ControlRig.h
+++ b/Engine/Plugins/Animation/ControlRig/Source/ControlRig/Public/ControlRig.h
@@ -141,7 +141,7 @@ public:
 	void ClearPoseBeforeBackwardsSolve();
 
 	/* For additive rigs, will set control values by inverting the pose found after the backwards solve */
-	void InvertInputPose(EControlRigSetKey InSetKey);
+	void InvertInputPose(const TArray<FRigElementKey>& InElements = TArray<FRigElementKey>(), EControlRigSetKey InSetKey = EControlRigSetKey::Never);
 
 	/** Setup bindings to a runtime object (or clear by passing in nullptr). */
 	void SetObjectBinding(TSharedPtr<IControlRigObjectBinding> InObjectBinding)
diff --git a/Engine/Plugins/Animation/ControlRig/Source/ControlRigEditor/Private/EditMode/ControlRigEditMode.cpp b/Engine/Plugins/Animation/ControlRig/Source/ControlRigEditor/Private/EditMode/ControlRigEditMode.cpp
--- a/Engine/Plugins/Animation/ControlRig/Source/ControlRigEditor/Private/EditMode/ControlRigEditMode.cpp
+++ b/Engine/Plugins/Animation/ControlRig/Source/ControlRigEditor/Private/EditMode/ControlRigEditMode.cpp
@@ -2369,14 +2369,17 @@ void FControlRigEditMode::BindCommands()
 		Commands.ToggleAllManipulators,
 		FExecuteAction::CreateRaw(this, &FControlRigEditMode::ToggleAllManipulators));
 	CommandBindings->MapAction(
-		Commands.ResetTransforms,
-		FExecuteAction::CreateRaw(this, &FControlRigEditMode::ResetTransforms, true));
+		Commands.ZeroTransforms,
+		FExecuteAction::CreateRaw(this, &FControlRigEditMode::ZeroTransforms, true));
 	CommandBindings->MapAction(
-		Commands.ResetAllTransforms,
-		FExecuteAction::CreateRaw(this, &FControlRigEditMode::ResetTransforms, false));
+		Commands.ZeroAllTransforms,
+		FExecuteAction::CreateRaw(this, &FControlRigEditMode::ZeroTransforms, false));
 	CommandBindings->MapAction(
-		Commands.InvertInputPose,
-		FExecuteAction::CreateRaw(this, &FControlRigEditMode::InvertInputPose));
+		Commands.InvertTransforms,
+		FExecuteAction::CreateRaw(this, &FControlRigEditMode::InvertInputPose, true));
+	CommandBindings->MapAction(
+		Commands.InvertAllTransforms,
+		FExecuteAction::CreateRaw(this, &FControlRigEditMode::InvertInputPose, false));
 	CommandBindings->MapAction(
 		Commands.ClearSelection,
 		FExecuteAction::CreateRaw(this, &FControlRigEditMode::ClearSelection));
@@ -2862,7 +2865,7 @@ void FControlRigEditMode::ToggleAllManipulators()
 	}
 }
 
-void FControlRigEditMode::ResetTransforms(bool bSelectionOnly)
+void FControlRigEditMode::ZeroTransforms(bool bSelectionOnly)
 {
 	// Gather up the control rigs for the selected controls
 	TArray<UControlRig*> ControlRigs;
@@ -2881,7 +2884,7 @@ void FControlRigEditMode::ResetTransforms(bool bSelectionOnly)
 		return;
 	}
 
-	FScopedTransaction Transaction(LOCTEXT("HierarchyResetTransforms", "Reset Transforms"));
+	FScopedTransaction Transaction(LOCTEXT("HierarchyZeroTransforms", "Zero Transforms"));
 
 	for (UControlRig* ControlRig : ControlRigs)
 	{
@@ -3062,20 +3065,44 @@ void FControlRigEditMode::ResetTransforms(bool bSelectionOnly)
 	}
 }
 
-void FControlRigEditMode::InvertInputPose()
+void FControlRigEditMode::InvertInputPose(bool bSelectionOnly)
 {
-	FScopedTransaction Transaction(LOCTEXT("HierarchyInvertInputPose", "Invert Input Pose"));
+	// Gather up the control rigs for the selected controls
+	TArray<UControlRig*> ControlRigs;
 	for (TWeakObjectPtr<UControlRig>& RuntimeRigPtr : RuntimeControlRigs)
 	{
-		if (RuntimeRigPtr.IsValid())
+		if (UControlRig* ControlRig = RuntimeRigPtr.Get())
 		{
-			if (RuntimeRigPtr->IsAdditive())
+			if (!bSelectionOnly || ControlRig->CurrentControlSelection().Num() > 0)
 			{
-				RuntimeRigPtr->InvertInputPose(EControlRigSetKey::Never);
-				RuntimeRigPtr->Evaluate_AnyThread();
+				ControlRigs.Add(ControlRig);
 			}
 		}
 	}
+	if (ControlRigs.Num() == 0)
+	{
+		return;
+	}
+
+	FScopedTransaction Transaction(LOCTEXT("HierarchyInvertTransformsToRestPose", "Invert Transforms to Rest Pose"));
+
+	for (UControlRig* ControlRig : ControlRigs)
+	{
+		if (!ControlRig->IsAdditive())
+		{
+			ZeroTransforms(bSelectionOnly);
+			continue;
+		}
+
+		TArray<FRigElementKey> SelectedRigElements;
+		if (bSelectionOnly)
+		{
+			SelectedRigElements = GetSelectedRigElements(ControlRig);
+		}
+
+		ControlRig->InvertInputPose(SelectedRigElements, EControlRigSetKey::Never);
+		ControlRig->Evaluate_AnyThread();
+	}
 }
 
 bool FControlRigEditMode::MouseMove(FEditorViewportClient* ViewportClient, FViewport* Viewport, int32 x, int32 y)
diff --git a/Engine/Plugins/Animation/ControlRig/Source/ControlRigEditor/Private/EditMode/ControlRigEditModeCommands.cpp b/Engine/Plugins/Animation/ControlRig/Source/ControlRigEditor/Private/EditMode/ControlRigEditModeCommands.cpp
--- a/Engine/Plugins/Animation/ControlRig/Source/ControlRigEditor/Private/EditMode/ControlRigEditModeCommands.cpp
+++ b/Engine/Plugins/Animation/ControlRig/Source/ControlRigEditor/Private/EditMode/ControlRigEditModeCommands.cpp
@@ -6,9 +6,10 @@
 
 void FControlRigEditModeCommands::RegisterCommands()
 {
-	UI_COMMAND(ResetTransforms, "Reset Transform", "Reset the Controls Transforms", EUserInterfaceActionType::Button, FInputChord(EKeys::G, EModifierKey::Control));
-	UI_COMMAND(ResetAllTransforms, "Reset All Transform", "Reset all of the Controls Transforms", EUserInterfaceActionType::Button, FInputChord(EKeys::G, EModifierKey::Control | EModifierKey::Shift));
-	UI_COMMAND(InvertInputPose, "Invert Input Pose", "Invert Input Pose for additive control rigs", EUserInterfaceActionType::Button, FInputChord(EKeys::I, EModifierKey::Control));
+	UI_COMMAND(InvertTransforms, "Invert to Rest Pose", "Invert the Controls Transforms to Rest Pose", EUserInterfaceActionType::Button, FInputChord(EKeys::G, EModifierKey::Control));
+	UI_COMMAND(InvertAllTransforms, "Invert All Transform to Rest Pose", "Invert all of the Controls Transforms to Rest Pose", EUserInterfaceActionType::Button, FInputChord(EKeys::G, EModifierKey::Control | EModifierKey::Shift));
+	UI_COMMAND(ZeroTransforms, "Zeros Transforms", "Resets Control Transforms to Zero", EUserInterfaceActionType::Button, FInputChord(EKeys::I, EModifierKey::Control));
+	UI_COMMAND(ZeroAllTransforms, "Zeros All Transforms", "Resets all of the Control Transforms to Zero", EUserInterfaceActionType::Button, FInputChord(EKeys::I, EModifierKey::Control | EModifierKey::Shift));
 	UI_COMMAND(ToggleManipulators, "Toggle Shapes", "Toggles visibility of active control rig shapes in the viewport", EUserInterfaceActionType::Button, FInputChord(EKeys::T));
 	UI_COMMAND(ToggleAllManipulators, "Toggle All Shapes", "Toggles visibility of all control rig shapes in the viewport", EUserInterfaceActionType::Button, FInputChord(EKeys::T, EModifierKey::Control | EModifierKey::Alt));
 	UI_COMMAND(FrameSelection, "Frame Selection", "Focus the viewport on the current selection", EUserInterfaceActionType::Button, FInputChord(EKeys::F));
diff --git a/Engine/Plugins/Animation/ControlRig/Source/ControlRigEditor/Private/EditMode/ControlRigEditModeCommands.h b/Engine/Plugins/Animation/ControlRig/Source/ControlRigEditor/Private/EditMode/ControlRigEditModeCommands.h
--- a/Engine/Plugins/Animation/ControlRig/Source/ControlRigEditor/Private/EditMode/ControlRigEditModeCommands.h
+++ b/Engine/Plugins/Animation/ControlRig/Source/ControlRigEditor/Private/EditMode/ControlRigEditModeCommands.h
@@ -32,14 +32,17 @@ public:
 	/**  Change current Anim Slider Tool */
 	TSharedPtr< FUICommandInfo > ChangeAnimSliderTool;
 
-	/** Reset Transforms for Controls */
-	TSharedPtr< FUICommandInfo > ResetTransforms;
+	/** Invert Transforms to Rest Pose */
+	TSharedPtr< FUICommandInfo > InvertTransforms;
 
-	/** Reset Transforms for Controls */
-	TSharedPtr< FUICommandInfo > ResetAllTransforms;
+	/** Invert All Transforms to Rest Pose */
+	TSharedPtr< FUICommandInfo > InvertAllTransforms;
 
-	/** Invert input pose for additive control rigs */
-	TSharedPtr< FUICommandInfo > InvertInputPose;
+	/** Set Transforms to Zero */
+	TSharedPtr< FUICommandInfo > ZeroTransforms;
+
+	/** Set All Transforms to Zero */
+	TSharedPtr< FUICommandInfo > ZeroAllTransforms;
 
 	/** Clear Selection*/
 	TSharedPtr< FUICommandInfo > ClearSelection;
diff --git a/Engine/Plugins/Animation/ControlRig/Source/ControlRigEditor/Private/Sequencer/ControlRigParameterTrackEditor.cpp b/Engine/Plugins/Animation/ControlRig/Source/ControlRigEditor/Private/Sequencer/ControlRigParameterTrackEditor.cpp
--- a/Engine/Plugins/Animation/ControlRig/Source/ControlRigEditor/Private/Sequencer/ControlRigParameterTrackEditor.cpp
+++ b/Engine/Plugins/Animation/ControlRig/Source/ControlRigEditor/Private/Sequencer/ControlRigParameterTrackEditor.cpp
@@ -1109,7 +1109,7 @@ void FControlRigParameterTrackEditor::BakeInvertedPose(UControlRig* InControlRig
 
 	auto DelegateHandle = InControlRig->OnPreAdditiveValuesApplication_AnyThread().AddLambda([](UControlRig* InControlRig, const FName& InEventName)
 	{
-		InControlRig->InvertInputPose(EControlRigSetKey::Never);
+		InControlRig->InvertInputPose();
 	});
 
 	auto KeyFrame = [this, ParentSequencer, InControlRig, SkelMeshComp](const FFrameNumber FrameNumber)
diff --git a/Engine/Plugins/Animation/ControlRig/Source/ControlRigEditor/Public/EditMode/ControlRigEditMode.h b/Engine/Plugins/Animation/ControlRig/Source/ControlRigEditor/Public/EditMode/ControlRigEditMode.h
--- a/Engine/Plugins/Animation/ControlRig/Source/ControlRigEditor/Public/EditMode/ControlRigEditMode.h
+++ b/Engine/Plugins/Animation/ControlRig/Source/ControlRigEditor/Public/EditMode/ControlRigEditMode.h
@@ -248,10 +248,10 @@ private:
 	bool CanFrameSelection();
 
 	/** Reset Transforms */
-	void ResetTransforms(bool bSelectionOnly);
+	void ZeroTransforms(bool bSelectionOnly);
 
 	/** Invert Input Pose */
-	void InvertInputPose();
+	void InvertInputPose(bool bSelectionOnly);
 
 	/** Increase Shape Size */
 	void IncreaseShapeSize();
