ControlRigLayerInstanceProxy: fix a bug with removing control rig track potentially lead to unremovable tracks due to a recent change to the ordering of nodes in the node array

diff --git a/Engine/Plugins/Animation/ControlRig/Source/ControlRig/Private/Sequencer/ControlRigLayerInstanceProxy.cpp b/Engine/Plugins/Animation/ControlRig/Source/ControlRig/Private/Sequencer/ControlRigLayerInstanceProxy.cpp
--- a/Engine/Plugins/Animation/ControlRig/Source/ControlRig/Private/Sequencer/ControlRigLayerInstanceProxy.cpp
+++ b/Engine/Plugins/Animation/ControlRig/Source/ControlRig/Private/Sequencer/ControlRigLayerInstanceProxy.cpp
@@ -345,6 +345,8 @@ void FControlRigLayerInstanceProxy::RemoveControlRigTrack(int32 ControlRigID)
 	if (FAnimNode_ControlRig_ExternalSource* Node = FindControlRigNode(ControlRigID))
 	{
 		FAnimNode_ControlRig_ExternalSource* Parent = nullptr;
+
+		// "ControlRigNodes" should have nodes sorted from parent(last to evaluate) to child(first to evaluate)
 		for (int32 Index = 0; Index < ControlRigNodes.Num(); ++Index)
 		{
 			FAnimNode_ControlRig_ExternalSource* Current = ControlRigNodes[Index].Get();
@@ -355,20 +357,15 @@ void FControlRigLayerInstanceProxy::RemoveControlRigTrack(int32 ControlRigID)
 				// find next child one
 				FAnimNode_ControlRig_ExternalSource* Child = (ControlRigNodes.IsValidIndex(Index + 1)) ? ControlRigNodes[Index + 1].Get() : nullptr;
 
-				// if no parent, change root
-				if (Parent == nullptr)
+				if (Parent)
 				{
-					// first one to delete
 					if (Child)
 					{
-						Child->Source.SetLinkNode(&InputPose);
+						Parent->Source.SetLinkNode(Child);
 					}
-				}
-				else
-				{
-					if (Child)
+					else
 					{
-						Child->Source.SetLinkNode(Parent);
+						Parent->Source.SetLinkNode(&InputPose);
 					}
 				}
 
@@ -385,7 +382,8 @@ void FControlRigLayerInstanceProxy::RemoveControlRigTrack(int32 ControlRigID)
 		}
 		else
 		{
-			CurrentRoot = ControlRigNodes.Last().Get();
+			// stay consistent with ConstructNodes()
+			CurrentRoot = ControlRigNodes[0].Get();
 		}
 		SequencerToControlRigNodeMap.Remove(ControlRigID);
 	}
