Control Rig: FIx selection in sequencer not syncronized with level selection
#rnx

diff --git a/Engine/Plugins/Animation/ControlRig/Source/ControlRigEditor/Private/Sequencer/ControlRigParameterTrackEditor.cpp b/Engine/Plugins/Animation/ControlRig/Source/ControlRigEditor/Private/Sequencer/ControlRigParameterTrackEditor.cpp
--- a/Engine/Plugins/Animation/ControlRig/Source/ControlRigEditor/Private/Sequencer/ControlRigParameterTrackEditor.cpp
+++ b/Engine/Plugins/Animation/ControlRig/Source/ControlRigEditor/Private/Sequencer/ControlRigParameterTrackEditor.cpp
@@ -2034,6 +2034,7 @@ void FControlRigParameterTrackEditor::OnSelectionChanged(TArray<UMovieSceneTrack
 	}
 
 	FControlRigEditMode* ControlRigEditMode = GetEditMode();
+	bool bEditModeExisted = ControlRigEditMode != nullptr;
 	UControlRig* ControlRig = nullptr;
 
 	TArray<const IKeyArea*> KeyAreas;
@@ -2088,6 +2089,12 @@ void FControlRigParameterTrackEditor::OnSelectionChanged(TArray<UMovieSceneTrack
 		return;
 	}
 	SelectRigsAndControls(ControlRig, KeyAreas);
+
+	// If the edit mode has been activated, we need to synchronize the external selection (possibly again to account for control rig control actors selection)
+	if (!bEditModeExisted && GetEditMode() != nullptr)
+	{
+		FSequencerUtilities::SynchronizeExternalSelectionWithSequencerSelection(GetSequencer().ToSharedRef());
+	}
 	
 }
 
