Remove MovieSceneControlRigParameter parameter masks hack (hotfix workaround)

diff --git a/Engine/Plugins/Animation/ControlRig/Source/ControlRig/Private/Sequencer/MovieSceneControlRigParameterTemplate.cpp b/Engine/Plugins/Animation/ControlRig/Source/ControlRig/Private/Sequencer/MovieSceneControlRigParameterTemplate.cpp
--- a/Engine/Plugins/Animation/ControlRig/Source/ControlRig/Private/Sequencer/MovieSceneControlRigParameterTemplate.cpp
+++ b/Engine/Plugins/Animation/ControlRig/Source/ControlRig/Private/Sequencer/MovieSceneControlRigParameterTemplate.cpp
@@ -1772,9 +1772,7 @@ void FMovieSceneControlRigParameterTemplate::Evaluate(const FMovieSceneEvaluatio
 		//Do blended tokens
 		FEvaluatedControlRigParameterSectionValues Values;
 
-		HACK_ChannelMasks = ChannelMasks;
-		EvaluateCurvesWithMasks(Context, Values);
-		HACK_ChannelMasks = nullptr;
+		EvaluateCurvesWithMasks(Context, *ChannelMasks, Values);
 
 		float Weight = EvaluateEasing(Context.GetTime());
 		if (EnumHasAllFlags(Section->TransformMask.GetChannels(), EMovieSceneTransformChannel::Weight))
@@ -1866,13 +1864,10 @@ void FMovieSceneControlRigParameterTemplate::Evaluate(const FMovieSceneEvaluatio
 }
 
 
-void FMovieSceneControlRigParameterTemplate::EvaluateCurvesWithMasks(const FMovieSceneContext& Context, FEvaluatedControlRigParameterSectionValues& Values) const
+void FMovieSceneControlRigParameterTemplate::EvaluateCurvesWithMasks(const FMovieSceneContext& Context, FEvaluatedControlRigParameterSectionChannelMasks& InMasks, FEvaluatedControlRigParameterSectionValues& Values) const
 {
 	const FFrameTime Time = Context.GetTime();
 
-	check(HACK_ChannelMasks);
-
-
 	const UMovieSceneControlRigParameterSection* Section = Cast<UMovieSceneControlRigParameterSection>(GetSourceSection());
 	if (Section)
 	{
@@ -1893,7 +1888,7 @@ void FMovieSceneControlRigParameterTemplate::EvaluateCurvesWithMasks(const FMovi
 			const FScalarParameterNameAndCurve& Scalar = this->Scalars[Index];
 			float Value = 0;
 
-			if (HACK_ChannelMasks->ScalarCurveMask[Index])
+			if (InMasks.ScalarCurveMask[Index])
 			{
 				Scalar.ParameterCurve.Evaluate(Time, Value);
 			}
@@ -1934,7 +1929,7 @@ void FMovieSceneControlRigParameterTemplate::EvaluateCurvesWithMasks(const FMovi
 		{
 			bool Value = false;
 			const FBoolParameterNameAndCurve& Bool = Bools[Index];
-			if (HACK_ChannelMasks->BoolCurveMask[Index])
+			if (InMasks.BoolCurveMask[Index])
 			{
 				Bool.ParameterCurve.Evaluate(Time, Value);
 			}
@@ -1951,7 +1946,7 @@ void FMovieSceneControlRigParameterTemplate::EvaluateCurvesWithMasks(const FMovi
 		{
 			int32  Value = 0;
 			const FIntegerParameterNameAndCurve& Integer = Integers[Index];
-			if (HACK_ChannelMasks->IntegerCurveMask[Index])
+			if (InMasks.IntegerCurveMask[Index])
 			{
 				Integer.ParameterCurve.Evaluate(Time, Value);
 			}
@@ -1968,7 +1963,7 @@ void FMovieSceneControlRigParameterTemplate::EvaluateCurvesWithMasks(const FMovi
 		{
 			uint8  Value = 0;
 			const FEnumParameterNameAndCurve& Enum = Enums[Index];
-			if (HACK_ChannelMasks->EnumCurveMask[Index])
+			if (InMasks.EnumCurveMask[Index])
 			{
 				Enum.ParameterCurve.Evaluate(Time, Value);
 			}
@@ -1987,7 +1982,7 @@ void FMovieSceneControlRigParameterTemplate::EvaluateCurvesWithMasks(const FMovi
 			FVector2f Value(ForceInitToZero);
 			const FVector2DParameterNameAndCurves& Vector2D = Vector2Ds[Index];
 
-			if (HACK_ChannelMasks->Vector2DCurveMask[Index])
+			if (InMasks.Vector2DCurveMask[Index])
 			{
 				Vector2D.XCurve.Evaluate(Time, Value.X);
 				Vector2D.YCurve.Evaluate(Time, Value.Y);
@@ -2015,7 +2010,7 @@ void FMovieSceneControlRigParameterTemplate::EvaluateCurvesWithMasks(const FMovi
 			FVector3f Value(ForceInitToZero);
 			const FVectorParameterNameAndCurves& Vector = Vectors[Index];
 
-			if (HACK_ChannelMasks->VectorCurveMask[Index])
+			if (InMasks.VectorCurveMask[Index])
 			{
 				Vector.XCurve.Evaluate(Time, Value.X);
 				Vector.YCurve.Evaluate(Time, Value.Y);
@@ -2046,7 +2041,7 @@ void FMovieSceneControlRigParameterTemplate::EvaluateCurvesWithMasks(const FMovi
 		{
 			FLinearColor ColorValue = FLinearColor::White;
 			const FColorParameterNameAndCurves& Color = Colors[Index];
-			if (HACK_ChannelMasks->ColorCurveMask[Index])
+			if (InMasks.ColorCurveMask[Index])
 			{
 				Color.RedCurve.Evaluate(Time, ColorValue.R);
 				Color.GreenCurve.Evaluate(Time, ColorValue.G);
@@ -2091,7 +2086,7 @@ void FMovieSceneControlRigParameterTemplate::EvaluateCurvesWithMasks(const FMovi
 			FRotator3f Rotator(0.0f, 0.0f, 0.0f);
 
 			const FTransformParameterNameAndCurves& Transform = Transforms[Index];
-			if (HACK_ChannelMasks->TransformCurveMask[Index])
+			if (InMasks.TransformCurveMask[Index])
 			{
 				if (EnumHasAllFlags(ChannelMask, EMovieSceneTransformChannel::TranslationX))
 				{
@@ -2270,10 +2265,7 @@ void FMovieSceneControlRigParameterTemplate::Interrogate(const FMovieSceneContex
 
 		//Do blended tokens
 		FEvaluatedControlRigParameterSectionValues Values;
-
-		HACK_ChannelMasks = &ChannelMasks;
-		EvaluateCurvesWithMasks(Context, Values);
-		HACK_ChannelMasks = nullptr;
+		EvaluateCurvesWithMasks(Context, ChannelMasks, Values);
 
 		FControlRigAnimTypeIDsPtr TypeIDs = FControlRigAnimTypeIDs::Get(Section->GetControlRig());
 
diff --git a/Engine/Plugins/Animation/ControlRig/Source/ControlRig/Public/Sequencer/MovieSceneControlRigParameterTemplate.h b/Engine/Plugins/Animation/ControlRig/Source/ControlRig/Public/Sequencer/MovieSceneControlRigParameterTemplate.h
--- a/Engine/Plugins/Animation/ControlRig/Source/ControlRig/Public/Sequencer/MovieSceneControlRigParameterTemplate.h
+++ b/Engine/Plugins/Animation/ControlRig/Source/ControlRig/Public/Sequencer/MovieSceneControlRigParameterTemplate.h
@@ -26,7 +26,7 @@ private:
 	virtual void Interrogate(const FMovieSceneContext& Context, FMovieSceneInterrogationData& Container, UObject* BindingOverride) const override;
 
 private:
-	void EvaluateCurvesWithMasks(const FMovieSceneContext& Context, FEvaluatedControlRigParameterSectionValues& Values) const;
+	void EvaluateCurvesWithMasks(const FMovieSceneContext& Context, struct FEvaluatedControlRigParameterSectionChannelMasks& InMasks, FEvaluatedControlRigParameterSectionValues& Values) const;
 
 protected:
 	/** The bool parameter names and their associated curves. */
