Add ability for ControlRigParameterTrack editing to happen on a UObject bound object, rather than only tied to USceneComponent
    - Helper functionality to iterate over tracks within MovieScene
    - Change functionality to operate on a UObject, and try to retrieve either a component or skeleton from it

diff --git a/Engine/Plugins/Animation/ControlRig/Source/ControlRigEditor/Private/Sequencer/ControlRigParameterTrackEditor.cpp b/Engine/Plugins/Animation/ControlRig/Source/ControlRigEditor/Private/Sequencer/ControlRigParameterTrackEditor.cpp
--- a/Engine/Plugins/Animation/ControlRig/Source/ControlRigEditor/Private/Sequencer/ControlRigParameterTrackEditor.cpp
+++ b/Engine/Plugins/Animation/ControlRig/Source/ControlRigEditor/Private/Sequencer/ControlRigParameterTrackEditor.cpp
@@ -245,7 +245,7 @@ FControlRigParameterTrackEditor::FControlRigParameterTrackEditor(TSharedRef<ISeq
 	{
 		//we check for two things, one if the control rig has been replaced if so we need to switch.
 		//the other is if bound object on the edit mode is null we request a re-evaluate which will reset it up.
-		FDelegateHandle OnObjectsReplacedHandle = FCoreUObjectDelegates::OnObjectsReplaced.AddLambda([this](const TMap<UObject*, UObject*>& ReplacementMap)
+		const FDelegateHandle OnObjectsReplacedHandle = FCoreUObjectDelegates::OnObjectsReplaced.AddLambda([this](const TMap<UObject*, UObject*>& ReplacementMap)
 		{
 			if (GetSequencer().IsValid())
 			{
@@ -299,14 +299,11 @@ FControlRigParameterTrackEditor::FControlRigParameterTrackEditor(TSharedRef<ISeq
 						GetSequencer()->EnterSilentMode();
 						bHasEnteredSilent = true;
 					}
-					UMovieScene* MovieScene = GetSequencer()->GetFocusedMovieSceneSequence()->GetMovieScene();
-					const TArray<FMovieSceneBinding>& Bindings = MovieScene->GetBindings();
-					for (const FMovieSceneBinding& Binding : Bindings)
+
+					IterateTracks([this, &OldToNewControlRigs, ControlRigEditMode, &bRequestEvaluate](UMovieSceneControlRigParameterTrack* Track)
 					{
-						UMovieSceneControlRigParameterTrack* Track = Cast<UMovieSceneControlRigParameterTrack>(MovieScene->FindTrack(UMovieSceneControlRigParameterTrack::StaticClass(), Binding.GetObjectGuid(), NAME_None));
-						if (Track && Track->GetControlRig())
+						if (UControlRig* OldControlRig = Track->GetControlRig())
 						{
-							UControlRig* OldControlRig = Track->GetControlRig();
 							UControlRig** NewControlRig = OldToNewControlRigs.Find(OldControlRig);
 							if (NewControlRig)
 							{
@@ -371,7 +368,10 @@ FControlRigParameterTrackEditor::FControlRigParameterTrackEditor(TSharedRef<ISeq
 								}
 							}
 						}
-					}
+
+						return false;
+					});
+
 					if (!ControlRigEditMode)
 					{
 						if (bRequestEvaluate)
@@ -398,16 +398,17 @@ FControlRigParameterTrackEditor::FControlRigParameterTrackEditor(TSharedRef<ISeq
 
 		AcquiredResources.Add([=] { FCoreUObjectDelegates::OnObjectsReplaced.Remove(OnObjectsReplacedHandle); });
 	}
-	//register all modified/selections for control rigs
-	const TArray<FMovieSceneBinding>& Bindings = MovieScene->GetBindings();
-	for (const FMovieSceneBinding& Binding : Bindings)
+
+	// Register all modified/selections for control rigs
+	IterateTracks([this](UMovieSceneControlRigParameterTrack* Track)
 	{
-		UMovieSceneControlRigParameterTrack* Track = Cast<UMovieSceneControlRigParameterTrack>(MovieScene->FindTrack(UMovieSceneControlRigParameterTrack::StaticClass(), Binding.GetObjectGuid(), NAME_None));
-		if (Track && Track->GetControlRig())
+		if (UControlRig* ControlRig = Track->GetControlRig())
 		{
-			BindControlRig(Track->GetControlRig());
+			BindControlRig(ControlRig);
 		}
-	}
+
+		return false;
+	});
 }
 
 FControlRigParameterTrackEditor::~FControlRigParameterTrackEditor()
@@ -1077,6 +1078,43 @@ void FControlRigParameterTrackEditor::BakeToControlRig(UClass* InClass, FGuid Ob
 	}
 }
 
+void FControlRigParameterTrackEditor::IterateTracks(TFunctionRef<bool(UMovieSceneControlRigParameterTrack*)> Callback) const
+{
+	UMovieScene* MovieScene = GetSequencer().IsValid() && GetSequencer()->GetFocusedMovieSceneSequence() ? GetSequencer()->GetFocusedMovieSceneSequence()->GetMovieScene() : nullptr;
+	if (!MovieScene)
+	{
+		return;
+	}
+
+	IterateTracksInMovieScene(*MovieScene, Callback);
+}
+
+void FControlRigParameterTrackEditor::IterateTracksInMovieScene(UMovieScene& MovieScene, TFunctionRef<bool(UMovieSceneControlRigParameterTrack*)> Callback) const
+{
+	TArray<UMovieSceneControlRigParameterTrack*> Tracks;
+	
+	const TArray<FMovieSceneBinding>& Bindings = MovieScene.GetBindings();
+	for (const FMovieSceneBinding& Binding : Bindings)
+	{
+		UMovieSceneControlRigParameterTrack* Track = Cast<UMovieSceneControlRigParameterTrack>(MovieScene.FindTrack(UMovieSceneControlRigParameterTrack::StaticClass(), Binding.GetObjectGuid(), NAME_None));
+		if (Callback(Track))
+		{
+			return;
+		}
+	}
+	
+	for (UMovieSceneTrack* Track : MovieScene.GetTracks())
+	{
+		if (UMovieSceneControlRigParameterTrack* CRTrack = Cast<UMovieSceneControlRigParameterTrack>(Track))
+		{			
+			if (Callback(CRTrack))
+			{
+				return;
+			}
+		}
+	}	
+}
+
 void FControlRigParameterTrackEditor::BakeInvertedPose(UControlRig* InControlRig, UMovieSceneControlRigParameterTrack* Track)
 {
 	if (!InControlRig->IsAdditive())
@@ -1570,14 +1608,14 @@ void FControlRigParameterTrackEditor::OnAddTransformKeysForSelectedObjects(EMovi
 		UControlRig* ControlRig = Selection.Key;
 		if (const TSharedPtr<IControlRigObjectBinding> ObjectBinding = ControlRig->GetObjectBinding())
 		{
-			if (USceneComponent* Component = Cast<USceneComponent>(ObjectBinding->GetBoundObject()))
+			if (UObject* Object = ObjectBinding->GetBoundObject())
 			{
 				const FName Name(*ControlRig->GetName());
 			
 				const TArray<FName> ControlNames = ControlRig->CurrentControlSelection();
 				for (const FName& ControlName : ControlNames)
 				{
-					AddControlKeys(Component, ControlRig, Name, ControlName, ChannelsToKey,
+					AddControlKeys(Object, ControlRig, Name, ControlName, ChannelsToKey,
 						ESequencerKeyMode::ManualKeyForced, FLT_MAX, bInConstraintSpace);
 				}
 			}
@@ -1935,23 +1973,21 @@ void FControlRigParameterTrackEditor::HandleActorAdded(AActor* Actor, FGuid Targ
 
 void FControlRigParameterTrackEditor::OnActivateSequenceChanged(FMovieSceneSequenceIDRef ID)
 {
-	UMovieScene* MovieScene = GetSequencer()->GetFocusedMovieSceneSequence()->GetMovieScene();
-	//register all modified/selections for control rigs
-	const TArray<FMovieSceneBinding>& Bindings = MovieScene->GetBindings();
-	for (const FMovieSceneBinding& Binding : Bindings)
+	IterateTracks([this](UMovieSceneControlRigParameterTrack* Track)
 	{
-		UMovieSceneControlRigParameterTrack* Track = Cast<UMovieSceneControlRigParameterTrack>(MovieScene->FindTrack(UMovieSceneControlRigParameterTrack::StaticClass(), Binding.GetObjectGuid(), NAME_None));
-		if (Track && Track->GetControlRig())
+		if (UControlRig* ControlRig = Track->GetControlRig())
 		{
-			BindControlRig(Track->GetControlRig());
+			BindControlRig(ControlRig);
 		}
-	}
+		
+		return false;
+	});
 }
 
 
 void FControlRigParameterTrackEditor::OnSequencerDataChanged(EMovieSceneDataChangeType DataChangeType)
 {
-	UMovieScene* MovieScene = GetSequencer()->GetFocusedMovieSceneSequence()->GetMovieScene();
+	const UMovieScene* MovieScene = GetSequencer()->GetFocusedMovieSceneSequence()->GetMovieScene();
 	FControlRigEditMode* ControlRigEditMode = GetEditMode();
 
 	//if we have a valid control rig edit mode need to check and see the control rig in that mode is still in a track
@@ -1960,30 +1996,25 @@ void FControlRigParameterTrackEditor::OnSequencerDataChanged(EMovieSceneDataChan
 		DataChangeType == EMovieSceneDataChangeType::Unknown))
 	{
 		TArray<UControlRig*> ControlRigs = ControlRigEditMode->GetControlRigsArray(false /*bIsVisible*/);
-		float FPS = 1.f / (float)GetSequencer()->GetFocusedDisplayRate().AsInterval();
+		const float FPS = (float)GetSequencer()->GetFocusedDisplayRate().AsDecimal();
 		for (UControlRig* ControlRig : ControlRigs)
 		{
 			if (ControlRig)
 			{
 				ControlRig->SetFramesPerSecond(FPS);
 
-		const TArray<FMovieSceneBinding>& Bindings = MovieScene->GetBindings();
 				bool bControlRigInTrack = false;
-		for (const FMovieSceneBinding& Binding : Bindings)
-		{
-			UMovieSceneControlRigParameterTrack* Track = Cast<UMovieSceneControlRigParameterTrack>(MovieScene->FindTrack(UMovieSceneControlRigParameterTrack::StaticClass(), Binding.GetObjectGuid(), NAME_None));
-					if (Track && Track->GetControlRig() == ControlRig)
-			{
+				IterateTracks([&bControlRigInTrack, ControlRig](const UMovieSceneControlRigParameterTrack* Track)
+				{
+					if(Track->GetControlRig() == ControlRig)
+			        {
 						bControlRigInTrack = true;
-						break;; //continue, we still have a good track
-			}
-		}
-				/* Nope don't do this anymore, todo mz
-		if (FEditorModeTools* Tools = GetEditorModeTools())
-		{
-			Tools->DeactivateMode(FControlRigEditMode::ModeName);
-		}
-				*/
+						return false;
+			        }
+
+					return true;
+				});
+
 				if (bControlRigInTrack == false)
 				{
 					ControlRigEditMode->RemoveControlRig(ControlRig);
@@ -1998,23 +2029,20 @@ void FControlRigParameterTrackEditor::PostEvaluation(UMovieScene* MovieScene, FF
 {
 	if (MovieScene)
 	{
-		const TArray<FMovieSceneBinding>& Bindings = MovieScene->GetBindings();
-		for (const FMovieSceneBinding& Binding : Bindings)
+		IterateTracksInMovieScene(*MovieScene, [this](const UMovieSceneControlRigParameterTrack* Track)
 		{
-			if (UMovieSceneControlRigParameterTrack* Track = Cast<UMovieSceneControlRigParameterTrack>(MovieScene->FindTrack(UMovieSceneControlRigParameterTrack::StaticClass(), Binding.GetObjectGuid(), NAME_None)))
-			{
-				if (UControlRig* ControlRig = Track->GetControlRig())
+			if (const UControlRig* ControlRig = Track->GetControlRig())
 				{
 					if (ControlRig->GetObjectBinding())
 					{
 						if (UControlRigComponent* ControlRigComponent = Cast<UControlRigComponent>(ControlRig->GetObjectBinding()->GetBoundObject()))
 						{
-							ControlRigComponent->Update(.1); //delta time doesn't matter.
-						}
+							ControlRigComponent->Update(.1f); //delta time doesn't matter.
 					}
 				}
 			}
-		}
+			return false;
+		});
 	}
 }
 
@@ -2298,18 +2326,18 @@ void FControlRigParameterTrackEditor::SelectRigsAndControls(UControlRig* Control
 }
 
 
-FMovieSceneTrackEditor::FFindOrCreateHandleResult FControlRigParameterTrackEditor::FindOrCreateHandleToSceneCompOrOwner(USceneComponent* InComp, UControlRig* InControlRig)
+FMovieSceneTrackEditor::FFindOrCreateHandleResult FControlRigParameterTrackEditor::FindOrCreateHandleToObject(UObject* InObj, UControlRig* InControlRig)
 {
 	const bool bCreateHandleIfMissing = false;
 	FName CreatedFolderName = NAME_None;
 
 	FFindOrCreateHandleResult Result;
-	bool bHandleWasValid = GetSequencer()->GetHandleToObject(InComp, bCreateHandleIfMissing).IsValid();
+	bool bHandleWasValid = GetSequencer()->GetHandleToObject(InObj, bCreateHandleIfMissing).IsValid();
 
-	Result.Handle = GetSequencer()->GetHandleToObject(InComp, bCreateHandleIfMissing, CreatedFolderName);
+	Result.Handle = GetSequencer()->GetHandleToObject(InObj, bCreateHandleIfMissing, CreatedFolderName);
 	Result.bWasCreated = bHandleWasValid == false && Result.Handle.IsValid();
 
-	UMovieScene* MovieScene = GetSequencer()->GetFocusedMovieSceneSequence()->GetMovieScene();
+	const UMovieScene* MovieScene = GetSequencer()->GetFocusedMovieSceneSequence()->GetMovieScene();
 
 	// Prioritize a control rig parameter track on this component if it matches the handle
 	if (Result.Handle.IsValid())
@@ -2324,29 +2352,32 @@ FMovieSceneTrackEditor::FFindOrCreateHandleResult FControlRigParameterTrackEdito
 	}
 
 	// If the owner has a control rig parameter track, let's use it
-	UObject* OwnerObject = InComp->GetOwner();
-	FGuid OwnerHandle = GetSequencer()->GetHandleToObject(OwnerObject, bCreateHandleIfMissing);
-	bHandleWasValid = OwnerHandle.IsValid();
-	if (OwnerHandle.IsValid())
-	{
-		if (UMovieSceneControlRigParameterTrack* Track = Cast<UMovieSceneControlRigParameterTrack>(MovieScene->FindTrack(UMovieSceneControlRigParameterTrack::StaticClass(), OwnerHandle, NAME_None)))
-		{
-			if (InControlRig == nullptr || (Track->GetControlRig() == InControlRig))
-			{
-				Result.Handle = OwnerHandle;
-				Result.bWasCreated = bHandleWasValid == false && Result.Handle.IsValid();
-				return Result;
-			}
+	if (const USceneComponent* SceneComponent = Cast<USceneComponent>(InObj))
+	{
+		// If the owner has a control rig parameter track, let's use it
+		UObject* OwnerObject = SceneComponent->GetOwner();
+		const FGuid OwnerHandle = GetSequencer()->GetHandleToObject(OwnerObject, bCreateHandleIfMissing);
+	    bHandleWasValid = OwnerHandle.IsValid();
+	    if (OwnerHandle.IsValid())
+	    {
+		    if (UMovieSceneControlRigParameterTrack* Track = Cast<UMovieSceneControlRigParameterTrack>(MovieScene->FindTrack(UMovieSceneControlRigParameterTrack::StaticClass(), OwnerHandle, NAME_None)))
+		    {
+			    if (InControlRig == nullptr || (Track->GetControlRig() == InControlRig))
+			    {
+				    Result.Handle = OwnerHandle;
+				    Result.bWasCreated = bHandleWasValid == false && Result.Handle.IsValid();
+				    return Result;
+			    }
+		    }
+	    }
+    
+	    // If the component handle doesn't exist, let's use the owner handle
+	    if (Result.Handle.IsValid() == false)
+	    {
+		    Result.Handle = OwnerHandle;
+		    Result.bWasCreated = bHandleWasValid == false && Result.Handle.IsValid();
 		}
 	}
-
-	// If the component handle doesn't exist, let's use the owner handle
-	if (Result.Handle.IsValid() == false)
-	{
-		Result.Handle = OwnerHandle;
-		Result.bWasCreated = bHandleWasValid == false && Result.Handle.IsValid();
-
-	}
 	return Result;
 }
 
@@ -2392,25 +2423,20 @@ FMovieSceneTrackEditor::FFindOrCreateTrackResult FControlRigParameterTrackEditor
 	FFindOrCreateTrackResult Result;
 	bool bTrackExisted = false;
 
-	UMovieScene* MovieScene = GetSequencer()->GetFocusedMovieSceneSequence()->GetMovieScene();
-
-	if (FMovieSceneBinding* Binding = MovieScene->FindBinding(ObjectHandle))
+	IterateTracks([&Result, &bTrackExisted, ControlRig](UMovieSceneControlRigParameterTrack* Track)
 	{
-		for (UMovieSceneTrack* Track : Binding->GetTracks())
+		if (Track->GetControlRig() == ControlRig)
 		{
-			if (UMovieSceneControlRigParameterTrack* ControlRigParameterTrack = Cast<UMovieSceneControlRigParameterTrack>(Track))
-			{
-				if (ControlRigParameterTrack->GetControlRig() == ControlRig)
-				{
-					Result.Track = ControlRigParameterTrack;
-					bTrackExisted = true;
-				}
-			}
+			Result.Track = Track;
+			bTrackExisted = true;
 		}
-	}
+		return false;
+	});
 
-	if (!Result.Track && bCreateTrackIfMissing)
+	// Only create track if the object handle is valid
+	if (!Result.Track && bCreateTrackIfMissing && ObjectHandle.IsValid())
 	{
+		UMovieScene* MovieScene = GetSequencer()->GetFocusedMovieSceneSequence()->GetMovieScene();
 		Result.Track = AddTrack(MovieScene, ObjectHandle, UMovieSceneControlRigParameterTrack::StaticClass(), PropertyName);
 	}
 
@@ -2426,24 +2452,18 @@ UMovieSceneControlRigParameterTrack* FControlRigParameterTrackEditor::FindTrack(
 		return nullptr;
 	}
 	
-	UMovieScene* MovieScene = GetSequencer()->GetFocusedMovieSceneSequence()->GetMovieScene();
-	if (MovieScene)
+	UMovieSceneControlRigParameterTrack* FoundTrack = nullptr;
+	IterateTracks([InControlRig, &FoundTrack](UMovieSceneControlRigParameterTrack* Track)
 	{
-		const TArray<FMovieSceneBinding>& Bindings = MovieScene->GetBindings();
-		for (const FMovieSceneBinding& Binding : Bindings)
+		if (Track->GetControlRig() == InControlRig)
 		{
-			TArray<UMovieSceneTrack*> Tracks = MovieScene->FindTracks(UMovieSceneControlRigParameterTrack::StaticClass(), Binding.GetObjectGuid(), NAME_None);
-			for (UMovieSceneTrack* AnyOleTrack : Tracks)
-			{
-				UMovieSceneControlRigParameterTrack* Track = Cast<UMovieSceneControlRigParameterTrack>(AnyOleTrack);
-				if (Track && Track->GetControlRig() == InControlRig)
-				{
-					return Track;
-				}
-			}
+			FoundTrack = Track;
+			return false;
 		}
-	}
-	return nullptr;
+		return true;
+	});
+
+	return FoundTrack;
 }
 
 void FControlRigParameterTrackEditor::HandleOnSpaceAdded(UMovieSceneControlRigParameterSection* Section, const FName& ControlName, FMovieSceneControlRigSpaceChannel* SpaceChannel)
@@ -2566,51 +2586,47 @@ void FControlRigParameterTrackEditor::ClearOutAllSpaceAndConstraintDelegates(con
 		return;
 	}
 
-	const TArray<FMovieSceneBinding>& Bindings = MovieScene->GetBindings();
-	for (const FMovieSceneBinding& Binding : Bindings)
+	IterateTracks([InOptionalControlRig](const UMovieSceneControlRigParameterTrack* Track)
 	{
-		const UMovieSceneTrack* Track = MovieScene->FindTrack(
-			UMovieSceneControlRigParameterTrack::StaticClass(), Binding.GetObjectGuid(), NAME_None);
-		if (const UMovieSceneControlRigParameterTrack* CRTrack = Cast<UMovieSceneControlRigParameterTrack>(Track))
+		if (InOptionalControlRig && Track->GetControlRig() != InOptionalControlRig)
 		{
-			if (InOptionalControlRig && CRTrack->GetControlRig() != InOptionalControlRig)
-			{
-				continue;
-			}
+			return true;
+		}
 				
-			for (UMovieSceneSection* Section : Track->GetAllSections())
+		for (UMovieSceneSection* Section : Track->GetAllSections())
+		{
+			if (UMovieSceneControlRigParameterSection* CRSection = Cast<UMovieSceneControlRigParameterSection>(Section))
 			{
-				if (UMovieSceneControlRigParameterSection* CRSection = Cast<UMovieSceneControlRigParameterSection>(Section))
+				// clear space channels
+				TArray<FSpaceControlNameAndChannel>& Channels = CRSection->GetSpaceChannels();
+				for (FSpaceControlNameAndChannel& SpaceAndChannel : Channels)
 				{
-					// clear space channels
-					TArray<FSpaceControlNameAndChannel>& Channels = CRSection->GetSpaceChannels();
-					for (FSpaceControlNameAndChannel& SpaceAndChannel : Channels)
-					{
-						SpaceAndChannel.SpaceCurve.OnKeyMovedEvent().Clear();
-						SpaceAndChannel.SpaceCurve.OnKeyDeletedEvent().Clear();
-					}
+					SpaceAndChannel.SpaceCurve.OnKeyMovedEvent().Clear();
+					SpaceAndChannel.SpaceCurve.OnKeyDeletedEvent().Clear();
+				}
 
-					// clear constraint channels
-					TArray<FConstraintAndActiveChannel>& ConstraintChannels = CRSection->GetConstraintsChannels();
-					for (FConstraintAndActiveChannel& Channel: ConstraintChannels)
-					{
-						Channel.ActiveChannel.OnKeyMovedEvent().Clear();
-						Channel.ActiveChannel.OnKeyDeletedEvent().Clear();
-					}
+				// clear constraint channels
+				TArray<FConstraintAndActiveChannel>& ConstraintChannels = CRSection->GetConstraintsChannels();
+				for (FConstraintAndActiveChannel& Channel: ConstraintChannels)
+				{
+					Channel.ActiveChannel.OnKeyMovedEvent().Clear();
+					Channel.ActiveChannel.OnKeyDeletedEvent().Clear();
+				}
 
-					if (CRSection->OnConstraintRemovedHandle.IsValid())
+				if (CRSection->OnConstraintRemovedHandle.IsValid())
+				{
+					if (const UControlRig* ControlRig = CRSection->GetControlRig())
 					{
-						if (const UControlRig* ControlRig = CRSection->GetControlRig())
-						{
-							FConstraintsManagerController& Controller = FConstraintsManagerController::Get(ControlRig->GetWorld());
-							Controller.GetNotifyDelegate().Remove(CRSection->OnConstraintRemovedHandle);
-							CRSection->OnConstraintRemovedHandle.Reset();
-						}
+						FConstraintsManagerController& Controller = FConstraintsManagerController::Get(ControlRig->GetWorld());
+						Controller.GetNotifyDelegate().Remove(CRSection->OnConstraintRemovedHandle);
+						CRSection->OnConstraintRemovedHandle.Reset();
 					}
 				}
 			}
 		}
-	}
+
+		return false;
+	});
 }
 
 namespace
@@ -2987,25 +3003,18 @@ void FControlRigParameterTrackEditor::HandleControlSelected(UControlRig* Subject
 
 	FControlRigEditMode* ControlRigEditMode = GetEditMode();
 
-	FName ControlRigName(*Subject->GetName());
+	const FName ControlRigName(*Subject->GetName());
 	if (TSharedPtr<IControlRigObjectBinding> ObjectBinding = Subject->GetObjectBinding())
 	{
-		UObject* ActorObject = nullptr;
-		USceneComponent* Component = Cast<USceneComponent>(ObjectBinding->GetBoundObject());
-		if (!Component)
-		{
-			return;
-		}
-		ActorObject = Component->GetOwner();
-		bool bCreateTrack = false;
-		FFindOrCreateHandleResult HandleResult = FindOrCreateHandleToSceneCompOrOwner(Component, Subject);
-		FGuid ObjectHandle = HandleResult.Handle;
-		if (!ObjectHandle.IsValid())
+		UObject* Object = ObjectBinding->GetBoundObject();
+		if (!Object)
 		{
 			return;
 		}
 
-		FFindOrCreateTrackResult TrackResult = FindOrCreateControlRigTrackForObject(ObjectHandle, Subject, ControlRigName, bCreateTrack);
+		const bool bCreateTrack = false;
+		const FFindOrCreateHandleResult HandleResult = FindOrCreateHandleToObject(Object, Subject);
+		FFindOrCreateTrackResult TrackResult = FindOrCreateControlRigTrackForObject(HandleResult.Handle, Subject, ControlRigName, bCreateTrack);
 		UMovieSceneControlRigParameterTrack* Track = CastChecked<UMovieSceneControlRigParameterTrack>(TrackResult.Track, ECastCheckedType::NullAllowed);
 		if (Track)
 		{
@@ -3039,8 +3048,6 @@ void FControlRigParameterTrackEditor::HandleControlSelected(UControlRig* Subject
 	}
 }
 
-
-
 void FControlRigParameterTrackEditor::HandleOnPostConstructed(UControlRig* Subject, const FName& InEventName)
 {
 	if (IsInGameThread())
@@ -3076,35 +3083,30 @@ void FControlRigParameterTrackEditor::HandleControlModified(UControlRig* Control
 	}
 	FTransform  Transform = ControlRig->GetControlLocalTransform(ControlElement->GetFName());
 	UMovieScene* MovieScene = GetSequencer()->GetFocusedMovieSceneSequence()->GetMovieScene();
-	const TArray<FMovieSceneBinding>& Bindings = MovieScene->GetBindings();
-	for (const FMovieSceneBinding& Binding : Bindings)
+
+	IterateTracks([this, ControlRig, ControlElement, Context](UMovieSceneControlRigParameterTrack* Track)
 	{
-		TArray<UMovieSceneTrack*> Tracks = MovieScene->FindTracks(UMovieSceneControlRigParameterTrack::StaticClass(), Binding.GetObjectGuid(), NAME_None);
-		for (UMovieSceneTrack* BaseTrack : Tracks)
+		if (Track && Track->GetControlRig() == ControlRig)
 		{
-			UMovieSceneControlRigParameterTrack* Track = Cast<UMovieSceneControlRigParameterTrack>(BaseTrack);
-			if (Track && Track->GetControlRig() == ControlRig)
+			FName Name(*ControlRig->GetName());
+			if (TSharedPtr<IControlRigObjectBinding> ObjectBinding = ControlRig->GetObjectBinding())
 			{
-				FName Name(*ControlRig->GetName());
-				if (TSharedPtr<IControlRigObjectBinding> ObjectBinding = ControlRig->GetObjectBinding())
+				if (UObject* Object = ObjectBinding->GetBoundObject())
 				{
-					USceneComponent* Component = Cast<USceneComponent>(ObjectBinding->GetBoundObject());
-					if (Component)
+					ESequencerKeyMode KeyMode = ESequencerKeyMode::AutoKey;
+					if (Context.SetKey == EControlRigSetKey::Always)
 					{
-						ESequencerKeyMode KeyMode = ESequencerKeyMode::AutoKey;
-						if (Context.SetKey == EControlRigSetKey::Always)
-						{
-							KeyMode = ESequencerKeyMode::ManualKeyForced;
-						}
-						AddControlKeys(Component, ControlRig, Name, ControlElement->GetFName(), (EControlRigContextChannelToKey)Context.KeyMask, 
-							KeyMode, Context.LocalTime);
-						ControlChangedDuringUndoBracket++;
+						KeyMode = ESequencerKeyMode::ManualKeyForced;
 					}
+					AddControlKeys(Object, ControlRig, Name, ControlElement->GetFName(), (EControlRigContextChannelToKey)Context.KeyMask, 
+						KeyMode, Context.LocalTime);
+					ControlChangedDuringUndoBracket++;
+					return true;
 				}
-				break;
 			}
 		}
-	}
+		return false;
+	});
 }
 
 void FControlRigParameterTrackEditor::HandleControlUndoBracket(UControlRig* Subject, bool bOpenUndoBracket)
@@ -3489,7 +3491,7 @@ void FControlRigParameterTrackEditor::GetControlRigKeys(
 	}
 }
 
-FKeyPropertyResult FControlRigParameterTrackEditor::AddKeysToControlRigHandle(USceneComponent* InSceneComp, UControlRig* InControlRig,
+FKeyPropertyResult FControlRigParameterTrackEditor::AddKeysToControlRigHandle(UObject* InObject, UControlRig* InControlRig,
 	FGuid ObjectHandle, FFrameNumber KeyTime, FFrameNumber EvaluateTime, FGeneratedTrackKeys& GeneratedKeys,
 	ESequencerKeyMode KeyMode, TSubclassOf<UMovieSceneTrack> TrackClass, FName ControlRigName, FName RigControlName)
 {
@@ -3538,9 +3540,9 @@ FKeyPropertyResult FControlRigParameterTrackEditor::AddKeysToControlRigHandle(US
 		{
 			if (!bTrackCreated)
 			{
-				ModifyOurGeneratedKeysByCurrentAndWeight(InSceneComp, InControlRig, RigControlName, Track, SectionToKey, EvaluateTime, GeneratedKeys, Weight);
+				ModifyOurGeneratedKeysByCurrentAndWeight(InObject, InControlRig, RigControlName, Track, SectionToKey, EvaluateTime, GeneratedKeys, Weight);
 			}
-			UMovieSceneControlRigParameterSection* ParamSection = Cast<UMovieSceneControlRigParameterSection>(SectionToKey);
+			const UMovieSceneControlRigParameterSection* ParamSection = Cast<UMovieSceneControlRigParameterSection>(SectionToKey);
 			if (!ParamSection->GetDoNotKey())
 			{
 				KeyPropertyResult |= AddKeysToSection(SectionToKey, KeyTime, GeneratedKeys, KeyMode, EKeyFrameTrackEditorSetDefault::SetDefaultOnAddKeys);
@@ -3572,7 +3574,7 @@ FKeyPropertyResult FControlRigParameterTrackEditor::AddKeysToControlRigHandle(US
 }
 
 FKeyPropertyResult FControlRigParameterTrackEditor::AddKeysToControlRig(
-	USceneComponent* InSceneComp, UControlRig* InControlRig, FFrameNumber KeyTime, FFrameNumber EvaluateTime, FGeneratedTrackKeys& GeneratedKeys,
+	UObject* InObject, UControlRig* InControlRig, FFrameNumber KeyTime, FFrameNumber EvaluateTime, FGeneratedTrackKeys& GeneratedKeys,
 	ESequencerKeyMode KeyMode, TSubclassOf<UMovieSceneTrack> TrackClass, FName ControlRigName, FName RigControlName)
 {
 	FKeyPropertyResult KeyPropertyResult;
@@ -3584,19 +3586,16 @@ FKeyPropertyResult FControlRigParameterTrackEditor::AddKeysToControlRig(
 		KeyMode == ESequencerKeyMode::ManualKeyForced ||
 		AllowEditsMode == EAllowEditsMode::AllowSequencerEditsOnly;
 
-	FFindOrCreateHandleResult HandleResult = FindOrCreateHandleToSceneCompOrOwner(InSceneComp, InControlRig);
+	FFindOrCreateHandleResult HandleResult = FindOrCreateHandleToObject(InObject, InControlRig);
 	FGuid ObjectHandle = HandleResult.Handle;
 	KeyPropertyResult.bHandleCreated = HandleResult.bWasCreated;
-	if (ObjectHandle.IsValid())
-	{
-		KeyPropertyResult |= AddKeysToControlRigHandle(InSceneComp, InControlRig, ObjectHandle, KeyTime, EvaluateTime, GeneratedKeys, KeyMode, TrackClass, ControlRigName, RigControlName);
-	}
+	KeyPropertyResult |= AddKeysToControlRigHandle(InObject, InControlRig, ObjectHandle, KeyTime, EvaluateTime, GeneratedKeys, KeyMode, TrackClass, ControlRigName, RigControlName);
 
 	return KeyPropertyResult;
 }
 
 void FControlRigParameterTrackEditor::AddControlKeys(
-	USceneComponent* InSceneComp,
+	UObject* InObject,
 	UControlRig* InControlRig,
 	FName ControlRigName,
 	FName RigControlName,
@@ -3611,12 +3610,8 @@ void FControlRigParameterTrackEditor::AddControlKeys(
 	}
 	bool bCreateTrack = false;
 	bool bCreateHandle = false;
-	FFindOrCreateHandleResult HandleResult = FindOrCreateHandleToSceneCompOrOwner(InSceneComp, InControlRig);
+	FFindOrCreateHandleResult HandleResult = FindOrCreateHandleToObject(InObject, InControlRig);
 	FGuid ObjectHandle = HandleResult.Handle;
-	if (!ObjectHandle.IsValid())
-	{
-		return;
-	}
 	FFindOrCreateTrackResult TrackResult = FindOrCreateControlRigTrackForObject(ObjectHandle, InControlRig, ControlRigName, bCreateTrack);
 	UMovieSceneControlRigParameterTrack* Track = CastChecked<UMovieSceneControlRigParameterTrack>(TrackResult.Track, ECastCheckedType::NullAllowed);
 	UMovieSceneControlRigParameterSection* ParamSection = nullptr;
@@ -3629,7 +3624,7 @@ void FControlRigParameterTrackEditor::AddControlKeys(
 			FFrameTime LocalFrameTime = GetSequencer()->GetFocusedTickResolution().AsFrameTime((double)InLocalTime);
 			BeginKeying(LocalFrameTime.RoundToFrame());
 		}
-		FFrameNumber  FrameTime = GetTimeForKey();
+		const FFrameNumber FrameTime = GetTimeForKey();
 		UMovieSceneSection* Section = Track->FindSection(FrameTime);
 		ParamSection = Cast<UMovieSceneControlRigParameterSection>(Section);
 
@@ -3663,7 +3658,7 @@ void FControlRigParameterTrackEditor::AddControlKeys(
 			EvaluateTime = LocalTime;
 		}
 		
-		return this->AddKeysToControlRig(InSceneComp, InControlRig, LocalTime, EvaluateTime, *GeneratedKeys, KeyMode, UMovieSceneControlRigParameterTrack::StaticClass(), ControlRigName, RigControlName);
+		return this->AddKeysToControlRig(InObject, InControlRig, LocalTime, EvaluateTime, *GeneratedKeys, KeyMode, UMovieSceneControlRigParameterTrack::StaticClass(), ControlRigName, RigControlName);
 	};
 
 	AnimatablePropertyChanged(FOnKeyProperty::CreateLambda(OnKeyProperty));
diff --git a/Engine/Plugins/Animation/ControlRig/Source/ControlRigEditor/Private/Sequencer/ControlRigParameterTrackEditor.h b/Engine/Plugins/Animation/ControlRig/Source/ControlRigEditor/Private/Sequencer/ControlRigParameterTrackEditor.h
--- a/Engine/Plugins/Animation/ControlRig/Source/ControlRigEditor/Private/Sequencer/ControlRigParameterTrackEditor.h
+++ b/Engine/Plugins/Animation/ControlRig/Source/ControlRigEditor/Private/Sequencer/ControlRigParameterTrackEditor.h
@@ -168,8 +168,8 @@ private:
 	/** Select control rig if not selected, select controls from key areas */
 	void SelectRigsAndControls(UControlRig* Subject, const TArray<const IKeyArea*>& KeyAreas);
 
-	/** Handle Creation for SkelMeshComp or Actor Owner, either may have a binding, use optional control rig to pick one*/
-	FMovieSceneTrackEditor::FFindOrCreateHandleResult FindOrCreateHandleToSceneCompOrOwner(USceneComponent* InComp, UControlRig* InControlRig = nullptr);
+	/** Handle Creation for Skeleton, SkelMeshComp or Actor Owner, they may or may not have a binding, use optional control rig to pick one */
+	FMovieSceneTrackEditor::FFindOrCreateHandleResult FindOrCreateHandleToObject(UObject* InObj, UControlRig* InControlRig = nullptr);
 
 	/** Handle Creation for control rig track given the object binding and the control rig */	
 	FMovieSceneTrackEditor::FFindOrCreateTrackResult FindOrCreateControlRigTrackForObject(FGuid ObjectBinding, UControlRig* ControlRig, FName PropertyName = NAME_None, bool bCreateTrackIfMissing = true);
@@ -202,24 +202,26 @@ private:
 	/** Set Up EditMode for Specified Control Rig*/
 	void SetUpEditModeIfNeeded(UControlRig* ControlRig);
 
+	/** Helper functions to iterate over UMovieSceneControlRigParameterTracks in the currently focussed MovieScene*/
+	void IterateTracks(TFunctionRef<bool(UMovieSceneControlRigParameterTrack*)> Callback) const;
+	void IterateTracksInMovieScene(UMovieScene& MovieScene, TFunctionRef<bool(UMovieSceneControlRigParameterTrack*)> Callback) const; 
 private:
-
 	/** Command Bindings added by the Transform Track Editor to Sequencer and curve editor. */
 	TSharedPtr<FUICommandList> CommandBindings;
 	FAcquiredResources AcquiredResources;
 
 public:
 
-	void AddControlKeys(USceneComponent *InSceneComp, UControlRig* InControlRig, FName PropertyName,
+	void AddControlKeys(UObject* InObject, UControlRig* InControlRig, FName PropertyName,
 		FName ParameterName, EControlRigContextChannelToKey ChannelsToKey, ESequencerKeyMode KeyMode,
 		float InLocalTime, const bool bInConstraintSpace = false);
 	void GetControlRigKeys(UControlRig* InControlRig, FName ParameterName, EControlRigContextChannelToKey ChannelsToKey,
 		ESequencerKeyMode KeyMode, UMovieSceneControlRigParameterSection* SectionToKey,	FGeneratedTrackKeys& OutGeneratedKeys,
 		const bool bInConstraintSpace = false);
 	FKeyPropertyResult AddKeysToControlRig(
-		USceneComponent *InSceneComp, UControlRig* InControlRig, FFrameNumber KeyTime, FFrameNumber EvaluateTime, FGeneratedTrackKeys& GeneratedKeys,
+		UObject* InObject, UControlRig* InControlRig, FFrameNumber KeyTime, FFrameNumber EvaluateTime, FGeneratedTrackKeys& GeneratedKeys,
 		ESequencerKeyMode KeyMode, TSubclassOf<UMovieSceneTrack> TrackClass, FName ControlRigName, FName RigControlName);
-	FKeyPropertyResult AddKeysToControlRigHandle(USceneComponent *InSceneComp, UControlRig* InControlRig,
+	FKeyPropertyResult AddKeysToControlRigHandle(UObject* InObject, UControlRig* InControlRig,
 		FGuid ObjectHandle, FFrameNumber KeyTime, FFrameNumber EvaluateTime, FGeneratedTrackKeys& GeneratedKeys,
 		ESequencerKeyMode KeyMode, TSubclassOf<UMovieSceneTrack> TrackClass, FName ControlRigName, FName RigControlName);
 	/**
