Control Rig: Fix ensure condition when parallel calls to GetAssetUserDataArray on Skeleton
#rnx

diff --git a/Engine/Plugins/Animation/ControlRig/Source/ControlRig/Private/ControlRig.cpp b/Engine/Plugins/Animation/ControlRig/Source/ControlRig/Private/ControlRig.cpp
--- a/Engine/Plugins/Animation/ControlRig/Source/ControlRig/Private/ControlRig.cpp
+++ b/Engine/Plugins/Animation/ControlRig/Source/ControlRig/Private/ControlRig.cpp
@@ -2783,6 +2783,7 @@ const TArray<UAssetUserData*>* UControlRig::GetAssetUserDataArray() const
 	if(HasAnyFlags(RF_ClassDefaultObject))
 	{
 #if WITH_EDITOR
+		static thread_local TArray<TObjectPtr<UAssetUserData>> CachedAssetUserData;
 		CachedAssetUserData.Reset();
 		CachedAssetUserData.Append(AssetUserData);
 		CachedAssetUserData.Append(AssetUserDataEditorOnly);
@@ -2860,8 +2861,7 @@ const TArray<UAssetUserData*>* UControlRig::GetAssetUserDataArray() const
 				{
 					if(const TArray<UAssetUserData*>* SkeletonUserDataArray = Skeleton->GetAssetUserDataArray())
 					{
-						const TArray<UAssetUserData*> LocalSkeletonUserDataArray = *SkeletonUserDataArray;
-						for(UAssetUserData* SkeletonUserData : LocalSkeletonUserDataArray)
+						for(UAssetUserData* SkeletonUserData : *SkeletonUserDataArray)
 						{
 							CombinedAssetUserData.Add(SkeletonUserData);
 						}
@@ -2869,8 +2869,7 @@ const TArray<UAssetUserData*>* UControlRig::GetAssetUserDataArray() const
 				}
 				if(const TArray<UAssetUserData*>* SkeletalMeshUserDataArray = SkeletalMesh->GetAssetUserDataArray())
 				{
-					const TArray<UAssetUserData*> LocalSkeletalMeshUserDataArray = *SkeletalMeshUserDataArray;
-					for(UAssetUserData* SkeletalMeshUserData : LocalSkeletalMeshUserDataArray)
+					for(UAssetUserData* SkeletalMeshUserData : *SkeletalMeshUserDataArray)
 					{
 						CombinedAssetUserData.Add(SkeletalMeshUserData);
 					}
@@ -2878,8 +2877,7 @@ const TArray<UAssetUserData*>* UControlRig::GetAssetUserDataArray() const
 			}
 			if (const TArray<UAssetUserData*>* ActorComponentUserDataArray = SkeletalMeshComponent->GetAssetUserDataArray())
 			{
-				const TArray<UAssetUserData*> LocalActorComponentUserDataArray = *ActorComponentUserDataArray;
-				for(UAssetUserData* ActorComponentUserData : LocalActorComponentUserDataArray)
+				for(UAssetUserData* ActorComponentUserData : *ActorComponentUserDataArray)
 				{
 					CombinedAssetUserData.Add(ActorComponentUserData);
 				}
diff --git a/Engine/Plugins/Runtime/RigVM/Source/RigVM/Private/RigVMHost.cpp b/Engine/Plugins/Runtime/RigVM/Source/RigVM/Private/RigVMHost.cpp
--- a/Engine/Plugins/Runtime/RigVM/Source/RigVM/Private/RigVMHost.cpp
+++ b/Engine/Plugins/Runtime/RigVM/Source/RigVM/Private/RigVMHost.cpp
@@ -1441,6 +1441,7 @@ void URigVMHost::RemoveUserDataOfClass(TSubclassOf<UAssetUserData> InUserDataCla
 const TArray<UAssetUserData*>* URigVMHost::GetAssetUserDataArray() const
 {
 #if WITH_EDITOR
+	static thread_local TArray<TObjectPtr<UAssetUserData>> CachedAssetUserData;
 	CachedAssetUserData.Reset();
 	CachedAssetUserData.Append(AssetUserData);
 	CachedAssetUserData.Append(AssetUserDataEditorOnly);
diff --git a/Engine/Plugins/Runtime/RigVM/Source/RigVM/Public/RigVMHost.h b/Engine/Plugins/Runtime/RigVM/Source/RigVM/Public/RigVMHost.h
--- a/Engine/Plugins/Runtime/RigVM/Source/RigVM/Public/RigVMHost.h
+++ b/Engine/Plugins/Runtime/RigVM/Source/RigVM/Public/RigVMHost.h
@@ -403,7 +403,6 @@ protected:
 	TArray<TObjectPtr<UAssetUserData>> AssetUserDataEditorOnly;
 #endif
 
-	mutable TArray<TObjectPtr<UAssetUserData>> CachedAssetUserData;
 
 
 protected:
diff --git a/Engine/Source/Runtime/Engine/Classes/Animation/Skeleton.h b/Engine/Source/Runtime/Engine/Classes/Animation/Skeleton.h
--- a/Engine/Source/Runtime/Engine/Classes/Animation/Skeleton.h
+++ b/Engine/Source/Runtime/Engine/Classes/Animation/Skeleton.h
@@ -1034,7 +1034,6 @@ protected:
 	TArray<TObjectPtr<UAssetUserData>> AssetUserDataEditorOnly;
 #endif
 
-	mutable TArray<TObjectPtr<UAssetUserData>> CachedAssetUserData;
 
 	friend struct FReferenceSkeletonModifier;
 	friend class FEditableSkeleton;
diff --git a/Engine/Source/Runtime/Engine/Classes/Components/ActorComponent.h b/Engine/Source/Runtime/Engine/Classes/Components/ActorComponent.h
--- a/Engine/Source/Runtime/Engine/Classes/Components/ActorComponent.h
+++ b/Engine/Source/Runtime/Engine/Classes/Components/ActorComponent.h
@@ -163,7 +163,6 @@ protected:
 	TArray<TObjectPtr<UAssetUserData>> AssetUserDataEditorOnly;
 #endif
 
-	mutable TArray<TObjectPtr<UAssetUserData>> CachedAssetUserData;
 
 private:
 	/** Used for fast removal of end of frame update */
diff --git a/Engine/Source/Runtime/Engine/Classes/Engine/SkeletalMesh.h b/Engine/Source/Runtime/Engine/Classes/Engine/SkeletalMesh.h
--- a/Engine/Source/Runtime/Engine/Classes/Engine/SkeletalMesh.h
+++ b/Engine/Source/Runtime/Engine/Classes/Engine/SkeletalMesh.h
@@ -2231,7 +2231,6 @@ protected:
 	FOnMeshChanged OnMeshChanged;
 #endif
 
-	mutable TArray<TObjectPtr<UAssetUserData>> CachedAssetUserData;
 
 	friend struct FSkeletalMeshUpdateContext;
 	friend class FSkeletalMeshUpdate;
diff --git a/Engine/Source/Runtime/Engine/Private/Animation/Skeleton.cpp b/Engine/Source/Runtime/Engine/Private/Animation/Skeleton.cpp
--- a/Engine/Source/Runtime/Engine/Private/Animation/Skeleton.cpp
+++ b/Engine/Source/Runtime/Engine/Private/Animation/Skeleton.cpp
@@ -2121,6 +2121,7 @@ void USkeleton::RemoveUserDataOfClass(TSubclassOf<UAssetUserData> InUserDataClas
 const TArray<UAssetUserData*>* USkeleton::GetAssetUserDataArray() const
 {
 #if WITH_EDITOR
+	static thread_local TArray<TObjectPtr<UAssetUserData>> CachedAssetUserData;
 	CachedAssetUserData.Reset();
 	CachedAssetUserData.Append(AssetUserData);
 	CachedAssetUserData.Append(AssetUserDataEditorOnly);
diff --git a/Engine/Source/Runtime/Engine/Private/Components/ActorComponent.cpp b/Engine/Source/Runtime/Engine/Private/Components/ActorComponent.cpp
--- a/Engine/Source/Runtime/Engine/Private/Components/ActorComponent.cpp
+++ b/Engine/Source/Runtime/Engine/Private/Components/ActorComponent.cpp
@@ -2088,6 +2088,7 @@ UAssetUserData* UActorComponent::GetAssetUserDataOfClass(TSubclassOf<UAssetUserD
 const TArray<UAssetUserData*>* UActorComponent::GetAssetUserDataArray() const
 {
 #if WITH_EDITOR
+	static thread_local TArray<TObjectPtr<UAssetUserData>> CachedAssetUserData;
 	CachedAssetUserData.Reset();
 	CachedAssetUserData.Append(AssetUserData);
 	CachedAssetUserData.Append(AssetUserDataEditorOnly);
diff --git a/Engine/Source/Runtime/Engine/Private/SkeletalMesh.cpp b/Engine/Source/Runtime/Engine/Private/SkeletalMesh.cpp
--- a/Engine/Source/Runtime/Engine/Private/SkeletalMesh.cpp
+++ b/Engine/Source/Runtime/Engine/Private/SkeletalMesh.cpp
@@ -4768,6 +4768,7 @@ void USkeletalMesh::RemoveUserDataOfClass(TSubclassOf<UAssetUserData> InUserData
 const TArray<UAssetUserData*>* USkeletalMesh::GetAssetUserDataArray() const
 {
 #if WITH_EDITOR
+	static thread_local TArray<TObjectPtr<UAssetUserData>> CachedAssetUserData;
 	CachedAssetUserData.Reset();
 	CachedAssetUserData.Append(AssetUserData);
 	CachedAssetUserData.Append(AssetUserDataEditorOnly);
