Sequener: Control Rig: Missing nullptr check. Getting this in asap to unblock folks.

diff --git a/Engine/Plugins/Animation/ControlRig/Source/ControlRigEditor/Private/Sequencer/ControlRigParameterTrackEditor.cpp b/Engine/Plugins/Animation/ControlRig/Source/ControlRigEditor/Private/Sequencer/ControlRigParameterTrackEditor.cpp
--- a/Engine/Plugins/Animation/ControlRig/Source/ControlRigEditor/Private/Sequencer/ControlRigParameterTrackEditor.cpp
+++ b/Engine/Plugins/Animation/ControlRig/Source/ControlRigEditor/Private/Sequencer/ControlRigParameterTrackEditor.cpp
@@ -1096,10 +1096,12 @@ void FControlRigParameterTrackEditor::IterateTracksInMovieScene(UMovieScene& Mov
 	const TArray<FMovieSceneBinding>& Bindings = MovieScene.GetBindings();
 	for (const FMovieSceneBinding& Binding : Bindings)
 	{
-		UMovieSceneControlRigParameterTrack* Track = Cast<UMovieSceneControlRigParameterTrack>(MovieScene.FindTrack(UMovieSceneControlRigParameterTrack::StaticClass(), Binding.GetObjectGuid(), NAME_None));
-		if (Callback(Track))
+		if (UMovieSceneControlRigParameterTrack* Track = Cast<UMovieSceneControlRigParameterTrack>(MovieScene.FindTrack(UMovieSceneControlRigParameterTrack::StaticClass(), Binding.GetObjectGuid(), NAME_None)))
 		{
-			return;
+			if (Callback(Track))
+			{
+				return;
+			}
 		}
 	}
 	
