Control Rig: Introducing modular rigging related nodes

diff --git a/Engine/Plugins/Animation/ControlRig/Source/ControlRig/Private/ControlRig.cpp b/Engine/Plugins/Animation/ControlRig/Source/ControlRig/Private/ControlRig.cpp
--- a/Engine/Plugins/Animation/ControlRig/Source/ControlRig/Private/ControlRig.cpp
+++ b/Engine/Plugins/Animation/ControlRig/Source/ControlRig/Private/ControlRig.cpp
@@ -1406,8 +1406,7 @@ bool UControlRig::Execute_Internal(const FName& InEventName)
 
 		// setup the module information
 		FControlRigExecuteContext& PublicContext = Context.GetPublicDataSafe<FControlRigExecuteContext>();
-		TGuardValue<FName> ModuleInstanceNameGuard(PublicContext.ModuleInstanceName, GetModuleInstanceName());
-		TGuardValue<FString> ModuleInstanceNameSpaceGuard(PublicContext.ModuleInstanceNameSpace, GetRigModuleNameSpace());
+		FControlRigExecuteContextRigModuleGuard RigModuleGuard(PublicContext, this);
 		FRigHierarchyRedirectorGuard ElementRedirectorGuard(this);
 
 		TArray<FRigVMMemoryStorageStruct*> LocalMemory = VM->GetLocalMemoryArray(Context);
@@ -1661,15 +1660,6 @@ UControlRig* UControlRig::GetParentRig() const
 	return GetTypedOuter<UControlRig>();
 }
 
-FName UControlRig::GetModuleInstanceName() const
-{
-	if(IsRigModuleInstance())
-	{
-		return GetFName();
-	}
-	return NAME_None;
-}
-
 const FString& UControlRig::GetRigModuleNameSpace() const
 {
 	if(IsRigModule())
diff --git a/Engine/Plugins/Animation/ControlRig/Source/ControlRig/Private/Rigs/RigHierarchy.cpp b/Engine/Plugins/Animation/ControlRig/Source/ControlRig/Private/Rigs/RigHierarchy.cpp
--- a/Engine/Plugins/Animation/ControlRig/Source/ControlRig/Private/Rigs/RigHierarchy.cpp
+++ b/Engine/Plugins/Animation/ControlRig/Source/ControlRig/Private/Rigs/RigHierarchy.cpp
@@ -1259,9 +1259,9 @@ FRigName URigHierarchy::GetSafeNewName(const FRigName& InPotentialNewName, ERigE
 	if(ExecuteContext)
 	{
 		const FControlRigExecuteContext& CRContext = ExecuteContext->GetPublicData<FControlRigExecuteContext>();
-		if(!CRContext.ModuleInstanceNameSpace.IsEmpty())
+		if(CRContext.IsRigModule())
 		{
-			SanitizedName = CRContext.ModuleInstanceNameSpace + SanitizedName.GetName();
+			SanitizedName = CRContext.GetRigModuleNameSpace() + SanitizedName.GetName();
 			bAllowNameSpaceWhenSanitizingName = true;
 		}
 	}
diff --git a/Engine/Plugins/Animation/ControlRig/Source/ControlRig/Private/Rigs/RigHierarchyController.cpp b/Engine/Plugins/Animation/ControlRig/Source/ControlRig/Private/Rigs/RigHierarchyController.cpp
--- a/Engine/Plugins/Animation/ControlRig/Source/ControlRig/Private/Rigs/RigHierarchyController.cpp
+++ b/Engine/Plugins/Animation/ControlRig/Source/ControlRig/Private/Rigs/RigHierarchyController.cpp
@@ -225,7 +225,7 @@ FRigElementKey URigHierarchyController::AddBone(FName InName, FRigElementKey InP
 	{
 		TGuardValue<bool> DisableCacheValidityChecks(Hierarchy->bEnableCacheValidityCheck, false);
 		NewElement->Key.Type = ERigElementType::Bone;
-		NewElement->Key.Name = Hierarchy->GetSafeNewName(InName, NewElement->Key.Type);
+		NewElement->Key.Name = GetSafeNewName(InName, NewElement->Key.Type);
 		NewElement->BoneType = InBoneType;
 		AddElement(NewElement, Hierarchy->Get(Hierarchy->GetIndex(InParent)), true, InName);
 
@@ -287,7 +287,7 @@ FRigElementKey URigHierarchyController::AddNull(FName InName, FRigElementKey InP
 	{
 		TGuardValue<bool> DisableCacheValidityChecks(Hierarchy->bEnableCacheValidityCheck, false);		
 		NewElement->Key.Type = ERigElementType::Null;
-		NewElement->Key.Name = Hierarchy->GetSafeNewName(InName, NewElement->Key.Type);
+		NewElement->Key.Name = GetSafeNewName(InName, NewElement->Key.Type);
 		AddElement(NewElement, Hierarchy->Get(Hierarchy->GetIndex(InParent)), false, InName);
 
 		if(bTransformInGlobal)
@@ -352,7 +352,7 @@ FRigElementKey URigHierarchyController::AddControl(
 	{
 		TGuardValue<bool> DisableCacheValidityChecks(Hierarchy->bEnableCacheValidityCheck, false);		
 		NewElement->Key.Type = ERigElementType::Control;
-		NewElement->Key.Name = Hierarchy->GetSafeNewName(InName, NewElement->Key.Type);
+		NewElement->Key.Name = GetSafeNewName(InName, NewElement->Key.Type);
 		NewElement->Settings = InSettings;
 		if(NewElement->Settings.LimitEnabled.IsEmpty())
 		{
@@ -366,7 +366,7 @@ FRigElementKey URigHierarchyController::AddControl(
 		else if(Hierarchy->HasExecuteContext())
 		{
 			const FControlRigExecuteContext& CRContext = Hierarchy->ExecuteContext->GetPublicData<FControlRigExecuteContext>();
-			if(!CRContext.ModuleInstanceNameSpace.IsEmpty())
+			if(!CRContext.GetRigModuleNameSpace().IsEmpty())
 			{
 				NewElement->Settings.DisplayName = Hierarchy->GetSafeNewDisplayName(InParent, InName);
 			}
@@ -454,7 +454,7 @@ FRigElementKey URigHierarchyController::AddCurve(FName InName, float InValue, bo
 	{
 		TGuardValue<bool> DisableCacheValidityChecks(Hierarchy->bEnableCacheValidityCheck, false);		
 		NewElement->Key.Type = ERigElementType::Curve;
-		NewElement->Key.Name = Hierarchy->GetSafeNewName(InName, NewElement->Key.Type);
+		NewElement->Key.Name = GetSafeNewName(InName, NewElement->Key.Type);
 		NewElement->Value = InValue;
 		AddElement(NewElement, nullptr, false, InName);
 	}
@@ -503,7 +503,7 @@ FRigElementKey URigHierarchyController::AddRigidBody(FName InName, FRigElementKe
 	{
 		TGuardValue<bool> DisableCacheValidityChecks(Hierarchy->bEnableCacheValidityCheck, false);		
 		NewElement->Key.Type = ERigElementType::RigidBody;
-		NewElement->Key.Name = Hierarchy->GetSafeNewName(InName, NewElement->Key.Type);
+		NewElement->Key.Name = GetSafeNewName(InName, NewElement->Key.Type);
 		NewElement->Settings = InSettings;
 		AddElement(NewElement, Hierarchy->Get(Hierarchy->GetIndex(InParent)), true, InName);
 
@@ -555,7 +555,7 @@ FRigElementKey URigHierarchyController::AddReference(FName InName, FRigElementKe
 	{
 		TGuardValue<bool> DisableCacheValidityChecks(Hierarchy->bEnableCacheValidityCheck, false);		
 		NewElement->Key.Type = ERigElementType::Reference;
-		NewElement->Key.Name = Hierarchy->GetSafeNewName(InName, NewElement->Key.Type);
+		NewElement->Key.Name = GetSafeNewName(InName, NewElement->Key.Type);
 		NewElement->GetWorldTransformDelegate = InDelegate;
 		AddElement(NewElement, Hierarchy->Get(Hierarchy->GetIndex(InParent)), true, InName);
 
@@ -593,7 +593,7 @@ FRigElementKey URigHierarchyController::AddConnector(FName InName, FTransform In
 	{
 		TGuardValue<bool> DisableCacheValidityChecks(Hierarchy->bEnableCacheValidityCheck, false);
 		NewElement->Key.Type = ERigElementType::Connector;
-		NewElement->Key.Name = Hierarchy->GetSafeNewName(InName, NewElement->Key.Type);
+		NewElement->Key.Name = GetSafeNewName(InName, NewElement->Key.Type);
 		NewElement->Settings = InSettings;
 		AddElement(NewElement, nullptr, true, InName);
 
@@ -1303,7 +1303,7 @@ TArray<FRigElementKey> URigHierarchyController::ImportFromText(FString InContent
 		}
 
 		const FName DesiredName = NewElement->Key.Name;
-		NewElement->Key.Name = Hierarchy->GetSafeNewName(DesiredName, NewElement->Key.Type);
+		NewElement->Key.Name = GetSafeNewName(DesiredName, NewElement->Key.Type);
 		AddElement(NewElement, nullptr, true, DesiredName);
 
 		KeyMap.FindOrAdd(PerElementData.Key) = NewElement->Key;
@@ -1783,6 +1783,20 @@ bool URigHierarchyController::IsValid() const
 	return Hierarchy.IsValid(bPendingKillAcceptable);
 }
 
+FName URigHierarchyController::GetSafeNewName(const FName& InDesiredName, ERigElementType InElementType) const
+{
+	FRigName Name(InDesiredName);
+
+	// remove potential namespaces from it
+	int32 Index = INDEX_NONE;
+	if(Name.GetName().FindLastChar(TEXT(':'), Index))
+	{
+		Name.SetName(Name.GetName().Mid(Index + 1));
+	}
+	
+	return Hierarchy->GetSafeNewName(Name, InElementType).GetFName();
+}
+
 int32 URigHierarchyController::AddElement(FRigBaseElement* InElementToAdd, FRigBaseElement* InFirstParent, bool bMaintainGlobalTransform, const FName& InDesiredName)
 {
 	ensure(IsValid());
@@ -1801,39 +1815,29 @@ int32 URigHierarchyController::AddElement(FRigBaseElement* InElementToAdd, FRigB
 	int32 NameSpaceTokenIndex = INDEX_NONE; 
 	if(InElementToAdd->GetName().FindLastChar(':', NameSpaceTokenIndex))
 	{
-		static const FName ShortNameMetaName = TEXT("ShortName");
 		const FString ShortName = InElementToAdd->GetName().RightChop(NameSpaceTokenIndex + 1);
 		if(!ShortName.IsEmpty())
 		{
-			Hierarchy->SetNameMetadata(InElementToAdd->Key, ShortNameMetaName, *ShortName);
+			Hierarchy->SetNameMetadata(InElementToAdd->Key, URigHierarchy::ShortNameMetadataName, *ShortName);
 		}
 	}
 
 	if(!InDesiredName.IsNone() &&
 		!InElementToAdd->GetFName().IsEqual(DesiredName.GetFName(), ENameCase::CaseSensitive))
 	{
-		static const FName DesiredNameMetaName = TEXT("DesiredName");
-		static const FName DesiredKeyMetaName = TEXT("DesiredKey");
-		Hierarchy->SetNameMetadata(InElementToAdd->Key, DesiredNameMetaName, DesiredName.GetFName());
-		Hierarchy->SetRigElementKeyMetadata(InElementToAdd->Key, DesiredKeyMetaName, FRigElementKey(DesiredName.GetFName(), InElementToAdd->Key.Type));
+		Hierarchy->SetNameMetadata(InElementToAdd->Key, URigHierarchy::DesiredNameMetadataName, DesiredName.GetFName());
+		Hierarchy->SetRigElementKeyMetadata(InElementToAdd->Key, URigHierarchy::DesiredKeyMetadataName, FRigElementKey(DesiredName.GetFName(), InElementToAdd->Key.Type));
 	}
 	
 	if(Hierarchy->HasExecuteContext())
 	{
 		const FControlRigExecuteContext& CRContext = Hierarchy->ExecuteContext->GetPublicData<FControlRigExecuteContext>();
 
-		if(!CRContext.ModuleInstanceName.IsNone())
-		{
-			static const FName ModuleMetaName = TEXT("Module");
-			Hierarchy->SetNameMetadata(InElementToAdd->Key, ModuleMetaName, CRContext.ModuleInstanceName);
-		}
-		
-		if(!CRContext.ModuleInstanceNameSpace.IsEmpty())
+		if(!CRContext.GetRigModuleNameSpace().IsEmpty())
 		{
-			if(InElementToAdd->GetName().StartsWith(CRContext.ModuleInstanceNameSpace, ESearchCase::CaseSensitive))
+			if(InElementToAdd->GetName().StartsWith(CRContext.GetRigModuleNameSpace(), ESearchCase::CaseSensitive))
 			{
-				static const FName NameSpaceMetaName = TEXT("NameSpace");
-				Hierarchy->SetNameMetadata(InElementToAdd->Key, NameSpaceMetaName, *CRContext.ModuleInstanceNameSpace);
+				Hierarchy->SetNameMetadata(InElementToAdd->Key, URigHierarchy::NameSpaceMetadataName, *CRContext.GetRigModuleNameSpace());
 
 				if(Hierarchy->ElementKeyRedirector)
 				{
@@ -2206,7 +2210,7 @@ bool URigHierarchyController::RenameElement(FRigBaseElement* InElement, const FN
 		TemporaryMap.Remove(OldKey);
    
 		TGuardValue<TMap<FRigElementKey, int32>> MapGuard(Hierarchy->IndexLookup, TemporaryMap);
-		InElement->Key.Name = Hierarchy->GetSafeNewName(InName, InElement->GetType());
+		InElement->Key.Name = GetSafeNewName(InName, InElement->GetType());
 		InElement->NameString.Reset();
 	}
 	
diff --git a/Engine/Plugins/Animation/ControlRig/Source/ControlRig/Private/Units/Execution/RigUnit_Hierarchy.cpp b/Engine/Plugins/Animation/ControlRig/Source/ControlRig/Private/Units/Execution/RigUnit_Hierarchy.cpp
--- a/Engine/Plugins/Animation/ControlRig/Source/ControlRig/Private/Units/Execution/RigUnit_Hierarchy.cpp
+++ b/Engine/Plugins/Animation/ControlRig/Source/ControlRig/Private/Units/Execution/RigUnit_Hierarchy.cpp
@@ -216,7 +216,13 @@ FRigUnit_HierarchyGetChainItemArray_Execute()
 				}
 
 				CachedChain = FRigElementKeyCollection(Keys);
-				CachedChain = FRigElementKeyCollection::MakeReversed(CachedChain);
+
+				// we are collecting the chain in reverse order (from tail to head)
+				// so we'll reverse it again if we are expecting the chain to be in order parent to child.
+				if(!bReverse)
+				{
+					CachedChain = FRigElementKeyCollection::MakeReversed(CachedChain);
+				}
 			}
 		}
 	}
diff --git a/Engine/Plugins/Animation/ControlRig/Source/ControlRig/Private/Units/Execution/RigUnit_Hierarchy.h b/Engine/Plugins/Animation/ControlRig/Source/ControlRig/Private/Units/Execution/RigUnit_Hierarchy.h
--- a/Engine/Plugins/Animation/ControlRig/Source/ControlRig/Private/Units/Execution/RigUnit_Hierarchy.h
+++ b/Engine/Plugins/Animation/ControlRig/Source/ControlRig/Private/Units/Execution/RigUnit_Hierarchy.h
@@ -271,6 +271,7 @@ struct CONTROLRIG_API FRigUnit_HierarchyGetChainItemArray : public FRigUnit_Hier
 		CachedStart = CachedEnd = FCachedRigElement();
 		CachedChain = FRigElementKeyCollection();
 		bIncludeStart = bIncludeEnd = true;
+		bReverse = false;
 	}
 
 	RIGVM_METHOD()
@@ -288,6 +289,9 @@ struct CONTROLRIG_API FRigUnit_HierarchyGetChainItemArray : public FRigUnit_Hier
 	UPROPERTY(meta = (Input))
 	bool bIncludeEnd;
 
+	UPROPERTY(meta = (Input))
+	bool bReverse;
+	
 	UPROPERTY(meta = (Output))
 	TArray<FRigElementKey> Chain;
 
diff --git a/Engine/Plugins/Animation/ControlRig/Source/ControlRig/Private/Units/Execution/RigUnit_RigModules.cpp b/Engine/Plugins/Animation/ControlRig/Source/ControlRig/Private/Units/Execution/RigUnit_RigModules.cpp
new file mode 100644
--- /dev/null
+++ b/Engine/Plugins/Animation/ControlRig/Source/ControlRig/Private/Units/Execution/RigUnit_RigModules.cpp
@@ -0,0 +1,127 @@
+// Copyright Epic Games, Inc. All Rights Reserved.
+
+#include "Units/Execution/RigUnit_RigModules.h"
+#include "Units/RigUnitContext.h"
+#include "ControlRig.h"
+
+#include UE_INLINE_GENERATED_CPP_BY_NAME(RigUnit_RigModules)
+
+FRigUnit_ResolveConnector_Execute()
+{
+	Result = Connector;
+	
+	if(const FRigBaseElement* ResolvedElement = ExecuteContext.Hierarchy->Find(Connector))
+	{
+		Result = ResolvedElement->GetKey();
+	}
+}
+
+FRigUnit_GetCurrentNameSpace_Execute()
+{
+	if(!ExecuteContext.IsRigModule())
+	{
+#if WITH_EDITOR
+		if(ExecuteContext.GetLog() != nullptr)
+		{
+			static const FString Message = TEXT("This node should only be used in a Rig Module."); 
+			ExecuteContext.GetLog()->Report(EMessageSeverity::Warning, ExecuteContext.GetFunctionName(), ExecuteContext.GetInstructionIndex(), Message);
+		}
+#endif
+	}
+	NameSpace = *ExecuteContext.GetRigModuleNameSpace();
+}
+
+FRigUnit_GetItemShortName_Execute()
+{
+	ShortName = NAME_None;
+
+	if(const URigHierarchy* Hierarchy = ExecuteContext.Hierarchy)
+	{
+		ShortName = Hierarchy->GetNameMetadata(Item, URigHierarchy::ShortNameMetadataName, NAME_None);
+	}
+
+	if(ShortName.IsNone())
+	{
+		ShortName = Item.Name;
+	}
+}
+
+FRigUnit_GetItemNameSpace_Execute()
+{
+	NameSpace.Reset();
+	HasNameSpace = false;
+
+	if(const URigHierarchy* Hierarchy = ExecuteContext.Hierarchy)
+	{
+		const FName NameSpaceName = Hierarchy->GetNameMetadata(Item, URigHierarchy::NameSpaceMetadataName, NAME_None);
+		if(!NameSpaceName.IsNone())
+		{
+			NameSpace = NameSpaceName.ToString();
+			HasNameSpace = true;
+		}
+	}
+}
+
+FRigUnit_IsItemInCurrentNameSpace_Execute()
+{
+	FString CurrentNameSpace;
+	FRigUnit_GetCurrentNameSpace::StaticExecute(ExecuteContext, CurrentNameSpace);
+	bool bHasNameSpace = false;
+	FString ItemNameSpace;
+	FRigUnit_GetItemNameSpace::StaticExecute(ExecuteContext, Item, bHasNameSpace, ItemNameSpace);
+
+	Result = false;
+	if(!CurrentNameSpace.IsEmpty() && !ItemNameSpace.IsEmpty())
+	{
+		Result = ItemNameSpace.Equals(CurrentNameSpace, ESearchCase::CaseSensitive); 
+	}
+}
+
+FRigUnit_GetItemsInNameSpace_Execute()
+{
+	const URigHierarchy* Hierarchy = ExecuteContext.Hierarchy;
+
+	FString NameSpace;
+	FRigUnit_GetCurrentNameSpace::StaticExecute(ExecuteContext, NameSpace);
+	if(NameSpace.IsEmpty())
+	{
+		Items.Reset();
+		return;
+	}
+	const FName NameSpaceName = *NameSpace;
+	
+	uint32 Hash = GetTypeHash(StaticStruct());
+	Hash = HashCombine(Hash, GetTypeHash((int32)TypeToSearch));
+	Hash = HashCombine(Hash, GetTypeHash(NameSpaceName));
+
+	if(const FRigElementKeyCollection* Cache = Hierarchy->FindCachedCollection(Hash))
+	{
+		Items = Cache->Keys;
+	}
+	else
+	{
+		FRigElementKeyCollection Collection;
+		Hierarchy->Traverse([Hierarchy, NameSpaceName, &Collection, TypeToSearch]
+			(const FRigBaseElement* InElement, bool &bContinue)
+			{
+				bContinue = true;
+						
+				const FRigElementKey Key = InElement->GetKey();
+				if(((uint8)TypeToSearch & (uint8)Key.Type) == (uint8)Key.Type)
+				{
+					const FName ItemNameSpace = Hierarchy->GetNameMetadata(Key, URigHierarchy::NameSpaceMetadataName, NAME_None);
+					if(!ItemNameSpace.IsNone())
+					{
+						if(ItemNameSpace.IsEqual(NameSpaceName, ENameCase::CaseSensitive))
+						{
+							Collection.AddUnique(Key);
+						}
+					}
+				}
+			}
+		);
+
+		Hierarchy->AddCachedCollection(Hash, Collection);
+		Items = Collection.Keys;
+	}
+}
\ No newline at end of file
diff --git a/Engine/Plugins/Animation/ControlRig/Source/ControlRig/Private/Units/Execution/RigUnit_RigModules.h b/Engine/Plugins/Animation/ControlRig/Source/ControlRig/Private/Units/Execution/RigUnit_RigModules.h
new file mode 100644
--- /dev/null
+++ b/Engine/Plugins/Animation/ControlRig/Source/ControlRig/Private/Units/Execution/RigUnit_RigModules.h
@@ -0,0 +1,196 @@
+// Copyright Epic Games, Inc. All Rights Reserved.
+
+#pragma once
+
+#include "Units/RigUnit.h"
+#include "ControlRigDefines.h" 
+#include "RigUnit_RigModules.generated.h"
+
+USTRUCT(meta = (Abstract, NodeColor="0.262745, 0.8, 0, 0.229412", Category = "Rig Modules"))
+struct CONTROLRIG_API FRigUnit_RigModulesBase : public FRigUnit
+{
+	GENERATED_BODY()
+};
+
+USTRUCT(meta = (Abstract, NodeColor="0.262745, 0.8, 0, 0.229412", Category = "Rig Modules"))
+struct CONTROLRIG_API FRigUnit_RigModulesBaseMutable : public FRigUnitMutable
+{
+	GENERATED_BODY()
+};
+
+/**
+* Returns the resolved item of the connector.
+*/
+USTRUCT(meta=(DisplayName="Resolve", Keywords="Connector,GetResolved,Target", Varying))
+struct CONTROLRIG_API FRigUnit_ResolveConnector : public FRigUnit_RigModulesBase
+{
+	GENERATED_BODY()
+
+	FRigUnit_ResolveConnector()
+	{
+		Connector = Result = FRigElementKey(NAME_None, ERigElementType::Connector);
+	}
+
+	RIGVM_METHOD()
+	virtual void Execute() override;
+
+	/*
+	 * The key of the connector to resolve
+	 */
+	UPROPERTY(meta = (Input, ExpandByDefault))
+	FRigElementKey Connector;
+
+	/*
+	 * The resulting item the connector is resolved to
+	 */
+	UPROPERTY(meta = (Output))
+	FRigElementKey Result;
+};
+
+/**
+ * Returns the currently used namespace of the module.
+ */
+USTRUCT(meta=(DisplayName="Get NameSpace", Keywords="NameSpace", Varying))
+struct CONTROLRIG_API FRigUnit_GetCurrentNameSpace : public FRigUnit_RigModulesBase
+{
+	GENERATED_BODY()
+
+	FRigUnit_GetCurrentNameSpace()
+	{
+		NameSpace = FString();;
+	}
+
+	RIGVM_METHOD()
+	virtual void Execute() override;
+
+	/*
+	 * The current namespace of the rig module
+	 */
+	UPROPERTY(meta = (Output))
+	FString NameSpace;
+};
+
+/**
+ * Returns the short name of the given item (without the namespace)
+ */
+USTRUCT(meta=(DisplayName="Get Short Name", Keywords="NameSpace", Varying))
+struct CONTROLRIG_API FRigUnit_GetItemShortName : public FRigUnit_RigModulesBase
+{
+	GENERATED_BODY()
+
+	FRigUnit_GetItemShortName()
+	{
+		Item = FRigElementKey();
+		ShortName = NAME_None;
+	}
+
+	RIGVM_METHOD()
+	virtual void Execute() override;
+
+	/*
+	 * The key of the item to return the short name for
+	 */
+	UPROPERTY(meta = (Input, ExpandByDefault))
+	FRigElementKey Item;
+
+	/*
+	 * The short name of the item (without the namespace)
+	 */
+	UPROPERTY(meta = (Output))
+	FName ShortName;
+};
+
+/**
+ * Returns the namespace of a given item. This may be an empty string if the item doesn't have a namespace.
+ */
+USTRUCT(meta=(DisplayName="Get Item NameSpace", Keywords="NameSpace", Varying))
+struct CONTROLRIG_API FRigUnit_GetItemNameSpace : public FRigUnit_RigModulesBase
+{
+	GENERATED_BODY()
+
+	FRigUnit_GetItemNameSpace()
+	{
+		Item = FRigElementKey();
+		HasNameSpace = false;
+		NameSpace = FString();
+	}
+
+	RIGVM_METHOD()
+	virtual void Execute() override;
+
+	/*
+	 * The key of the item to return the namespace for
+	 */
+	UPROPERTY(meta = (Input, ExpandByDefault))
+	FRigElementKey Item;
+
+	/*
+	 * True if the item has a valid namespace
+	 */
+	UPROPERTY(meta = (Output))
+	bool HasNameSpace;
+
+	/*
+	 * The namespace of the item
+	 */
+	UPROPERTY(meta = (Output))
+	FString NameSpace;
+};
+
+/**
+ * Returns true if the given item has been created by this module, 
+ * which means that the item's namespace is the module namespace.
+ */
+USTRUCT(meta=(DisplayName="Is In Current NameSpace", Keywords="NameSpace", Varying))
+struct CONTROLRIG_API FRigUnit_IsItemInCurrentNameSpace : public FRigUnit_RigModulesBase
+{
+	GENERATED_BODY()
+
+	FRigUnit_IsItemInCurrentNameSpace()
+	{
+		Item = FRigElementKey();
+		Result = false;
+	}
+
+	RIGVM_METHOD()
+	virtual void Execute() override;
+
+	/*
+	 * The key of the item to return the namespace for
+	 */
+	UPROPERTY(meta = (Input, ExpandByDefault))
+	FRigElementKey Item;
+
+	/*
+	 * True if the item is in this namespace / owned by this module.
+	 */
+	UPROPERTY(meta = (Output))
+	bool Result;
+};
+
+/**
+ * Returns all items (based on a filter) in the current namespace. 
+ */
+USTRUCT(meta=(DisplayName="Get Items In NameSpace", Keywords="NameSpace", Varying))
+struct CONTROLRIG_API FRigUnit_GetItemsInNameSpace : public FRigUnit_RigModulesBase
+{
+	GENERATED_BODY()
+
+	FRigUnit_GetItemsInNameSpace()
+	{
+		TypeToSearch = ERigElementType::All;
+		Items.Reset();
+	}
+
+	RIGVM_METHOD()
+	virtual void Execute() override;
+
+	UPROPERTY(meta = (Input))
+	ERigElementType TypeToSearch;
+
+	/*
+	 * The items within the namespace
+	 */
+	UPROPERTY(meta = (Output))
+	TArray<FRigElementKey> Items;
+};
\ No newline at end of file
diff --git a/Engine/Plugins/Animation/ControlRig/Source/ControlRig/Private/Units/RigUnitContext.cpp b/Engine/Plugins/Animation/ControlRig/Source/ControlRig/Private/Units/RigUnitContext.cpp
--- a/Engine/Plugins/Animation/ControlRig/Source/ControlRig/Private/Units/RigUnitContext.cpp
+++ b/Engine/Plugins/Animation/ControlRig/Source/ControlRig/Private/Units/RigUnitContext.cpp
@@ -1,12 +1,13 @@
 // Copyright Epic Games, Inc. All Rights Reserved.
 
 #include "Units/RigUnitContext.h"
+#include "ControlRig.h"
 
 #include UE_INLINE_GENERATED_CPP_BY_NAME(RigUnitContext)
 
 FName FControlRigExecuteContext::AddRigModuleNameSpace(const FName& InName) const
 {
-	if(ModuleInstanceName.IsNone())
+	if(IsRigModule())
 	{
 		return InName;
 	}
@@ -15,17 +16,17 @@ FName FControlRigExecuteContext::AddRigModuleNameSpace(const FName& InName) cons
 
 FString FControlRigExecuteContext::AddRigModuleNameSpace(const FString& InName) const
 {
-	if(ModuleInstanceName.IsNone())
+	if(IsRigModule())
 	{
 		return InName;
 	}
-	check(!ModuleInstanceNameSpace.IsEmpty());
-	return ModuleInstanceNameSpace + InName;
+	check(!RigModuleNameSpace.IsEmpty());
+	return RigModuleNameSpace + InName;
 }
 
 FName FControlRigExecuteContext::RemoveRigModuleNameSpace(const FName& InName) const
 {
-	if(ModuleInstanceName.IsNone())
+	if(IsRigModule())
 	{
 		return InName;
 	}
@@ -34,15 +35,30 @@ FName FControlRigExecuteContext::RemoveRigModuleNameSpace(const FName& InName) c
 
 FString FControlRigExecuteContext::RemoveRigModuleNameSpace(const FString& InName) const
 {
-	if(ModuleInstanceName.IsNone())
+	if(IsRigModule())
 	{
 		return InName;
 	}
-	check(!ModuleInstanceNameSpace.IsEmpty());
+	check(!RigModuleNameSpace.IsEmpty());
 
-	if(InName.StartsWith(ModuleInstanceNameSpace, ESearchCase::CaseSensitive))
+	if(InName.StartsWith(RigModuleNameSpace, ESearchCase::CaseSensitive))
 	{
-		return InName.Mid(ModuleInstanceNameSpace.Len());
+		return InName.Mid(RigModuleNameSpace.Len());
 	}
 	return InName;
 }
+
+FControlRigExecuteContextRigModuleGuard::FControlRigExecuteContextRigModuleGuard(FControlRigExecuteContext& InContext, const UControlRig* InControlRig)
+	: Context(InContext)
+	, PreviousRigModuleNameSpace(InContext.RigModuleNameSpace)
+	, PreviousRigModuleNameSpaceHash(InContext.RigModuleNameSpaceHash)
+{
+	Context.RigModuleNameSpace = InControlRig->GetRigModuleNameSpace();
+	Context.RigModuleNameSpaceHash = GetTypeHash(Context.RigModuleNameSpace);
+}
+	
+FControlRigExecuteContextRigModuleGuard::~FControlRigExecuteContextRigModuleGuard()
+{
+	Context.RigModuleNameSpace = PreviousRigModuleNameSpace;
+	Context.RigModuleNameSpaceHash = PreviousRigModuleNameSpaceHash; 
+}
diff --git a/Engine/Plugins/Animation/ControlRig/Source/ControlRig/Public/ControlRig.h b/Engine/Plugins/Animation/ControlRig/Source/ControlRig/Public/ControlRig.h
--- a/Engine/Plugins/Animation/ControlRig/Source/ControlRig/Public/ControlRig.h
+++ b/Engine/Plugins/Animation/ControlRig/Source/ControlRig/Public/ControlRig.h
@@ -98,9 +98,6 @@ public:
 	// Returns the parent rig hosting this module instance
 	UControlRig* GetParentRig() const;
 
-	// Returns the name of this module instance
-	FName GetModuleInstanceName() const;
-
 	// Returns the namespace of this module (for example ArmModule::)
 	const FString& GetRigModuleNameSpace() const;
 
diff --git a/Engine/Plugins/Animation/ControlRig/Source/ControlRig/Public/Rigs/RigHierarchy.h b/Engine/Plugins/Animation/ControlRig/Source/ControlRig/Public/Rigs/RigHierarchy.h
--- a/Engine/Plugins/Animation/ControlRig/Source/ControlRig/Public/Rigs/RigHierarchy.h
+++ b/Engine/Plugins/Animation/ControlRig/Source/ControlRig/Public/Rigs/RigHierarchy.h
@@ -164,6 +164,11 @@ public:
 	typedef TPair<int32, TArray<int32>> TElementDependencyMapPair;
 	typedef TTuple<int32, int32, int32, ERigTransformType::Type> TInstructionSliceElement;
 	inline static const FName TagMetadataName = TEXT("Tags");
+	inline static const FName ShortNameMetadataName = TEXT("ShortName");
+	inline static const FName DesiredNameMetadataName = TEXT("DesiredName");
+	inline static const FName DesiredKeyMetadataName = TEXT("DesiredKey");
+	inline static const FName ModuleMetadataName = TEXT("Module");
+	inline static const FName NameSpaceMetadataName = TEXT("NameSpace");
 
 	URigHierarchy();
 
@@ -4372,10 +4377,21 @@ protected:
 
 public:
 
-	const FRigElementKeyCollection* FindCachedCollection(uint32 InHash) const { return KeyCollectionCache.Find(InHash); }
-	FRigElementKeyCollection& FindOrAddCachedCollection(uint32 InHash) const { return KeyCollectionCache.FindOrAdd(InHash); };
-	void AddCachedCollection(uint32 InHash, const FRigElementKeyCollection& InCollection) const { KeyCollectionCache.Add(InHash, InCollection); }
+	const FRigElementKeyCollection* FindCachedCollection(uint32 InHash) const
+	{
+		return KeyCollectionCache.Find(InHash);
+	}
+	
+	FRigElementKeyCollection& FindOrAddCachedCollection(uint32 InHash) const
+	{
+		return KeyCollectionCache.FindOrAdd(InHash);
+	};
 	
+	void AddCachedCollection(uint32 InHash, const FRigElementKeyCollection& InCollection) const
+	{
+		KeyCollectionCache.Add(InHash, InCollection);
+	}
+
 private:
 	
 	mutable TMap<uint32, FRigElementKeyCollection> KeyCollectionCache;
diff --git a/Engine/Plugins/Animation/ControlRig/Source/ControlRig/Public/Rigs/RigHierarchyController.h b/Engine/Plugins/Animation/ControlRig/Source/ControlRig/Public/Rigs/RigHierarchyController.h
--- a/Engine/Plugins/Animation/ControlRig/Source/ControlRig/Public/Rigs/RigHierarchyController.h
+++ b/Engine/Plugins/Animation/ControlRig/Source/ControlRig/Public/Rigs/RigHierarchyController.h
@@ -615,6 +615,18 @@ private:
 	 */
 	bool IsValid() const;
 
+	/**
+	 * Determine a safe new name for an element. If a name is passed which contains a namespace
+	 * (for example "MyNameSpace:Control") we'll remove the namespace and just use the short name
+	 * prefixed with the current namespace (for example: "MyOtherNameSpace:Control"). Name clashes
+	 * will be resolved for the full name. So two modules with two different namespaces can have
+	 * elements with the same short name (for example "MyModuleA:Control" and "MyModuleB:Control").
+	 * @param InDesiredName The name provided by the user
+	 * @param InElementType The kind of element we are about to create
+	 * @return The safe name of the element to create.
+	 */
+	FName GetSafeNewName(const FName& InDesiredName, ERigElementType InElementType) const;
+	
 	/**
 	 * Adds a new element to the hierarchy
 	 * @param InElementToAdd The new element to add to the hierarchy 
diff --git a/Engine/Plugins/Animation/ControlRig/Source/ControlRig/Public/Units/RigUnitContext.h b/Engine/Plugins/Animation/ControlRig/Source/ControlRig/Public/Units/RigUnitContext.h
--- a/Engine/Plugins/Animation/ControlRig/Source/ControlRig/Public/Units/RigUnitContext.h
+++ b/Engine/Plugins/Animation/ControlRig/Source/ControlRig/Public/Units/RigUnitContext.h
@@ -11,7 +11,9 @@
 #include "RigVMCore/RigVMExecuteContext.h"
 #include "RigUnitContext.generated.h"
 
+class UControlRig;
 class UControlRigShapeLibrary;
+
 /**
  * The type of interaction happening on a rig
  */
@@ -97,11 +99,13 @@ struct FControlRigExecuteContext : public FRigVMExecuteContext
 {
 	GENERATED_BODY()
 
+public:
+	
 	FControlRigExecuteContext()
 		: FRigVMExecuteContext()
 		, Hierarchy(nullptr)
-		, ModuleInstanceName(NAME_None)
-		, ModuleInstanceNameSpace()
+		, RigModuleNameSpace()
+		, RigModuleNameSpaceHash(0)
 	{
 	}
 
@@ -149,6 +153,16 @@ struct FControlRigExecuteContext : public FRigVMExecuteContext
 	FName RemoveRigModuleNameSpace(const FName& InName) const;
 	FString RemoveRigModuleNameSpace(const FString& InName) const;
 
+	bool IsRigModule() const
+	{
+		return !GetRigModuleNameSpace().IsEmpty();
+	}
+
+	FString GetRigModuleNameSpace() const
+	{
+		return RigModuleNameSpace;
+	}
+
 	/** The list of available asset user data object */
 	TArray<const UAssetUserData*> AssetUserData;
 
@@ -160,8 +174,25 @@ struct FControlRigExecuteContext : public FRigVMExecuteContext
 	
 	FRigUnitContext UnitContext;
 	URigHierarchy* Hierarchy;
-	FName ModuleInstanceName;
-	FString ModuleInstanceNameSpace;
+
+private:
+	FString RigModuleNameSpace;
+	uint32 RigModuleNameSpaceHash;
+
+	friend class FControlRigExecuteContextRigModuleGuard;
+};
+
+class CONTROLRIG_API FControlRigExecuteContextRigModuleGuard
+{
+public:
+	FControlRigExecuteContextRigModuleGuard(FControlRigExecuteContext& InContext, const UControlRig* InControlRig);
+	~FControlRigExecuteContextRigModuleGuard();
+
+private:
+
+	FControlRigExecuteContext& Context;
+	FString PreviousRigModuleNameSpace;
+	uint32 PreviousRigModuleNameSpaceHash;
 };
 
 #if WITH_EDITOR
diff --git a/Engine/Plugins/Animation/ControlRig/Source/ControlRigEditor/Private/Editor/SRigHierarchyTreeView.cpp b/Engine/Plugins/Animation/ControlRig/Source/ControlRigEditor/Private/Editor/SRigHierarchyTreeView.cpp
--- a/Engine/Plugins/Animation/ControlRig/Source/ControlRigEditor/Private/Editor/SRigHierarchyTreeView.cpp
+++ b/Engine/Plugins/Animation/ControlRig/Source/ControlRigEditor/Private/Editor/SRigHierarchyTreeView.cpp
@@ -34,6 +34,7 @@ FRigTreeDisplaySettings FRigTreeDelegates::DefaultDisplaySettings;
 FRigTreeElement::FRigTreeElement(const FRigElementKey& InKey, TWeakPtr<SRigHierarchyTreeView> InTreeView, bool InSupportsRename, ERigTreeFilterResult InFilterResult)
 {
 	Key = InKey;
+	ShortName = InKey.Name;
 	ChannelName = NAME_None;
 	bIsTransient = false;
 	bIsAnimationChannel = false;
@@ -45,6 +46,8 @@ FRigTreeElement::FRigTreeElement(const FRigElementKey& InKey, TWeakPtr<SRigHiera
 	{
 		if(const URigHierarchy* Hierarchy = InTreeView.Pin()->GetRigTreeDelegates().GetHierarchy())
 		{
+			ShortName = Hierarchy->GetNameMetadata(InKey, URigHierarchy::ShortNameMetadataName, ShortName);
+			
 			const FRigTreeDisplaySettings& Settings = InTreeView.Pin()->GetRigTreeDelegates().GetDisplaySettings();
 			RefreshDisplaySettings(Hierarchy, Settings);
 		}
@@ -169,7 +172,8 @@ void SRigHierarchyItem::Construct(const FArguments& InArgs, const TSharedRef<STa
 			.VAlign(VAlign_Center)
 			[
 				SAssignNew(InlineWidget, SInlineEditableTextBlock)
-				.Text(this, &SRigHierarchyItem::GetName)
+				.Text(this, &SRigHierarchyItem::GetName, true)
+				.ToolTipText(this, &SRigHierarchyItem::GetItemTooltip)
 				.OnVerifyTextChanged(this, &SRigHierarchyItem::OnVerifyNameChanged)
 				.OnTextCommitted(this, &SRigHierarchyItem::OnNameCommitted)
 				.MultiLine(false)
@@ -196,7 +200,7 @@ void SRigHierarchyItem::Construct(const FArguments& InArgs, const TSharedRef<STa
 	InRigTreeElement->OnRenameRequested.BindSP(InlineWidget.Get(), &SInlineEditableTextBlock::EnterEditingMode);
 }
 
-FText SRigHierarchyItem::GetName() const
+FText SRigHierarchyItem::GetName(bool bUseShortName) const
 {
 	if(WeakRigTreeElement.Pin()->bIsTransient)
 	{
@@ -207,9 +211,24 @@ FText SRigHierarchyItem::GetName() const
 	{
 		return FText::FromName(WeakRigTreeElement.Pin()->ChannelName);
 	}
+	if(bUseShortName)
+	{
+		return (FText::FromName(WeakRigTreeElement.Pin()->ShortName));
+	}
 	return (FText::FromName(WeakRigTreeElement.Pin()->Key.Name));
 }
 
+FText SRigHierarchyItem::GetItemTooltip() const
+{
+	const FText FullName = GetName(false);
+	const FText ShortName = GetName(true);
+	if(FullName.EqualTo(ShortName))
+	{
+		return FText();
+	}
+	return FullName;
+}
+
 //////////////////////////////////////////////////////////////
 /// SRigHierarchyTreeView
 ///////////////////////////////////////////////////////////
diff --git a/Engine/Plugins/Animation/ControlRig/Source/ControlRigEditor/Private/Editor/SRigHierarchyTreeView.h b/Engine/Plugins/Animation/ControlRig/Source/ControlRigEditor/Private/Editor/SRigHierarchyTreeView.h
--- a/Engine/Plugins/Animation/ControlRig/Source/ControlRigEditor/Private/Editor/SRigHierarchyTreeView.h
+++ b/Engine/Plugins/Animation/ControlRig/Source/ControlRigEditor/Private/Editor/SRigHierarchyTreeView.h
@@ -189,6 +189,7 @@ public:
 public:
 	/** Element Data to display */
 	FRigElementKey Key;
+	FName ShortName;
 	FName ChannelName;
 	bool bIsTransient;
 	bool bIsAnimationChannel;
@@ -236,7 +237,8 @@ private:
 	TWeakPtr<FRigTreeElement> WeakRigTreeElement;
  	FRigTreeDelegates Delegates;
 
-	FText GetName() const;
+	FText GetName(bool bUseShortName) const;
+	FText GetItemTooltip() const;
 
 	friend class SRigHierarchyTreeView; 
 };
