Control Rig: Refactor Collection Rig Units to rely on new Hierarchy chain node

diff --git a/Engine/Plugins/Animation/ControlRig/Source/ControlRig/Private/Units/Execution/RigUnit_Collection.cpp b/Engine/Plugins/Animation/ControlRig/Source/ControlRig/Private/Units/Execution/RigUnit_Collection.cpp
--- a/Engine/Plugins/Animation/ControlRig/Source/ControlRig/Private/Units/Execution/RigUnit_Collection.cpp
+++ b/Engine/Plugins/Animation/ControlRig/Source/ControlRig/Private/Units/Execution/RigUnit_Collection.cpp
@@ -2,6 +2,7 @@
 
 #include "Units/Execution/RigUnit_Collection.h"
 #include "Units/Execution/RigUnit_Item.h"
+#include "Units/Execution/RigUnit_Hierarchy.h"
 #include "Units/RigUnitContext.h"
 
 #include UE_INLINE_GENERATED_CPP_BY_NAME(RigUnit_Collection)
@@ -10,19 +11,6 @@
 #include "Units/RigUnitTest.h"
 #endif
 
-#define FRigUnit_CollectionChain_Hash 1
-#define FRigUnit_CollectionNameSearch_Hash 2
-#define FRigUnit_CollectionChildren_Hash 3
-#define FRigUnit_CollectionReplaceItems_Hash 4
-#define FRigUnit_CollectionItems_Hash 5
-#define FRigUnit_CollectionUnion_Hash 6
-#define FRigUnit_CollectionIntersection_Hash 7
-#define FRigUnit_CollectionDifference_Hash 8
-#define FRigUnit_CollectionReverse_Hash 9
-#define FRigUnit_CollectionCount_Hash 10
-#define FRigUnit_CollectionItemAtIndex_Hash 11
-#define FRigUnit_CollectionGetAll_Hash 12
-
 FRigUnit_CollectionChain_Execute()
 {
 	FRigUnit_CollectionChainArray::StaticExecute(ExecuteContext, FirstItem, LastItem, Reverse, Collection.Keys);
@@ -30,50 +18,30 @@ FRigUnit_CollectionChain_Execute()
 
 FRigVMStructUpgradeInfo FRigUnit_CollectionChain::GetUpgradeInfo() const
 {
-	FRigUnit_CollectionChainArray NewNode;
-	NewNode.FirstItem = FirstItem;
-	NewNode.LastItem = LastItem;
-	NewNode.Reverse = Reverse;
-
+	FRigUnit_HierarchyGetChainItemArray NewNode;
+	NewNode.Start = FirstItem;
+	NewNode.End = LastItem;
+	NewNode.bReverse = Reverse;
+	NewNode.bIncludeStart = NewNode.bIncludeEnd = true;
 	return FRigVMStructUpgradeInfo(*this, NewNode);
 }
 
 FRigUnit_CollectionChainArray_Execute()
 {
-	if(ExecuteContext.Hierarchy == nullptr)
-	{
-		Items.Reset();
-		return;
-	}
-	
-	uint32 Hash = FRigUnit_CollectionChain_Hash + ExecuteContext.Hierarchy->GetTopologyVersion() * 17;
-	Hash = HashCombine(Hash, GetTypeHash(FirstItem));
-	Hash = HashCombine(Hash, GetTypeHash(LastItem));
-	Hash = HashCombine(Hash, Reverse ? 1 : 0);
+	FCachedRigElement FirstCache, LastCache;
+	FRigElementKeyCollection CachedChain;
+	FRigUnit_HierarchyGetChainItemArray::StaticExecute(ExecuteContext, FirstItem, LastItem, true, true, false, Items, FirstCache, LastCache, CachedChain);
+}
 
-	FRigElementKeyCollection Collection;
-	if(const FRigElementKeyCollection* Cache = ExecuteContext.Hierarchy->FindCachedCollection(Hash))
-	{
-		Collection = *Cache;
-	}
-	else
-	{
-		Collection = FRigElementKeyCollection::MakeFromChain(ExecuteContext.Hierarchy, FirstItem, LastItem, Reverse);
-		if (Collection.IsEmpty())
-		{
-			if (ExecuteContext.Hierarchy->GetIndex(FirstItem) == INDEX_NONE)
-			{
-				UE_CONTROLRIG_RIGUNIT_REPORT_WARNING(TEXT("First Item '%s' is not valid."), *FirstItem.ToString());
-			}
-			if (ExecuteContext.Hierarchy->GetIndex(LastItem) == INDEX_NONE)
-			{
-				UE_CONTROLRIG_RIGUNIT_REPORT_WARNING(TEXT("Last Item '%s' is not valid."), *LastItem.ToString());
-			}
-		}
-		ExecuteContext.Hierarchy->AddCachedCollection(Hash, Collection);
-	}
+FRigVMStructUpgradeInfo FRigUnit_CollectionChainArray::GetUpgradeInfo() const
+{
+	FRigUnit_HierarchyGetChainItemArray NewNode;
+	NewNode.Start = FirstItem;
+	NewNode.End = LastItem;
+	NewNode.bReverse = Reverse;
+	NewNode.bIncludeStart = NewNode.bIncludeEnd = true;
 
-	Items = Collection.Keys;
+	return FRigVMStructUpgradeInfo(*this, NewNode);
 }
 
 FRigUnit_CollectionNameSearch_Execute()
@@ -100,7 +68,7 @@ FRigUnit_CollectionNameSearchArray_Execute()
 		return;
 	}
 
-	uint32 Hash = FRigUnit_CollectionNameSearch_Hash + ExecuteContext.Hierarchy->GetTopologyVersion() * 17;
+	uint32 Hash = GetTypeHash(StaticStruct()) + ExecuteContext.Hierarchy->GetTopologyVersion() * 17;
 	Hash = HashCombine(Hash, GetTypeHash(PartialName));
 	Hash = HashCombine(Hash, (int32)TypeToSearch * 8);
 
@@ -146,7 +114,7 @@ FRigUnit_CollectionChildrenArray_Execute()
 		return;
 	}
 	
-	uint32 Hash = FRigUnit_CollectionChildren_Hash + ExecuteContext.Hierarchy->GetTopologyVersion() * 17;
+	uint32 Hash = GetTypeHash(StaticStruct()) + ExecuteContext.Hierarchy->GetTopologyVersion() * 17;
 	Hash = HashCombine(Hash, GetTypeHash(Parent));
 	Hash = HashCombine(Hash, bRecursive ? 2 : 0);
 	Hash = HashCombine(Hash, bIncludeParent ? 1 : 0);
@@ -183,7 +151,7 @@ FRigUnit_CollectionGetAll_Execute()
 		return;
 	}
 	
-	uint32 Hash = FRigUnit_CollectionGetAll_Hash + ExecuteContext.Hierarchy->GetTopologyVersion() * 17;
+	uint32 Hash = GetTypeHash(StaticStruct()) + ExecuteContext.Hierarchy->GetTopologyVersion() * 17;
 	Hash = HashCombine(Hash, (int32)TypeToSearch * 8);
 
 	FRigElementKeyCollection Collection;
@@ -280,7 +248,7 @@ FRigUnit_CollectionReplaceItemsArray_Execute()
 		return;
 	}
 	
-	uint32 Hash = FRigUnit_CollectionReplaceItems_Hash + ExecuteContext.Hierarchy->GetTopologyVersion() * 17;
+	uint32 Hash = GetTypeHash(StaticStruct()) + ExecuteContext.Hierarchy->GetTopologyVersion() * 17;
 	Hash = HashCombine(Hash, GetTypeHash(Items));
 	Hash = HashCombine(Hash, 12 * GetTypeHash(Old));
 	Hash = HashCombine(Hash, 13 * GetTypeHash(New));
diff --git a/Engine/Plugins/Animation/ControlRig/Source/ControlRig/Private/Units/Execution/RigUnit_Collection.h b/Engine/Plugins/Animation/ControlRig/Source/ControlRig/Private/Units/Execution/RigUnit_Collection.h
--- a/Engine/Plugins/Animation/ControlRig/Source/ControlRig/Private/Units/Execution/RigUnit_Collection.h
+++ b/Engine/Plugins/Animation/ControlRig/Source/ControlRig/Private/Units/Execution/RigUnit_Collection.h
@@ -56,7 +56,7 @@ struct CONTROLRIG_API FRigUnit_CollectionChain : public FRigUnit_CollectionBase
 * Creates an item array based on a first and last item within a chain.
 * Chains can refer to bone chains or chains within a control hierarchy.
 */
-USTRUCT(meta=(DisplayName="Item Chain", Keywords="Bone,Joint,Collection", Varying))
+USTRUCT(meta=(DisplayName="Item Chain", Keywords="Bone,Joint,Collection", Varying, Deprecated = "5.4"))
 struct CONTROLRIG_API FRigUnit_CollectionChainArray : public FRigUnit_CollectionBase
 {
 	GENERATED_BODY()
@@ -81,6 +81,9 @@ struct CONTROLRIG_API FRigUnit_CollectionChainArray : public FRigUnit_Collection
 
 	UPROPERTY(meta = (Output))
 	TArray<FRigElementKey> Items;
+
+	RIGVM_METHOD()
+	virtual FRigVMStructUpgradeInfo GetUpgradeInfo() const override;
 };
 
 /**
