Control Rig: Metadata nodes to respect namespaces

diff --git a/Engine/Plugins/Animation/ControlRig/Source/ControlRig/Private/Units/Hierarchy/RigUnit_Metadata.cpp b/Engine/Plugins/Animation/ControlRig/Source/ControlRig/Private/Units/Hierarchy/RigUnit_Metadata.cpp
--- a/Engine/Plugins/Animation/ControlRig/Source/ControlRig/Private/Units/Hierarchy/RigUnit_Metadata.cpp
+++ b/Engine/Plugins/Animation/ControlRig/Source/ControlRig/Private/Units/Hierarchy/RigUnit_Metadata.cpp
@@ -13,6 +13,7 @@
 
 FName FRigDispatch_MetadataBase::ItemArgName = TEXT("Item");
 FName FRigDispatch_MetadataBase::NameArgName = TEXT("Name");
+FName FRigDispatch_MetadataBase::UseNameSpaceArgName = TEXT("UseNameSpace");
 FName FRigDispatch_MetadataBase::CacheArgName = TEXT("Cache");
 FName FRigDispatch_MetadataBase::DefaultArgName = TEXT("Default");
 FName FRigDispatch_MetadataBase::ValueArgName = TEXT("Value");
@@ -71,6 +72,7 @@ const TArray<FRigVMTemplateArgument>& FRigDispatch_MetadataBase::GetArguments()
 	{
 		ItemArgIndex = Arguments.Emplace(ItemArgName, ERigVMPinDirection::Input, FRigVMRegistry::Get().GetTypeIndex<FRigElementKey>());
 		NameArgIndex = Arguments.Emplace(NameArgName, ERigVMPinDirection::Input, RigVMTypeUtils::TypeIndex::FName);
+		UseNameSpaceArgIndex = Arguments.Emplace(UseNameSpaceArgName, ERigVMPinDirection::Input, RigVMTypeUtils::TypeIndex::Bool);
 		CacheArgIndex = Arguments.Emplace(CacheArgName, ERigVMPinDirection::Hidden, FRigVMRegistry::Get().GetTypeIndex<FCachedRigElement>());
 	};
 	return Arguments;
@@ -88,6 +90,10 @@ FText FRigDispatch_MetadataBase::GetArgumentTooltip(const FName& InArgumentName,
 	{
 		return NSLOCTEXT("FRigDispatch_MetadataBase", "NameArgTooltip", "The name of the metadata");
 	}
+	if(InArgumentName == UseNameSpaceArgName)
+	{
+		return NSLOCTEXT("FRigDispatch_MetadataBase", "UseNameSpaceArgTooltip", "If checked the metadata will be available only in this name space / module");
+	}
 	if(InArgumentName == DefaultArgName)
 	{
 		return NSLOCTEXT("FRigDispatch_MetadataBase", "DefaultArgTooltip", "The default value used as a fallback if the metadata does not exist");
@@ -107,9 +113,19 @@ FText FRigDispatch_MetadataBase::GetArgumentTooltip(const FName& InArgumentName,
 	return FRigDispatchFactory::GetArgumentTooltip(InArgumentName, InTypeIndex);
 }
 
+FString FRigDispatch_MetadataBase::GetArgumentDefaultValue(const FName& InArgumentName,
+	TRigVMTypeIndex InTypeIndex) const
+{
+	if(InArgumentName == UseNameSpaceArgName)
+	{
+		return TrueString;
+	}
+	return FRigDispatchFactory::GetArgumentDefaultValue(InArgumentName, InTypeIndex);
+}
+
 FString FRigDispatch_MetadataBase::GetArgumentMetaData(const FName& InArgumentName, const FName& InMetaDataKey) const
 {
-	if(InArgumentName == TEXT("Name"))
+	if(InArgumentName == NameArgName)
 	{
 		if(InMetaDataKey == FRigVMStruct::CustomWidgetMetaName)
 		{
@@ -168,15 +184,23 @@ const TArray<FRigVMTemplateArgument>& FRigDispatch_GetMetadata::GetArguments() c
 }
 
 FRigBaseMetadata* FRigDispatch_GetMetadata::FindMetadata(const FRigVMExtendedExecuteContext& InContext,
-                                                         const FRigElementKey& InKey, const FName& InName, ERigMetadataType InType, FCachedRigElement& Cache)
+                                                         const FRigElementKey& InKey, const FName& InName,
+                                                         ERigMetadataType InType, bool bUseNameSpace, FCachedRigElement& Cache)
 {
 	const FControlRigExecuteContext& ExecuteContext = InContext.GetPublicData<FControlRigExecuteContext>();
 	if(Cache.UpdateCache(InKey, ExecuteContext.Hierarchy))
 	{
-		if(FRigBaseElement* Element = ExecuteContext.Hierarchy->Get(Cache.GetIndex()))
+		if(const FRigBaseElement* Element = ExecuteContext.Hierarchy->Get(Cache.GetIndex()))
 		{
+			// first try to find the metadata in the namespace
+			const FName Name = ExecuteContext.AdaptMetadataName(bUseNameSpace, InName);
+			if(FRigBaseMetadata* Metadata = Element->GetMetadata(Name, InType))
+			{
+				return Metadata;
+			}
+			
 			return Element->GetMetadata(InName, InType);
-		}
+		};
 	}
 	return nullptr;
 }
@@ -292,13 +316,15 @@ const TArray<FRigVMExecuteArgument>& FRigDispatch_SetMetadata::GetExecuteArgumen
 }
 
 FRigBaseMetadata* FRigDispatch_SetMetadata::FindOrAddMetadata(FControlRigExecuteContext& InContext,
-                                                              const FRigElementKey& InKey, const FName& InName, ERigMetadataType InType, FCachedRigElement& Cache)
+                                                              const FRigElementKey& InKey, const FName& InName, ERigMetadataType InType,
+                                                              bool bUseNameSpace, FCachedRigElement& Cache)
 {
 	if(Cache.UpdateCache(InKey, InContext.Hierarchy))
 	{
 		if(FRigBaseElement* Element = InContext.Hierarchy->Get(Cache.GetIndex()))
 		{
-			return Element->SetupValidMetadata(InName, InType);
+			const FName Name = InContext.AdaptMetadataName(bUseNameSpace, InName);
+			return Element->SetupValidMetadata(Name, InType);
 		}
 	}
 	return nullptr;
@@ -407,7 +433,12 @@ FRigUnit_RemoveMetadata_Execute()
 	{
 		if(FRigBaseElement* Element = Hierarchy->Get(CachedIndex))
 		{
-			Removed = Element->RemoveMetadata(Name);
+			const FName LocalName = ExecuteContext.AdaptMetadataName(UseNameSpace, Name);
+			Removed = Element->RemoveMetadata(LocalName);
+			if(!Removed)
+			{
+				Removed = Element->RemoveMetadata(Name);
+			}
 		}
 	}
 }
@@ -426,6 +457,28 @@ FRigUnit_RemoveAllMetadata_Execute()
 	{
 		if(FRigBaseElement* Element = Hierarchy->Get(CachedIndex))
 		{
+			if(UseNameSpace)
+			{
+				// only remove the metadata within this module / with this namespace
+				if(ExecuteContext.IsRigModule())
+				{
+					const TArray<FName> MetadataNames = Hierarchy->GetMetadataNames(Element->GetKey());
+
+					Removed = false;
+					const FString NameSpace = ExecuteContext.GetRigModuleNameSpace();
+					for(const FName& MetadataName : MetadataNames)
+					{
+						if(MetadataName.ToString().StartsWith(NameSpace, ESearchCase::CaseSensitive))
+						{
+							if(Element->RemoveMetadata(MetadataName))
+							{
+								Removed = true;
+							}
+						}
+					}
+					return;
+				}
+			}
 			Removed = Element->RemoveAllMetadata();
 		}
 	}
@@ -445,7 +498,12 @@ FRigUnit_HasMetadata_Execute()
 	{
 		if(const FRigBaseElement* Element = Hierarchy->Get(CachedIndex))
 		{
-			Found = Element->GetMetadata(Name, Type) != nullptr;
+			const FName LocalName = ExecuteContext.AdaptMetadataName(UseNameSpace, Name);
+			Found = Element->GetMetadata(LocalName, Type) != nullptr;
+			if(!Found)
+			{
+				Found = Element->GetMetadata(Name, Type) != nullptr;
+			}
 		}
 	}
 }
@@ -460,9 +518,10 @@ FRigUnit_FindItemsWithMetadata_Execute()
 		return;
 	}
 
-	Hierarchy->Traverse([&Items, Name, Type](const FRigBaseElement* Element, bool& bContinue)
+	const FName LocalName = ExecuteContext.AdaptMetadataName(UseNameSpace, Name);
+	Hierarchy->Traverse([&Items, Name, LocalName, Type](const FRigBaseElement* Element, bool& bContinue)
 	{
-		if(Element->GetMetadata(Name, Type) != nullptr)
+		if(Element->GetMetadata(LocalName, Type) != nullptr)
 		{
 			Items.AddUnique(Element->GetKey());
 		}
@@ -479,7 +538,7 @@ FRigUnit_GetMetadataTags_Execute()
 	{
 		return;
 	}
-	
+
 	if(CachedIndex.UpdateCache(Item, Hierarchy))
 	{
 		if(const FRigBaseElement* Element = Hierarchy->Get(CachedIndex))
@@ -499,7 +558,7 @@ FRigUnit_SetMetadataTag_Execute()
 	{
 		return;
 	}
-	
+
 	if(CachedIndex.UpdateCache(Item, Hierarchy))
 	{
 		if(FRigBaseElement* Element = Hierarchy->Get(CachedIndex))
@@ -507,9 +566,10 @@ FRigUnit_SetMetadataTag_Execute()
 			if(FRigNameArrayMetadata* Md = Cast<FRigNameArrayMetadata>(Element->SetupValidMetadata(URigHierarchy::TagMetadataName, ERigMetadataType::NameArray)))
 			{
 				const int32 LastIndex = Md->GetValue().Num(); 
-				if(Md->GetValue().AddUnique(Tag) == LastIndex)
+				const FName LocalTag = ExecuteContext.AdaptMetadataName(UseNameSpace, Tag);
+				if(Md->GetValue().AddUnique(LocalTag) == LastIndex)
 				{
-					Element->NotifyMetadataTagChanged(Tag, true);
+					Element->NotifyMetadataTagChanged(LocalTag, true);
 				}
 			}
 		}
@@ -523,7 +583,7 @@ FRigUnit_SetMetadataTagArray_Execute()
 	{
 		return;
 	}
-	
+
 	if(CachedIndex.UpdateCache(Item, Hierarchy))
 	{
 		if(FRigBaseElement* Element = Hierarchy->Get(CachedIndex))
@@ -533,9 +593,10 @@ FRigUnit_SetMetadataTagArray_Execute()
 				for(const FName& Tag : Tags)
 				{
 					const int32 LastIndex = Md->GetValue().Num(); 
-					if(Md->GetValue().AddUnique(Tag) == LastIndex)
+					const FName LocalTag = ExecuteContext.AdaptMetadataName(UseNameSpace, Tag);
+					if(Md->GetValue().AddUnique(LocalTag) == LastIndex)
 					{
-						Element->NotifyMetadataTagChanged(Tag, true);
+						Element->NotifyMetadataTagChanged(LocalTag, true);
 					}
 				}
 			}
@@ -551,17 +612,18 @@ FRigUnit_RemoveMetadataTag_Execute()
 	{
 		return;
 	}
-	
+
 	if(CachedIndex.UpdateCache(Item, Hierarchy))
 	{
 		if(FRigBaseElement* Element = Hierarchy->Get(CachedIndex))
 		{
 			if(FRigNameArrayMetadata* Md = Cast<FRigNameArrayMetadata>(Element->GetMetadata(URigHierarchy::TagMetadataName, ERigMetadataType::NameArray)))
 			{
-				Removed = Md->GetValue().Remove(Tag) > 0;
+				const FName LocalTag = ExecuteContext.AdaptMetadataName(UseNameSpace, Tag);
+				Removed = Md->GetValue().Remove(LocalTag) > 0;
 				if(Removed)
 				{
-					Element->NotifyMetadataTagChanged(Tag, false);
+					Element->NotifyMetadataTagChanged(LocalTag, false);
 				}
 			}
 		}
@@ -577,14 +639,15 @@ FRigUnit_HasMetadataTag_Execute()
 	{
 		return;
 	}
-	
+
 	if(CachedIndex.UpdateCache(Item, Hierarchy))
 	{
 		if(const FRigBaseElement* Element = Hierarchy->Get(CachedIndex))
 		{
 			if(const FRigNameArrayMetadata* Md = Cast<FRigNameArrayMetadata>(Element->GetMetadata(URigHierarchy::TagMetadataName, ERigMetadataType::NameArray)))
 			{
-				Found = Md->GetValue().Contains(Tag);
+				const FName LocalTag = ExecuteContext.AdaptMetadataName(UseNameSpace, Tag);
+				Found = Md->GetValue().Contains(LocalTag);
 			}
 		}
 	}
@@ -599,7 +662,7 @@ FRigUnit_HasMetadataTagArray_Execute()
 	{
 		return;
 	}
-	
+
 	if(CachedIndex.UpdateCache(Item, Hierarchy))
 	{
 		if(const FRigBaseElement* Element = Hierarchy->Get(CachedIndex))
@@ -609,7 +672,8 @@ FRigUnit_HasMetadataTagArray_Execute()
 				Found = true;
 				for(const FName& Tag : Tags)
 				{
-					if(!Md->GetValue().Contains(Tag))
+					const FName LocalTag = ExecuteContext.AdaptMetadataName(UseNameSpace, Tag);
+					if(!Md->GetValue().Contains(LocalTag))
 					{
 						Found = false;
 						break;
@@ -630,11 +694,12 @@ FRigUnit_FindItemsWithMetadataTag_Execute()
 		return;
 	}
 
-	Hierarchy->Traverse([&Items, Tag](const FRigBaseElement* Element, bool& bContinue)
+	const FName LocalTag = ExecuteContext.AdaptMetadataName(UseNameSpace, Tag);
+	Hierarchy->Traverse([&Items, LocalTag](const FRigBaseElement* Element, bool& bContinue)
 	{
 		if(const FRigNameArrayMetadata* Md = Cast<FRigNameArrayMetadata>(Element->GetMetadata(URigHierarchy::TagMetadataName, ERigMetadataType::NameArray)))
 		{
-			if(Md->GetValue().Contains(Tag))
+			if(Md->GetValue().Contains(LocalTag))
 			{
 				Items.AddUnique(Element->GetKey());
 			}
@@ -653,12 +718,25 @@ FRigUnit_FindItemsWithMetadataTagArray_Execute()
 		return;
 	}
 
-	Hierarchy->Traverse([&Items, Tags](const FRigBaseElement* Element, bool& bContinue)
+	TArrayView<const FName> LocalTags(Tags);
+	TArray<FName> AdaptedTags;
+	if(UseNameSpace && ExecuteContext.IsRigModule())
+	{
+		AdaptedTags.Reserve(Tags.Num());
+		for(const FName& Tag : Tags)
+		{
+			const FName LocalTag = ExecuteContext.AdaptMetadataName(UseNameSpace, Tag);
+			AdaptedTags.Add(LocalTag);
+		}
+		LocalTags = TArrayView<const FName>(AdaptedTags);
+	}
+	
+	Hierarchy->Traverse([&Items, LocalTags](const FRigBaseElement* Element, bool& bContinue)
 	{
 		if(const FRigNameArrayMetadata* Md = Cast<FRigNameArrayMetadata>(Element->GetMetadata(URigHierarchy::TagMetadataName, ERigMetadataType::NameArray)))
 		{
 			bool bFoundAll = true;
-			for(const FName& Tag : Tags)
+			for(const FName& Tag : LocalTags)
 			{
 				if(!Md->GetValue().Contains(Tag))
 				{
@@ -692,6 +770,19 @@ FRigUnit_FilterItemsByMetadataTags_Execute()
 		CachedIndices.SetNumZeroed(Items.Num());
 	}
 
+	TArrayView<const FName> LocalTags(Tags);
+	TArray<FName> AdaptedTags;
+	if(UseNameSpace && ExecuteContext.IsRigModule())
+	{
+		AdaptedTags.Reserve(Tags.Num());
+		for(const FName& Tag : Tags)
+		{
+			const FName LocalTag = ExecuteContext.AdaptMetadataName(UseNameSpace, Tag);
+			AdaptedTags.Add(LocalTag);
+		}
+		LocalTags = TArrayView<const FName>(AdaptedTags);
+	}
+
 	for(int32 Index = 0; Index < Items.Num(); Index++)
 	{
 		if(CachedIndices[Index].UpdateCache(Items[Index], Hierarchy))
@@ -703,7 +794,7 @@ FRigUnit_FilterItemsByMetadataTags_Execute()
 					if(Inclusive)
 					{
 						bool bFoundAll = true;
-						for(const FName& Tag : Tags)
+						for(const FName& Tag : LocalTags)
 						{
 							if(!Md->GetValue().Contains(Tag))
 							{
@@ -719,7 +810,7 @@ FRigUnit_FilterItemsByMetadataTags_Execute()
 					else
 					{
 						bool bFoundAny = false;
-						for(const FName& Tag : Tags)
+						for(const FName& Tag : LocalTags)
 						{
 							if(Md->GetValue().Contains(Tag))
 							{
diff --git a/Engine/Plugins/Animation/ControlRig/Source/ControlRig/Private/Units/Hierarchy/RigUnit_Metadata.h b/Engine/Plugins/Animation/ControlRig/Source/ControlRig/Private/Units/Hierarchy/RigUnit_Metadata.h
--- a/Engine/Plugins/Animation/ControlRig/Source/ControlRig/Private/Units/Hierarchy/RigUnit_Metadata.h
+++ b/Engine/Plugins/Animation/ControlRig/Source/ControlRig/Private/Units/Hierarchy/RigUnit_Metadata.h
@@ -22,6 +22,7 @@ struct CONTROLRIG_API FRigDispatch_MetadataBase : public FRigDispatchFactory
 	virtual bool IsSetMetadata() const { return false; }
 
 #if WITH_EDITOR
+	virtual FString GetArgumentDefaultValue(const FName& InArgumentName, TRigVMTypeIndex InTypeIndex) const override;
 	virtual FString GetArgumentMetaData(const FName& InArgumentName, const FName& InMetaDataKey) const override;
 	virtual FText GetArgumentTooltip(const FName& InArgumentName, TRigVMTypeIndex InTypeIndex) const override;
 #endif
@@ -34,6 +35,7 @@ protected:
 	
 	mutable int32 ItemArgIndex = INDEX_NONE;
 	mutable int32 NameArgIndex = INDEX_NONE;
+	mutable int32 UseNameSpaceArgIndex = INDEX_NONE;
 	mutable int32 CacheArgIndex = INDEX_NONE;
 	mutable int32 DefaultArgIndex = INDEX_NONE;
 	mutable int32 ValueArgIndex = INDEX_NONE;
@@ -41,6 +43,7 @@ protected:
 	mutable int32 SuccessArgIndex = INDEX_NONE;
 	static FName ItemArgName;
 	static FName NameArgName;
+	static FName UseNameSpaceArgName;
 	static FName CacheArgName;
 	static FName DefaultArgName;
 	static FName ValueArgName;
@@ -60,7 +63,7 @@ struct CONTROLRIG_API FRigDispatch_GetMetadata : public FRigDispatch_MetadataBas
 
 protected:
 
-	static FRigBaseMetadata* FindMetadata(const FRigVMExtendedExecuteContext& InContext, const FRigElementKey& InKey, const FName& InName, ERigMetadataType InType, FCachedRigElement& Cache);
+	static FRigBaseMetadata* FindMetadata(const FRigVMExtendedExecuteContext& InContext, const FRigElementKey& InKey, const FName& InName, ERigMetadataType InType, bool bUseNameSpace, FCachedRigElement& Cache);
 	virtual FRigVMFunctionPtr GetDispatchFunctionImpl(const FRigVMTemplateTypeMap& InTypes) const override;
 
 #if WITH_EDITOR
@@ -91,13 +94,14 @@ protected:
 		// unpack the memory
 		const FRigElementKey& Item = *(const FRigElementKey*)Handles[Factory->ItemArgIndex].GetData();
 		const FName& Name = *(const FName*)Handles[Factory->NameArgIndex].GetData();
+		const bool bUseNameSpace = *(const bool*)Handles[Factory->UseNameSpaceArgIndex].GetData();
 		FCachedRigElement& Cache = *(FCachedRigElement*)Handles[Factory->CacheArgIndex].GetData(false, InContext.GetSlice().GetIndex());
 		const ValueType& Default = *(const ValueType*)Handles[Factory->DefaultArgIndex].GetData();
 		ValueType& Value = *(ValueType*)Handles[Factory->ValueArgIndex].GetData();
 		bool& Found = *(bool*)Handles[Factory->FoundArgIndex].GetData();
 
 		// extract the metadata
-		if (MetadataType* Md = Cast<MetadataType>(FindMetadata(InContext, Item, Name, EnumValue, Cache)))
+		if (MetadataType* Md = Cast<MetadataType>(FindMetadata(InContext, Item, Name, EnumValue, bUseNameSpace, Cache)))
 		{
 			Value = Md->GetValue();
 			Found = true;
@@ -124,7 +128,7 @@ struct CONTROLRIG_API FRigDispatch_SetMetadata : public FRigDispatch_MetadataBas
 
 protected:
 
-	static FRigBaseMetadata* FindOrAddMetadata(FControlRigExecuteContext& InContext, const FRigElementKey& InKey, const FName& InName, ERigMetadataType InType, FCachedRigElement& Cache);
+	static FRigBaseMetadata* FindOrAddMetadata(FControlRigExecuteContext& InContext, const FRigElementKey& InKey, const FName& InName, ERigMetadataType InType, bool bUseNameSpace, FCachedRigElement& Cache);
 	virtual FRigVMFunctionPtr GetDispatchFunctionImpl(const FRigVMTemplateTypeMap& InTypes) const override;
 
 #if WITH_EDITOR
@@ -155,12 +159,13 @@ protected:
 		FControlRigExecuteContext& ExecuteContext = InContext.GetPublicData<FControlRigExecuteContext>();
 		const FRigElementKey& Item = *(const FRigElementKey*)Handles[Factory->ItemArgIndex].GetData();
 		const FName& Name = *(const FName*)Handles[Factory->NameArgIndex].GetData();
+		const bool bUseNameSpace = *(const bool*)Handles[Factory->UseNameSpaceArgIndex].GetData();
 		FCachedRigElement& Cache = *(FCachedRigElement*)Handles[Factory->CacheArgIndex].GetData(false, InContext.GetSlice().GetIndex());
 		ValueType& Value = *(ValueType*)Handles[Factory->ValueArgIndex].GetData();
 		bool& Success = *(bool*)Handles[Factory->SuccessArgIndex].GetData();
 
 		// extract the metadata
-		if (MetadataType* Md = Cast<MetadataType>(FindOrAddMetadata(ExecuteContext, Item, Name, EnumValue, Cache)))
+		if (MetadataType* Md = Cast<MetadataType>(FindOrAddMetadata(ExecuteContext, Item, Name, EnumValue, bUseNameSpace, Cache)))
 		{
 			Md->GetValue() = Value;
 			Success = true;
@@ -183,6 +188,7 @@ struct CONTROLRIG_API FRigUnit_RemoveMetadata : public FRigUnitMutable
 	FRigUnit_RemoveMetadata()
 		: Item(NAME_None, ERigElementType::Bone)
 		, Name(NAME_None)
+		, UseNameSpace(true)
 		, Removed(false)
 		, CachedIndex()
 	{}
@@ -202,6 +208,12 @@ struct CONTROLRIG_API FRigUnit_RemoveMetadata : public FRigUnitMutable
 	UPROPERTY(meta = (Input, CustomWidget="MetadataName"))
 	FName Name;
 
+	/**
+	 * If checked the metadata will be available only in this name space / module
+	 */
+	UPROPERTY(meta = (Input))
+	bool UseNameSpace;
+
 	// True if the metadata has been removed
 	UPROPERTY(meta=(Output))
 	bool Removed;
@@ -221,6 +233,7 @@ struct CONTROLRIG_API FRigUnit_RemoveAllMetadata : public FRigUnitMutable
 
 	FRigUnit_RemoveAllMetadata()
 		: Item(NAME_None, ERigElementType::Bone)
+		, UseNameSpace(true)
 		, Removed(false)
 		, CachedIndex()
 	{}
@@ -234,6 +247,12 @@ struct CONTROLRIG_API FRigUnit_RemoveAllMetadata : public FRigUnitMutable
 	UPROPERTY(meta = (Input, ExpandByDefault))
 	FRigElementKey Item;
 
+	/**
+	 * If checked the metadata will be available only in this name space / module
+	 */
+	UPROPERTY(meta = (Input))
+	bool UseNameSpace;
+
 	// True if any metadata has been removed
 	UPROPERTY(meta=(Output))
 	bool Removed;
@@ -255,6 +274,7 @@ struct CONTROLRIG_API FRigUnit_HasMetadata : public FRigUnit
 		: Item(NAME_None, ERigElementType::Bone)
 		, Name(NAME_None)
 		, Type(ERigMetadataType::Float)
+		, UseNameSpace(true)
 		, Found(false)
 		, CachedIndex()
 	{}
@@ -280,6 +300,12 @@ struct CONTROLRIG_API FRigUnit_HasMetadata : public FRigUnit
 	UPROPERTY(meta = (Input))
 	ERigMetadataType Type;
 
+	/**
+	 * If checked the metadata will be available only in this name space / module
+	 */
+	UPROPERTY(meta = (Input))
+	bool UseNameSpace;
+
 	// True if the item has the metadata
 	UPROPERTY(meta=(Output))
 	bool Found;
@@ -300,6 +326,7 @@ struct CONTROLRIG_API FRigUnit_FindItemsWithMetadata : public FRigUnit
 	FRigUnit_FindItemsWithMetadata()
 		: Name(NAME_None)
 		, Type(ERigMetadataType::Float)
+		, UseNameSpace(true)
 		, Items()
 	{}
 
@@ -318,6 +345,12 @@ struct CONTROLRIG_API FRigUnit_FindItemsWithMetadata : public FRigUnit
 	UPROPERTY(meta = (Input))
 	ERigMetadataType Type;
 
+	/**
+	 * If checked the metadata will be available only in this name space / module
+	 */
+	UPROPERTY(meta = (Input))
+	bool UseNameSpace;
+
 	// The items containing the metadata
 	UPROPERTY(meta=(Output))
 	TArray<FRigElementKey> Items;
@@ -368,6 +401,7 @@ struct CONTROLRIG_API FRigUnit_SetMetadataTag : public FRigUnitMutable
 	FRigUnit_SetMetadataTag()
 		: Item(NAME_None, ERigElementType::Bone)
 		, Tag(NAME_None)
+		, UseNameSpace(true)
 		, CachedIndex()
 	{}
 
@@ -386,6 +420,12 @@ struct CONTROLRIG_API FRigUnit_SetMetadataTag : public FRigUnitMutable
 	UPROPERTY(meta = (Input, CustomWidget="MetadataTagName"))
 	FName Tag;
 
+	/**
+	 * If checked the metadata will be available only in this name space / module
+	 */
+	UPROPERTY(meta = (Input))
+	bool UseNameSpace;
+
 	// Used to cache the internally used index
 	UPROPERTY()
 	FCachedRigElement CachedIndex;
@@ -402,6 +442,7 @@ struct CONTROLRIG_API FRigUnit_SetMetadataTagArray : public FRigUnitMutable
 	FRigUnit_SetMetadataTagArray()
 		: Item(NAME_None, ERigElementType::Bone)
 		, Tags()
+		, UseNameSpace(true)
 		, CachedIndex()
 	{}
 
@@ -419,6 +460,12 @@ struct CONTROLRIG_API FRigUnit_SetMetadataTagArray : public FRigUnitMutable
 	 */ 
 	UPROPERTY(meta = (Input, CustomWidget="MetadataTagName"))
 	TArray<FName> Tags;
+	
+	/**
+	 * If checked the metadata will be available only in this name space / module
+	 */
+	UPROPERTY(meta = (Input))
+	bool UseNameSpace;
 
 	// Used to cache the internally used index
 	UPROPERTY()
@@ -436,6 +483,7 @@ struct CONTROLRIG_API FRigUnit_RemoveMetadataTag : public FRigUnitMutable
 	FRigUnit_RemoveMetadataTag()
 		: Item(NAME_None, ERigElementType::Bone)
 		, Tag(NAME_None)
+		, UseNameSpace(true)
 		, Removed(false)
 		, CachedIndex()
 	{}
@@ -455,6 +503,12 @@ struct CONTROLRIG_API FRigUnit_RemoveMetadataTag : public FRigUnitMutable
 	UPROPERTY(meta = (Input, CustomWidget="MetadataTagName"))
 	FName Tag;
 
+	/**
+	 * If checked the metadata will be available only in this name space / module
+	 */
+	UPROPERTY(meta = (Input))
+	bool UseNameSpace;
+	
 	/**
 	 * Returns true if the removal was successful
 	 */ 
@@ -477,6 +531,7 @@ struct CONTROLRIG_API FRigUnit_HasMetadataTag : public FRigUnit
 	FRigUnit_HasMetadataTag()
 		: Item(NAME_None, ERigElementType::Bone)
 		, Tag(NAME_None)
+		, UseNameSpace(true)
 		, Found(false)
 		, CachedIndex()
 	{}
@@ -496,6 +551,12 @@ struct CONTROLRIG_API FRigUnit_HasMetadataTag : public FRigUnit
 	UPROPERTY(meta = (Input, CustomWidget="MetadataTagName"))
 	FName Tag;
 
+	/**
+	 * If checked the metadata will be available only in this name space / module
+	 */
+	UPROPERTY(meta = (Input))
+	bool UseNameSpace;
+	
 	// True if the item has the metadata
 	UPROPERTY(meta=(Output))
 	bool Found;
@@ -516,6 +577,7 @@ struct CONTROLRIG_API FRigUnit_HasMetadataTagArray : public FRigUnit
 	FRigUnit_HasMetadataTagArray()
 		: Item(NAME_None, ERigElementType::Bone)
 		, Tags()
+		, UseNameSpace(true)
 		, Found(false)
 		, CachedIndex()
 	{}
@@ -535,6 +597,12 @@ struct CONTROLRIG_API FRigUnit_HasMetadataTagArray : public FRigUnit
 	UPROPERTY(meta = (Input, CustomWidget="MetadataTagName"))
 	TArray<FName> Tags;
 
+	/**
+	 * If checked the metadata will be available only in this name space / module
+	 */
+	UPROPERTY(meta = (Input))
+	bool UseNameSpace;
+	
 	// True if the item has the metadata
 	UPROPERTY(meta=(Output))
 	bool Found;
@@ -554,6 +622,7 @@ struct CONTROLRIG_API FRigUnit_FindItemsWithMetadataTag : public FRigUnit
 
 	FRigUnit_FindItemsWithMetadataTag()
 		: Tag(NAME_None)
+		, UseNameSpace(true)
 		, Items()
 	{}
 
@@ -566,6 +635,12 @@ struct CONTROLRIG_API FRigUnit_FindItemsWithMetadataTag : public FRigUnit
 	UPROPERTY(meta = (Input, CustomWidget="MetadataTagName"))
 	FName Tag;
 
+	/**
+	 * If checked the metadata will be available only in this name space / module
+	 */
+	UPROPERTY(meta = (Input))
+	bool UseNameSpace;
+
 	// The items containing the metadata
 	UPROPERTY(meta=(Output))
 	TArray<FRigElementKey> Items;
@@ -581,6 +656,7 @@ struct CONTROLRIG_API FRigUnit_FindItemsWithMetadataTagArray : public FRigUnit
 
 	FRigUnit_FindItemsWithMetadataTagArray()
 		: Tags()
+		, UseNameSpace(true)
 		, Items()
 	{}
 
@@ -593,6 +669,12 @@ struct CONTROLRIG_API FRigUnit_FindItemsWithMetadataTagArray : public FRigUnit
 	UPROPERTY(meta = (Input, CustomWidget="MetadataTagName"))
 	TArray<FName> Tags;
 
+	/**
+	 * If checked the metadata will be available only in this name space / module
+	 */
+	UPROPERTY(meta = (Input))
+	bool UseNameSpace;
+
 	// The items containing the metadata
 	UPROPERTY(meta=(Output))
 	TArray<FRigElementKey> Items;
@@ -609,6 +691,7 @@ struct CONTROLRIG_API FRigUnit_FilterItemsByMetadataTags : public FRigUnit
 	FRigUnit_FilterItemsByMetadataTags()
 		: Items()
 		, Tags()
+		, UseNameSpace(true)
 		, Inclusive(true)
 		, Result()
 	{}
@@ -626,6 +709,12 @@ struct CONTROLRIG_API FRigUnit_FilterItemsByMetadataTags : public FRigUnit
 	UPROPERTY(meta = (Input, CustomWidget="MetadataTagName"))
 	TArray<FName> Tags;
 
+	/**
+	 * If checked the metadata will be available only in this name space / module
+	 */
+	UPROPERTY(meta = (Input))
+	bool UseNameSpace;
+
 	/**
      * If set to true only items with ALL of tags will be returned,
      * if set to false items with ANY of the tags will be removed
diff --git a/Engine/Plugins/Animation/ControlRig/Source/ControlRig/Private/Units/RigUnitContext.cpp b/Engine/Plugins/Animation/ControlRig/Source/ControlRig/Private/Units/RigUnitContext.cpp
--- a/Engine/Plugins/Animation/ControlRig/Source/ControlRig/Private/Units/RigUnitContext.cpp
+++ b/Engine/Plugins/Animation/ControlRig/Source/ControlRig/Private/Units/RigUnitContext.cpp
@@ -48,6 +48,30 @@ FString FControlRigExecuteContext::RemoveRigModuleNameSpace(const FString& InNam
 	return InName;
 }
 
+FName FControlRigExecuteContext::AdaptMetadataName(bool bUseNameSpace, const FName& InMetadataName) const
+{
+	// only if we are within a rig module let's adapt the meta data name
+	if(bUseNameSpace && IsRigModule() && !InMetadataName.IsNone())
+	{
+		// if the metadata name already contains a namespace - we are just going
+		// to use it as is. this means that modules have access to other module's metadata,
+		// and that's ok. the user will require the full path to it anyway so it is a
+		// conscious user decision.
+		const FString MetadataNameString = InMetadataName.ToString();
+		int32 Index = INDEX_NONE;
+		if(MetadataNameString.FindChar(TEXT(':'), Index))
+		{
+			return InMetadataName;
+		}
+
+		// prefix the meta data name with the namespace to allow modules to store their
+		// metadata in a way that doesn't collide with other modules' metadata.
+		const FString JoinedMetadataName = GetRigModuleNameSpace() + MetadataNameString;
+		return *JoinedMetadataName;
+	}
+	return InMetadataName;
+}
+
 FControlRigExecuteContextRigModuleGuard::FControlRigExecuteContextRigModuleGuard(FControlRigExecuteContext& InContext, const UControlRig* InControlRig)
 	: Context(InContext)
 	, PreviousRigModuleNameSpace(InContext.RigModuleNameSpace)
diff --git a/Engine/Plugins/Animation/ControlRig/Source/ControlRig/Public/Units/RigUnitContext.h b/Engine/Plugins/Animation/ControlRig/Source/ControlRig/Public/Units/RigUnitContext.h
--- a/Engine/Plugins/Animation/ControlRig/Source/ControlRig/Public/Units/RigUnitContext.h
+++ b/Engine/Plugins/Animation/ControlRig/Source/ControlRig/Public/Units/RigUnitContext.h
@@ -153,16 +153,27 @@ public:
 	FName RemoveRigModuleNameSpace(const FName& InName) const;
 	FString RemoveRigModuleNameSpace(const FString& InName) const;
 
+	/**
+	 * Returns true if this context is used on a module currently
+	 */
 	bool IsRigModule() const
 	{
 		return !GetRigModuleNameSpace().IsEmpty();
 	}
 
+	/**
+	 * Returns the namespace of the currently running rig module
+	 */
 	FString GetRigModuleNameSpace() const
 	{
 		return RigModuleNameSpace;
 	}
 
+	/**
+	 * Adapts a metadata name according to rig module namespace.
+	 */
+	FName AdaptMetadataName(bool bUseNameSpace, const FName& InMetadataName) const;
+
 	/** The list of available asset user data object */
 	TArray<const UAssetUserData*> AssetUserData;
 
diff --git a/Engine/Plugins/Animation/ControlRig/Source/ControlRigEditor/Private/Graph/ControlRigGraphPanelPinFactory.cpp b/Engine/Plugins/Animation/ControlRig/Source/ControlRigEditor/Private/Graph/ControlRigGraphPanelPinFactory.cpp
--- a/Engine/Plugins/Animation/ControlRig/Source/ControlRigEditor/Private/Graph/ControlRigGraphPanelPinFactory.cpp
+++ b/Engine/Plugins/Animation/ControlRig/Source/ControlRigEditor/Private/Graph/ControlRigGraphPanelPinFactory.cpp
@@ -19,6 +19,7 @@
 #include "RigVMModel/Nodes/RigVMUnitNode.h"
 #include "RigVMCore/RigVMExecuteContext.h"
 #include "Units/Execution/RigUnit_DynamicHierarchy.h"
+#include "Units/Hierarchy/RigUnit_Metadata.h"
 #include "ControlRigElementDetails.h"
 #include "IPropertyAccessEditor.h"
 
@@ -333,7 +334,19 @@ TSharedPtr<SGraphPin> FControlRigGraphPanelPinFactory::CreatePin_Internal(UEdGra
 								if(UControlRig* ControlRig = Cast<UControlRig>(Blueprint->GetObjectBeingDebugged()))
 								{
 									const FString MapHash = Blueprint->GetPathName();
-									const int32 MetadataVersion = ControlRig->GetHierarchy()->GetMetadataVersion(); 
+									const int32 MetadataVersion = ControlRig->GetHierarchy()->GetMetadataVersion();
+									
+									bool bUseShortNames = false;
+									if (const URigVMNode* ModelNode = InPin->GetNode())
+									{
+										if(const URigVMPin* UseNameSpacePin = ModelNode->FindPin(GET_MEMBER_NAME_STRING_CHECKED(FRigUnit_HasMetadata, UseNameSpace)))
+										{
+											if(UseNameSpacePin->GetDefaultValue().Equals(TEXT("true"), ESearchCase::IgnoreCase))
+											{
+												bUseShortNames = true;
+											}
+										}
+									}
 
 									static TMap<FString, FCachedMetadataNames> MetadataNameLists;
 									FCachedMetadataNames& MetadataNames = MetadataNameLists.FindOrAdd(MapHash);
@@ -358,7 +371,16 @@ TSharedPtr<SGraphPin> FControlRigGraphPanelPinFactory::CreatePin_Internal(UEdGra
 
 										for(const FName& Name : Names)
 										{
-											MetadataNames.Names->Add(MakeShareable(new FRigVMStringWithTag(Name.ToString())));
+											FString NameString = Name.ToString();
+											if(bUseShortNames)
+											{
+												int32 Index = INDEX_NONE;
+												if(NameString.FindChar(TEXT(':'), Index))
+												{
+													NameString.MidInline(Index + 1);
+												}
+											}
+											MetadataNames.Names->Add(MakeShareable(new FRigVMStringWithTag(NameString)));
 										}
 
 										MetadataNames.Names->Sort([](const TSharedPtr<FRigVMStringWithTag>& A, const TSharedPtr<FRigVMStringWithTag>& B)
@@ -402,6 +424,18 @@ TSharedPtr<SGraphPin> FControlRigGraphPanelPinFactory::CreatePin_Internal(UEdGra
 									const FString MapHash = Blueprint->GetPathName();
 									const int32 MetadataTagVersion = ControlRig->GetHierarchy()->GetMetadataTagVersion(); 
 
+									bool bUseShortNames = false;
+									if (const URigVMNode* ModelNode = InPin->GetNode())
+									{
+										if(const URigVMPin* UseNameSpacePin = ModelNode->FindPin(GET_MEMBER_NAME_STRING_CHECKED(FRigUnit_HasMetadata, UseNameSpace)))
+										{
+											if(UseNameSpacePin->GetDefaultValue().Equals(TEXT("true"), ESearchCase::IgnoreCase))
+											{
+												bUseShortNames = true;
+											}
+										}
+									}
+
 									static TMap<FString, FCachedMetadataTagNames> MetadataTagNameLists;
 									FCachedMetadataTagNames& MetadataTagNames = MetadataTagNameLists.FindOrAdd(MapHash);
 
@@ -428,7 +462,16 @@ TSharedPtr<SGraphPin> FControlRigGraphPanelPinFactory::CreatePin_Internal(UEdGra
 
 										for(const FName& Tag : Tags)
 										{
-											MetadataTagNames.Names->Add(MakeShareable(new FRigVMStringWithTag(Tag.ToString())));
+											FString TagString = Tag.ToString();
+											if(bUseShortNames)
+											{
+												int32 Index = INDEX_NONE;
+												if(TagString.FindChar(TEXT(':'), Index))
+												{
+													TagString.MidInline(Index + 1);
+												}
+											}
+											MetadataTagNames.Names->Add(MakeShareable(new FRigVMStringWithTag(TagString)));
 										}
 										MetadataTagNames.Names->Sort([](const TSharedPtr<FRigVMStringWithTag>& A, const TSharedPtr<FRigVMStringWithTag>& B)
 										{
