Constraints: ensure that we cannot constrain the skeletal mesh component bounded to the CR to one of its controls

#rnx

diff --git a/Engine/Plugins/Animation/ControlRig/Source/ControlRig/Private/Constraints/ControlRigTransformableHandle.cpp b/Engine/Plugins/Animation/ControlRig/Source/ControlRig/Private/Constraints/ControlRigTransformableHandle.cpp
--- a/Engine/Plugins/Animation/ControlRig/Source/ControlRig/Private/Constraints/ControlRigTransformableHandle.cpp
+++ b/Engine/Plugins/Animation/ControlRig/Source/ControlRig/Private/Constraints/ControlRigTransformableHandle.cpp
@@ -205,6 +205,12 @@ bool UTransformableControlHandle::HasDirectDependencyWith(const UTransformableHa
 	// check whether the other handle is one of the skeletal mesh parent
 	if (const USkeletalMeshComponent* SkeletalMeshComponent = GetSkeletalMesh())
 	{
+		if (GetTypeHash(SkeletalMeshComponent) == OtherHash)
+		{
+			// we cannot constrain the skeletal mesh component to one of ControlRig's controls
+			return true;
+		}
+		
 		for (const USceneComponent* Comp=SkeletalMeshComponent->GetAttachParent(); Comp!=nullptr; Comp=Comp->GetAttachParent() )
 		{
 			const uint32 AttachParentHash = GetTypeHash(Comp);
