Control Rig: Add ability to create additive rig tracks through python
#rnx

diff --git a/Engine/Plugins/Animation/ControlRig/Source/ControlRigEditor/Private/ControlRigSequencerEditorLibrary.cpp b/Engine/Plugins/Animation/ControlRig/Source/ControlRigEditor/Private/ControlRigSequencerEditorLibrary.cpp
--- a/Engine/Plugins/Animation/ControlRig/Source/ControlRigEditor/Private/ControlRigSequencerEditorLibrary.cpp
+++ b/Engine/Plugins/Animation/ControlRig/Source/ControlRigEditor/Private/ControlRigSequencerEditorLibrary.cpp
@@ -199,79 +199,116 @@ static TSharedPtr<ISequencer> GetSequencerFromAsset()
 	return Sequencer;
 }
 
-static UMovieSceneControlRigParameterTrack* AddControlRig(ULevelSequence* LevelSequence,const UClass* InClass, UObject* BoundActor, FGuid ObjectBinding, UControlRig* InExistingControlRig)
+static UMovieSceneControlRigParameterTrack* AddControlRig(ULevelSequence* LevelSequence,const UClass* InClass, UObject* BoundActor, FGuid ObjectBinding, UControlRig* InExistingControlRig, bool bIsAdditiveControlRig)
 {
 	FSlateApplication::Get().DismissAllMenus();
+	if (!InClass || !InClass->IsChildOf(UControlRig::StaticClass()) ||
+		!LevelSequence || !LevelSequence->GetMovieScene())
+	{
+		return nullptr;
+	}
+	
+	UMovieScene* OwnerMovieScene = LevelSequence->GetMovieScene();
+	TSharedPtr<ISequencer> SharedSequencer = GetSequencerFromAsset();
+	ISequencer* Sequencer = nullptr; // will be valid  if we have a ISequencer AND it's focused.
+	if (SharedSequencer.IsValid() && SharedSequencer->GetFocusedMovieSceneSequence() == LevelSequence)
+	{
+		Sequencer = SharedSequencer.Get();
+	}
+	LevelSequence->Modify();
+	OwnerMovieScene->Modify();
+	
+	if (bIsAdditiveControlRig && InClass != UFKControlRig::StaticClass() && !InClass->GetDefaultObject<UControlRig>()->SupportsEvent(FRigUnit_InverseExecution::EventName))
+	{
+		UE_LOG(LogControlRigEditor, Error, TEXT("Cannot add an additive control rig which does not contain a backwards solve event."));
+		return nullptr;
+	}
+	
+	FScopedTransaction AddControlRigTrackTransaction(LOCTEXT("AddControlRigTrack", "Add Control Rig Track"));
 
-	if (InClass && InClass->IsChildOf(UControlRig::StaticClass()) && LevelSequence && LevelSequence->GetMovieScene())
+	UMovieSceneControlRigParameterTrack* Track = Cast<UMovieSceneControlRigParameterTrack>(OwnerMovieScene->AddTrack(UMovieSceneControlRigParameterTrack::StaticClass(), ObjectBinding));
+	if (Track)
 	{
-		UMovieScene* OwnerMovieScene = LevelSequence->GetMovieScene();
-		TSharedPtr<ISequencer> SharedSequencer = GetSequencerFromAsset();
-		ISequencer* Sequencer = nullptr; // will be valid  if we have a ISequencer AND it's focused.
-		if (SharedSequencer.IsValid() && SharedSequencer->GetFocusedMovieSceneSequence() == LevelSequence)
+		FString ObjectName = InClass->GetName(); //GetDisplayNameText().ToString();
+		ObjectName.RemoveFromEnd(TEXT("_C"));
+
+		bool bSequencerOwnsControlRig = false;
+		UControlRig* ControlRig = InExistingControlRig;
+		if (ControlRig == nullptr)
 		{
-			Sequencer = SharedSequencer.Get();
+			ControlRig = NewObject<UControlRig>(Track, InClass, FName(*ObjectName), RF_Transactional);
+			bSequencerOwnsControlRig = true;
 		}
-		LevelSequence->Modify();
-		OwnerMovieScene->Modify();
-		UMovieSceneControlRigParameterTrack* Track = Cast<UMovieSceneControlRigParameterTrack>(OwnerMovieScene->AddTrack(UMovieSceneControlRigParameterTrack::StaticClass(), ObjectBinding));
-		if (Track)
-		{
-			FString ObjectName = InClass->GetName(); //GetDisplayNameText().ToString();
-			ObjectName.RemoveFromEnd(TEXT("_C"));
 
-			bool bSequencerOwnsControlRig = false;
-			UControlRig* ControlRig = InExistingControlRig;
-			if (ControlRig == nullptr)
+		ControlRig->Modify();
+		if (UFKControlRig* FKControlRig = Cast<UFKControlRig>(ControlRig))
+		{
+			if (bIsAdditiveControlRig)
 			{
-				ControlRig = NewObject<UControlRig>(Track, InClass, FName(*ObjectName), RF_Transactional);
-				bSequencerOwnsControlRig = true;
+				FKControlRig->SetApplyMode(EControlRigFKRigExecuteMode::Additive);
 			}
-
-			ControlRig->Modify();
-			ControlRig->SetObjectBinding(MakeShared<FControlRigObjectBinding>());
-			ControlRig->GetObjectBinding()->BindToObject(BoundActor);
-			ControlRig->GetDataSourceRegistry()->RegisterDataSource(UControlRig::OwnerComponent, ControlRig->GetObjectBinding()->GetBoundObject());
+		}
+		else
+		{
+			ControlRig->SetIsAdditive(bIsAdditiveControlRig);
+		}
+		ControlRig->SetObjectBinding(MakeShared<FControlRigObjectBinding>());
+		ControlRig->GetObjectBinding()->BindToObject(BoundActor);
+		ControlRig->GetDataSourceRegistry()->RegisterDataSource(UControlRig::OwnerComponent, ControlRig->GetObjectBinding()->GetBoundObject());
+		// Do not re-initialize existing control rig
+		if (!InExistingControlRig)
+		{
 			ControlRig->Initialize();
-			ControlRig->Evaluate_AnyThread();
+		}
+		ControlRig->Evaluate_AnyThread();
 
+		SharedSequencer->NotifyMovieSceneDataChanged(EMovieSceneDataChangeType::MovieSceneStructureItemsChanged);
 
-			Track->Modify();
-			UMovieSceneSection* NewSection = Track->CreateControlRigSection(0, ControlRig, bSequencerOwnsControlRig);
-			NewSection->Modify();
+		Track->Modify();
+		UMovieSceneSection* NewSection = Track->CreateControlRigSection(0, ControlRig, bSequencerOwnsControlRig);
+		NewSection->Modify();
 
+		if (bIsAdditiveControlRig)
+		{
+			const FString AdditiveObjectName = ObjectName + TEXT(" (Additive)");
+			Track->SetTrackName(FName(*ObjectName));
+			Track->SetDisplayName(FText::FromString(AdditiveObjectName));
+			Track->SetColorTint(UMovieSceneControlRigParameterTrack::AdditiveRigTrackColor);
+		}
+		else
+		{
 			//mz todo need to have multiple rigs with same class
 			Track->SetTrackName(FName(*ObjectName));
 			Track->SetDisplayName(FText::FromString(ObjectName));
+			Track->SetColorTint(UMovieSceneControlRigParameterTrack::AbsoluteRigTrackColor);
+		}
 
-			if (Sequencer)
-			{
-				Sequencer->NotifyMovieSceneDataChanged(EMovieSceneDataChangeType::MovieSceneStructureItemsChanged);
-				Sequencer->EmptySelection();
-				Sequencer->SelectSection(NewSection);
-				Sequencer->ThrobSectionSelection();
-				Sequencer->NotifyMovieSceneDataChanged(EMovieSceneDataChangeType::MovieSceneStructureItemAdded);
-				Sequencer->ObjectImplicitlyAdded(ControlRig);
-			}
+		if (SharedSequencer.IsValid())
+		{
+			SharedSequencer->EmptySelection();
+			SharedSequencer->SelectSection(NewSection);
+			SharedSequencer->ThrobSectionSelection();
+			SharedSequencer->NotifyMovieSceneDataChanged(EMovieSceneDataChangeType::MovieSceneStructureItemAdded);
+			SharedSequencer->ObjectImplicitlyAdded(ControlRig);
+		}
 
-			FControlRigEditMode* ControlRigEditMode = static_cast<FControlRigEditMode*>(GLevelEditorModeTools().GetActiveMode(FControlRigEditMode::ModeName));
-			if (!ControlRigEditMode)
-			{
-				GLevelEditorModeTools().ActivateMode(FControlRigEditMode::ModeName);
-				ControlRigEditMode = static_cast<FControlRigEditMode*>(GLevelEditorModeTools().GetActiveMode(FControlRigEditMode::ModeName));
+		FControlRigEditMode* ControlRigEditMode = static_cast<FControlRigEditMode*>(GLevelEditorModeTools().GetActiveMode(FControlRigEditMode::ModeName));
+		if (!ControlRigEditMode)
+		{
+			GLevelEditorModeTools().ActivateMode(FControlRigEditMode::ModeName);
+			ControlRigEditMode = static_cast<FControlRigEditMode*>(GLevelEditorModeTools().GetActiveMode(FControlRigEditMode::ModeName));
 
-			}
-			if (ControlRigEditMode)
-			{
-				ControlRigEditMode->AddControlRigObject(ControlRig, SharedSequencer);
-			}
-			return Track;
 		}
+		if (ControlRigEditMode)
+		{
+			ControlRigEditMode->AddControlRigObject(ControlRig, SharedSequencer);
+		}
+		return Track;
 	}
 	return nullptr;
 }
 
-UMovieSceneTrack* UControlRigSequencerEditorLibrary::FindOrCreateControlRigTrack(UWorld* World, ULevelSequence* LevelSequence, const UClass* ControlRigClass, const FMovieSceneBindingProxy& InBinding)
+UMovieSceneTrack* UControlRigSequencerEditorLibrary::FindOrCreateControlRigTrack(UWorld* World, ULevelSequence* LevelSequence, const UClass* ControlRigClass, const FMovieSceneBindingProxy& InBinding, bool bIsAdditiveControlRig)
 {
 	UMovieScene* MovieScene = InBinding.Sequence ? InBinding.Sequence->GetMovieScene() : nullptr;
 	UMovieSceneTrack* BaseTrack = nullptr;
@@ -304,11 +341,11 @@ UMovieSceneTrack* UControlRigSequencerEditorLibrary::FindOrCreateControlRigTrack
 				UControlRig* ControlRig = nullptr;
 				if (Skeleton && SkeletalMeshComponent)
 				{
-					UMovieSceneControlRigParameterTrack* Track = AddControlRig(LevelSequence, ControlRigClass, SkeletalMeshComponent, InBinding.BindingID, nullptr);
+					UMovieSceneControlRigParameterTrack* Track = AddControlRig(LevelSequence, ControlRigClass, SkeletalMeshComponent, InBinding.BindingID, nullptr, bIsAdditiveControlRig);
 
 					if (Track)
 					{
-						BaseTrack = Track;								
+						BaseTrack = Track;
 					}
 				}
 			}
@@ -370,7 +407,7 @@ TArray<UMovieSceneTrack*> UControlRigSequencerEditorLibrary::FindOrCreateControl
 
 					if (GoodTrack == nullptr)
 					{
-						GoodTrack = AddControlRig(LevelSequence, CR->GetClass(), BoundActor, InBinding.BindingID, CR);
+						GoodTrack = AddControlRig(LevelSequence, CR->GetClass(), BoundActor, InBinding.BindingID, CR, false);
 					}
 					Tracks.Add(GoodTrack);
 				}
@@ -3011,6 +3048,11 @@ bool UControlRigSequencerEditorLibrary::IsFKControlRig(UControlRig* InControlRig
 	return (InControlRig && InControlRig->IsA<UFKControlRig>());
 }
 
+bool UControlRigSequencerEditorLibrary::IsAdditiveControlRig(UControlRig* InControlRig)
+{
+	return (InControlRig && InControlRig->IsAdditive());
+}
+
 EControlRigFKRigExecuteMode UControlRigSequencerEditorLibrary::GetFKControlRigApplyMode(UControlRig* InControlRig)
 {
 	EControlRigFKRigExecuteMode ApplyMode = EControlRigFKRigExecuteMode::Direct;
diff --git a/Engine/Plugins/Animation/ControlRig/Source/ControlRigEditor/Private/Sequencer/ControlRigParameterTrackEditor.cpp b/Engine/Plugins/Animation/ControlRig/Source/ControlRigEditor/Private/Sequencer/ControlRigParameterTrackEditor.cpp
--- a/Engine/Plugins/Animation/ControlRig/Source/ControlRigEditor/Private/Sequencer/ControlRigParameterTrackEditor.cpp
+++ b/Engine/Plugins/Animation/ControlRig/Source/ControlRigEditor/Private/Sequencer/ControlRigParameterTrackEditor.cpp
@@ -92,6 +92,8 @@
 #include "Editor/UnrealEd/Private/FbxExporter.h"
 #include "Sequencer/ControlRigSequencerHelpers.h"
 #include "Widgets/Layout/SSpacer.h"
+#include "ControlRigSequencerEditorLibrary.h"
+#include "LevelSequence.h"
 
 #define LOCTEXT_NAMESPACE "FControlRigParameterTrackEditor"
 
@@ -1264,8 +1266,8 @@ void FControlRigParameterTrackEditor::HandleAddTrackSubMenu(FMenuBuilder& MenuBu
 		EUserInterfaceActionType::ToggleButton);
 
 	MenuBuilder.AddSubMenu(
-		LOCTEXT("AddAssetControlRig", "Asset-Based Control Rig"),
-		LOCTEXT("AddAsetControlRigTooltip", "Adds an asset based Control Rig track"),
+		LOCTEXT("AddControlRigClass", "Control Rig Classes"),
+		LOCTEXT("AddControlRigClassTooltip", "Adds a Control Rig track based on selected class"),
 		FNewMenuDelegate::CreateRaw(this, &FControlRigParameterTrackEditor::HandleAddControlRigSubMenu, ObjectBindings, Track)
 	);
 }
@@ -1427,96 +1429,15 @@ void FControlRigParameterTrackEditor::OnControlRigAssetSelected(const FAssetData
 	}
 */
 
-void FControlRigParameterTrackEditor::AddControlRig(UClass* InClass, UObject* BoundActor, FGuid ObjectBinding, UControlRig* InExistingControlRig)
+void FControlRigParameterTrackEditor::AddControlRig(const UClass* InClass, UObject* BoundActor, FGuid ObjectBinding, UControlRig* InExistingControlRig)
 {
-	FSlateApplication::Get().DismissAllMenus();
-	const TSharedPtr<ISequencer> SequencerParent = GetSequencer();
-
-	if (InClass && InClass->IsChildOf(UControlRig::StaticClass()) && SequencerParent.IsValid())
+	UWorld* World = GCurrentLevelEditingViewportClient ? GCurrentLevelEditingViewportClient->GetWorld() : nullptr;
+	ULevelSequence* LevelSequence = Cast<ULevelSequence>(GetSequencer()->GetFocusedMovieSceneSequence());
+	UMovieSceneSequence* Sequence = GetSequencer()->GetFocusedMovieSceneSequence();
+	FMovieSceneBindingProxy BindingProxy(ObjectBinding, Sequence);
+	if (UMovieSceneTrack* Track = UControlRigSequencerEditorLibrary::FindOrCreateControlRigTrack(World, LevelSequence, InClass, BindingProxy, bIsAdditiveControlRig))
 	{
-		if (bIsAdditiveControlRig && InClass != UFKControlRig::StaticClass() && !InClass->GetDefaultObject<UControlRig>()->SupportsEvent(FRigUnit_InverseExecution::EventName))
-		{
-			UE_LOG(LogControlRigEditor, Error, TEXT("Cannot add an additive control rig which does not contain a backwards solve event."));
-			return;
-		}
-		
-		UMovieSceneSequence* OwnerSequence = GetSequencer()->GetFocusedMovieSceneSequence();
-		UMovieScene* OwnerMovieScene = OwnerSequence->GetMovieScene();
-		FScopedTransaction AddControlRigTrackTransaction(LOCTEXT("AddControlRigTrack", "Add Control Rig Track"));
-
-		OwnerSequence->Modify();
-		OwnerMovieScene->Modify();
-		UMovieSceneControlRigParameterTrack* Track = Cast<UMovieSceneControlRigParameterTrack>(AddTrack(OwnerMovieScene, ObjectBinding, UMovieSceneControlRigParameterTrack::StaticClass(), NAME_None));
-		if (Track)
-		{
-			FString ObjectName = InClass->GetName(); //GetDisplayNameText().ToString();
-			ObjectName.RemoveFromEnd(TEXT("_C"));
-
-			bool bSequencerOwnsControlRig = false;
-			UControlRig* ControlRig = InExistingControlRig;
-			if (ControlRig == nullptr)
-			{
-				ControlRig = NewObject<UControlRig>(Track, InClass, FName(*ObjectName), RF_Transactional);
-				bSequencerOwnsControlRig = true;
-			}
-
-			ControlRig->Modify();
-			if (UFKControlRig* FKControlRig = Cast<UFKControlRig>(ControlRig))
-			{
-				if (bIsAdditiveControlRig)
-				{
-					FKControlRig->SetApplyMode(EControlRigFKRigExecuteMode::Additive);
-				}
-			}
-			else
-			{
-				ControlRig->SetIsAdditive(bIsAdditiveControlRig);
-			}
-			ControlRig->SetObjectBinding(MakeShared<FControlRigObjectBinding>());
-			ControlRig->GetObjectBinding()->BindToObject(BoundActor);
-			ControlRig->GetDataSourceRegistry()->RegisterDataSource(UControlRig::OwnerComponent, ControlRig->GetObjectBinding()->GetBoundObject());
-			// Do not re-initialize existing control rig
-			if (!InExistingControlRig)
-			{
-				ControlRig->Initialize();
-			}
-			ControlRig->Evaluate_AnyThread();
-
-			SequencerParent->NotifyMovieSceneDataChanged(EMovieSceneDataChangeType::MovieSceneStructureItemsChanged);
-
-			Track->Modify();
-			UMovieSceneSection* NewSection = Track->CreateControlRigSection(0, ControlRig, bSequencerOwnsControlRig);
-			NewSection->Modify();
-
-			if (bIsAdditiveControlRig)
-			{
-				const FString AdditiveObjectName = ObjectName + TEXT(" (Additive)");
-				Track->SetTrackName(FName(*ObjectName));
-				Track->SetDisplayName(FText::FromString(AdditiveObjectName));
-				Track->SetColorTint(UMovieSceneControlRigParameterTrack::AdditiveRigTrackColor);
-			}
-			else
-			{
-				//mz todo need to have multiple rigs with same class
-				Track->SetTrackName(FName(*ObjectName));
-				Track->SetDisplayName(FText::FromString(ObjectName));
-				Track->SetColorTint(UMovieSceneControlRigParameterTrack::AbsoluteRigTrackColor);
-			}
-
-			GetSequencer()->EmptySelection();
-			GetSequencer()->SelectSection(NewSection);
-			GetSequencer()->ThrobSectionSelection();
-			GetSequencer()->NotifyMovieSceneDataChanged(EMovieSceneDataChangeType::MovieSceneStructureItemAdded);
-
-			FControlRigEditMode* ControlRigEditMode = GetEditMode(true);
-
-			if (ControlRigEditMode)
-			{
-				ControlRigEditMode->AddControlRigObject(ControlRig, GetSequencer());
-			}
-			BindControlRig(ControlRig);
-
-		}
+		BindControlRig(CastChecked<UMovieSceneControlRigParameterTrack>(Track)->GetControlRig());
 	}
 }
 
diff --git a/Engine/Plugins/Animation/ControlRig/Source/ControlRigEditor/Private/Sequencer/ControlRigParameterTrackEditor.h b/Engine/Plugins/Animation/ControlRig/Source/ControlRigEditor/Private/Sequencer/ControlRigParameterTrackEditor.h
--- a/Engine/Plugins/Animation/ControlRig/Source/ControlRigEditor/Private/Sequencer/ControlRigParameterTrackEditor.h
+++ b/Engine/Plugins/Animation/ControlRig/Source/ControlRigEditor/Private/Sequencer/ControlRigParameterTrackEditor.h
@@ -101,7 +101,7 @@ private:
 
 	/** Control Rig Picked */
 	void AddControlRig(UClass* InClass, UObject* BoundActor, FGuid ObjectBinding);
-	void AddControlRig(UClass* InClass, UObject* BoundActor, FGuid ObjectBinding, UControlRig* InExistingControlRig);
+	void AddControlRig(const UClass* InClass, UObject* BoundActor, FGuid ObjectBinding, UControlRig* InExistingControlRig);
 	void AddControlRigFromComponent(FGuid InGuid);
 	
 	/** Delegate for Selection Changed Event */
diff --git a/Engine/Plugins/Animation/ControlRig/Source/ControlRigEditor/Public/ControlRigSequencerEditorLibrary.h b/Engine/Plugins/Animation/ControlRig/Source/ControlRigEditor/Public/ControlRigSequencerEditorLibrary.h
--- a/Engine/Plugins/Animation/ControlRig/Source/ControlRigEditor/Public/ControlRigSequencerEditorLibrary.h
+++ b/Engine/Plugins/Animation/ControlRig/Source/ControlRigEditor/Public/ControlRigSequencerEditorLibrary.h
@@ -96,7 +96,7 @@ public:
 	* @return returns Return the found or created track
 	*/
 	UFUNCTION(BlueprintCallable, Category = "Editor Scripting | Sequencer Tools | Control Rig")
-	static UMovieSceneTrack* FindOrCreateControlRigTrack(UWorld* World, ULevelSequence* LevelSequence, const UClass* ControlRigClass, const FMovieSceneBindingProxy& InBinding);
+	static UMovieSceneTrack* FindOrCreateControlRigTrack(UWorld* World, ULevelSequence* LevelSequence, const UClass* ControlRigClass, const FMovieSceneBindingProxy& InBinding, bool bIsAdditiveControlRig = false);
 
 	/**
 	* Find or create a Control Rig Component
@@ -1046,6 +1046,12 @@ public:
 	UFUNCTION(BlueprintCallable, Category = "Editor Scripting | Sequencer Tools | Control Rig")
 	static bool IsFKControlRig(UControlRig* InControlRig);
 
+	/**	Whether or not the control rig is an Additive Control Rig.
+	@param InControlRig Rig to test to see if Additive Control Rig
+	**/
+	UFUNCTION(BlueprintCallable, Category = "Editor Scripting | Sequencer Tools | Control Rig")
+	static bool IsAdditiveControlRig(UControlRig* InControlRig);
+
 	/**	Get FKControlRig Apply Mode.
 	@param InControlRig Rig to test
 	@return The EControlRigFKRigExecuteMode mode it is in, either Replace,Additive or Direct
