Control Rig: Fix bug where new rig nodes are not added into the array of ControlRigNodes in the FControlRigLayerInstanceProxy
#rnx

diff --git a/Engine/Plugins/Animation/ControlRig/Source/ControlRig/Private/Sequencer/ControlRigLayerInstanceProxy.cpp b/Engine/Plugins/Animation/ControlRig/Source/ControlRig/Private/Sequencer/ControlRigLayerInstanceProxy.cpp
--- a/Engine/Plugins/Animation/ControlRig/Source/ControlRig/Private/Sequencer/ControlRigLayerInstanceProxy.cpp
+++ b/Engine/Plugins/Animation/ControlRig/Source/ControlRig/Private/Sequencer/ControlRigLayerInstanceProxy.cpp
@@ -248,10 +248,11 @@ void FControlRigLayerInstanceProxy::AddControlRigTrack(int32 ControlRigID, UCont
 				}
 			}
 			else
-			{ 
-				for (int32 Index = 0; Index < ControlRigNodes.Num(); ++Index)
+			{
+				int32 NewNodePosition = 0;
+				for (; NewNodePosition < ControlRigNodes.Num(); ++NewNodePosition)
 				{
-					Node = ControlRigNodes[Index].Get();
+					Node = ControlRigNodes[NewNodePosition].Get();
 	
 					if (!Node || !Node->GetControlRig())
 					{
@@ -265,27 +266,28 @@ void FControlRigLayerInstanceProxy::AddControlRigTrack(int32 ControlRigID, UCont
 					//not additive but this one is so we add it here
 					else if (!InControlRig->IsAdditive() && Node->GetControlRig()->IsAdditive())
 					{
-						Node = ControlRigNodes.Insert_GetRef(MakeShared<FAnimNode_ControlRig_ExternalSource>(), Index).Get();
 						break;
 					}
 					if (UMovieSceneControlRigParameterTrack* OtherTrack = Node->GetControlRig()->GetTypedOuter<UMovieSceneControlRigParameterTrack>())
 					{
 						if (PriorityOrder >= OtherTrack->GetPriorityOrder())
 						{
-							//if last one then add it
-							if (Index == ControlRigNodes.Num() - 1)
-							{
-								Node = ControlRigNodes.Add_GetRef(MakeShared<FAnimNode_ControlRig_ExternalSource>()).Get();
-								break;
-							}
+							continue;
 						}
 						else
 						{
-							Node = ControlRigNodes.Insert_GetRef(MakeShared<FAnimNode_ControlRig_ExternalSource>(), Index).Get();
 							break;
 						}
 					}
 				}
+				if (NewNodePosition >= ControlRigNodes.Num() - 1)
+				{
+					Node = ControlRigNodes.Add_GetRef(MakeShared<FAnimNode_ControlRig_ExternalSource>()).Get();
+				}
+				else
+				{
+					Node = ControlRigNodes.Insert_GetRef(MakeShared<FAnimNode_ControlRig_ExternalSource>(), NewNodePosition).Get();
+				}
 			}
 		}
 		else
