Constraints: Support Control Rig Component

diff --git a/Engine/Plugins/Animation/ControlRig/Source/ControlRig/Private/Constraints/ControlRigTransformableHandle.cpp b/Engine/Plugins/Animation/ControlRig/Source/ControlRig/Private/Constraints/ControlRigTransformableHandle.cpp
--- a/Engine/Plugins/Animation/ControlRig/Source/ControlRig/Private/Constraints/ControlRigTransformableHandle.cpp
+++ b/Engine/Plugins/Animation/ControlRig/Source/ControlRig/Private/Constraints/ControlRigTransformableHandle.cpp
@@ -4,11 +4,10 @@
 #include "Constraints/ControlRigTransformableHandle.h"
 
 #include "Components/SkeletalMeshComponent.h"
+#include "ControlRigComponent.h"
 #include "ControlRig.h"
-#include "ControlRigObjectBinding.h"
 #include "IControlRigObjectBinding.h"
 #include "Rigs/RigHierarchyElements.h"
-#include "Channels/MovieSceneChannelProxy.h"
 #include "Sequencer/MovieSceneControlRigParameterSection.h"
 #include "Sections/MovieScene3DTransformSection.h"
 #include "Sequencer/ControlRigSequencerHelpers.h"
@@ -38,8 +37,8 @@ bool UTransformableControlHandle::IsValid() const
 		return false;
 	}
 
-	const USkeletalMeshComponent* SkeletalMeshComponent = GetSkeletalMesh();
-	if (!SkeletalMeshComponent)
+	const USceneComponent* BoundComponent = GetBoundComponent();
+	if (!BoundComponent)
 	{
 		return false;
 	}
@@ -55,8 +54,7 @@ bool UTransformableControlHandle::IsValid() const
 
 void UTransformableControlHandle::TickForBaking() const
 {
-	USkeletalMeshComponent* SkeletalMeshComponent = GetSkeletalMesh();
-	if (SkeletalMeshComponent)
+	if (const USkeletalMeshComponent* SkeletalMeshComponent = GetSkeletalMesh())
 	{
 		const AActor* Parent = SkeletalMeshComponent->GetOwner();
 		while (Parent)
@@ -77,6 +75,12 @@ void UTransformableControlHandle::TickForBaking() const
 
 			Parent = Parent->GetAttachParentActor();
 		}
+		return;
+	}
+
+	if (UControlRigComponent* ControlRigComponent = GetControlRigComponent())
+	{
+		ControlRigComponent->Update();
 	}
 }
 
@@ -90,18 +94,21 @@ void UTransformableControlHandle::SetGlobalTransform(const FTransform& InGlobal)
 		return;
 	}
 
-	const USkeletalMeshComponent* SkeletalMeshComponent = GetSkeletalMesh();
-	if (!SkeletalMeshComponent)
+	const USceneComponent* BoundComponent = GetBoundComponent();
+	if (!BoundComponent)
 	{
 		return;
 	}
 	
 	const FRigElementKey& ControlKey = ControlElement->GetKey();
-	
-	const FTransform& ComponentTransform = SkeletalMeshComponent->GetComponentTransform();
-	//use this function so we don't set the preferred angles
-	ControlRig->SetControlGlobalTransform(ControlKey.Name, InGlobal.GetRelativeTransform(ComponentTransform), false/*bNotify*/, FRigControlModifiedContext(), false/*bSetupUndo*/, false /*bPrintPython*/, false/* bFixEulerFlips*/);
+	const FTransform& ComponentTransform = BoundComponent->GetComponentTransform();
+
+	static const FRigControlModifiedContext Context;
+	static constexpr bool bNotify = false, bSetupUndo = false, bPrintPython = false, bFixEulerFlips = false;
 
+	//use this function so we don't set the preferred angles
+	ControlRig->SetControlGlobalTransform(ControlKey.Name, InGlobal.GetRelativeTransform(ComponentTransform),
+		bNotify, Context, bSetupUndo, bPrintPython, bFixEulerFlips);
 }
 
 void UTransformableControlHandle::SetLocalTransform(const FTransform& InLocal) const
@@ -129,8 +136,8 @@ FTransform UTransformableControlHandle::GetGlobalTransform() const
 		return FTransform::Identity;
 	}
 	
-	const USkeletalMeshComponent* SkeletalMeshComponent = GetSkeletalMesh();
-	if (!SkeletalMeshComponent)
+	const USceneComponent* BoundComponent = GetBoundComponent();
+	if (!BoundComponent)
 	{
 		return FTransform::Identity;
 	}
@@ -139,7 +146,7 @@ FTransform UTransformableControlHandle::GetGlobalTransform() const
 	const URigHierarchy* Hierarchy = ControlRig->GetHierarchy();
 	const int32 CtrlIndex = Hierarchy->GetIndex(ControlKey);
 
-	const FTransform& ComponentTransform = SkeletalMeshComponent->GetComponentTransform();
+	const FTransform& ComponentTransform = BoundComponent->GetComponentTransform();
 	return Hierarchy->GetGlobalTransform(CtrlIndex) * ComponentTransform;
 }
 
@@ -160,13 +167,13 @@ FTransform UTransformableControlHandle::GetLocalTransform() const
 
 UObject* UTransformableControlHandle::GetPrerequisiteObject() const
 {
-	return GetSkeletalMesh(); 
+	return GetBoundComponent(); 
 }
 
 FTickFunction* UTransformableControlHandle::GetTickFunction() const
 {
-	USkeletalMeshComponent* SkelMeshComponent = GetSkeletalMesh();
-	return SkelMeshComponent ? &SkelMeshComponent->PrimaryComponentTick : nullptr;
+	USceneComponent* BoundComponent = GetBoundComponent();
+	return BoundComponent ? &BoundComponent->PrimaryComponentTick : nullptr;
 }
 
 uint32 UTransformableControlHandle::ComputeHash(const UControlRig* InControlRig, const FName& InControlName)
@@ -185,7 +192,16 @@ uint32 UTransformableControlHandle::GetHash() const
 
 TWeakObjectPtr<UObject> UTransformableControlHandle::GetTarget() const
 {
-	return GetSkeletalMesh();
+	return GetBoundComponent();
+}
+
+USceneComponent* UTransformableControlHandle::GetBoundComponent() const
+{
+	if (USkeletalMeshComponent* SkeletalMeshComponent = GetSkeletalMesh())
+	{
+		return SkeletalMeshComponent;	
+	}
+	return GetControlRigComponent();
 }
 
 USkeletalMeshComponent* UTransformableControlHandle::GetSkeletalMesh() const
@@ -194,6 +210,12 @@ USkeletalMeshComponent* UTransformableControlHandle::GetSkeletalMesh() const
    	return ObjectBinding ? Cast<USkeletalMeshComponent>(ObjectBinding->GetBoundObject()) : nullptr;
 }
 
+UControlRigComponent* UTransformableControlHandle::GetControlRigComponent() const
+{
+	const TSharedPtr<IControlRigObjectBinding> ObjectBinding = ControlRig.IsValid() ? ControlRig->GetObjectBinding() : nullptr;
+	return ObjectBinding ? Cast<UControlRigComponent>(ObjectBinding->GetBoundObject()) : nullptr;
+}
+
 bool UTransformableControlHandle::HasDirectDependencyWith(const UTransformableHandle& InOther) const
 {
 	const uint32 OtherHash = InOther.GetHash();
@@ -203,15 +225,15 @@ bool UTransformableControlHandle::HasDirectDependencyWith(const UTransformableHa
 	}
 
 	// check whether the other handle is one of the skeletal mesh parent
-	if (const USkeletalMeshComponent* SkeletalMeshComponent = GetSkeletalMesh())
+	if (const USceneComponent* BoundComponent = GetBoundComponent())
 	{
-		if (GetTypeHash(SkeletalMeshComponent) == OtherHash)
+		if (GetTypeHash(BoundComponent) == OtherHash)
 		{
 			// we cannot constrain the skeletal mesh component to one of ControlRig's controls
 			return true;
 		}
 		
-		for (const USceneComponent* Comp=SkeletalMeshComponent->GetAttachParent(); Comp!=nullptr; Comp=Comp->GetAttachParent() )
+		for (const USceneComponent* Comp=BoundComponent->GetAttachParent(); Comp!=nullptr; Comp=Comp->GetAttachParent() )
 		{
 			const uint32 AttachParentHash = GetTypeHash(Comp);
 			if (AttachParentHash == OtherHash)
@@ -277,7 +299,7 @@ FTickPrerequisite UTransformableControlHandle::GetPrimaryPrerequisite() const
 {
 	if (FTickFunction* TickFunction = GetTickFunction())
 	{
-		return FTickPrerequisite( GetSkeletalMesh(), *TickFunction); 
+		return FTickPrerequisite( GetBoundComponent(), *TickFunction); 
 	}
 	
 	static const FTickPrerequisite DummyPrerex;
@@ -571,15 +593,15 @@ FString UTransformableControlHandle::GetLabel() const
 
 FString UTransformableControlHandle::GetFullLabel() const
 {
-	const USkeletalMeshComponent* SkeletalMesh = GetSkeletalMesh();
-	if (!SkeletalMesh)
+	const USceneComponent* BoundComponent = GetBoundComponent();
+	if (!BoundComponent)
 	{
 		static const FString DummyLabel;
 		return DummyLabel;
 	}
 	
-	const AActor* Actor = SkeletalMesh->GetOwner();
-	const FString ControlRigLabel = Actor ? Actor->GetActorLabel() : SkeletalMesh->GetName();
+	const AActor* Actor = BoundComponent->GetOwner();
+	const FString ControlRigLabel = Actor ? Actor->GetActorLabel() : BoundComponent->GetName();
 	return FString::Printf(TEXT("%s/%s"), *ControlRigLabel, *ControlName.ToString() );
 }
 
diff --git a/Engine/Plugins/Animation/ControlRig/Source/ControlRig/Public/Constraints/ControlRigTransformableHandle.h b/Engine/Plugins/Animation/ControlRig/Source/ControlRig/Public/Constraints/ControlRigTransformableHandle.h
--- a/Engine/Plugins/Animation/ControlRig/Source/ControlRig/Public/Constraints/ControlRigTransformableHandle.h
+++ b/Engine/Plugins/Animation/ControlRig/Source/ControlRig/Public/Constraints/ControlRigTransformableHandle.h
@@ -13,6 +13,7 @@ struct FRigControlElement;
 
 class UControlRig;
 class USkeletalMeshComponent;
+class UControlRigComponent;
 class URigHierarchy;
 
 /**
@@ -74,9 +75,6 @@ public:
 	/** tick things when baking*/ 
 	virtual void TickForBaking() const override;
 
-	/** Returns the skeletal mesh bound to ControlRig. */
-	USkeletalMeshComponent* GetSkeletalMesh() const;
-
 	/** Registers/Unregisters useful delegates to track changes in the control's transform. */
 	void UnregisterDelegates() const;
 	void RegisterDelegates();
@@ -109,7 +107,12 @@ public:
 	
 private:
 
-	/** @todo document */
+	/** Returns the component bounded to ControlRig. */
+	USceneComponent* GetBoundComponent() const;
+	USkeletalMeshComponent* GetSkeletalMesh() const;
+	UControlRigComponent* GetControlRigComponent() const;
+	
+	/** Handles notifications coming from the ControlRig's hierarchy */
 	void OnHierarchyModified(
 		ERigHierarchyNotification InNotif,
 		URigHierarchy* InHierarchy,
