Sequencer: New Tween Tools for Blend Relative, Blend To Ease, and Smooth/Rough

diff --git a/Engine/Plugins/Animation/ControlRig/Source/ControlRigEditor/Private/ControlRigSequencerEditorLibrary.cpp b/Engine/Plugins/Animation/ControlRig/Source/ControlRigEditor/Private/ControlRigSequencerEditorLibrary.cpp
--- a/Engine/Plugins/Animation/ControlRig/Source/ControlRigEditor/Private/ControlRigSequencerEditorLibrary.cpp
+++ b/Engine/Plugins/Animation/ControlRig/Source/ControlRigEditor/Private/ControlRigSequencerEditorLibrary.cpp
@@ -480,6 +480,27 @@ bool UControlRigSequencerEditorLibrary::BlendValuesOnSelected(ULevelSequence* Le
 					BlendTool.Blend(WeakSequencer, BlendValue);
 					return true;
 				}
+			case EAnimToolBlendOperation::BlendRelative:
+				{
+					FBlendRelativeSlider BlendTool;
+					BlendTool.Setup(WeakSequencer, WeakMode);
+					BlendTool.Blend(WeakSequencer, BlendValue);
+					return true;
+				}
+			case EAnimToolBlendOperation::BlendToEase:
+				{
+					FBlendToEaseSlider BlendTool;
+					BlendTool.Setup(WeakSequencer, WeakMode);
+					BlendTool.Blend(WeakSequencer, BlendValue);
+					return true;
+				}
+			case EAnimToolBlendOperation::SmoothRough:
+				{
+					FSmoothRoughSlider BlendTool;
+					BlendTool.Setup(WeakSequencer, WeakMode);
+					BlendTool.Blend(WeakSequencer, BlendValue);
+					return true;
+				}
 			}
 		}
 	}
diff --git a/Engine/Plugins/Animation/ControlRig/Source/ControlRigEditor/Private/EditMode/SControlRigTweenWidget.cpp b/Engine/Plugins/Animation/ControlRig/Source/ControlRigEditor/Private/EditMode/SControlRigTweenWidget.cpp
--- a/Engine/Plugins/Animation/ControlRig/Source/ControlRigEditor/Private/EditMode/SControlRigTweenWidget.cpp
+++ b/Engine/Plugins/Animation/ControlRig/Source/ControlRigEditor/Private/EditMode/SControlRigTweenWidget.cpp
@@ -252,6 +252,12 @@ void SControlRigTweenWidget::Construct(const FArguments& InArgs)
 	
 	TSharedPtr<FBaseAnimSlider> TweenPtr = MakeShareable(new FControlsToTween());
 	AnimBlendTools.RegisterAnimSlider(TweenPtr);
+	TSharedPtr<FBaseAnimSlider> BlendRelativePtr = MakeShareable(new FBlendRelativeSlider());
+	AnimBlendTools.RegisterAnimSlider(BlendRelativePtr);
+	TSharedPtr<FBaseAnimSlider> BlendToEasePtr = MakeShareable(new FBlendToEaseSlider());
+	AnimBlendTools.RegisterAnimSlider(BlendToEasePtr);
+	TSharedPtr<FBaseAnimSlider> SmoothRoughPtr = MakeShareable(new FSmoothRoughSlider());
+	AnimBlendTools.RegisterAnimSlider(SmoothRoughPtr);
 
 	// Combo Button to swap sliders 
 	TSharedRef<SComboButton> SliderComoboBtn = SNew(SComboButton)
diff --git a/Engine/Plugins/Animation/ControlRig/Source/ControlRigEditor/Private/EditMode/SControlRigTweenWidget.h b/Engine/Plugins/Animation/ControlRig/Source/ControlRigEditor/Private/EditMode/SControlRigTweenWidget.h
--- a/Engine/Plugins/Animation/ControlRig/Source/ControlRigEditor/Private/EditMode/SControlRigTweenWidget.h
+++ b/Engine/Plugins/Animation/ControlRig/Source/ControlRigEditor/Private/EditMode/SControlRigTweenWidget.h
@@ -93,7 +93,7 @@ private:
 
 	TWeakPtr<ISequencer> WeakSequencer;
 	TWeakPtr<FControlRigEditModeToolkit> OwningToolkit;
-	FAnimBlendTooLManager  AnimBlendTools;
+	FAnimBlendToolManager  AnimBlendTools;
 
 	TWeakPtr<FControlRigEditMode> OwningEditMode;
 
diff --git a/Engine/Plugins/Animation/ControlRig/Source/ControlRigEditor/Private/Tools/ControlRigTweener.cpp b/Engine/Plugins/Animation/ControlRig/Source/ControlRigEditor/Private/Tools/ControlRigTweener.cpp
--- a/Engine/Plugins/Animation/ControlRig/Source/ControlRigEditor/Private/Tools/ControlRigTweener.cpp
+++ b/Engine/Plugins/Animation/ControlRig/Source/ControlRigEditor/Private/Tools/ControlRigTweener.cpp
@@ -326,27 +326,27 @@ bool FAnimSliderObjectSelection::Setup(const TArray<UControlRig*>& SelectedContr
 						{
 							switch (ControlElement->Settings.ControlType)
 							{
-							case ERigControlType::Float:
-							case ERigControlType::ScaleFloat:
+								case ERigControlType::Float:
+								case ERigControlType::ScaleFloat:
 								{
 									return 1;
 								}
-							case ERigControlType::Vector2D:
+								case ERigControlType::Vector2D:
 								{
 									return 2;
 								}
-							case ERigControlType::Position:
-							case ERigControlType::Scale:
-							case ERigControlType::Rotator:
+								case ERigControlType::Position:
+								case ERigControlType::Scale:
+								case ERigControlType::Rotator:
 								{
 									return 3;
 								}
-							case ERigControlType::TransformNoScale:
+								case ERigControlType::TransformNoScale:
 								{
 									return 6;
 								}
-							case ERigControlType::Transform:
-							case ERigControlType::EulerTransform:
+								case ERigControlType::Transform:
+								case ERigControlType::EulerTransform:
 								{
 									return 9;
 								}
@@ -362,7 +362,7 @@ bool FAnimSliderObjectSelection::Setup(const TArray<UControlRig*>& SelectedContr
 						{
 							FMovieSceneFloatChannel* Channel = Channels[ChannelIdx];
 							SetupChannel(CurrentTime.Time.GetFrame(), KeyTimes, Handles, Channel, nullptr,
-								ObjectChannels.KeyBounds[BoundIndex]);
+							ObjectChannels.KeyBounds[BoundIndex]);
 							if (ObjectChannels.KeyBounds[BoundIndex].bValid)
 							{
 								++NumValidChannels;
@@ -376,7 +376,7 @@ bool FAnimSliderObjectSelection::Setup(const TArray<UControlRig*>& SelectedContr
 					}
 				}
 			}
-		}		
+		}
 	};
 	
 	// Handle MovieScene bindings
@@ -476,7 +476,6 @@ bool FBaseAnimSlider::Setup(TWeakPtr<ISequencer>& InSequencer, TWeakPtr<FControl
 	{
 		return ObjectSelection.Setup(InSequencer, InEditMode);
 	}
-
 	return true;
 }
 
@@ -517,16 +516,20 @@ bool FBasicBlendSlider::Blend(TWeakPtr<ISequencer>& InSequencer, const double Bl
 					const int32 NumIndices = Keys.Indices.Num();
 					KeyHandles.SetNum(NumIndices);
 					KeyPositions.SetNum(NumIndices);
+					const double FirstValue = KeysArray.Value.AllKeyPositions[Keys.Indices[0]].OutputValue;
+					const double LastValue = KeysArray.Value.AllKeyPositions[Keys.Indices[Keys.Indices.Num() -1]].OutputValue;
+					FBlendStruct BlendStruct(KeysArray.Value.AllKeyPositions, Keys.Indices);
+					int32 CurrentIndex = 0;
 					for (int32 Index = 0; Index < NumIndices; ++Index)
 					{
 						const int32& KeyIndex = Keys.Indices[Index];
 						KeyHandles[Index] = KeysArray.Value.AllKeyHandles[KeyIndex];
 						const double CurrentValue = KeysArray.Value.AllKeyPositions[KeyIndex].OutputValue;
 						const double CurrentTime = KeysArray.Value.AllKeyPositions[KeyIndex].InputValue;
-
-						const double NewValue = DoBlend(PreviousTime,PreviousValue, CurrentTime, CurrentValue,
-							NextTime, NextValue, BlendValue);
-
+						BlendStruct.SetValues(PreviousTime, PreviousValue, CurrentTime, CurrentValue,
+							NextTime, NextValue, BlendValue, FirstValue, LastValue, CurrentIndex);
+						++CurrentIndex;
+						double NewValue = DoBlend(BlendStruct);
 						KeyPositions[Index] = FKeyPosition(KeysArray.Value.AllKeyPositions[KeyIndex].InputValue, NewValue);
 					}
 					Curve->SetKeyPositions(KeyHandles, KeyPositions);
@@ -537,10 +540,17 @@ bool FBasicBlendSlider::Blend(TWeakPtr<ISequencer>& InSequencer, const double Bl
 		return bDidBlend;
 	}
 
+	const FFrameTime FrameTime = Sequencer->GetLocalTime().Time;
+	const FFrameRate TickResolution = Sequencer->GetFocusedTickResolution();
+
 	if(ObjectSelection.ChannelsArray.Num() > 0)
-	{
-		const FFrameTime FrameTime = Sequencer->GetLocalTime().Time;
-		const FFrameRate TickResolution = Sequencer->GetFocusedTickResolution();
+	{ 
+		//empty key array for the blend struct
+		const TArray<FKeyPosition> AllKeyPositions;
+		const TArray <int32> Indices;
+		const int32 CurrentIndex = -1;
+		FBlendStruct BlendStruct(AllKeyPositions, Indices);
+
 		for (const FAnimSliderObjectSelection::FObjectChannels& ObjectChannels : ObjectSelection.ChannelsArray)
 		{
 			if (ObjectChannels.Section)
@@ -557,8 +567,9 @@ bool FBasicBlendSlider::Blend(TWeakPtr<ISequencer>& InSequencer, const double Bl
 					const double NextTime = TickResolution.AsSeconds(FFrameTime(ObjectChannels.KeyBounds[Index].NextFrame));
 					const double CurrentValue = ObjectChannels.KeyBounds[Index].CurrentValue;
 					const double CurrentTime = TickResolution.AsSeconds(FFrameTime(ObjectChannels.KeyBounds[Index].CurrentFrame));
-					const double NewValue = DoBlend(PreviousTime, PreviousValue, CurrentTime, CurrentValue,
-						NextTime, NextValue, BlendValue);
+					BlendStruct.SetValues(PreviousTime, PreviousValue, CurrentTime, CurrentValue,
+						NextTime, NextValue, BlendValue, CurrentValue, CurrentValue, CurrentIndex);
+					const double NewValue = DoBlend(BlendStruct);
 					using namespace UE::MovieScene;
 					if (ObjectChannels.KeyBounds[Index].FloatChannel)
 					{
@@ -584,21 +595,20 @@ bool FBasicBlendSlider::Blend(TWeakPtr<ISequencer>& InSequencer, const double Bl
 *
 */
 
-double FControlsToTween::DoBlend(const double PreviousTime, const double PreviousValue, const double CurrentTime, const double CurrentValue,
-	const double NextTime, const double NextValue, const double BlendValue)
+double FControlsToTween::DoBlend(const FBlendStruct& BlendStruct)
 {
 	//clasic tween will move all to same location, not based on current time at all just blend and values
-	const double NormalizedBlendValue = (BlendValue + 1.0f) * 0.5f;
-	const double Value = PreviousValue + (NextValue - PreviousValue) * (NormalizedBlendValue);
+	const double NormalizedBlendValue = (BlendStruct.BlendValue + 1.0f) * 0.5f;
+	const double Value = BlendStruct.PreviousValue + (BlendStruct.NextValue - BlendStruct.PreviousValue) * (NormalizedBlendValue);
 	return Value;
 }
 
-FText FControlsToTween::GetText()
+FText FControlsToTween::GetText() const
 {
 	return LOCTEXT("TW", "TW");
 }
 
-FText FControlsToTween::GetTooltipText()
+FText FControlsToTween::GetTooltipText() const
 {
 	return LOCTEXT("TweenControllerTooltip", "Tween between the next and previous keys");
 }
@@ -620,30 +630,29 @@ bool FControlsToTween::Setup(TWeakPtr<ISequencer>& InSequencer, TWeakPtr<FContro
 */
 
 
-double FPushPullSlider::DoBlend(const double PreviousTime, const double PreviousValue, const double CurrentTime, const double CurrentValue,
-	const double NextTime, const double NextValue, const double BlendValue)
+double FPushPullSlider::DoBlend(const FBlendStruct& BlendStruct)
 {
-	const double T = (CurrentTime - PreviousTime) / (NextTime - PreviousTime);
-	const double ValueAtT = PreviousValue + T * (NextValue - PreviousValue);
+	const double T = (BlendStruct.CurrentTime - BlendStruct.PreviousTime) / (BlendStruct.NextTime - BlendStruct.PreviousTime);
+	const double ValueAtT = BlendStruct.PreviousValue + T * (BlendStruct.NextValue - BlendStruct.PreviousValue);
 	double NewValue;
-	if (BlendValue < 0.0)
+	if (BlendStruct.BlendValue < 0.0)
 	{
-		NewValue = CurrentValue + (-1.0 * BlendValue) * (ValueAtT - CurrentValue);
+		NewValue = BlendStruct.CurrentValue + (-1.0 * BlendStruct.BlendValue) * (ValueAtT - BlendStruct.CurrentValue);
 	}
 	else
 	{
-		const double AmplifyValueAtT = CurrentValue + (CurrentValue - ValueAtT);
-		NewValue = CurrentValue + BlendValue * (AmplifyValueAtT - CurrentValue);
+		const double AmplifyValueAtT = BlendStruct.CurrentValue + (BlendStruct.CurrentValue - ValueAtT);
+		NewValue = BlendStruct.CurrentValue + BlendStruct.BlendValue * (AmplifyValueAtT - BlendStruct.CurrentValue);
 	}
 	return NewValue;
 }
 
-FText FPushPullSlider::GetText()
+FText FPushPullSlider::GetText() const
 {
 	return LOCTEXT("PP", "PP");
 }
 
-FText FPushPullSlider::GetTooltipText()
+FText FPushPullSlider::GetTooltipText() const
 {
 	return LOCTEXT("PushPullTooltip", "Push or pull the values to the interpolation between the previous and next keys");
 }
@@ -655,30 +664,139 @@ FText FPushPullSlider::GetTooltipText()
 *
 */
 
-double FBlendNeighborSlider::DoBlend(const double PreviousTime, const double PreviousValue, const double CurrentTime, const double CurrentValue,
-	const double NextTime, const double NextValue, const double BlendValue)
+double FBlendNeighborSlider::DoBlend(const FBlendStruct& BlendStruct)
 {
-	double NewValue;
-	if (BlendValue < 0.0)
+	double NewValue = BlendStruct.CurrentValue;
+	if (BlendStruct.BlendValue < 0.0)
 	{
-		NewValue = CurrentValue + (-1.0 * BlendValue) * (PreviousValue - CurrentValue);
+		NewValue = BlendStruct.CurrentValue + (-1.0 * BlendStruct.BlendValue) * (BlendStruct.PreviousValue - BlendStruct.CurrentValue);
 	}
 	else
 	{
-		NewValue = CurrentValue + BlendValue * (NextValue - CurrentValue);
+		NewValue = BlendStruct.CurrentValue + BlendStruct.BlendValue * (BlendStruct.NextValue - BlendStruct.CurrentValue);
 	}
 	return NewValue;
 }
 
-FText FBlendNeighborSlider::GetText()
+FText FBlendNeighborSlider::GetText() const
 {
 	return LOCTEXT("BN", "BN");
 }
 
-FText FBlendNeighborSlider::GetTooltipText()
+FText FBlendNeighborSlider::GetTooltipText() const
 {
 	return LOCTEXT("BlendToNeighborsSliderTooltip", "Blend to the next or previous values for selected keys or objects");
 }
 
+/*
+*
+*  Blend Relative Slider
+*
+*/
+
+double FBlendRelativeSlider::DoBlend(const FBlendStruct& BlendStruct)
+{
+	double NewValue = BlendStruct.CurrentValue;
+	if (BlendStruct.BlendValue < 0.0)
+	{
+		NewValue = BlendStruct.CurrentValue + (-1.0 * BlendStruct.BlendValue) * (BlendStruct.PreviousValue - BlendStruct.FirstValue);
+	}
+	else
+	{
+		NewValue = BlendStruct.CurrentValue + BlendStruct.BlendValue * (BlendStruct.NextValue - BlendStruct.LastValue);
+	}
+	return NewValue;
+}
+
+FText FBlendRelativeSlider::GetText() const
+{
+	return LOCTEXT("BR", "BR");
+}
+
+FText FBlendRelativeSlider::GetTooltipText() const
+{
+	return LOCTEXT("BlendRelativeSliderTooltip", "Blend relative to the next or previous value for selected keys or objects");
+}
+
+/*
+*
+*  Blend To Ease Slider
+*
+*/
+namespace BlendToEase
+{ 
+static float ExpIn(float InTime)
+{
+	return FMath::Pow(2, 10 * (InTime - 1.f));
+}
+static float ExpOut(float InTime)
+{
+	return 1.f - ExpIn(1.f - InTime);
+}
+}
+double FBlendToEaseSlider::DoBlend(const FBlendStruct& BlendStruct)
+{
+	double NewValue = BlendStruct.CurrentValue;
+	const double FullTimeDiff = BlendStruct.NextTime - BlendStruct.PreviousTime;
+	if (BlendStruct.BlendValue < 0.0)
+	{
+		const double MyTimeDiff = ((BlendStruct.CurrentTime - BlendStruct.PreviousTime) / FullTimeDiff);
+		double NewBlend = FMath::Clamp((- 1.0 * BlendStruct.BlendValue) / MyTimeDiff, 0.0, 1.0);
+		NewBlend *= (1.0 - BlendToEase::ExpIn(MyTimeDiff));
+		NewValue = BlendStruct.CurrentValue + NewBlend * (BlendStruct.PreviousValue - BlendStruct.CurrentValue);
+	}
+	else
+	{
+		const double MyTimeDiff =  ((BlendStruct.NextTime - BlendStruct.CurrentTime) / FullTimeDiff);
+		double NewBlend = FMath::Clamp(BlendStruct.BlendValue / MyTimeDiff, 0.0, 1.0);
+		NewBlend *= (1.0 - BlendToEase::ExpIn(MyTimeDiff));
+		NewValue = BlendStruct.CurrentValue + NewBlend * (BlendStruct.NextValue - BlendStruct.CurrentValue);
+	}
+	return NewValue;
+}
+
+FText FBlendToEaseSlider::GetText() const
+{
+	return LOCTEXT("BE", "BE");
+}
+
+FText FBlendToEaseSlider::GetTooltipText() const
+{
+	return LOCTEXT("BlendToEaseTooltip", "Blend with an ease falloff to the next or previous value for selected keys or objects");
+}
+
+/*
+*
+* Smooth/Rough Slider
+*
+*/
+
+double FSmoothRoughSlider::DoBlend(const FBlendStruct& BlendStruct)
+{
+	double PrevVal = BlendStruct.CurrentIndex > 0 ? BlendStruct.AllKeyPositions[BlendStruct.Indices[BlendStruct.CurrentIndex - 1]].OutputValue : BlendStruct.PreviousValue;
+	double CurVal = BlendStruct.CurrentValue;
+	double NextVal = BlendStruct.CurrentIndex < BlendStruct.Indices.Num() - 1 ? BlendStruct.AllKeyPositions[BlendStruct.Indices[BlendStruct.CurrentIndex + 1]].OutputValue : BlendStruct.NextValue;
+	double NewValue = (PrevVal * 0.25 + CurVal * 0.5 + NextVal * 0.25);
+	if (BlendStruct.BlendValue < 0.0)
+	{
+		NewValue = BlendStruct.CurrentValue + (-1.0 * BlendStruct.BlendValue) * (NewValue - BlendStruct.CurrentValue);
+	}
+	else
+	{
+		NewValue = BlendStruct.CurrentValue + BlendStruct.BlendValue * (BlendStruct.CurrentValue - NewValue);
+	}
+	return NewValue;
+}
+
+FText FSmoothRoughSlider::GetText() const
+{
+	return LOCTEXT("SR", "SR");
+}
+
+FText FSmoothRoughSlider::GetTooltipText() const
+{
+	return LOCTEXT("SmoothRoughTooltip", "Make selected keys smooth or rough");
+}
+
 #undef LOCTEXT_NAMESPACE
 
diff --git a/Engine/Plugins/Animation/ControlRig/Source/ControlRigEditor/Public/ControlRigSequencerEditorLibrary.h b/Engine/Plugins/Animation/ControlRig/Source/ControlRigEditor/Public/ControlRigSequencerEditorLibrary.h
--- a/Engine/Plugins/Animation/ControlRig/Source/ControlRigEditor/Public/ControlRigSequencerEditorLibrary.h
+++ b/Engine/Plugins/Animation/ControlRig/Source/ControlRigEditor/Public/ControlRigSequencerEditorLibrary.h
@@ -57,7 +57,10 @@ enum class EAnimToolBlendOperation : uint8
 {
 	Tween,
 	BlendToNeighbor,
-	PushPull
+	PushPull,
+	BlendRelative,
+	BlendToEase,
+	SmoothRough,
 };
 
 /**
diff --git a/Engine/Plugins/Animation/ControlRig/Source/ControlRigEditor/Public/Tools/ControlRigTweener.h b/Engine/Plugins/Animation/ControlRig/Source/ControlRigEditor/Public/Tools/ControlRigTweener.h
--- a/Engine/Plugins/Animation/ControlRig/Source/ControlRigEditor/Public/Tools/ControlRigTweener.h
+++ b/Engine/Plugins/Animation/ControlRig/Source/ControlRigEditor/Public/Tools/ControlRigTweener.h
@@ -122,12 +122,12 @@ struct FBaseAnimSlider
 	/**
 	* @return the name of the slider
 	*/
-	virtual FText GetText() = 0;
+	virtual FText GetText() const = 0;
 
 	/**
 	* @return the tooltip text for the slider
 	*/
-	virtual FText GetTooltipText() = 0;
+	virtual FText GetTooltipText() const = 0;
 
 protected:
 
@@ -139,10 +139,10 @@ protected:
 * Simple manager that holds all of the registered blend tools
 * This will then be used by whatever the view wants to do to display these tool
 */
-class FAnimBlendTooLManager
+class FAnimBlendToolManager
 {
 public:
-	FAnimBlendTooLManager() {};
+	FAnimBlendToolManager() {};
 	void RegisterAnimSlider(TSharedPtr<FBaseAnimSlider> AnimSlider)
 	{
 		AnimSliders.Add(AnimSlider);
@@ -159,6 +159,50 @@ private:
 	TArray<TSharedPtr<FBaseAnimSlider>> AnimSliders;
 };
 
+/**
+*
+*  Struct containing data that is used by Blending
+*
+**/
+struct FBlendStruct
+{
+	//Previous and Next are the keys that are before/after the contiguous set of keys, 
+	//to get the Previous key for example you use AllKeyPostions[Indices[CurrentIndex -1]]
+	double PreviousTime; 
+	double PreviousValue; 
+	double CurrentTime; 
+	double CurrentValue;
+	double NextTime;
+	double NextValue;
+	double BlendValue;
+	double FirstValue;
+	double LastValue;
+	TArray<FKeyPosition> AllKeyPositions;
+	TArray <int32> Indices;
+	int32 CurrentIndex; 
+
+	FBlendStruct(const TArray<FKeyPosition>& InAllKeyPositions, const TArray<int32>& InIndices) 
+		: AllKeyPositions(InAllKeyPositions), Indices(InIndices) {};
+	void SetValues(double InPreviousTime, double InPreviousValue, double InCurrentTime, double InCurrentValue,
+		double InNextTime, double InNextValue, double InBlendValue, double InFirstValue, double InLastValue,
+		int32 InCurrentIndex)
+	{
+		PreviousTime = InPreviousTime;
+		PreviousValue = InPreviousValue;
+		CurrentTime = InCurrentTime;
+		CurrentValue = InCurrentValue;
+		NextTime = InNextTime;
+		NextValue = InNextValue;
+		BlendValue = InBlendValue;
+		FirstValue = InFirstValue;
+		LastValue = InLastValue;
+		CurrentIndex = InCurrentIndex;
+	}
+
+private:
+	FBlendStruct() = delete;
+	
+};
 /**
 *
 *  Slider that used the next and previous values to do the blend
@@ -166,8 +210,7 @@ private:
 **/
 struct FBasicBlendSlider : public FBaseAnimSlider
 {
-	virtual double DoBlend(const double PreviousTime, const double PreviousValue, const double CurrentTime, const double CurrentValue,
-		const double NextTime, const double NextValue, const double BlendValue) = 0;
+	virtual double DoBlend(const FBlendStruct& BlendStruct) = 0;
 	//BaseAnimSlider overrides
 	virtual bool Blend(TWeakPtr<ISequencer>& InSequencer, const double BlendValue) override;
 };
@@ -180,13 +223,12 @@ struct FBasicBlendSlider : public FBaseAnimSlider
 struct FControlsToTween :public FBasicBlendSlider
 {
 	//BaseAnimSlider overrides
-	virtual FText GetText() override;
-	virtual FText GetTooltipText() override;
+	virtual FText GetText() const override;
+	virtual FText GetTooltipText() const override;
 	virtual bool Setup(TWeakPtr<ISequencer>& InSequencer, TWeakPtr<FControlRigEditMode>& InEditMode) override;
 
 	//BasicBlendSlider overrides
-	virtual double DoBlend(const double PreviousTime, const double PreviousValue, const double CurrentTime, const double CurrentValue,
-		const double NextTime, const double NextValue, const double BlendValue) override;
+	virtual double DoBlend(const FBlendStruct& BlendStruct) override;
 
 	bool Setup(const TArray<UControlRig*>& SelectedControlRigs, TWeakPtr<ISequencer>& InSequencer);
 
@@ -200,12 +242,11 @@ struct FControlsToTween :public FBasicBlendSlider
 struct FPushPullSlider :public FBasicBlendSlider
 {
 	//BaseAnimSlider overrides
-	virtual FText GetText();
-	virtual FText GetTooltipText();
+	virtual FText GetText() const;
+	virtual FText GetTooltipText() const;
 
 	//BasicBlendSlider overrides
-	virtual double DoBlend(const double PreviousTime, const double PreviousValue, const double CurrentTime, const double CurrentValue,
-		const double NextTime, const double NextValue, const double BlendValue) override;
+	virtual double DoBlend(const FBlendStruct& BlendStruct) override;
 };
 
 /**
@@ -216,14 +257,58 @@ struct FPushPullSlider :public FBasicBlendSlider
 struct FBlendNeighborSlider : public FBasicBlendSlider
 {
 	//BaseAnimSlider overrides
-	virtual FText GetText();
-	virtual FText GetTooltipText();
+	virtual FText GetText() const;
+	virtual FText GetTooltipText() const;
+
+	//BasicBlendSlider overrides
+	virtual double DoBlend(const FBlendStruct& BlendStruct) override;
+};
+
+
+/**
+*
+*   Will offset keys from previous or next values
+*
+**/
+struct FBlendRelativeSlider: public FBasicBlendSlider
+{
+	//BaseAnimSlider overrides
+	virtual FText GetText() const;
+	virtual FText GetTooltipText() const;
 
 	//BasicBlendSlider overrides
-	virtual double DoBlend(const double PreviousTime, const double PreviousValue, const double CurrentTime, const double CurrentValue,
-		const double NextTime, const double NextValue, const double BlendValue) override;
+	virtual double DoBlend(const FBlendStruct& BlendStruct) override;
 };
 
 
+/**
+*
+*   Will do an eased blend
+*
+**/
+struct FBlendToEaseSlider : public FBasicBlendSlider
+{
+	//BaseAnimSlider overrides
+	virtual FText GetText() const;
+	virtual FText GetTooltipText() const;
+
+	//BasicBlendSlider overrides
+	virtual double DoBlend(const FBlendStruct& BlendStruct) override;
+};
 
 
+/**
+*
+*   Will do a smooth or rough blend
+*
+**/
+struct FSmoothRoughSlider : public FBasicBlendSlider
+{
+	//BaseAnimSlider overrides
+	virtual FText GetText() const;
+	virtual FText GetTooltipText() const;
+
+	//BasicBlendSlider overrides
+	virtual double DoBlend(const FBlendStruct& BlendStruct) override;
+};
+
