Control Rig: Fix zeroing or inverting transforms in additive rigs affect boolean controls which are not additive controls
#rnx

diff --git a/Engine/Plugins/Animation/ControlRig/Source/ControlRigEditor/Private/EditMode/ControlRigEditMode.cpp b/Engine/Plugins/Animation/ControlRig/Source/ControlRigEditor/Private/EditMode/ControlRigEditMode.cpp
--- a/Engine/Plugins/Animation/ControlRig/Source/ControlRigEditor/Private/EditMode/ControlRigEditMode.cpp
+++ b/Engine/Plugins/Animation/ControlRig/Source/ControlRigEditor/Private/EditMode/ControlRigEditMode.cpp
@@ -2889,6 +2889,18 @@ void FControlRigEditMode::ZeroTransforms(bool bSelectionOnly)
 	for (UControlRig* ControlRig : ControlRigs)
 	{
 		TArray<FRigElementKey> SelectedRigElements = GetSelectedRigElements(ControlRig);
+		if (ControlRig->IsAdditive())
+		{
+			// For additive rigs, ignore boolean controls
+			SelectedRigElements = SelectedRigElements.FilterByPredicate([ControlRig](const FRigElementKey& Key)
+			{
+				if (FRigControlElement* Element = ControlRig->FindControl(Key.Name))
+				{
+					return Element->Settings.ControlType != ERigControlType::Bool;
+				}
+				return true;
+			});
+		}
 		TArray<FRigElementKey> ControlsToReset = SelectedRigElements;
 		TArray<FRigElementKey> ControlsInteracting = SelectedRigElements;
 		TArray<FRigElementKey> TransformElementsToReset = SelectedRigElements;
@@ -2899,6 +2911,14 @@ void FControlRigEditMode::ZeroTransforms(bool bSelectionOnly)
 			TransformElementsToReset.Reserve(Elements.Num());
 			for (const FRigBaseElement* Element : Elements)
 			{
+				// For additive rigs, ignore boolean controls
+				if (const FRigControlElement* Control = Cast<FRigControlElement>(Element))
+				{
+					if (ControlRig->IsAdditive() && Control->Settings.ControlType == ERigControlType::Bool)
+					{
+						continue;
+					}
+				}
 				TransformElementsToReset.Add(Element->GetKey());
 			}
 			
@@ -2908,6 +2928,11 @@ void FControlRigEditMode::ZeroTransforms(bool bSelectionOnly)
 			ControlsInteracting.SetNum(0);
 			for (const FRigControlElement* Control : Controls)
 			{
+				// For additive rigs, ignore boolean controls
+				if (ControlRig->IsAdditive() && Control->Settings.ControlType == ERigControlType::Bool)
+				{
+					continue;
+				}
 				ControlsToReset.Add(Control->GetKey());
 				if(Control->Settings.AnimationType == ERigControlAnimationType::AnimationControl ||
 					Control->IsAnimationChannel())
@@ -3098,6 +3123,14 @@ void FControlRigEditMode::InvertInputPose(bool bSelectionOnly)
 		if (bSelectionOnly)
 		{
 			SelectedRigElements = GetSelectedRigElements(ControlRig);
+			SelectedRigElements = SelectedRigElements.FilterByPredicate([ControlRig](const FRigElementKey& Key)
+			{
+				if (FRigControlElement* Element = ControlRig->FindControl(Key.Name))
+				{
+					return Element->Settings.ControlType != ERigControlType::Bool;
+				}
+				return true;
+			});
 		}
 
 		ControlRig->InvertInputPose(SelectedRigElements, EControlRigSetKey::Never);
