Control Rig: Additive rigs behaving strangely when moving several controls at the same time
#rnx

diff --git a/Engine/Plugins/Animation/ControlRig/Source/ControlRigEditor/Private/EditMode/ControlRigEditMode.cpp b/Engine/Plugins/Animation/ControlRig/Source/ControlRigEditor/Private/EditMode/ControlRigEditMode.cpp
--- a/Engine/Plugins/Animation/ControlRig/Source/ControlRigEditor/Private/EditMode/ControlRigEditMode.cpp
+++ b/Engine/Plugins/Animation/ControlRig/Source/ControlRigEditor/Private/EditMode/ControlRigEditMode.cpp
@@ -1844,6 +1844,18 @@ bool FControlRigEditMode::InputDelta(FEditorViewportClient* InViewportClient, FV
 					FTransform InOutLocal = FTransform::Identity;
 					
 					bool const bJustStartedManipulation = !bManipulatorMadeChange;
+					bool bAnyAdditiveRig = false;
+					for (AControlRigShapeActor* ShapeActor : Pairs.Value)
+					{
+						if (ShapeActor->ControlRig.IsValid())
+						{
+							if (ShapeActor->ControlRig->IsAdditive())
+							{
+								bAnyAdditiveRig = true;
+								break;
+							}
+						}
+					}
 
 					for (AControlRigShapeActor* ShapeActor : Pairs.Value)
 					{
@@ -1854,16 +1866,21 @@ bool FControlRigEditMode::InputDelta(FEditorViewportClient* InViewportClient, FV
 							{
 								GEditor->BeginTransaction(LOCTEXT("MoveControlTransaction", "Move Control"));
 							}
-							if (bFirstTime)
-							{
-								bFirstTime = false;
-							}
-							else
+
+							// Cannot benefit of same local transform when applying to additive rigs
+							if (!bAnyAdditiveRig)
 							{
-								if (bDoLocal)
+								if (bFirstTime)
 								{
-									bUseLocal = true;
-									bDoLocal = false;
+									bFirstTime = false;
+								}
+								else
+								{
+									if (bDoLocal)
+									{
+										bUseLocal = true;
+										bDoLocal = false;
+									}
 								}
 							}
 
