Control Rig: Rename ControlRig files to BaseControlRig

diff --git a/Engine/Plugins/Animation/AnimationData/Source/AnimationData/Private/AnimSequencerDataModel.cpp b/Engine/Plugins/Animation/AnimationData/Source/AnimationData/Private/AnimSequencerDataModel.cpp
--- a/Engine/Plugins/Animation/AnimationData/Source/AnimationData/Private/AnimSequencerDataModel.cpp
+++ b/Engine/Plugins/Animation/AnimationData/Source/AnimationData/Private/AnimSequencerDataModel.cpp
@@ -7,7 +7,7 @@
 #include "AnimDataController.h"
 #include "Async/ParallelFor.h"
 #include "IAnimationEditor.h"
-#include "ControlRig.h"
+#include "BaseControlRig.h"
 #include "ControlRigObjectBinding.h"
 #include "Algo/Accumulate.h"
 #include "Animation/AnimSequenceHelpers.h"
diff --git a/Engine/Plugins/Animation/ControlRig/Source/ControlRig/Private/AnimNode_ControlRig.cpp b/Engine/Plugins/Animation/ControlRig/Source/ControlRig/Private/AnimNode_ControlRig.cpp
--- a/Engine/Plugins/Animation/ControlRig/Source/ControlRig/Private/AnimNode_ControlRig.cpp
+++ b/Engine/Plugins/Animation/ControlRig/Source/ControlRig/Private/AnimNode_ControlRig.cpp
@@ -1,7 +1,7 @@
 // Copyright Epic Games, Inc. All Rights Reserved.
 
 #include "AnimNode_ControlRig.h"
-#include "ControlRig.h"
+#include "BaseControlRig.h"
 #include "ControlRigComponent.h"
 #include "Components/SkeletalMeshComponent.h"
 #include "Animation/AnimInstanceProxy.h"
diff --git a/Engine/Plugins/Animation/ControlRig/Source/ControlRig/Private/AnimNode_ControlRigBase.cpp b/Engine/Plugins/Animation/ControlRig/Source/ControlRig/Private/AnimNode_ControlRigBase.cpp
--- a/Engine/Plugins/Animation/ControlRig/Source/ControlRig/Private/AnimNode_ControlRigBase.cpp
+++ b/Engine/Plugins/Animation/ControlRig/Source/ControlRig/Private/AnimNode_ControlRigBase.cpp
@@ -1,7 +1,7 @@
 // Copyright Epic Games, Inc. All Rights Reserved.
 
 #include "AnimNode_ControlRigBase.h"
-#include "ControlRig.h"
+#include "BaseControlRig.h"
 #include "Components/SkeletalMeshComponent.h"
 #include "Animation/AnimInstanceProxy.h"
 #include "Animation/NodeMappingContainer.h"
diff --git a/Engine/Plugins/Animation/ControlRig/Source/ControlRig/Private/AnimNode_ControlRig_ExternalSource.cpp b/Engine/Plugins/Animation/ControlRig/Source/ControlRig/Private/AnimNode_ControlRig_ExternalSource.cpp
--- a/Engine/Plugins/Animation/ControlRig/Source/ControlRig/Private/AnimNode_ControlRig_ExternalSource.cpp
+++ b/Engine/Plugins/Animation/ControlRig/Source/ControlRig/Private/AnimNode_ControlRig_ExternalSource.cpp
@@ -1,7 +1,7 @@
 // Copyright Epic Games, Inc. All Rights Reserved.
 
 #include "AnimNode_ControlRig_ExternalSource.h"
-#include "ControlRig.h"
+#include "BaseControlRig.h"
 
 #include UE_INLINE_GENERATED_CPP_BY_NAME(AnimNode_ControlRig_ExternalSource)
 
diff --git a/Engine/Plugins/Animation/ControlRig/Source/ControlRig/Private/BaseControlRig.cpp b/Engine/Plugins/Animation/ControlRig/Source/ControlRig/Private/BaseControlRig.cpp
new file mode 100644
--- /dev/null
+++ b/Engine/Plugins/Animation/ControlRig/Source/ControlRig/Private/BaseControlRig.cpp
@@ -0,0 +1,3285 @@
+// Copyright Epic Games, Inc. All Rights Reserved.
+
+#include "BaseControlRig.h"
+#include "Engine/SkeletalMesh.h"
+#include "GameFramework/Actor.h"
+#include "Misc/RuntimeErrors.h"
+#include "IControlRigObjectBinding.h"
+#include "HelperUtil.h"
+#include "ObjectTrace.h"
+#include "RigVMBlueprintGeneratedClass.h"
+#include "ControlRigObjectVersion.h"
+#include "Rigs/RigHierarchyController.h"
+#include "Units/Execution/RigUnit_BeginExecution.h"
+#include "Units/Execution/RigUnit_PrepareForExecution.h"
+#include "Units/Execution/RigUnit_InverseExecution.h"
+#include "Units/Execution/RigUnit_InteractionExecution.h"
+#include "Components/SkeletalMeshComponent.h"
+#include "Animation/AnimInstance.h"
+#include "Animation/AnimationPoseData.h"
+#include "Internationalization/TextKey.h"
+#include "UObject/Package.h"
+#if WITH_EDITOR
+#include "ControlRigModule.h"
+#include "Modules/ModuleManager.h"
+#include "RigVMModel/RigVMNode.h"
+#include "Engine/Blueprint.h"
+#include "EdGraphSchema_K2.h"
+#include "Kismet2/BlueprintEditorUtils.h"
+#include "RigVMDeveloperTypeUtils.h"
+#include "Editor.h"
+#include "RigVMBlueprint.h"
+#include "RigVMCompiler/RigVMCompiler.h"
+#endif// WITH_EDITOR
+#include "ControlRigComponent.h"
+#include "Constraints/ControlRigTransformableHandle.h"
+#include "RigVMCore/RigVMNativized.h"
+#include "UObject/UObjectIterator.h"
+#include "RigVMCore/RigVMAssetUserData.h"
+
+#include UE_INLINE_GENERATED_CPP_BY_NAME(BaseControlRig)
+
+#define LOCTEXT_NAMESPACE "ControlRig"
+
+DEFINE_LOG_CATEGORY(LogControlRig);
+
+DECLARE_STATS_GROUP(TEXT("ControlRig"), STATGROUP_ControlRig, STATCAT_Advanced);
+DECLARE_CYCLE_STAT_EXTERN(TEXT("Control Rig Execution"), STAT_RigExecution, STATGROUP_ControlRig, );
+DEFINE_STAT(STAT_RigExecution);
+
+const FName UControlRig::OwnerComponent("OwnerComponent");
+
+//CVar to specify if we should create a float control for each curve in the curve container
+//By default we don't but it may be useful to do so for debugging
+static TAutoConsoleVariable<int32> CVarControlRigCreateFloatControlsForCurves(
+	TEXT("ControlRig.CreateFloatControlsForCurves"),
+	0,
+	TEXT("If nonzero we create a float control for each curve in the curve container, useful for debugging low level controls."),
+	ECVF_Default);
+
+//CVar to specify if we should allow debug drawing in game (during shipped game or PIE)
+//By default we don't for performance 
+static TAutoConsoleVariable<float> CVarControlRigEnableDrawInterfaceInGame(
+	TEXT("ControlRig.EnableDrawInterfaceInGame"),
+	0,
+	TEXT("If nonzero debug drawing will be enabled during play."),
+	ECVF_Default);
+
+
+UControlRig::UControlRig(const FObjectInitializer& ObjectInitializer)
+	: Super(ObjectInitializer)
+#if WITH_EDITOR
+	, bEnableAnimAttributeTrace(false)
+#endif 
+	, DataSourceRegistry(nullptr)
+#if WITH_EDITOR
+	, PreviewInstance(nullptr)
+#endif
+	, bCopyHierarchyBeforeConstruction(true)
+	, bResetInitialTransformsBeforeConstruction(true)
+	, bManipulationEnabled(false)
+	, PreConstructionBracket(0)
+	, PostConstructionBracket(0)
+	, InteractionBracket(0)
+	, InterRigSyncBracket(0)
+	, ControlUndoBracketIndex(0)
+	, InteractionType((uint8)EControlRigInteractionType::None)
+	, bInteractionJustBegan(false)
+	, DebugBoneRadiusMultiplier(1.f)
+#if WITH_EDITOR
+	, bRecordSelectionPoseForConstructionMode(true)
+	, bIsClearingTransientControls(false)
+#endif
+{
+	EventQueue.Add(FRigUnit_BeginExecution::EventName);
+
+	SetRigVMExtendedExecuteContext(&RigVMExtendedExecuteContext);
+}
+
+void UControlRig::BeginDestroy()
+{
+	Super::BeginDestroy();
+	SetRigVMExtendedExecuteContext(nullptr);
+
+	PreConstructionEvent.Clear();
+	PostConstructionEvent.Clear();
+	PreForwardsSolveEvent.Clear();
+	PostForwardsSolveEvent.Clear();
+	PreAdditiveValuesApplicationEvent.Clear();
+
+#if WITH_EDITOR
+	if (!HasAnyFlags(RF_ClassDefaultObject))
+	{
+		if(UControlRig* CDO = GetClass()->GetDefaultObject<UControlRig>())
+		{
+			if (!CDO->HasAnyFlags(RF_BeginDestroyed))
+			{
+				if (CDO->GetHierarchy())
+				{
+					CDO->GetHierarchy()->UnregisterListeningHierarchy(GetHierarchy());
+				}
+			}
+		}
+	}
+#endif
+
+	TRACE_OBJECT_LIFETIME_END(this);
+}
+
+UWorld* UControlRig::GetWorld() const
+{
+	if (ObjectBinding.IsValid())
+	{
+		AActor* HostingActor = ObjectBinding->GetHostingActor();
+		if (HostingActor)
+		{
+			return HostingActor->GetWorld();
+		}
+
+		UObject* Owner = ObjectBinding->GetBoundObject();
+		if (Owner)
+		{
+			return Owner->GetWorld();
+		}
+	}
+	return Super::GetWorld();
+}
+
+void UControlRig::Initialize(bool bRequestInit)
+{
+	TRACE_OBJECT_LIFETIME_BEGIN(this);
+
+	DECLARE_SCOPE_HIERARCHICAL_COUNTER_FUNC()
+	QUICK_SCOPE_CYCLE_COUNTER(STAT_ControlRig_Initialize);
+
+	Super::Initialize(bRequestInit);
+
+	// Create the data source registry here to avoid UObject creation from Non-Game Threads
+	GetDataSourceRegistry();
+
+	// Create the Hierarchy Controller here to avoid UObject creation from Non-Game Threads
+	GetHierarchy()->GetController(true);
+	
+	// should refresh mapping 
+	RequestConstruction();
+	
+	GetHierarchy()->OnModified().RemoveAll(this);
+	GetHierarchy()->OnModified().AddUObject(this, &UControlRig::HandleHierarchyModified);
+	GetHierarchy()->OnEventReceived().RemoveAll(this);
+	GetHierarchy()->OnEventReceived().AddUObject(this, &UControlRig::HandleHierarchyEvent);
+	GetHierarchy()->UpdateVisibilityOnProxyControls();
+}
+
+void UControlRig::OnAddShapeLibrary(const FControlRigExecuteContext* InContext, const FString& InLibraryName, UControlRigShapeLibrary* InShapeLibrary, bool bReplaceExisting, bool bLogResults)
+{
+	// don't ever change the CDO
+	if (HasAnyFlags(RF_ClassDefaultObject))
+	{
+		return;
+	}
+
+	if(InShapeLibrary == nullptr)
+	{
+		if(InContext)
+		{
+			static constexpr TCHAR Message[] = TEXT("No shape library provided.");
+			InContext->Log(EMessageSeverity::Error, Message);
+		}
+		return;
+	}
+
+	// remove the shape library in case it was there before
+	ShapeLibraryNameMap.Remove(InShapeLibrary->GetName());
+	ShapeLibraries.Remove(InShapeLibrary);
+
+	// if we've removed all shape libraries - let's add the ones from the CDO back
+	if (ShapeLibraries.IsEmpty())
+	{
+		UControlRig* CDO = GetClass()->GetDefaultObject<UControlRig>();
+		ShapeLibraries = CDO->ShapeLibraries;
+	}
+
+	// if we are supposed to replace the library and the library name is empty
+	FString LibraryName = InLibraryName;
+	if(LibraryName.IsEmpty() && bReplaceExisting)
+	{
+		if(ShapeLibraries.Num() == 1)
+		{
+			LibraryName = ShapeLibraries[0]->GetName();
+		}
+	}
+
+	if(LibraryName.IsEmpty())
+	{
+		LibraryName = InShapeLibrary->GetName();
+	}
+
+	if(LibraryName != InShapeLibrary->GetName())
+	{
+		ShapeLibraryNameMap.FindOrAdd(InShapeLibrary->GetName()) = LibraryName;
+	}
+
+	if(bReplaceExisting)
+	{
+		for(int32 Index = 0; Index < ShapeLibraries.Num(); Index++)
+		{
+			const TSoftObjectPtr<UControlRigShapeLibrary>& ExistingShapeLibrary = ShapeLibraries[Index];
+			if(ExistingShapeLibrary.IsNull())
+			{
+				continue;
+			}
+			FString ExistingName = ExistingShapeLibrary->GetName();
+			if (FString* MapName = ShapeLibraryNameMap.Find(ExistingName))
+			{
+				ExistingName = *MapName;
+			}
+			if(ExistingName.Equals(LibraryName, ESearchCase::IgnoreCase))
+			{
+				ShapeLibraryNameMap.Remove(ExistingShapeLibrary->GetName());
+				ShapeLibraries[Index] = InShapeLibrary;
+				break;
+			}
+		}
+	}
+
+	ShapeLibraries.AddUnique(InShapeLibrary);
+
+#if WITH_EDITOR
+	if(bLogResults)
+	{
+		static constexpr TCHAR MapFormat[] = TEXT("Control Rig '%s': Shape Library Name Map: '%s' -> '%s'");
+		static constexpr TCHAR LibraryFormat[] = TEXT("Control Rig '%s': Shape Library '%s' uses asset '%s'");
+		static constexpr TCHAR DefaultShapeFormat[] = TEXT("Control Rig '%s': Shape Library '%s' has default shape '%s'");
+		static constexpr TCHAR ShapeFormat[] = TEXT("Control Rig '%s': Shape Library '%s' contains shape %03d: '%s'");
+
+		const FString PathName = GetPathName();
+
+		for(const TPair<FString, FString>& Pair: ShapeLibraryNameMap)
+		{
+			UE_LOG(LogControlRig, Display, MapFormat, *PathName, *Pair.Key, *Pair.Value);
+		}
+
+		const int32 NumShapeLibraries = ShapeLibraries.Num();
+		
+		for(int32 Index = 0; Index <NumShapeLibraries; Index++)
+		{
+			const TSoftObjectPtr<UControlRigShapeLibrary>& ShapeLibrary = ShapeLibraries[Index];
+			if(ShapeLibrary.IsNull())
+			{
+				continue;
+			}
+
+			UE_LOG(LogControlRig, Display, LibraryFormat, *PathName, *ShapeLibrary->GetName(), *ShapeLibrary->GetPathName());
+
+			const bool bUseNameSpace = ShapeLibraries.Num() > 1;
+			const FString DefaultShapeName = UControlRigShapeLibrary::GetShapeName(ShapeLibrary.Get(), bUseNameSpace, ShapeLibraryNameMap, ShapeLibrary->DefaultShape);
+			UE_LOG(LogControlRig, Display, DefaultShapeFormat, *PathName, *ShapeLibrary->GetName(), *DefaultShapeName);
+
+			for(int32 ShapeIndex = 0; ShapeIndex < ShapeLibrary->Shapes.Num(); ShapeIndex++)
+			{
+				const FString ShapeName = UControlRigShapeLibrary::GetShapeName(ShapeLibrary.Get(), bUseNameSpace, ShapeLibraryNameMap, ShapeLibrary->Shapes[ShapeIndex]);
+				UE_LOG(LogControlRig, Display, ShapeFormat, *PathName, *ShapeLibrary->GetName(), ShapeIndex, *ShapeName);
+			}
+		}
+	}
+#endif
+}
+
+bool UControlRig::OnShapeExists(const FName& InShapeName) const
+{
+	const TArray<TSoftObjectPtr<UControlRigShapeLibrary>>& Libraries = GetShapeLibraries();
+	if (UControlRigShapeLibrary::GetShapeByName(InShapeName, GetShapeLibraries(), ShapeLibraryNameMap))
+	{
+		return true;
+	}
+	
+	return false;
+}
+
+bool UControlRig::InitializeVM(const FName& InEventName)
+{
+	if(!Super::InitializeVM(InEventName))
+	{
+		return false;
+	}
+
+	RequestConstruction();
+
+#if WITH_EDITOR
+	// setup the hierarchy's controller log function
+	if (URigHierarchyController* HierarchyController = GetHierarchy()->GetController(true))
+	{
+		HierarchyController->LogFunction = [this](EMessageSeverity::Type InSeverity, const FString& Message)
+		{
+			const FRigVMExecuteContext& PublicContext = GetRigVMExtendedExecuteContext().GetPublicData<>();
+			if(RigVMLog)
+			{
+				RigVMLog->Report(InSeverity,PublicContext.GetFunctionName(),PublicContext.GetInstructionIndex(), Message);
+			}
+			else
+			{
+				LogOnce(InSeverity,PublicContext.GetInstructionIndex(), Message);
+			}
+		};
+	}
+#endif
+
+	return true;
+}
+
+void UControlRig::Evaluate_AnyThread()
+{
+	if (bIsAdditive)
+	{
+		if (bRequiresInitExecution)
+		{
+			Super::Evaluate_AnyThread();
+		}
+		
+		// we can have other systems trying to poke into running instances of Control Rigs
+		// on the anim thread and query data, such as
+		// URigVMHostSkeletalMeshComponent::RebuildDebugDrawSkeleton,
+		// using a lock here to prevent them from having an inconsistent view of the rig at some
+		// intermediate stage of evaluation, for example, during evaluate, we can have a call
+		// to copy hierarchy, which empties the hierarchy for a short period of time
+		// and we don't want other systems to see that.
+		FScopeLock EvaluateLock(&GetEvaluateMutex());
+
+		URigHierarchy* Hierarchy = GetHierarchy();
+
+		if (PoseBeforeBackwardsSolve.Num() == 0)
+		{
+			// If the pose is empty, this is an indication that a new pose is coming in
+			uint8 Types = (uint8)ERigElementType::Bone | (uint8)ERigElementType::Curve;
+			PoseBeforeBackwardsSolve = Hierarchy->GetPose(false, (ERigElementType)Types, TArrayView<const FRigElementKey>());
+		}
+		else
+		{
+			// Restore the pose from the anim sequence
+			Hierarchy->SetPose(PoseBeforeBackwardsSolve);
+		}
+		
+		// Backwards solve
+		{
+			TGuardValue<bool> UpdatePreferredAngles(Hierarchy->bUpdatePreferedEulerAngleWhenSettingTransform, false);
+			Execute(FRigUnit_InverseExecution::EventName);
+		}
+
+		// Store control pose after backwards solve to figure out additive local transforms based on animation
+		ControlsAfterBackwardsSolve = Hierarchy->GetPose(false, ERigElementType::Control, TArrayView<const FRigElementKey>());
+
+		if (PreAdditiveValuesApplicationEvent.IsBound())
+		{
+			FControlRigBracketScope BracketScope(PreAdditiveValuesApplicationBracket);
+			PreAdditiveValuesApplicationEvent.Broadcast(this, TEXT("Additive"));
+		}
+		
+		// Apply additive controls
+		for (TPair<FRigElementKey, FRigSetControlValueInfo>& Value : ControlValues)
+		{
+			if (FRigBaseElement* Element = Hierarchy->Find(Value.Key))
+			{
+				if (FRigControlElement* Control = Cast<FRigControlElement>(Element))
+				{
+					FRigSetControlValueInfo& Info = Value.Value;
+
+					// A bool value is not an additive property. We just overwrite the value.
+					if (Control->Settings.ControlType == ERigControlType::Bool)
+					{
+						const bool bSetupUndo = false; // Rely on the sequencer track to handle undo/redo
+						Hierarchy->SetControlValue(Control, Info.Value, ERigControlValueType::Current, bSetupUndo, false, Info.bPrintPythonCommnds, false);
+					}
+					else
+					{
+						// Transform from animation
+						const FRigControlValue PreviousValue = Hierarchy->GetControlValue(Control, ERigControlValueType::Current, false);
+						const FTransform PreviousTransform = PreviousValue.GetAsTransform(Control->Settings.ControlType, Control->Settings.PrimaryAxis);
+
+						// Additive transform from controls
+						const FRigControlValue& AdditiveValue = Info.Value;
+						const FTransform AdditiveTransform = AdditiveValue.GetAsTransform(Control->Settings.ControlType, Control->Settings.PrimaryAxis);
+
+						// Add them to find the final value
+						FTransform FinalTransform = AdditiveTransform * PreviousTransform;
+
+						FRigControlValue FinalValue;
+						FinalValue.SetFromTransform(FinalTransform, Control->Settings.ControlType, Control->Settings.PrimaryAxis);
+
+						const bool bSetupUndo = false; // Rely on the sequencer track to handle undo/redo
+						Hierarchy->SetControlValue(Control, FinalValue, ERigControlValueType::Current, bSetupUndo, false, Info.bPrintPythonCommnds, false);
+						Hierarchy->SetPreferredEulerAnglesFromValue(Control, AdditiveValue, ERigControlValueType::Current, Info.bFixEulerFlips);
+					}
+
+					if (Info.bNotify && OnControlModified.IsBound())
+					{
+						OnControlModified.Broadcast(this, Control, Info.Context);
+						Info.bNotify = false;
+					}
+				}
+			}
+		}
+		
+		// Forward solve
+		{
+			TGuardValue<bool> UpdatePreferredAngles(Hierarchy->bUpdatePreferedEulerAngleWhenSettingTransform, false);
+			Execute(FRigUnit_BeginExecution::EventName);
+		}
+	}
+	else
+	{
+		Super::Evaluate_AnyThread();
+	}
+}
+
+bool UControlRig::EvaluateSkeletalMeshComponent(double InDeltaTime)
+{
+	if (USkeletalMeshComponent* SkelMeshComp = Cast<USkeletalMeshComponent>(GetObjectBinding()->GetBoundObject()))
+	{
+		SkelMeshComp->TickAnimation(InDeltaTime, false);
+
+		SkelMeshComp->RefreshBoneTransforms();
+		SkelMeshComp->RefreshFollowerComponents();
+		SkelMeshComp->UpdateComponentToWorld();
+		SkelMeshComp->FinalizeBoneTransform();
+		SkelMeshComp->MarkRenderTransformDirty();
+		SkelMeshComp->MarkRenderDynamicDataDirty();
+		return true;
+	}
+	return false;
+}
+
+void UControlRig::ResetControlValues()
+{
+	ControlValues.Reset();
+}
+
+void UControlRig::ClearPoseBeforeBackwardsSolve()
+{
+	PoseBeforeBackwardsSolve.Reset();
+}
+
+void UControlRig::InvertInputPose(const TArray<FRigElementKey>& InElements, EControlRigSetKey InSetKey)
+{
+	for (const FRigPoseElement& PoseElement : ControlsAfterBackwardsSolve)
+	{
+		if (FRigControlElement* ControlElement = Cast<FRigControlElement>(DynamicHierarchy->Get(PoseElement.Index)))
+		{
+			if (InElements.IsEmpty() || InElements.Contains(ControlElement->GetKey()))
+			{
+				FRigControlValue Value;
+				Value.SetFromTransform(PoseElement.LocalTransform.Inverse(), ControlElement->Settings.ControlType, ControlElement->Settings.PrimaryAxis);
+				FRigSetControlValueInfo Info = {Value, true, FRigControlModifiedContext(InSetKey), false, false, false};
+				ControlValues.Add(ControlElement->GetKey(), Info);
+			}
+		}
+	}
+}
+
+void UControlRig::InitializeFromCDO()
+{
+	Super::InitializeFromCDO();
+	
+	DECLARE_SCOPE_HIERARCHICAL_COUNTER_FUNC()
+
+	// copy CDO property you need to here
+	if (!HasAnyFlags(RF_ClassDefaultObject))
+	{
+		// similar to FControlRigBlueprintCompilerContext::CopyTermDefaultsToDefaultObject,
+		// where CDO is initialized from BP there,
+		// we initialize all other instances of Control Rig from the CDO here
+		UControlRig* CDO = GetClass()->GetDefaultObject<UControlRig>();
+		URigHierarchy* Hierarchy = GetHierarchy();
+
+		// copy hierarchy
+		{
+			FRigHierarchyValidityBracket ValidityBracketA(Hierarchy);
+			FRigHierarchyValidityBracket ValidityBracketB(CDO->GetHierarchy());
+			
+			TGuardValue<bool> Guard(Hierarchy->GetSuspendNotificationsFlag(), true);
+			Hierarchy->CopyHierarchy(CDO->GetHierarchy());
+			Hierarchy->ResetPoseToInitial(ERigElementType::All);
+		}
+
+#if WITH_EDITOR
+		// current hierarchy should always mirror CDO's hierarchy whenever a change of interest happens
+		CDO->GetHierarchy()->RegisterListeningHierarchy(Hierarchy);
+#endif
+
+		// notify clients that the hierarchy has changed
+		Hierarchy->Notify(ERigHierarchyNotification::HierarchyReset, nullptr);
+
+		// copy hierarchy settings
+		HierarchySettings = CDO->HierarchySettings;
+		ElementKeyRedirector = FRigElementKeyRedirector(CDO->ElementKeyRedirector, Hierarchy); 
+		
+		// increment the procedural limit based on the number of elements in the CDO
+		if(const URigHierarchy* CDOHierarchy = CDO->GetHierarchy())
+		{
+			HierarchySettings.ProceduralElementLimit += CDOHierarchy->Num();
+		}
+
+		ExternalVariableDataAssetLinks.Reset();
+	}
+}
+
+UControlRig::FAnimAttributeContainerPtrScope::FAnimAttributeContainerPtrScope(UControlRig* InControlRig,
+	UE::Anim::FStackAttributeContainer& InExternalContainer)
+{
+	ControlRig = InControlRig;
+	ControlRig->ExternalAnimAttributeContainer = &InExternalContainer;
+}
+
+UControlRig::FAnimAttributeContainerPtrScope::~FAnimAttributeContainerPtrScope()
+{
+	// control rig should not hold on to this container since it is stack allocated
+	// and should not be used outside of stack, see FPoseContext
+	ControlRig->ExternalAnimAttributeContainer = nullptr;
+}
+
+AActor* UControlRig::GetHostingActor() const
+{
+	return ObjectBinding ? ObjectBinding->GetHostingActor() : nullptr;
+}
+
+#if WITH_EDITOR
+FText UControlRig::GetCategory() const
+{
+	return LOCTEXT("DefaultControlRigCategory", "Animation|ControlRigs");
+}
+
+FText UControlRig::GetToolTipText() const
+{
+	return LOCTEXT("DefaultControlRigTooltip", "ControlRig");
+}
+#endif
+
+bool UControlRig::AllConnectorsAreResolved(FString* OutFailureReason, FRigElementKey* OutConnector) const
+{
+	if(const URigHierarchy* Hierarchy = GetHierarchy())
+	{
+		// todo: introduce a cache here based on ElementKeyDirector hash and
+		// topology hash of the hierarchy
+		
+		const TArray<FRigConnectorElement*> Connectors = Hierarchy->GetConnectors(false);
+
+		// collect the connection map
+		TMap<FRigElementKey, FRigElementKey> ConnectionMap;
+		for(const FRigConnectorElement* Connector : Connectors)
+		{
+			if(const FCachedRigElement* Cache = ElementKeyRedirector.Find(Connector->GetKey()))
+			{
+				if(const_cast<FCachedRigElement*>(Cache)->UpdateCache(Hierarchy))
+				{
+					ConnectionMap.Add(Connector->GetKey(), Cache->GetKey());
+				}
+				else
+				{
+					if(OutFailureReason)
+					{
+						static constexpr TCHAR Format[] = TEXT("Connector '%s' has invalid target '%s'.");
+						*OutFailureReason = FString::Printf(Format, *Connector->GetName(), *Cache->GetKey().ToString());
+					}
+					if(OutConnector)
+					{
+						*OutConnector = Connector->GetKey();
+					}
+					return false;
+				}
+			}
+			else
+			{
+				if(OutFailureReason)
+				{
+					static constexpr TCHAR Format[] = TEXT("Connector '%s' is not resolved.");
+					*OutFailureReason = FString::Printf(Format, *Connector->GetName());
+				}
+				if(OutConnector)
+				{
+					*OutConnector = Connector->GetKey();
+				}
+				return false;
+			}
+		}
+
+		// now that we have the connection try if we can connect this way
+		const FRigConnectionInfo ConnectionInfo(ConnectionMap, Hierarchy, Hierarchy);
+		if(!Hierarchy->CanConnect(&ConnectionInfo, OutFailureReason, OutConnector))
+		{
+			return false;
+		}
+	}
+	return true;
+}
+
+bool UControlRig::Execute(const FName& InEventName)
+{
+	if(!CanExecute())
+	{
+		return false;
+	}
+
+	bool bJustRanInit = false;
+	if(bRequiresInitExecution)
+	{
+		const TGuardValue<float> AbsoluteTimeGuard(AbsoluteTime, AbsoluteTime);
+		const TGuardValue<float> DeltaTimeGuard(DeltaTime, DeltaTime);
+		if(!InitializeVM(InEventName))
+		{
+			return false;
+		}
+		bJustRanInit = true;
+	}
+
+	if(EventQueueToRun.IsEmpty())
+	{
+		EventQueueToRun = EventQueue;
+	}
+
+	const bool bIsEventInQueue = EventQueueToRun.Contains(InEventName);
+	const bool bIsEventFirstInQueue = !EventQueueToRun.IsEmpty() && EventQueueToRun[0] == InEventName; 
+	const bool bIsEventLastInQueue = !EventQueueToRun.IsEmpty() && EventQueueToRun.Last() == InEventName;
+	const bool bIsConstructionEvent = InEventName == FRigUnit_PrepareForExecution::EventName;
+	const bool bIsForwardSolve = InEventName == FRigUnit_BeginExecution::EventName;
+	const bool bIsInteractionEvent = InEventName == FRigUnit_InteractionExecution::EventName;
+
+	ensure(!HasAnyFlags(RF_ClassDefaultObject));
+	
+	DECLARE_SCOPE_HIERARCHICAL_COUNTER_FUNC()
+	QUICK_SCOPE_CYCLE_COUNTER(STAT_ControlRig_Execute);
+	
+	FRigVMExtendedExecuteContext& ExtendedExecuteContext = GetRigVMExtendedExecuteContext();
+
+	FControlRigExecuteContext& PublicContext = ExtendedExecuteContext.GetPublicDataSafe<FControlRigExecuteContext>();
+
+#if WITH_EDITOR
+	PublicContext.SetLog(RigVMLog); // may be nullptr
+#endif
+
+	PublicContext.SetDeltaTime(DeltaTime);
+	PublicContext.SetAbsoluteTime(AbsoluteTime);
+	PublicContext.SetFramesPerSecond(GetCurrentFramesPerSecond());
+
+#if UE_RIGVM_DEBUG_EXECUTION
+	PublicContext.bDebugExecution = bDebugExecutionEnabled;
+#endif
+
+	if (VM)
+	{
+#if WITH_EDITOR
+		// default to always clear data after each execution
+		// only set a valid first entry event later when execution
+		// has passed the initialization stage and there are multiple events present in one evaluation
+		// first entry event is used to determined when to clear data during an evaluation
+		VM->SetFirstEntryEventInEventQueue(ExtendedExecuteContext, NAME_None);
+#endif
+	}
+
+#if WITH_EDITOR
+	if (IsInDebugMode())
+	{
+		if (UControlRig* CDO = GetClass()->GetDefaultObject<UControlRig>())
+		{
+			// Copy the breakpoints. This will not override the state of the breakpoints
+			DebugInfo.SetBreakpoints(CDO->DebugInfo.GetBreakpoints());
+
+			// If there are any breakpoints, create the Snapshot VM if it hasn't been created yet
+			if (DebugInfo.GetBreakpoints().Num() > 0)
+			{
+				GetSnapshotVM();
+			}
+		}
+
+		ExtendedExecuteContext.SetDebugInfo(&DebugInfo);
+	}
+	else
+	{
+		ExtendedExecuteContext.SetDebugInfo(nullptr);
+		GetSnapshotContext().Reset();
+	}
+#endif
+
+	FRigUnitContext& Context = PublicContext.UnitContext;
+
+	bool bEnableDrawInterface = false;
+#if WITH_EDITOR
+	const bool bEnabledDuringGame = CVarControlRigEnableDrawInterfaceInGame->GetInt() != 0;
+	const bool bInGame = !(GEditor && !GEditor->GetPIEWorldContext());
+	if (bEnabledDuringGame || !bInGame)
+	{
+		bEnableDrawInterface = true;
+	}	
+#endif
+
+	// setup the draw interface for debug drawing
+	if(!bIsEventInQueue || bIsEventFirstInQueue)
+	{
+		DrawInterface.Reset();
+	}
+
+	if (bEnableDrawInterface)
+	{
+		PublicContext.SetDrawInterface(&DrawInterface);
+	}
+	else
+	{
+		PublicContext.SetDrawInterface(nullptr);
+	}
+
+	// setup the animation attribute container
+	Context.AnimAttributeContainer = ExternalAnimAttributeContainer;
+
+	// draw container contains persistent draw instructions, 
+	// so we cannot call Reset(), which will clear them,
+	// instead, we re-initialize them from the CDO
+	if (UControlRig* CDO = GetClass()->GetDefaultObject<UControlRig>())
+	{
+		DrawContainer = CDO->DrawContainer;
+	}
+	PublicContext.SetDrawContainer(&DrawContainer);
+
+	// setup the data source registry
+	Context.DataSourceRegistry = GetDataSourceRegistry();
+
+	// setup the context with further fields
+	Context.InteractionType = InteractionType;
+	Context.ElementsBeingInteracted = ElementsBeingInteracted;
+	PublicContext.Hierarchy = GetHierarchy();
+
+	// allow access to the hierarchy
+	Context.HierarchySettings = HierarchySettings;
+	check(PublicContext.Hierarchy);
+
+	// allow access to the shape libraries
+	PublicContext.OnAddShapeLibraryDelegate.BindUObject(this, &UControlRig::OnAddShapeLibrary);
+	PublicContext.OnShapeExistsDelegate.BindUObject(this, &UControlRig::OnShapeExists);
+
+	// allow access to the default hierarchy to allow to reset
+	if(!HasAnyFlags(RF_ClassDefaultObject))
+	{
+		if (UControlRig* CDO = Cast<UControlRig>(GetClass()->GetDefaultObject()))
+		{
+			if(URigHierarchy* DefaultHierarchy = CDO->GetHierarchy())
+			{
+				PublicContext.Hierarchy->DefaultHierarchyPtr = DefaultHierarchy;
+			}
+		}
+	}
+
+	// disable any controller access outside of the construction event
+	FRigHierarchyEnableControllerBracket DisableHierarchyController(PublicContext.Hierarchy, bIsConstructionEvent);
+
+	// given the outer scene component configure
+	// the transform lookups to map transforms from rig space to world space
+	if (USceneComponent* SceneComponent = GetOwningSceneComponent())
+	{
+		PublicContext.SetOwningComponent(SceneComponent);
+	}
+	else
+	{
+		PublicContext.SetOwningComponent(nullptr);
+		
+		if (ObjectBinding.IsValid())
+		{
+			AActor* HostingActor = ObjectBinding->GetHostingActor();
+			if (HostingActor)
+			{
+				PublicContext.SetOwningActor(HostingActor);
+			}
+			else if (UObject* Owner = ObjectBinding->GetBoundObject())
+			{
+				PublicContext.SetWorld(Owner->GetWorld());
+			}
+		}
+
+		if (PublicContext.GetWorld() == nullptr)
+		{
+			if (UObject* Outer = GetOuter())
+			{
+				PublicContext.SetWorld(Outer->GetWorld());
+			}
+		}
+	}
+	
+	// setup the user data
+	// first using the data from the rig,
+	// followed by the data coming from the skeleton
+	// then the data coming from the skeletal mesh,
+	// and in the future we'll also add the data from the control rig component
+	PublicContext.AssetUserData.Reset();
+	if(const TArray<UAssetUserData*>* ControlRigUserDataArray = GetAssetUserDataArray())
+	{
+		for(const UAssetUserData* ControlRigUserData : *ControlRigUserDataArray)
+		{
+			PublicContext.AssetUserData.Add(ControlRigUserData);
+		}
+	}
+	PublicContext.AssetUserData.Remove(nullptr);
+
+	// if we have any referenced elements dirty them
+	if(GetHierarchy())
+	{
+		GetHierarchy()->UpdateReferences(&PublicContext);
+	}
+
+	// guard against recursion
+	if(IsExecuting())
+	{
+		UE_LOG(LogControlRig, Warning, TEXT("%s: Execute is being called recursively."), *GetPathName());
+		return false;
+	}
+	if(bIsConstructionEvent)
+	{
+		if(IsRunningPreConstruction() || IsRunningPostConstruction())
+		{
+			UE_LOG(LogControlRig, Warning, TEXT("%s: Construction is being called recursively."), *GetPathName());
+			return false;
+		}
+	}
+
+	bool bSuccess = true;
+
+	// we'll special case the construction event here
+	if (bIsConstructionEvent)
+	{
+		// remember the previous selection
+		const TArray<FRigElementKey> PreviousSelection = GetHierarchy()->GetSelectedKeys();
+
+		// construction mode means that we are running the construction event
+		// constantly for testing purposes.
+		const bool bConstructionModeEnabled = IsConstructionModeEnabled();
+		{
+			if (PreConstructionForUIEvent.IsBound())
+			{
+				FControlRigBracketScope BracketScope(PreConstructionBracket);
+				PreConstructionForUIEvent.Broadcast(this, FRigUnit_PrepareForExecution::EventName);
+			}
+
+#if WITH_EDITOR
+			// apply the selection pose for the construction mode
+			// we are doing this here to make sure the brackets below have the right data
+			if(bConstructionModeEnabled)
+			{
+				ApplySelectionPoseForConstructionMode(InEventName);
+			}
+#endif
+
+			// disable selection notifications from the hierarchy
+			TGuardValue<bool> DisableSelectionNotifications(GetHierarchy()->GetController(true)->bSuspendSelectionNotifications, true);
+			{
+				{
+					// Copy the hierarchy from the default object onto this one
+#if WITH_EDITOR
+					FTransientControlScope TransientControlScope(GetHierarchy());
+	#endif
+					{
+						// maintain the initial pose if it ever was set by the client
+						FRigPose InitialPose;
+						if(!bResetInitialTransformsBeforeConstruction)
+						{
+							InitialPose = GetHierarchy()->GetPose(true, ERigElementType::ToResetAfterConstructionEvent, FRigElementKeyCollection());
+						}
+
+						if(bCopyHierarchyBeforeConstruction)
+						{
+							GetHierarchy()->ResetToDefault();
+						}
+
+						if(InitialPose.Num() > 0)
+						{
+							GetHierarchy()->SetPose(InitialPose, ERigTransformType::InitialLocal);
+						}
+					}
+
+					{
+	#if WITH_EDITOR
+						TUniquePtr<FTransientControlPoseScope> TransientControlPoseScope;
+						if (bConstructionModeEnabled)
+						{
+							// save the transient control value, it should not be constantly reset in construction mode
+							TransientControlPoseScope = MakeUnique<FTransientControlPoseScope>(this);
+						}
+	#endif
+						// reset the pose to initial such that construction event can run from a deterministic initial state
+						GetHierarchy()->ResetPoseToInitial(ERigElementType::All);
+					}
+
+					// clone the shape libraries again from the CDO 
+					if (!HasAnyFlags(RF_ClassDefaultObject) && ShapeLibraries.IsEmpty())
+					{
+						UControlRig* CDO = GetClass()->GetDefaultObject<UControlRig>();
+						ShapeLibraries = CDO->ShapeLibraries;
+					}
+
+					if (PreConstructionEvent.IsBound())
+					{
+						FControlRigBracketScope BracketScope(PreConstructionBracket);
+						PreConstructionEvent.Broadcast(this, FRigUnit_PrepareForExecution::EventName);
+					}
+
+					bSuccess = Execute_Internal(FRigUnit_PrepareForExecution::EventName);
+					
+				} // destroy FTransientControlScope
+				
+				if (PostConstructionEvent.IsBound())
+				{
+					FControlRigBracketScope BracketScope(PostConstructionBracket);
+					PostConstructionEvent.Broadcast(this, FRigUnit_PrepareForExecution::EventName);
+				}
+			}
+			
+			// set it here to reestablish the selection. the notifications
+			// will be eaten since we still have the bSuspend flag on in the controller.
+			GetHierarchy()->GetController()->SetSelection(PreviousSelection);
+			
+		} // destroy DisableSelectionNotifications
+
+		if (bConstructionModeEnabled)
+		{
+#if WITH_EDITOR
+			TUniquePtr<FTransientControlPoseScope> TransientControlPoseScope;
+			if (bConstructionModeEnabled)
+			{
+				// save the transient control value, it should not be constantly reset in construction mode
+				TransientControlPoseScope = MakeUnique<FTransientControlPoseScope>(this);
+			}
+#endif
+			GetHierarchy()->ResetPoseToInitial(ERigElementType::Bone);
+		}
+
+		// synchronize the selection now with the new hierarchy after running construction
+		const TArray<const FRigBaseElement*> CurrentSelection = GetHierarchy()->GetSelectedElements();
+		for(const FRigBaseElement* SelectedElement : CurrentSelection)
+		{
+			if(!PreviousSelection.Contains(SelectedElement->GetKey()))
+			{
+				GetHierarchy()->Notify(ERigHierarchyNotification::ElementSelected, SelectedElement);
+			}
+		}
+		for(const FRigElementKey& PreviouslySelectedKey : PreviousSelection)
+		{
+			if(const FRigBaseElement* PreviouslySelectedElement = GetHierarchy()->Find(PreviouslySelectedKey))
+			{
+				if(!CurrentSelection.Contains(PreviouslySelectedElement))
+				{
+					GetHierarchy()->Notify(ERigHierarchyNotification::ElementDeselected, PreviouslySelectedElement);
+				}
+			}
+		}
+	}
+	else
+	{
+#if WITH_EDITOR
+		// only set a valid first entry event when execution
+		// has passed the initialization stage and there are multiple events present
+		if (EventQueueToRun.Num() >= 2 && VM)
+		{
+			VM->SetFirstEntryEventInEventQueue(ExtendedExecuteContext, EventQueueToRun[0]);
+		}
+
+		// Transform Overrride is generated using a Transient Control 
+		ApplyTransformOverrideForUserCreatedBones();
+
+		if (bEnableAnimAttributeTrace && ExternalAnimAttributeContainer != nullptr)
+		{
+			InputAnimAttributeSnapshot.CopyFrom(*ExternalAnimAttributeContainer);
+		}
+#endif
+		
+		if (bIsForwardSolve)
+		{
+			if (PreForwardsSolveEvent.IsBound())
+			{
+				FControlRigBracketScope BracketScope(PreForwardsSolveBracket);
+				PreForwardsSolveEvent.Broadcast(this, FRigUnit_BeginExecution::EventName);
+			}
+		}
+
+		bSuccess = Execute_Internal(InEventName);
+
+#if WITH_EDITOR
+		if (bEnableAnimAttributeTrace && ExternalAnimAttributeContainer != nullptr)
+		{
+			OutputAnimAttributeSnapshot.CopyFrom(*ExternalAnimAttributeContainer);
+		}
+#endif
+
+		if (bIsForwardSolve)
+		{
+			if (PostForwardsSolveEvent.IsBound())
+			{
+				FControlRigBracketScope BracketScope(PostForwardsSolveBracket);
+				PostForwardsSolveEvent.Broadcast(this, FRigUnit_BeginExecution::EventName);
+			}
+		}
+	}
+
+#if WITH_EDITOR
+
+	// for the last event in the queue - clear the log message queue
+	if (RigVMLog != nullptr && bEnableLogging)
+	{
+		if (bJustRanInit)
+		{
+			RigVMLog->KnownMessages.Reset();
+			LoggedMessages.Reset();
+		}
+		else if(bIsEventLastInQueue)
+		{
+			for (const FRigVMLog::FLogEntry& Entry : RigVMLog->Entries)
+			{
+				if (Entry.FunctionName == NAME_None || Entry.InstructionIndex == INDEX_NONE || Entry.Message.IsEmpty())
+				{
+					continue;
+				}
+
+				FString PerInstructionMessage = 
+					FString::Printf(
+						TEXT("Instruction[%d] '%s': '%s'"),
+						Entry.InstructionIndex,
+						*Entry.FunctionName.ToString(),
+						*Entry.Message
+					);
+
+				LogOnce(Entry.Severity, Entry.InstructionIndex, PerInstructionMessage);
+			}
+		}
+	}
+#endif
+
+	if(!bIsEventInQueue || bIsEventLastInQueue) 
+	{
+		DeltaTime = 0.f;
+	}
+
+	if (ExecutedEvent.IsBound())
+	{
+		if (!IsAdditive() || InEventName != FRigUnit_InverseExecution::EventName)
+		{
+			FControlRigBracketScope BracketScope(ExecuteBracket);
+			ExecutedEvent.Broadcast(this, InEventName);
+		}
+	}
+
+	// close remaining undo brackets from hierarchy
+	while(ControlUndoBracketIndex > 0)
+	{
+		FRigEventContext EventContext;
+		EventContext.Event = ERigEvent::CloseUndoBracket;
+		EventContext.SourceEventName = InEventName;
+		EventContext.LocalTime = PublicContext.GetAbsoluteTime();
+		HandleHierarchyEvent(GetHierarchy(), EventContext);
+	}
+
+	if (PublicContext.GetDrawInterface() && PublicContext.GetDrawContainer() && bIsEventLastInQueue) 
+	{
+		PublicContext.GetDrawInterface()->Instructions.Append(PublicContext.GetDrawContainer()->Instructions);
+
+		FRigHierarchyValidityBracket ValidityBracket(GetHierarchy());
+		
+		GetHierarchy()->ForEach<FRigControlElement>([this](FRigControlElement* ControlElement) -> bool
+		{
+			const FRigControlSettings& Settings = ControlElement->Settings;
+
+			if (Settings.IsVisible() &&
+				!Settings.bIsTransientControl &&
+				Settings.bDrawLimits &&
+				Settings.LimitEnabled.Contains(FRigControlLimitEnabled(true, true)))
+			{
+				FTransform Transform = GetHierarchy()->GetGlobalControlOffsetTransformByIndex(ControlElement->GetIndex());
+				FRigVMDrawInstruction Instruction(ERigVMDrawSettings::Lines, Settings.ShapeColor, 0.f, Transform);
+
+				switch (Settings.ControlType)
+				{
+					case ERigControlType::Float:
+					case ERigControlType::ScaleFloat:
+					{
+						if(Settings.LimitEnabled[0].IsOff())
+						{
+							break;
+						}
+
+						FVector MinPos = FVector::ZeroVector;
+						FVector MaxPos = FVector::ZeroVector;
+
+						switch (Settings.PrimaryAxis)
+						{
+							case ERigControlAxis::X:
+							{
+								MinPos.X = Settings.MinimumValue.Get<float>();
+								MaxPos.X = Settings.MaximumValue.Get<float>();
+								break;
+							}
+							case ERigControlAxis::Y:
+							{
+								MinPos.Y = Settings.MinimumValue.Get<float>();
+								MaxPos.Y = Settings.MaximumValue.Get<float>();
+								break;
+							}
+							case ERigControlAxis::Z:
+							{
+								MinPos.Z = Settings.MinimumValue.Get<float>();
+								MaxPos.Z = Settings.MaximumValue.Get<float>();
+								break;
+							}
+						}
+
+						Instruction.Positions.Add(MinPos);
+						Instruction.Positions.Add(MaxPos);
+						break;
+					}
+					case ERigControlType::Integer:
+					{
+						if(Settings.LimitEnabled[0].IsOff())
+						{
+							break;
+						}
+
+						FVector MinPos = FVector::ZeroVector;
+						FVector MaxPos = FVector::ZeroVector;
+
+						switch (Settings.PrimaryAxis)
+						{
+							case ERigControlAxis::X:
+							{
+								MinPos.X = (float)Settings.MinimumValue.Get<int32>();
+								MaxPos.X = (float)Settings.MaximumValue.Get<int32>();
+								break;
+							}
+							case ERigControlAxis::Y:
+							{
+								MinPos.Y = (float)Settings.MinimumValue.Get<int32>();
+								MaxPos.Y = (float)Settings.MaximumValue.Get<int32>();
+								break;
+							}
+							case ERigControlAxis::Z:
+							{
+								MinPos.Z = (float)Settings.MinimumValue.Get<int32>();
+								MaxPos.Z = (float)Settings.MaximumValue.Get<int32>();
+								break;
+							}
+						}
+
+						Instruction.Positions.Add(MinPos);
+						Instruction.Positions.Add(MaxPos);
+						break;
+					}
+					case ERigControlType::Vector2D:
+					{
+						if(Settings.LimitEnabled.Num() < 2)
+						{
+							break;
+						}
+						if(Settings.LimitEnabled[0].IsOff() && Settings.LimitEnabled[1].IsOff())
+						{
+							break;
+						}
+
+						Instruction.PrimitiveType = ERigVMDrawSettings::LineStrip;
+						FVector3f MinPos = Settings.MinimumValue.Get<FVector3f>();
+						FVector3f MaxPos = Settings.MaximumValue.Get<FVector3f>();
+
+						switch (Settings.PrimaryAxis)
+						{
+							case ERigControlAxis::X:
+							{
+								Instruction.Positions.Add(FVector(0.f, MinPos.X, MinPos.Y));
+								Instruction.Positions.Add(FVector(0.f, MaxPos.X, MinPos.Y));
+								Instruction.Positions.Add(FVector(0.f, MaxPos.X, MaxPos.Y));
+								Instruction.Positions.Add(FVector(0.f, MinPos.X, MaxPos.Y));
+								Instruction.Positions.Add(FVector(0.f, MinPos.X, MinPos.Y));
+								break;
+							}
+							case ERigControlAxis::Y:
+							{
+								Instruction.Positions.Add(FVector(MinPos.X, 0.f, MinPos.Y));
+								Instruction.Positions.Add(FVector(MaxPos.X, 0.f, MinPos.Y));
+								Instruction.Positions.Add(FVector(MaxPos.X, 0.f, MaxPos.Y));
+								Instruction.Positions.Add(FVector(MinPos.X, 0.f, MaxPos.Y));
+								Instruction.Positions.Add(FVector(MinPos.X, 0.f, MinPos.Y));
+								break;
+							}
+							case ERigControlAxis::Z:
+							{
+								Instruction.Positions.Add(FVector(MinPos.X, MinPos.Y, 0.f));
+								Instruction.Positions.Add(FVector(MaxPos.X, MinPos.Y, 0.f));
+								Instruction.Positions.Add(FVector(MaxPos.X, MaxPos.Y, 0.f));
+								Instruction.Positions.Add(FVector(MinPos.X, MaxPos.Y, 0.f));
+								Instruction.Positions.Add(FVector(MinPos.X, MinPos.Y, 0.f));
+								break;
+							}
+						}
+						break;
+					}
+					case ERigControlType::Position:
+					case ERigControlType::Scale:
+					case ERigControlType::Transform:
+					case ERigControlType::TransformNoScale:
+					case ERigControlType::EulerTransform:
+					{
+						FVector3f MinPos = FVector3f::ZeroVector;
+						FVector3f MaxPos = FVector3f::ZeroVector;
+
+						// we only check the first three here
+						// since we only consider translation anyway
+						// for scale it's also the first three
+						if(Settings.LimitEnabled.Num() < 3)
+						{
+							break;
+						}
+						if(!Settings.LimitEnabled[0].IsOn() && !Settings.LimitEnabled[1].IsOn() && !Settings.LimitEnabled[2].IsOn())
+						{
+							break;
+						}
+
+						switch (Settings.ControlType)
+						{
+							case ERigControlType::Position:
+							case ERigControlType::Scale:
+							{
+								MinPos = Settings.MinimumValue.Get<FVector3f>();
+								MaxPos = Settings.MaximumValue.Get<FVector3f>();
+								break;
+							}
+							case ERigControlType::Transform:
+							{
+								MinPos = Settings.MinimumValue.Get<FRigControlValue::FTransform_Float>().GetTranslation();
+								MaxPos = Settings.MaximumValue.Get<FRigControlValue::FTransform_Float>().GetTranslation();
+								break;
+							}
+							case ERigControlType::TransformNoScale:
+							{
+								MinPos = Settings.MinimumValue.Get<FRigControlValue::FTransformNoScale_Float>().GetTranslation();
+								MaxPos = Settings.MaximumValue.Get<FRigControlValue::FTransformNoScale_Float>().GetTranslation();
+								break;
+							}
+							case ERigControlType::EulerTransform:
+							{
+								MinPos = Settings.MinimumValue.Get<FRigControlValue::FEulerTransform_Float>().GetTranslation();
+								MaxPos = Settings.MaximumValue.Get<FRigControlValue::FEulerTransform_Float>().GetTranslation();
+								break;
+							}
+						}
+
+						Instruction.Positions.Add(FVector(MinPos.X, MinPos.Y, MinPos.Z));
+						Instruction.Positions.Add(FVector(MaxPos.X, MinPos.Y, MinPos.Z));
+						Instruction.Positions.Add(FVector(MinPos.X, MaxPos.Y, MinPos.Z));
+						Instruction.Positions.Add(FVector(MaxPos.X, MaxPos.Y, MinPos.Z));
+						Instruction.Positions.Add(FVector(MinPos.X, MinPos.Y, MaxPos.Z));
+						Instruction.Positions.Add(FVector(MaxPos.X, MinPos.Y, MaxPos.Z));
+						Instruction.Positions.Add(FVector(MinPos.X, MaxPos.Y, MaxPos.Z));
+						Instruction.Positions.Add(FVector(MaxPos.X, MaxPos.Y, MaxPos.Z));
+
+						Instruction.Positions.Add(FVector(MinPos.X, MinPos.Y, MinPos.Z));
+						Instruction.Positions.Add(FVector(MinPos.X, MaxPos.Y, MinPos.Z));
+						Instruction.Positions.Add(FVector(MaxPos.X, MinPos.Y, MinPos.Z));
+						Instruction.Positions.Add(FVector(MaxPos.X, MaxPos.Y, MinPos.Z));
+						Instruction.Positions.Add(FVector(MinPos.X, MinPos.Y, MaxPos.Z));
+						Instruction.Positions.Add(FVector(MinPos.X, MaxPos.Y, MaxPos.Z));
+						Instruction.Positions.Add(FVector(MaxPos.X, MinPos.Y, MaxPos.Z));
+						Instruction.Positions.Add(FVector(MaxPos.X, MaxPos.Y, MaxPos.Z));
+
+						Instruction.Positions.Add(FVector(MinPos.X, MinPos.Y, MinPos.Z));
+						Instruction.Positions.Add(FVector(MinPos.X, MinPos.Y, MaxPos.Z));
+						Instruction.Positions.Add(FVector(MaxPos.X, MinPos.Y, MinPos.Z));
+						Instruction.Positions.Add(FVector(MaxPos.X, MinPos.Y, MaxPos.Z));
+						Instruction.Positions.Add(FVector(MinPos.X, MaxPos.Y, MinPos.Z));
+						Instruction.Positions.Add(FVector(MinPos.X, MaxPos.Y, MaxPos.Z));
+						Instruction.Positions.Add(FVector(MaxPos.X, MaxPos.Y, MinPos.Z));
+						Instruction.Positions.Add(FVector(MaxPos.X, MaxPos.Y, MaxPos.Z));
+						break;
+					}
+				}
+
+				if (Instruction.Positions.Num() > 0)
+				{
+					DrawInterface.DrawInstruction(Instruction);
+				}
+			}
+
+			return true;
+		});
+	}
+
+	if(bIsInteractionEvent)
+	{
+		bInteractionJustBegan = false;
+	}
+
+	if(bIsConstructionEvent)
+	{
+		RemoveRunOnceEvent(FRigUnit_PrepareForExecution::EventName);
+	}
+
+	return bSuccess;
+}
+
+bool UControlRig::Execute_Internal(const FName& InEventName)
+{
+	if(IsRigModule())
+	{
+		FString ConnectorWarning;
+		if(!AllConnectorsAreResolved(&ConnectorWarning))
+		{
+#if WITH_EDITOR
+			LogOnce(EMessageSeverity::Warning, INDEX_NONE, ConnectorWarning);
+#endif
+			return false;
+		}
+	}
+	
+	if (VM)
+	{
+		FRigVMExtendedExecuteContext& Context = GetRigVMExtendedExecuteContext();
+
+		static constexpr TCHAR InvalidatedVMFormat[] = TEXT("%s: Invalidated VM - aborting execution.");
+		if(VM->IsNativized())
+		{
+			if(!IsValidLowLevel() ||
+				!VM->IsValidLowLevel())
+			{
+				UE_LOG(LogControlRig, Warning, InvalidatedVMFormat, *GetClass()->GetName());
+				return false;
+			}
+		}
+		else
+		{
+			// sanity check the validity of the VM to ensure stability.
+			if(!VM->IsContextValidForExecution(Context)
+				|| !IsValidLowLevel()
+				|| !VM->IsValidLowLevel()
+			)
+			{
+				UE_LOG(LogControlRig, Warning, InvalidatedVMFormat, *GetClass()->GetName());
+				return false;
+			}
+		}
+		
+#if UE_RIGVM_PROFILE_EXECUTE_UNITS_NUM
+		const uint64 StartCycles = FPlatformTime::Cycles64();
+		if(ProfilingRunsLeft <= 0)
+		{
+			ProfilingRunsLeft = UE_RIGVM_PROFILE_EXECUTE_UNITS_NUM;
+			AccumulatedCycles = 0;
+		}
+#endif
+		
+		const bool bUseDebuggingSnapshots = !VM->IsNativized();
+		
+#if WITH_EDITOR
+		if(bUseDebuggingSnapshots)
+		{
+			if(URigVM* SnapShotVM = GetSnapshotVM(false)) // don't create it for normal runs
+			{
+				const bool bIsEventFirstInQueue = !EventQueueToRun.IsEmpty() && EventQueueToRun[0] == InEventName; 
+				const bool bIsEventLastInQueue = !EventQueueToRun.IsEmpty() && EventQueueToRun.Last() == InEventName;
+
+				if (GetHaltedAtBreakpoint().IsValid())
+				{
+					if(bIsEventFirstInQueue)
+					{
+						CopyVMMemory(GetRigVMExtendedExecuteContext(), GetSnapshotContext());
+					}
+				}
+				else if(bIsEventLastInQueue)
+				{
+					CopyVMMemory(GetSnapshotContext(), GetRigVMExtendedExecuteContext());
+				}
+			}
+		}
+#endif
+
+		URigHierarchy* Hierarchy = GetHierarchy();
+#if WITH_EDITOR
+
+		bool bRecordTransformsAtRuntime = true;
+		if(const UObject* Outer = GetOuter())
+		{
+			if(Outer->IsA<UControlRigComponent>())
+			{
+				bRecordTransformsAtRuntime = false;
+			}
+		}
+		TGuardValue<bool> RecordTransformsPerInstructionGuard(Hierarchy->bRecordTransformsAtRuntime, bRecordTransformsAtRuntime);
+		
+		if(Hierarchy->bRecordTransformsAtRuntime)
+		{
+			Hierarchy->ReadTransformsAtRuntime.Reset();
+			Hierarchy->WrittenTransformsAtRuntime.Reset();
+		}
+		
+#endif
+		FRigHierarchyExecuteContextBracket HierarchyContextGuard(Hierarchy, &Context);
+
+		// setup the module information
+		FControlRigExecuteContext& PublicContext = Context.GetPublicDataSafe<FControlRigExecuteContext>();
+		FControlRigExecuteContextRigModuleGuard RigModuleGuard(PublicContext, this);
+		FRigHierarchyRedirectorGuard ElementRedirectorGuard(this);
+
+		const bool bSuccess = VM->ExecuteVM(Context, InEventName) != ERigVMExecuteResult::Failed;
+
+#if UE_RIGVM_PROFILE_EXECUTE_UNITS_NUM
+		const uint64 EndCycles = FPlatformTime::Cycles64();
+		const uint64 Cycles = EndCycles - StartCycles;
+		AccumulatedCycles += Cycles;
+		ProfilingRunsLeft--;
+		if(ProfilingRunsLeft == 0)
+		{
+			const double Milliseconds = FPlatformTime::ToMilliseconds64(AccumulatedCycles);
+			UE_LOG(LogControlRig, Display, TEXT("%s: %d runs took %.03lfms."), *GetClass()->GetName(), UE_RIGVM_PROFILE_EXECUTE_UNITS_NUM, Milliseconds);
+		}
+#endif
+
+		return bSuccess;
+	}
+	return false;
+}
+
+void UControlRig::RequestInit()
+{
+	Super::RequestInit();
+	RequestConstruction();
+}
+
+void UControlRig::RequestConstruction()
+{
+	RequestRunOnceEvent(FRigUnit_PrepareForExecution::EventName, 0);
+}
+
+bool UControlRig::IsConstructionRequired() const
+{
+	return IsRunOnceEvent(FRigUnit_PrepareForExecution::EventName);
+}
+
+bool UControlRig::SupportsBackwardsSolve() const
+{
+	return SupportsEvent(FRigUnit_InverseExecution::EventName);
+}
+
+void UControlRig::AdaptEventQueueForEvaluate(TArray<FName>& InOutEventQueueToRun)
+{
+	Super::AdaptEventQueueForEvaluate(InOutEventQueueToRun);
+
+	if(InteractionType != (uint8)EControlRigInteractionType::None)
+	{
+		static const FName& InteractionEvent = FRigUnit_InteractionExecution::EventName;
+		if(SupportsEvent(InteractionEvent))
+		{
+			if(InOutEventQueueToRun.IsEmpty())
+			{
+				InOutEventQueueToRun.Add(InteractionEvent);
+			}
+			else if(!InOutEventQueueToRun.Contains(FRigUnit_PrepareForExecution::EventName))
+			{
+				// insert just before the last event so the interaction runs prior to
+				// forward solve or backwards solve.
+				InOutEventQueueToRun.Insert(InteractionEvent, InOutEventQueueToRun.Num() - 1);
+			}
+		}
+	}
+}
+
+void UControlRig::GetMappableNodeData(TArray<FName>& OutNames, TArray<FNodeItem>& OutNodeItems) const
+{
+	DECLARE_SCOPE_HIERARCHICAL_COUNTER_FUNC()
+
+	OutNames.Reset();
+	OutNodeItems.Reset();
+
+	check(DynamicHierarchy);
+
+	// now add all nodes
+	DynamicHierarchy->ForEach<FRigBoneElement>([&OutNames, &OutNodeItems, this](FRigBoneElement* BoneElement) -> bool
+    {
+		OutNames.Add(BoneElement->GetFName());
+		FRigElementKey ParentKey = DynamicHierarchy->GetFirstParent(BoneElement->GetKey());
+		if(ParentKey.Type != ERigElementType::Bone)
+		{
+			ParentKey.Name = NAME_None;
+		}
+
+		const FTransform GlobalInitial = DynamicHierarchy->GetGlobalTransformByIndex(BoneElement->GetIndex(), true);
+		OutNodeItems.Add(FNodeItem(ParentKey.Name, GlobalInitial));
+		return true;
+	});
+}
+
+UAnimationDataSourceRegistry* UControlRig::GetDataSourceRegistry()
+{
+	if (DataSourceRegistry)
+	{
+		if (DataSourceRegistry->GetOuter() != this)
+		{
+			DataSourceRegistry = nullptr;
+		}
+	}
+	if (DataSourceRegistry == nullptr)
+	{
+		DataSourceRegistry = NewObject<UAnimationDataSourceRegistry>(this, NAME_None, RF_Transient);
+
+		if (HasAnyFlags(RF_ClassDefaultObject) && GetClass()->IsNative())
+		{
+			DataSourceRegistry->AddToRoot();
+		}
+	}
+
+	return DataSourceRegistry;
+}
+
+#if WITH_EDITORONLY_DATA
+
+void UControlRig::PostReinstanceCallback(const UControlRig* Old)
+{
+	DECLARE_SCOPE_HIERARCHICAL_COUNTER_FUNC()
+
+	ObjectBinding = Old->ObjectBinding;
+	Initialize();
+}
+
+#endif // WITH_EDITORONLY_DATA
+
+TArray<UControlRig*> UControlRig::FindControlRigs(UObject* Outer, TSubclassOf<UControlRig> OptionalClass)
+{
+	TArray<UControlRig*> Result;
+	
+	if(Outer == nullptr)
+	{
+		return Result; 
+	}
+	
+	AActor* OuterActor = Cast<AActor>(Outer);
+	if(OuterActor == nullptr)
+	{
+		OuterActor = Outer->GetTypedOuter<AActor>();
+	}
+	
+	for (TObjectIterator<UControlRig> Itr; Itr; ++Itr)
+	{
+		UControlRig* RigInstance = *Itr;
+		const UClass* RigInstanceClass = RigInstance ? RigInstance->GetClass() : nullptr;
+		if (OptionalClass == nullptr || (RigInstanceClass && RigInstanceClass->IsChildOf(OptionalClass)))
+		{
+			if(RigInstance->IsInOuter(Outer))
+			{
+				Result.Add(RigInstance);
+				continue;
+			}
+
+			if(OuterActor)
+			{
+				if(RigInstance->IsInOuter(OuterActor))
+				{
+					Result.Add(RigInstance);
+					continue;
+				}
+
+				if (TSharedPtr<IControlRigObjectBinding> Binding = RigInstance->GetObjectBinding())
+				{
+					if (AActor* Actor = Binding->GetHostingActor())
+					{
+						if (Actor == OuterActor)
+						{
+							Result.Add(RigInstance);
+							continue;
+						}
+					}
+				}
+			}
+		}
+	}
+
+	return Result;
+}
+
+void UControlRig::Serialize(FArchive& Ar)
+{
+	DECLARE_SCOPE_HIERARCHICAL_COUNTER_FUNC()
+
+	Super::Serialize(Ar);
+
+	Ar.UsingCustomVersion(FControlRigObjectVersion::GUID);
+}
+
+void UControlRig::PostLoad()
+{
+	Super::PostLoad();
+
+	if(HasAnyFlags(RF_ClassDefaultObject))
+	{
+		if(DynamicHierarchy)
+		{
+			// Some dynamic hierarchy objects have been created using NewObject<> instead of CreateDefaultSubObjects.
+			// Assets from that version require the dynamic hierarchy to be flagged as below.
+			DynamicHierarchy->SetFlags(DynamicHierarchy->GetFlags() | RF_Public | RF_DefaultSubObject);
+		}
+	}
+
+	FRigInfluenceMapPerEvent NewInfluences;
+	for(int32 MapIndex = 0; MapIndex < Influences.Num(); MapIndex++)
+	{
+		FRigInfluenceMap Map = Influences[MapIndex];
+		FName EventName = Map.GetEventName();
+
+		if(EventName == TEXT("Update"))
+		{
+			EventName = FRigUnit_BeginExecution::EventName;
+		}
+		else if(EventName == TEXT("Inverse"))
+		{
+			EventName = FRigUnit_InverseExecution::EventName;
+		}
+		
+		NewInfluences.FindOrAdd(EventName).Merge(Map, true);
+	}
+	Influences = NewInfluences;
+}
+
+const FRigModuleSettings& UControlRig::GetRigModuleSettings() const
+{
+	if(HasAnyFlags(RF_ClassDefaultObject))
+	{
+		return RigModuleSettings;
+	}
+	if (const UControlRig* CDO = Cast<UControlRig>(GetClass()->GetDefaultObject()))
+	{
+		return CDO->GetRigModuleSettings();
+	}
+	return RigModuleSettings;
+}
+
+bool UControlRig::IsRigModule() const
+{
+	return GetRigModuleSettings().IsValidModule(false);
+}
+
+bool UControlRig::IsRigModuleInstance() const
+{
+	if(IsRigModule())
+	{
+		return GetParentRig() != nullptr;
+	}
+	return false;
+}
+
+UControlRig* UControlRig::GetParentRig() const
+{
+	return GetTypedOuter<UControlRig>();
+}
+
+const FString& UControlRig::GetRigModuleNameSpace() const
+{
+	if(IsRigModule())
+	{
+		if(const UControlRig* ParentRig = GetParentRig())
+		{
+			const FString& ParentNameSpace = ParentRig->GetRigModuleNameSpace();
+			static constexpr TCHAR JoinFormat[] = TEXT("%s%s:");
+			RigModuleNameSpace = FString::Printf(JoinFormat, *ParentNameSpace, *GetName());
+			return RigModuleNameSpace;
+		}
+
+		// this means we are not an instance - so we are authoring the rig module right now.
+		// for visualization purposes we'll use the name of the rig module chosen in the settings.
+		static constexpr TCHAR NameSpaceFormat[] = TEXT("%s:"); 
+		RigModuleNameSpace = FString::Printf(NameSpaceFormat, *GetRigModuleSettings().Identifier.Name);
+		return RigModuleNameSpace;
+	}
+
+	static const FString EmptyNameSpace;
+	return EmptyNameSpace;
+}
+
+FRigElementKeyRedirector& UControlRig::GetElementKeyRedirector()
+{
+	// if we are an instance on a modular rig, use our local info
+	if(IsRigModuleInstance())
+	{
+		return ElementKeyRedirector;
+	}
+
+	// if we are a rig module ( but not an instance on a rig)
+	if(IsRigModule())
+	{
+		return ElementKeyRedirector;
+	}
+
+	static FRigElementKeyRedirector EmptyRedirector = FRigElementKeyRedirector();
+	return EmptyRedirector;
+}
+
+TArray<FRigControlElement*> UControlRig::AvailableControls() const
+{
+	if(DynamicHierarchy)
+	{
+		return DynamicHierarchy->GetElementsOfType<FRigControlElement>();
+	}
+	return TArray<FRigControlElement*>();
+}
+
+FRigControlElement* UControlRig::FindControl(const FName& InControlName) const
+{
+	if(DynamicHierarchy == nullptr)
+	{
+		return nullptr;
+	}
+	return DynamicHierarchy->Find<FRigControlElement>(FRigElementKey(InControlName, ERigElementType::Control));
+}
+
+bool UControlRig::IsConstructionModeEnabled() const
+{
+	return EventQueueToRun.Num() == 1 && EventQueueToRun.Contains(FRigUnit_PrepareForExecution::EventName);
+}
+
+FTransform UControlRig::SetupControlFromGlobalTransform(const FName& InControlName, const FTransform& InGlobalTransform)
+{
+	if (IsConstructionModeEnabled())
+	{
+		FRigControlElement* ControlElement = FindControl(InControlName);
+		if (ControlElement && !ControlElement->Settings.bIsTransientControl)
+		{
+			const FTransform ParentTransform = GetHierarchy()->GetParentTransform(ControlElement, ERigTransformType::CurrentGlobal);
+			const FTransform OffsetTransform = InGlobalTransform.GetRelativeTransform(ParentTransform);
+			GetHierarchy()->SetControlOffsetTransform(ControlElement, OffsetTransform, ERigTransformType::InitialLocal, true, false);
+			GetHierarchy()->SetControlOffsetTransform(ControlElement, OffsetTransform, ERigTransformType::CurrentLocal, true, false);
+
+			if(URigHierarchy* DefaultHierarchy = GetHierarchy()->GetDefaultHierarchy())
+			{
+				if(FRigControlElement* DefaultControlElement = DefaultHierarchy->Find<FRigControlElement>(ControlElement->GetKey()))
+				{
+					DefaultHierarchy->SetControlOffsetTransform(DefaultControlElement, OffsetTransform, ERigTransformType::InitialLocal, true, true);
+					DefaultHierarchy->SetControlOffsetTransform(DefaultControlElement, OffsetTransform, ERigTransformType::CurrentLocal, true, true);
+				}
+			}
+		}
+	}
+	return InGlobalTransform;
+}
+
+void UControlRig::CreateRigControlsForCurveContainer()
+{
+	const bool bCreateFloatControls = CVarControlRigCreateFloatControlsForCurves->GetInt() == 0 ? false : true;
+	if(bCreateFloatControls && DynamicHierarchy)
+	{
+		URigHierarchyController* Controller = DynamicHierarchy->GetController(true);
+		if(Controller == nullptr)
+		{
+			return;
+		}
+		static const FString CtrlPrefix(TEXT("CTRL_"));
+
+		DynamicHierarchy->ForEach<FRigCurveElement>([this, Controller](FRigCurveElement* CurveElement) -> bool
+        {
+			const FString Name = CurveElement->GetFName().ToString();
+			
+			if (Name.Contains(CtrlPrefix) && !DynamicHierarchy->Contains(FRigElementKey(*Name, ERigElementType::Curve))) //-V1051
+			{
+				FRigControlSettings Settings;
+				Settings.AnimationType = ERigControlAnimationType::AnimationChannel;
+				Settings.ControlType = ERigControlType::Float;
+				Settings.bIsCurve = true;
+				Settings.bDrawLimits = false;
+
+				FRigControlValue Value;
+				Value.Set<float>(CurveElement->Value);
+
+				Controller->AddControl(CurveElement->GetFName(), FRigElementKey(), Settings, Value, FTransform::Identity, FTransform::Identity); 
+			}
+
+			return true;
+		});
+
+		ControlModified().AddUObject(this, &UControlRig::HandleOnControlModified);
+	}
+}
+
+void UControlRig::HandleOnControlModified(UControlRig* Subject, FRigControlElement* Control, const FRigControlModifiedContext& Context)
+{
+	if (Control->Settings.bIsCurve && DynamicHierarchy)
+	{
+		const FRigControlValue Value = DynamicHierarchy->GetControlValue(Control, IsConstructionModeEnabled() ? ERigControlValueType::Initial : ERigControlValueType::Current);
+		DynamicHierarchy->SetCurveValue(FRigElementKey(Control->GetFName(), ERigElementType::Curve), Value.Get<float>());
+	}	
+}
+
+bool UControlRig::IsCurveControl(const FRigControlElement* InControlElement) const
+{
+	return InControlElement->Settings.bIsCurve;
+}
+
+FTransform UControlRig::GetControlGlobalTransform(const FName& InControlName) const
+{
+	if(DynamicHierarchy == nullptr)
+	{
+		return FTransform::Identity;
+	}
+	return DynamicHierarchy->GetGlobalTransform(FRigElementKey(InControlName, ERigElementType::Control), false);
+}
+
+FRigControlValue UControlRig::GetControlValue(FRigControlElement* InControl, const ERigControlValueType& InValueType)
+{
+	if (bIsAdditive && InValueType == ERigControlValueType::Current)
+	{
+		const int32 ControlIndex = ControlsAfterBackwardsSolve.GetIndex(InControl->GetKey());
+		if (ControlIndex != INDEX_NONE)
+		{
+			// Booleans are not additive properties, just return the current value
+			if (InControl->Settings.ControlType == ERigControlType::Bool)
+			{
+				return GetHierarchy()->GetControlValue(InControl, InValueType);
+			}
+			
+			// return local space control value (the one to be added after backwards solve)
+			const FRigPoseElement& AnimPose = ControlsAfterBackwardsSolve[ControlIndex];
+			const FRigControlValue& CurrentValue = GetHierarchy()->GetControlValue(InControl, InValueType);
+			const FTransform FinalTransform = CurrentValue.GetAsTransform(InControl->Settings.ControlType, InControl->Settings.PrimaryAxis);
+			const FTransform AdditiveTransform = FinalTransform * AnimPose.LocalTransform.Inverse();
+			FRigControlValue AdditiveValue;
+			AdditiveValue.SetFromTransform(AdditiveTransform, InControl->Settings.ControlType, InControl->Settings.PrimaryAxis);
+			return AdditiveValue;
+		}
+	}
+	return GetHierarchy()->GetControlValue(InControl, InValueType);
+}
+
+void UControlRig::SetControlValueImpl(const FName& InControlName, const FRigControlValue& InValue, bool bNotify,
+	const FRigControlModifiedContext& Context, bool bSetupUndo, bool bPrintPythonCommnds, bool bFixEulerFlips)
+{
+	const FRigElementKey Key(InControlName, ERigElementType::Control);
+
+	FRigControlElement* ControlElement = DynamicHierarchy->Find<FRigControlElement>(Key);
+	if(ControlElement == nullptr)
+	{
+		return;
+	}
+	if (bIsAdditive)
+	{
+		// Store the value to apply it after the backwards solve
+		FRigSetControlValueInfo Info = {InValue, bNotify, Context, bSetupUndo, bPrintPythonCommnds, bFixEulerFlips};
+		ControlValues.Add(ControlElement->GetKey(), Info);
+	}
+	else
+	{
+		DynamicHierarchy->SetControlValue(ControlElement, InValue, ERigControlValueType::Current, bSetupUndo, false, bPrintPythonCommnds, bFixEulerFlips);
+
+		if (bNotify && OnControlModified.IsBound())
+		{
+			OnControlModified.Broadcast(this, ControlElement, Context);
+		}
+	}
+}
+
+bool UControlRig::SetControlGlobalTransform(const FName& InControlName, const FTransform& InGlobalTransform, bool bNotify, const FRigControlModifiedContext& Context, bool bSetupUndo, bool bPrintPythonCommands, bool bFixEulerFlips)
+{
+	FTransform GlobalTransform = InGlobalTransform;
+	ERigTransformType::Type TransformType = ERigTransformType::CurrentGlobal;
+	if (IsConstructionModeEnabled())
+	{
+#if WITH_EDITOR
+		if(bRecordSelectionPoseForConstructionMode)
+		{
+			SelectionPoseForConstructionMode.FindOrAdd(FRigElementKey(InControlName, ERigElementType::Control)) = GlobalTransform;
+		}
+#endif
+		TransformType = ERigTransformType::InitialGlobal;
+		GlobalTransform = SetupControlFromGlobalTransform(InControlName, GlobalTransform);
+	}
+
+	FRigControlElement* Control = FindControl(InControlName);
+	if (IsAdditive())
+	{
+		if (Control)
+		{
+			// For additive control rigs, proxy controls should not be stored and applied after the backwards solve like other controls
+			// Instead, apply the transform as global, run the forward solve and let the driven controls request keying
+			if (Control->Settings.AnimationType == ERigControlAnimationType::ProxyControl)
+			{
+				DynamicHierarchy->SetTransform(Control, GlobalTransform, ERigTransformType::CurrentGlobal, true);
+				Execute(FRigUnit_BeginExecution::EventName);
+				return true;
+			}
+		}
+	}
+
+	FRigControlValue Value = GetControlValueFromGlobalTransform(InControlName, GlobalTransform, TransformType);
+	if (OnFilterControl.IsBound())
+	{
+		if (Control)
+		{
+			OnFilterControl.Broadcast(this, Control, Value);
+		}
+	}
+
+	SetControlValue(InControlName, Value, bNotify, Context, bSetupUndo, bPrintPythonCommands, bFixEulerFlips);
+	return true;
+}
+
+FRigControlValue UControlRig::GetControlValueFromGlobalTransform(const FName& InControlName, const FTransform& InGlobalTransform, ERigTransformType::Type InTransformType)
+{
+	FRigControlValue Value;
+
+	if (FRigControlElement* ControlElement = FindControl(InControlName))
+	{
+		if(DynamicHierarchy)
+		{
+			if (bIsAdditive)
+			{
+				const int32 ControlIndex = ControlsAfterBackwardsSolve.GetIndex(ControlElement->GetKey());
+				if (ControlIndex != INDEX_NONE)
+				{
+					// ParentGlobalTransform = ParentAnimationGlobalTransform + ParentAdditiveTransform
+					const FTransform& ParentGlobalTransform = DynamicHierarchy->GetParentTransform(ControlElement, ERigTransformType::CurrentGlobal);
+
+					// LocalTransform = InGlobal - ParentGlobalTransform - OffsetLocal
+					const FTransform OffsetLocalTransform = InGlobalTransform.GetRelativeTransform(ParentGlobalTransform);
+					const FTransform LocalTransform = OffsetLocalTransform.GetRelativeTransform(DynamicHierarchy->GetControlOffsetTransform(ControlElement, ERigTransformType::CurrentLocal));
+
+					// Additive = LocalTransform - AnimLocalTransform
+					const FTransform& AnimLocalTransform = ControlsAfterBackwardsSolve[ControlIndex].LocalTransform;
+					const FTransform AdditiveTransform = LocalTransform.GetRelativeTransform(AnimLocalTransform);
+
+					Value.SetFromTransform(AdditiveTransform, ControlElement->Settings.ControlType, ControlElement->Settings.PrimaryAxis);
+				}
+			}
+			else
+			{
+				FTransform Transform = DynamicHierarchy->ComputeLocalControlValue(ControlElement, InGlobalTransform, InTransformType);
+				Value.SetFromTransform(Transform, ControlElement->Settings.ControlType, ControlElement->Settings.PrimaryAxis);
+			}
+
+			if (ShouldApplyLimits())
+			{
+				ControlElement->Settings.ApplyLimits(Value);
+			}
+		}
+	}
+
+	return Value;
+}
+
+void UControlRig::SetControlLocalTransform(const FName& InControlName, const FTransform& InLocalTransform, bool bNotify, const FRigControlModifiedContext& Context, bool bSetupUndo, bool bFixEulerFlips)
+{
+	if (FRigControlElement* ControlElement = FindControl(InControlName))
+	{
+		FRigControlValue Value;
+		Value.SetFromTransform(InLocalTransform, ControlElement->Settings.ControlType, ControlElement->Settings.PrimaryAxis);
+
+		if (OnFilterControl.IsBound())
+		{
+			OnFilterControl.Broadcast(this, ControlElement, Value);
+			
+		}
+		SetControlValue(InControlName, Value, bNotify, Context, bSetupUndo, false /*bPrintPython not defined!*/, bFixEulerFlips);
+	}
+}
+
+FTransform UControlRig::GetControlLocalTransform(const FName& InControlName)
+{
+	if(DynamicHierarchy == nullptr)
+	{
+		return FTransform::Identity;
+	}
+	if (bIsAdditive)
+	{
+		FRigElementKey ControlKey (InControlName, ERigElementType::Control);
+		if (FRigBaseElement* Element = DynamicHierarchy->Find(ControlKey))
+		{
+			if (FRigControlElement* ControlElement = Cast<FRigControlElement>(Element))
+			{
+				return GetControlValue(ControlElement, ERigControlValueType::Current).GetAsTransform(ControlElement->Settings.ControlType, ControlElement->Settings.PrimaryAxis);
+			}
+		}
+	}
+	return DynamicHierarchy->GetLocalTransform(FRigElementKey(InControlName, ERigElementType::Control));
+}
+
+const TArray<TSoftObjectPtr<UControlRigShapeLibrary>>& UControlRig::GetShapeLibraries() const
+{
+	const TArray<TSoftObjectPtr<UControlRigShapeLibrary>>* LibrariesPtr = &ShapeLibraries;
+
+	if(!GetClass()->IsNative() && ShapeLibraries.IsEmpty())
+	{
+		if (UControlRig* CDO = Cast<UControlRig>(GetClass()->GetDefaultObject()))
+		{
+			LibrariesPtr = &CDO->ShapeLibraries;
+		}
+	}
+
+	const TArray<TSoftObjectPtr<UControlRigShapeLibrary>>& Libraries = *LibrariesPtr;
+	for(const TSoftObjectPtr<UControlRigShapeLibrary>& ShapeLibrary : Libraries)
+	{
+		if (!ShapeLibrary.IsValid())
+		{
+			ShapeLibrary.LoadSynchronous();
+		}
+	}
+
+	return Libraries;
+}
+
+void UControlRig::SelectControl(const FName& InControlName, bool bSelect)
+{
+	if(DynamicHierarchy)
+	{
+		if(URigHierarchyController* Controller = DynamicHierarchy->GetController(true))
+		{
+			Controller->SelectElement(FRigElementKey(InControlName, ERigElementType::Control), bSelect);
+		}
+	}
+}
+
+bool UControlRig::ClearControlSelection()
+{
+	if(DynamicHierarchy)
+	{
+		if(URigHierarchyController* Controller = DynamicHierarchy->GetController(true))
+		{
+			return Controller->ClearSelection();
+		}
+	}
+	return false;
+}
+
+TArray<FName> UControlRig::CurrentControlSelection() const
+{
+	TArray<FName> SelectedControlNames;
+
+	if(DynamicHierarchy)
+	{
+		TArray<const FRigBaseElement*> SelectedControls = DynamicHierarchy->GetSelectedElements(ERigElementType::Control);
+		for (const FRigBaseElement* SelectedControl : SelectedControls)
+		{
+			SelectedControlNames.Add(SelectedControl->GetFName());
+		}
+#if WITH_EDITOR
+		for(const TSharedPtr<FRigDirectManipulationInfo>& ManipulationInfo : RigUnitManipulationInfos)
+		{
+			SelectedControlNames.Add(ManipulationInfo->ControlKey.Name);
+		}
+#endif
+	}
+	return SelectedControlNames;
+}
+
+bool UControlRig::IsControlSelected(const FName& InControlName)const
+{
+	if(DynamicHierarchy)
+	{
+		if(FRigControlElement* ControlElement = FindControl(InControlName))
+		{
+			return DynamicHierarchy->IsSelected(ControlElement);
+		}
+	}
+	return false;
+}
+
+void UControlRig::HandleHierarchyModified(ERigHierarchyNotification InNotification, URigHierarchy* InHierarchy,
+    const FRigBaseElement* InElement)
+{
+	switch(InNotification)
+	{
+		case ERigHierarchyNotification::ElementSelected:
+		case ERigHierarchyNotification::ElementDeselected:
+		{
+			bool bClearTransientControls = true;
+			if(FRigControlElement* ControlElement = Cast<FRigControlElement>((FRigBaseElement*)InElement))
+			{
+				const bool bSelected = InNotification == ERigHierarchyNotification::ElementSelected;
+				ControlSelected().Broadcast(this, ControlElement, bSelected);
+
+				OnControlSelected_BP.Broadcast(this, *ControlElement, bSelected);
+				bClearTransientControls = !ControlElement->Settings.bIsTransientControl;
+			}
+
+#if WITH_EDITOR
+			if(IsConstructionModeEnabled())
+			{
+				if(InElement->GetType() == ERigElementType::Control)
+				{
+					if(InNotification == ERigHierarchyNotification::ElementSelected)
+					{
+						SelectionPoseForConstructionMode.FindOrAdd(InElement->GetKey()) = GetHierarchy()->GetGlobalTransform(InElement->GetKey());
+					}
+					else
+					{
+						SelectionPoseForConstructionMode.Remove(InElement->GetKey());
+					}
+				}
+			}
+			if(bClearTransientControls)
+			{
+				ClearTransientControls();
+			}
+#endif
+			break;
+		}
+		case ERigHierarchyNotification::ControlSettingChanged:
+		case ERigHierarchyNotification::ControlShapeTransformChanged:
+		{
+			if(FRigControlElement* ControlElement = Cast<FRigControlElement>((FRigBaseElement*)InElement))
+			{
+				ControlModified().Broadcast(this, ControlElement, FRigControlModifiedContext(EControlRigSetKey::Never));
+			}
+			break;
+		}
+		default:
+		{
+			break;
+		}
+	}
+}
+
+#if WITH_EDITOR
+
+bool UControlRig::CanAddTransientControl(const URigVMUnitNode* InNode, const FRigDirectManipulationTarget& InTarget, FString* OutFailureReason)
+{
+	if (InNode == nullptr)
+	{
+		if(OutFailureReason)
+		{
+			static const FString Reason = TEXT("Provided node is nullptr.");
+			*OutFailureReason = Reason;
+		}
+		return false;
+	}
+	if (DynamicHierarchy == nullptr)
+	{
+		if(OutFailureReason)
+		{
+			static const FString Reason = TEXT("The rig does not contain a hierarchy.");
+			*OutFailureReason = Reason;
+		}
+		return false;
+	}
+
+	const UScriptStruct* UnitStruct = InNode->GetScriptStruct();
+	if(UnitStruct == nullptr)
+	{
+		if(OutFailureReason)
+		{
+			static const FString Reason = TEXT("The node is not resolved.");
+			*OutFailureReason = Reason;
+		}
+		return false;
+	}
+
+	const TSharedPtr<FStructOnScope> NodeInstance = InNode->ConstructLiveStructInstance(this);
+	if(!NodeInstance.IsValid() || !NodeInstance->IsValid())
+	{
+		if(OutFailureReason)
+		{
+			static const FString Reason = TEXT("Unexpected error: Node instance could not be constructed.");
+			*OutFailureReason = Reason;
+		}
+		return false;
+	}
+
+	const FRigUnit* UnitInstance = GetRigUnitInstanceFromScope(NodeInstance);
+	check(UnitInstance);
+
+	TArray<FRigDirectManipulationTarget> Targets;
+	if(UnitInstance->GetDirectManipulationTargets(InNode, NodeInstance, DynamicHierarchy, Targets, OutFailureReason))
+	{
+		return Targets.Contains(InTarget);
+	}
+	return false;
+}
+
+FName UControlRig::AddTransientControl(const URigVMUnitNode* InNode, const FRigDirectManipulationTarget& InTarget)
+{
+	if(!CanAddTransientControl(InNode, InTarget, nullptr))
+	{
+		return NAME_None;
+	}
+	RemoveTransientControl(InNode, InTarget);
+
+	URigHierarchyController* Controller = DynamicHierarchy->GetController(true);
+	if(Controller == nullptr)
+	{
+		return NAME_None;
+	}
+
+	const UScriptStruct* UnitStruct = InNode->GetScriptStruct();
+	check(UnitStruct);
+
+	const TSharedPtr<FStructOnScope> NodeInstance = InNode->ConstructLiveStructInstance(this);
+	check(NodeInstance.IsValid() && NodeInstance->IsValid());
+
+	const FRigUnit* UnitInstance = GetRigUnitInstanceFromScope(NodeInstance);
+	check(UnitInstance);
+
+	const FName ControlName = GetNameForTransientControl(InNode, InTarget);
+	FTransform ShapeTransform = FTransform::Identity;
+	ShapeTransform.SetScale3D(FVector::ZeroVector);
+
+	FRigControlSettings Settings;
+	Settings.ControlType = ERigControlType::Transform;
+	Settings.bIsTransientControl = true;
+	Settings.DisplayName = TEXT("Temporary Control");
+
+	TSharedPtr<FRigDirectManipulationInfo> Info = MakeShareable(new FRigDirectManipulationInfo());
+	Info->Target = InTarget;
+	Info->Node = TWeakObjectPtr<const URigVMUnitNode>(InNode);
+
+	FRigControlValue Value = FRigControlValue::Make(FTransform::Identity);
+	UnitInstance->ConfigureDirectManipulationControl(InNode, Info, Settings, Value);
+
+	Controller->ClearSelection();
+
+    const FRigElementKey ControlKey = Controller->AddControl(
+    	ControlName,
+    	FRigElementKey(),
+    	Settings,
+    	Value,
+    	FTransform::Identity,
+    	ShapeTransform, false);
+
+	Info->ControlKey = ControlKey;
+	RigUnitManipulationInfos.Add(Info);
+	SetTransientControlValue(InNode, Info);
+
+	return Info->ControlKey.Name;
+}
+
+bool UControlRig::SetTransientControlValue(const URigVMUnitNode* InNode, TSharedPtr<FRigDirectManipulationInfo> InInfo)
+{
+	check(InNode);
+	check(DynamicHierarchy);
+	
+	const TSharedPtr<FStructOnScope> NodeInstance = InNode->ConstructLiveStructInstance(this);
+	check(NodeInstance.IsValid());
+	check(NodeInstance->IsValid());
+
+	const UScriptStruct* UnitStruct = InNode->GetScriptStruct();
+	check(UnitStruct);
+
+	FRigUnit* UnitInstance = GetRigUnitInstanceFromScope(NodeInstance);
+	check(UnitInstance);
+
+	const FName ControlName = GetNameForTransientControl(InNode, InInfo->Target);
+	const FRigControlElement* ControlElement = DynamicHierarchy->Find<FRigControlElement>({ ControlName, ERigElementType::Control });
+	if(ControlElement == nullptr)
+	{
+		return false;
+	}
+
+	FControlRigExecuteContext& PublicContext = GetRigVMExtendedExecuteContext().GetPublicDataSafe<FControlRigExecuteContext>();
+	const bool bResult = UnitInstance->UpdateHierarchyForDirectManipulation(InNode, NodeInstance, PublicContext, InInfo);
+	InInfo->bInitialized = true;
+	return bResult;
+}
+
+FName UControlRig::RemoveTransientControl(const URigVMUnitNode* InNode, const FRigDirectManipulationTarget& InTarget)
+{
+	if ((InNode == nullptr) || (DynamicHierarchy == nullptr))
+	{
+		return NAME_None;
+	}
+
+	RigUnitManipulationInfos.RemoveAll([InTarget](const TSharedPtr<FRigDirectManipulationInfo>& Info) -> bool
+	{
+		return Info->Target == InTarget;
+	});
+
+	URigHierarchyController* Controller = DynamicHierarchy->GetController(true);
+	if(Controller == nullptr)
+	{
+		return NAME_None;
+	}
+
+	const FName ControlName = GetNameForTransientControl(InNode, InTarget);
+	if(FRigControlElement* ControlElement = FindControl(ControlName))
+	{
+		DynamicHierarchy->Notify(ERigHierarchyNotification::ElementDeselected, ControlElement);
+		if(Controller->RemoveElement(ControlElement))
+		{
+#if WITH_EDITOR
+			SelectionPoseForConstructionMode.Remove(FRigElementKey(ControlName, ERigElementType::Control));
+#endif
+			return ControlName;
+		}
+	}
+
+	return NAME_None;
+}
+
+FName UControlRig::AddTransientControl(const FRigElementKey& InElement)
+{
+	if (!InElement.IsValid())
+	{
+		return NAME_None;
+	}
+
+	if(DynamicHierarchy == nullptr)
+	{
+		return NAME_None;
+	}
+
+	URigHierarchyController* Controller = DynamicHierarchy->GetController(true);
+	if(Controller == nullptr)
+	{
+		return NAME_None;
+	}
+
+	const FName ControlName = GetNameForTransientControl(InElement);
+	if(DynamicHierarchy->Contains(FRigElementKey(ControlName, ERigElementType::Control)))
+	{
+		SetTransientControlValue(InElement);
+		return ControlName;
+	}
+
+	const int32 ElementIndex = DynamicHierarchy->GetIndex(InElement);
+	if (ElementIndex == INDEX_NONE)
+	{
+		return NAME_None;
+	}
+
+	FTransform ShapeTransform = FTransform::Identity;
+	ShapeTransform.SetScale3D(FVector::ZeroVector);
+
+	FRigControlSettings Settings;
+	Settings.ControlType = ERigControlType::Transform;
+	Settings.bIsTransientControl = true;
+	Settings.DisplayName = TEXT("Temporary Control");
+
+	FRigElementKey Parent;
+	switch (InElement.Type)
+	{
+		case ERigElementType::Bone:
+		{
+			Parent = DynamicHierarchy->GetFirstParent(InElement);
+			break;
+		}
+		case ERigElementType::Null:
+		{
+			Parent = InElement;
+			break;
+		}
+		default:
+		{
+			break;
+		}
+	}
+
+	const FRigElementKey ControlKey = Controller->AddControl(
+        ControlName,
+        Parent,
+        Settings,
+        FRigControlValue::Make(FTransform::Identity),
+        FTransform::Identity,
+        ShapeTransform, false);
+
+	if (InElement.Type == ERigElementType::Bone)
+	{
+		// don't allow transient control to modify forward mode poses when we
+		// already switched to the construction mode
+		if (!IsConstructionModeEnabled())
+		{
+			if(FRigBoneElement* BoneElement = DynamicHierarchy->Find<FRigBoneElement>(InElement))
+			{
+				// add a modify bone AnimNode internally that the transient control controls for imported bones only
+				// for user created bones, refer to UControlRig::TransformOverrideForUserCreatedBones 
+				if (BoneElement->BoneType == ERigBoneType::Imported)
+				{ 
+					if (PreviewInstance)
+					{
+						PreviewInstance->ModifyBone(InElement.Name);
+					}
+				}
+				else if (BoneElement->BoneType == ERigBoneType::User)
+				{
+					// add an empty entry, which will be given the correct value in
+					// SetTransientControlValue(InElement);
+					TransformOverrideForUserCreatedBones.FindOrAdd(InElement.Name);
+				}
+			}
+		}
+	}
+
+	SetTransientControlValue(InElement);
+
+	return ControlKey.Name;
+}
+
+bool UControlRig::SetTransientControlValue(const FRigElementKey& InElement)
+{
+	if (!InElement.IsValid())
+	{
+		return false;
+	}
+
+	if(DynamicHierarchy == nullptr)
+	{
+		return false;
+	}
+
+	const FName ControlName = GetNameForTransientControl(InElement);
+	if(FRigControlElement* ControlElement = FindControl(ControlName))
+	{
+		if (InElement.Type == ERigElementType::Bone)
+		{
+			if (IsConstructionModeEnabled())
+			{
+				// need to get initial because that is what construction mode uses
+				// specifically, when user change the initial from the details panel
+				// this will allow the transient control to react to that change
+				const FTransform InitialLocalTransform = DynamicHierarchy->GetInitialLocalTransform(InElement);
+				DynamicHierarchy->SetTransform(ControlElement, InitialLocalTransform, ERigTransformType::InitialLocal, true, false);
+				DynamicHierarchy->SetTransform(ControlElement, InitialLocalTransform, ERigTransformType::CurrentLocal, true, false);
+			}
+			else
+			{
+				const FTransform LocalTransform = DynamicHierarchy->GetLocalTransform(InElement);
+				DynamicHierarchy->SetTransform(ControlElement, LocalTransform, ERigTransformType::InitialLocal, true, false);
+				DynamicHierarchy->SetTransform(ControlElement, LocalTransform, ERigTransformType::CurrentLocal, true, false);
+
+				if (FRigBoneElement* BoneElement = DynamicHierarchy->Find<FRigBoneElement>(InElement))
+				{
+					if (BoneElement->BoneType == ERigBoneType::Imported)
+					{
+						if (PreviewInstance)
+						{
+							if (FAnimNode_ModifyBone* Modify = PreviewInstance->FindModifiedBone(InElement.Name))
+							{
+								Modify->Translation = LocalTransform.GetTranslation();
+								Modify->Rotation = LocalTransform.GetRotation().Rotator();
+								Modify->TranslationSpace = EBoneControlSpace::BCS_ParentBoneSpace;
+								Modify->RotationSpace = EBoneControlSpace::BCS_ParentBoneSpace;
+							}
+						}	
+					}
+					else if (BoneElement->BoneType == ERigBoneType::User)
+					{
+						if (FTransform* TransformOverride = TransformOverrideForUserCreatedBones.Find(InElement.Name))
+						{
+							*TransformOverride = LocalTransform;
+						}	
+					}
+				}
+			}
+		}
+		else if (InElement.Type == ERigElementType::Null)
+		{
+			const FTransform GlobalTransform = DynamicHierarchy->GetGlobalTransform(InElement);
+			DynamicHierarchy->SetTransform(ControlElement, GlobalTransform, ERigTransformType::InitialGlobal, true, false);
+			DynamicHierarchy->SetTransform(ControlElement, GlobalTransform, ERigTransformType::CurrentGlobal, true, false);
+		}
+
+		return true;
+	}
+	return false;
+}
+
+FName UControlRig::RemoveTransientControl(const FRigElementKey& InElement)
+{
+	if (!InElement.IsValid())
+	{
+		return NAME_None;
+	}
+
+	if(DynamicHierarchy == nullptr)
+	{
+		return NAME_None;
+	}
+
+	URigHierarchyController* Controller = DynamicHierarchy->GetController(true);
+	if(Controller == nullptr)
+	{
+		return NAME_None;
+	}
+
+	const FName ControlName = GetNameForTransientControl(InElement);
+	if(FRigControlElement* ControlElement = FindControl(ControlName))
+	{
+		DynamicHierarchy->Notify(ERigHierarchyNotification::ElementDeselected, ControlElement);
+		if(Controller->RemoveElement(ControlElement))
+		{
+#if WITH_EDITOR
+			SelectionPoseForConstructionMode.Remove(FRigElementKey(ControlName, ERigElementType::Control));
+#endif
+			return ControlName;
+		}
+	}
+
+	return NAME_None;
+}
+
+FName UControlRig::GetNameForTransientControl(const URigVMUnitNode* InNode, const FRigDirectManipulationTarget& InTarget) const
+{
+	check(InNode);
+	check(DynamicHierarchy);
+	
+	const FString NodeName = InNode->GetName();
+	return DynamicHierarchy->GetSanitizedName(FRigName(FString::Printf(TEXT("ControlForNode|%s|%s"), *NodeName, *InTarget.Name)));
+}
+
+FString UControlRig::GetNodeNameFromTransientControl(const FRigElementKey& InKey)
+{
+	FString Name = InKey.Name.ToString();
+	if(Name.StartsWith(TEXT("ControlForNode|")))
+	{
+		Name.RightChopInline(15);
+		Name.LeftInline(Name.Find(TEXT("|")));
+	}
+	else
+	{
+		return FString();
+	}
+	return Name;
+}
+
+FString UControlRig::GetTargetFromTransientControl(const FRigElementKey& InKey)
+{
+	FString Name = InKey.Name.ToString();
+	if(Name.StartsWith(TEXT("ControlForNode|")))
+	{
+		Name.RightChopInline(15);
+		Name.RightChopInline(Name.Find(TEXT("|")) + 1);
+	}
+	else
+	{
+		return FString();
+	}
+	return Name;
+}
+
+TSharedPtr<FRigDirectManipulationInfo> UControlRig::GetRigUnitManipulationInfoForTransientControl(
+	const FRigElementKey& InKey)
+{
+	const TSharedPtr<FRigDirectManipulationInfo>* InfoPtr = RigUnitManipulationInfos.FindByPredicate(
+		[InKey](const TSharedPtr<FRigDirectManipulationInfo>& Info) -> bool
+		{
+			return Info->ControlKey == InKey;
+		}) ;
+
+	if(InfoPtr)
+	{
+		return *InfoPtr;
+	}
+
+	return TSharedPtr<FRigDirectManipulationInfo>();
+}
+
+FName UControlRig::GetNameForTransientControl(const FRigElementKey& InElement)
+{
+	if (InElement.Type == ERigElementType::Control)
+	{
+		return InElement.Name;
+	}
+
+	const FName EnumName = *StaticEnum<ERigElementType>()->GetDisplayNameTextByValue((int64)InElement.Type).ToString();
+	return *FString::Printf(TEXT("ControlForRigElement_%s_%s"), *EnumName.ToString(), *InElement.Name.ToString());
+}
+
+FRigElementKey UControlRig::GetElementKeyFromTransientControl(const FRigElementKey& InKey)
+{
+	if(InKey.Type != ERigElementType::Control)
+	{
+		return FRigElementKey();
+	}
+	
+	static FString ControlRigForElementBoneName;
+	static FString ControlRigForElementNullName;
+
+	if (ControlRigForElementBoneName.IsEmpty())
+	{
+		ControlRigForElementBoneName = FString::Printf(TEXT("ControlForRigElement_%s_"),
+            *StaticEnum<ERigElementType>()->GetDisplayNameTextByValue((int64)ERigElementType::Bone).ToString());
+		ControlRigForElementNullName = FString::Printf(TEXT("ControlForRigElement_%s_"),
+            *StaticEnum<ERigElementType>()->GetDisplayNameTextByValue((int64)ERigElementType::Null).ToString());
+	}
+	
+	FString Name = InKey.Name.ToString();
+	if(Name.StartsWith(ControlRigForElementBoneName))
+	{
+		Name.RightChopInline(ControlRigForElementBoneName.Len());
+		return FRigElementKey(*Name, ERigElementType::Bone);
+	}
+	if(Name.StartsWith(ControlRigForElementNullName))
+	{
+		Name.RightChopInline(ControlRigForElementNullName.Len());
+		return FRigElementKey(*Name, ERigElementType::Null);
+	}
+	
+	return FRigElementKey();;
+}
+
+void UControlRig::ClearTransientControls()
+{
+	if(DynamicHierarchy == nullptr)
+	{
+		return;
+	}
+
+	URigHierarchyController* Controller = DynamicHierarchy->GetController(true);
+	if(Controller == nullptr)
+	{
+		return;
+	}
+
+	if(bIsClearingTransientControls)
+	{
+		return;
+	}
+	TGuardValue<bool> ReEntryGuard(bIsClearingTransientControls, true);
+
+	RigUnitManipulationInfos.Reset();
+
+	const TArray<FRigControlElement*> ControlsToRemove = DynamicHierarchy->GetTransientControls();
+	for (FRigControlElement* ControlToRemove : ControlsToRemove)
+	{
+		const FRigElementKey KeyToRemove = ControlToRemove->GetKey();
+		if(Controller->RemoveElement(ControlToRemove))
+		{
+			SelectionPoseForConstructionMode.Remove(KeyToRemove);
+		}
+	}
+}
+
+void UControlRig::ApplyTransformOverrideForUserCreatedBones()
+{
+	if(DynamicHierarchy == nullptr)
+	{
+		return;
+	}
+	
+	for (const auto& Entry : TransformOverrideForUserCreatedBones)
+	{
+		DynamicHierarchy->SetLocalTransform(FRigElementKey(Entry.Key, ERigElementType::Bone), Entry.Value, false);
+	}
+}
+
+void UControlRig::ApplySelectionPoseForConstructionMode(const FName& InEventName)
+{
+	FRigControlModifiedContext ControlValueContext;
+	ControlValueContext.EventName = InEventName;
+
+	TGuardValue<bool> DisableRecording(bRecordSelectionPoseForConstructionMode, false);
+	for(const TPair<FRigElementKey, FTransform>& Pair : SelectionPoseForConstructionMode)
+	{
+		const FName ControlName = Pair.Key.Name;
+		const FRigControlValue Value = GetControlValueFromGlobalTransform(ControlName, Pair.Value, ERigTransformType::InitialGlobal);
+		SetControlValue(ControlName, Value, true, ControlValueContext, false, false, false);
+	}
+}
+
+#endif
+
+void UControlRig::HandleHierarchyEvent(URigHierarchy* InHierarchy, const FRigEventContext& InEvent)
+{
+	if (RigEventDelegate.IsBound())
+	{
+		RigEventDelegate.Broadcast(InHierarchy, InEvent);
+	}
+
+	switch (InEvent.Event)
+	{
+		case ERigEvent::RequestAutoKey:
+		{
+			int32 Index = InHierarchy->GetIndex(InEvent.Key);
+			if (Index != INDEX_NONE && InEvent.Key.Type == ERigElementType::Control)
+			{
+				if(FRigControlElement* ControlElement = InHierarchy->GetChecked<FRigControlElement>(Index))
+				{
+					FRigControlModifiedContext Context;
+					Context.SetKey = EControlRigSetKey::Always;
+					Context.LocalTime = InEvent.LocalTime;
+					Context.EventName = InEvent.SourceEventName;
+					ControlModified().Broadcast(this, ControlElement, Context);
+				}
+			}
+			break;
+		}
+		case ERigEvent::OpenUndoBracket:
+		case ERigEvent::CloseUndoBracket:
+		{
+			const bool bOpenUndoBracket = InEvent.Event == ERigEvent::OpenUndoBracket;
+			ControlUndoBracketIndex = FMath::Max<int32>(0, ControlUndoBracketIndex + (bOpenUndoBracket ? 1 : -1));
+			ControlUndoBracket().Broadcast(this, bOpenUndoBracket);
+			break;
+		}
+		default:
+		{
+			break;
+		}
+	}
+}
+
+void UControlRig::GetControlsInOrder(TArray<FRigControlElement*>& SortedControls) const
+{
+	SortedControls.Reset();
+
+	if(DynamicHierarchy == nullptr)
+	{
+		return;
+	}
+
+	SortedControls = DynamicHierarchy->GetControls(true);
+}
+
+const FRigInfluenceMap* UControlRig::FindInfluenceMap(const FName& InEventName)
+{
+	if (UControlRig* CDO = GetClass()->GetDefaultObject<UControlRig>())
+	{
+		return CDO->Influences.Find(InEventName);
+	}
+	return nullptr;
+}
+
+#if WITH_EDITOR
+
+void UControlRig::PreEditChange(FProperty* PropertyAboutToChange)
+{
+	// for BP user authored properties let's ignore changes since they
+	// will be distributed from the BP anyway to all archetype instances.
+	if(PropertyAboutToChange && !PropertyAboutToChange->IsNative())
+	{
+		if(const UBlueprint* Blueprint = Cast<UBlueprint>(GetClass()->ClassGeneratedBy))
+		{
+			const struct FBPVariableDescription* FoundVariable = Blueprint->NewVariables.FindByPredicate(
+				[PropertyAboutToChange](const struct FBPVariableDescription& NewVariable)
+				{
+					return NewVariable.VarName == PropertyAboutToChange->GetFName();
+				}
+			);
+
+			if(FoundVariable)
+			{
+				return;
+			}
+		}
+	}
+	
+	Super::PreEditChange(PropertyAboutToChange);
+}
+
+void UControlRig::PostEditChangeProperty(FPropertyChangedEvent& PropertyChangedEvent)
+{
+	Super::PostEditChangeProperty(PropertyChangedEvent);
+}
+#endif
+
+FRigUnit* UControlRig::GetRigUnitInstanceFromScope(TSharedPtr<FStructOnScope> InScope)
+{
+	if(InScope.IsValid())
+	{
+		if(InScope->IsValid())
+		{
+			if(InScope->GetStruct()->IsChildOf(FRigUnit::StaticStruct()))
+			{
+				return (FRigUnit*)InScope->GetStructMemory(); 
+			}
+		}
+	}
+	static FStructOnScope DefaultRigUnitInstance(FRigUnit::StaticStruct());
+	return (FRigUnit*)DefaultRigUnitInstance.GetStructMemory();
+}
+
+const TArray<UAssetUserData*>* UControlRig::GetAssetUserDataArray() const
+{
+	if(HasAnyFlags(RF_ClassDefaultObject))
+	{
+#if WITH_EDITOR
+		static thread_local TArray<TObjectPtr<UAssetUserData>> CachedAssetUserData;
+		CachedAssetUserData.Reset();
+		CachedAssetUserData.Append(AssetUserData);
+		CachedAssetUserData.Append(AssetUserDataEditorOnly);
+		return &ToRawPtrTArrayUnsafe(CachedAssetUserData);
+#else
+		return &ToRawPtrTArrayUnsafe(AssetUserData);
+#endif
+	}
+
+	CombinedAssetUserData.Reset();
+
+	if (UControlRig* CDO = Cast<UControlRig>(GetClass()->GetDefaultObject(false)))
+	{
+		CombinedAssetUserData.Append(*CDO->GetAssetUserDataArray());
+	}
+	else
+	{
+		CombinedAssetUserData.Append(AssetUserData);
+#if WITH_EDITOR
+		CombinedAssetUserData.Append(AssetUserDataEditorOnly);
+#endif
+	}
+	
+	if(GetExternalAssetUserDataDelegate.IsBound())
+	{
+		CombinedAssetUserData.Append(GetExternalAssetUserDataDelegate.Execute());
+	}
+
+	// find the data assets on the external variable list
+	TArray<FRigVMExternalVariable> ExternalVariables = GetExternalVariablesImpl(false);
+	for(const FRigVMExternalVariable& ExternalVariable : ExternalVariables)
+	{
+		if(ExternalVariable.Memory != nullptr)
+		{
+			if(const UClass* Class = Cast<UClass>(ExternalVariable.TypeObject))
+			{
+				if(Class->IsChildOf(UDataAsset::StaticClass()))
+				{
+					TObjectPtr<UDataAsset>& DataAsset = *(TObjectPtr<UDataAsset>*)ExternalVariable.Memory;
+					if(IsValid(DataAsset) && !DataAsset->GetFName().IsNone())
+					{
+						if(const TObjectPtr<UDataAssetLink>* ExistingDataAssetLink =
+							ExternalVariableDataAssetLinks.Find(ExternalVariable.Name))
+						{
+							if(!IsValid(*ExistingDataAssetLink) ||
+								(*ExistingDataAssetLink)->HasAnyFlags(EObjectFlags::RF_BeginDestroyed))
+							{
+								ExternalVariableDataAssetLinks.Remove(ExternalVariable.Name);
+							}
+						}
+						if(!ExternalVariableDataAssetLinks.Contains(ExternalVariable.Name))
+						{
+							ExternalVariableDataAssetLinks.Add(
+								ExternalVariable.Name,
+								NewObject<UDataAssetLink>(GetTransientPackage(), UDataAssetLink::StaticClass(), NAME_None, RF_Transient));
+						}
+
+						TObjectPtr<UDataAssetLink>& DataAssetLink = ExternalVariableDataAssetLinks.FindChecked(ExternalVariable.Name);
+						DataAssetLink->NameSpace = ExternalVariable.Name.ToString();
+						DataAssetLink->SetDataAsset(DataAsset);
+						CombinedAssetUserData.Add(DataAssetLink);
+					}
+				}
+			}
+		}
+	}
+
+	if(OuterSceneComponent.IsValid())
+	{
+		if(const USkeletalMeshComponent* SkeletalMeshComponent = Cast<USkeletalMeshComponent>(OuterSceneComponent))
+		{
+			if(USkeletalMesh* SkeletalMesh = SkeletalMeshComponent->GetSkeletalMeshAsset())
+			{
+				if(const USkeleton* Skeleton = SkeletalMesh->GetSkeleton())
+				{
+					if(const TArray<UAssetUserData*>* SkeletonUserDataArray = Skeleton->GetAssetUserDataArray())
+					{
+						for(UAssetUserData* SkeletonUserData : *SkeletonUserDataArray)
+						{
+							CombinedAssetUserData.Add(SkeletonUserData);
+						}
+					}
+				}
+				if(const TArray<UAssetUserData*>* SkeletalMeshUserDataArray = SkeletalMesh->GetAssetUserDataArray())
+				{
+					for(UAssetUserData* SkeletalMeshUserData : *SkeletalMeshUserDataArray)
+					{
+						CombinedAssetUserData.Add(SkeletalMeshUserData);
+					}
+				}
+			}
+			if (const TArray<UAssetUserData*>* ActorComponentUserDataArray = SkeletalMeshComponent->GetAssetUserDataArray())
+			{
+				for(UAssetUserData* ActorComponentUserData : *ActorComponentUserDataArray)
+				{
+					CombinedAssetUserData.Add(ActorComponentUserData);
+				}
+			}
+		}
+	}
+
+	return &ToRawPtrTArrayUnsafe(CombinedAssetUserData);
+}
+
+void UControlRig::CopyPoseFromOtherRig(UControlRig* Subject)
+{
+	check(DynamicHierarchy);
+	check(Subject);
+	URigHierarchy* OtherHierarchy = Subject->GetHierarchy();
+	check(OtherHierarchy);
+
+	for (FRigBaseElement* Element : *DynamicHierarchy)
+	{
+		FRigBaseElement* OtherElement = OtherHierarchy->Find(Element->GetKey());
+		if(OtherElement == nullptr)
+		{
+			continue;
+		}
+
+		if(OtherElement->GetType() != Element->GetType())
+		{
+			continue;
+		}
+
+		if(FRigBoneElement* BoneElement = Cast<FRigBoneElement>(Element))
+		{
+			FRigBoneElement* OtherBoneElement = CastChecked<FRigBoneElement>(OtherElement);
+			const FTransform Transform = OtherHierarchy->GetTransform(OtherBoneElement, ERigTransformType::CurrentLocal);
+			DynamicHierarchy->SetTransform(BoneElement, Transform, ERigTransformType::CurrentLocal, true, false);
+		}
+		else if(FRigCurveElement* CurveElement = Cast<FRigCurveElement>(Element))
+		{
+			FRigCurveElement* OtherCurveElement = CastChecked<FRigCurveElement>(OtherElement);
+			const float Value = OtherHierarchy->GetCurveValue(OtherCurveElement);
+			DynamicHierarchy->SetCurveValue(CurveElement, Value, false);
+		}
+	}
+}
+
+void UControlRig::SetBoneInitialTransformsFromAnimInstance(UAnimInstance* InAnimInstance)
+{
+	FMemMark Mark(FMemStack::Get());
+	FCompactPose OutPose;
+	OutPose.ResetToRefPose(InAnimInstance->GetRequiredBones());
+	SetBoneInitialTransformsFromCompactPose(&OutPose);
+}
+
+void UControlRig::SetBoneInitialTransformsFromAnimInstanceProxy(const FAnimInstanceProxy* InAnimInstanceProxy)
+{
+	FMemMark Mark(FMemStack::Get());
+	FCompactPose OutPose;
+	OutPose.ResetToRefPose(InAnimInstanceProxy->GetRequiredBones());
+	SetBoneInitialTransformsFromCompactPose(&OutPose);
+}
+
+void UControlRig::SetBoneInitialTransformsFromSkeletalMeshComponent(USkeletalMeshComponent* InSkelMeshComp, bool bUseAnimInstance)
+{
+	check(InSkelMeshComp);
+	check(DynamicHierarchy);
+	
+	if (!bUseAnimInstance && (InSkelMeshComp->GetAnimInstance() != nullptr))
+	{
+		SetBoneInitialTransformsFromAnimInstance(InSkelMeshComp->GetAnimInstance());
+	}
+	else
+	{
+		SetBoneInitialTransformsFromSkeletalMesh(InSkelMeshComp->GetSkeletalMeshAsset());
+	}
+}
+
+
+void UControlRig::SetBoneInitialTransformsFromSkeletalMesh(USkeletalMesh* InSkeletalMesh)
+{
+	if (ensure(InSkeletalMesh))
+	{ 
+		SetBoneInitialTransformsFromRefSkeleton(InSkeletalMesh->GetRefSkeleton());
+	}
+}
+
+void UControlRig::SetBoneInitialTransformsFromRefSkeleton(const FReferenceSkeleton& InReferenceSkeleton)
+{
+	check(DynamicHierarchy);
+
+	DynamicHierarchy->ForEach<FRigBoneElement>([this, InReferenceSkeleton](FRigBoneElement* BoneElement) -> bool
+	{
+		if(BoneElement->BoneType == ERigBoneType::Imported)
+		{
+			const int32 BoneIndex = InReferenceSkeleton.FindBoneIndex(BoneElement->GetFName());
+			if (BoneIndex != INDEX_NONE)
+			{
+				const FTransform LocalInitialTransform = InReferenceSkeleton.GetRefBonePose()[BoneIndex];
+				DynamicHierarchy->SetTransform(BoneElement, LocalInitialTransform, ERigTransformType::InitialLocal, true, false);
+				DynamicHierarchy->SetTransform(BoneElement, LocalInitialTransform, ERigTransformType::CurrentLocal, true, false);
+			}
+		}
+		return true;
+	});
+	bResetInitialTransformsBeforeConstruction = false;
+	RequestConstruction();
+}
+
+void UControlRig::SetBoneInitialTransformsFromCompactPose(FCompactPose* InCompactPose)
+{
+	check(InCompactPose);
+
+	if(!InCompactPose->IsValid())
+	{
+		return;
+	}
+	if(!InCompactPose->GetBoneContainer().IsValid())
+	{
+		return;
+	}
+	
+	FMemMark Mark(FMemStack::Get());
+
+	DynamicHierarchy->ForEach<FRigBoneElement>([this, InCompactPose](FRigBoneElement* BoneElement) -> bool
+		{
+			if (BoneElement->BoneType == ERigBoneType::Imported)
+			{
+				int32 MeshIndex = InCompactPose->GetBoneContainer().GetPoseBoneIndexForBoneName(BoneElement->GetFName());
+				if (MeshIndex != INDEX_NONE)
+				{
+					FCompactPoseBoneIndex CPIndex = InCompactPose->GetBoneContainer().MakeCompactPoseIndex(FMeshPoseBoneIndex(MeshIndex));
+					if (CPIndex != INDEX_NONE)
+					{
+						FTransform LocalInitialTransform = InCompactPose->GetRefPose(CPIndex);
+						DynamicHierarchy->SetTransform(BoneElement, LocalInitialTransform, ERigTransformType::InitialLocal, true, false);
+					}
+				}
+			}
+			return true;
+		});
+
+	bResetInitialTransformsBeforeConstruction = false;
+	RequestConstruction();
+}
+
+const FRigControlElementCustomization* UControlRig::GetControlCustomization(const FRigElementKey& InControl) const
+{
+	check(InControl.Type == ERigElementType::Control);
+
+	if(const FRigControlElementCustomization* Customization = ControlCustomizations.Find(InControl))
+	{
+		return Customization;
+	}
+
+	if(DynamicHierarchy)
+	{
+		if(const FRigControlElement* ControlElement = DynamicHierarchy->Find<FRigControlElement>(InControl))
+		{
+			return &ControlElement->Settings.Customization;
+		}
+	}
+
+	return nullptr;
+}
+
+void UControlRig::SetControlCustomization(const FRigElementKey& InControl, const FRigControlElementCustomization& InCustomization)
+{
+	check(InControl.Type == ERigElementType::Control);
+	
+	ControlCustomizations.FindOrAdd(InControl) = InCustomization;
+}
+
+void UControlRig::PostInitInstanceIfRequired()
+{
+	if(GetHierarchy() == nullptr || VM == nullptr)
+	{
+		if(HasAnyFlags(RF_ClassDefaultObject))
+		{
+			PostInitInstance(nullptr);
+		}
+		else
+		{
+			UControlRig* CDO = GetClass()->GetDefaultObject<UControlRig>();
+			ensure(VM == nullptr || VM == CDO->VM);
+			PostInitInstance(CDO);
+		}
+	}
+}
+
+#if WITH_EDITORONLY_DATA
+void UControlRig::DeclareConstructClasses(TArray<FTopLevelAssetPath>& OutConstructClasses, const UClass* SpecificSubclass)
+{
+	Super::DeclareConstructClasses(OutConstructClasses, SpecificSubclass);
+	OutConstructClasses.Add(FTopLevelAssetPath(URigHierarchy::StaticClass()));
+	OutConstructClasses.Add(FTopLevelAssetPath(UAnimationDataSourceRegistry::StaticClass()));
+}
+
+#endif
+
+USceneComponent* UControlRig::GetOwningSceneComponent()
+{
+	if(OuterSceneComponent == nullptr)
+	{
+		const FControlRigExecuteContext& PublicContext = GetRigVMExtendedExecuteContext().GetPublicDataSafe<FControlRigExecuteContext>();
+		const FRigUnitContext& Context = PublicContext.UnitContext;
+
+		USceneComponent* SceneComponentFromRegistry = Context.DataSourceRegistry->RequestSource<USceneComponent>(UControlRig::OwnerComponent);
+		if (SceneComponentFromRegistry)
+		{
+			OuterSceneComponent = SceneComponentFromRegistry;
+		}
+
+		if(OuterSceneComponent == nullptr)
+		{
+			OuterSceneComponent = Super::GetOwningSceneComponent();
+		}
+	}
+	return OuterSceneComponent.Get();
+}
+
+void UControlRig::PostInitInstance(URigVMHost* InCDO)
+{
+	const EObjectFlags SubObjectFlags =
+	HasAnyFlags(RF_ClassDefaultObject) ?
+		RF_Public | RF_DefaultSubObject :
+		RF_Transient | RF_Transactional;
+
+	FRigVMExtendedExecuteContext& Context = GetRigVMExtendedExecuteContext();
+	
+	Context.SetContextPublicDataStruct(FControlRigExecuteContext::StaticStruct());
+
+	Context.ExecutionReachedExit().RemoveAll(this);
+	Context.ExecutionReachedExit().AddUObject(this, &UControlRig::HandleExecutionReachedExit);
+	UpdateVMSettings();
+
+	// set up the hierarchy
+	{
+		// If this is not a CDO, it should have never saved its hieararchy. However, we have found that some rigs in the past
+		// did save their hiearchies. If that's the case, let's mark them as garbage and rename them before creating our own hierarchy.
+		if (!HasAnyFlags(RF_ClassDefaultObject))
+		{
+			UObject* ObjectFound = StaticFindObjectFast(URigHierarchy::StaticClass(), this, TEXT("DynamicHierarchy"), true, RF_DefaultSubObject);
+			if (ObjectFound)
+			{
+				FName NewName = MakeUniqueObjectName(GetTransientPackage(), URigHierarchy::StaticClass(), TEXT("DynamicHierarchy_Deleted"));
+				ObjectFound->Rename(*NewName.ToString(), GetTransientPackage(), REN_ForceNoResetLoaders | REN_DoNotDirty | REN_DontCreateRedirectors | REN_NonTransactional);
+				ObjectFound->MarkAsGarbage();
+			}
+		}
+
+		DynamicHierarchy = NewObject<URigHierarchy>(this, TEXT("DynamicHierarchy"), SubObjectFlags);
+	}
+
+#if WITH_EDITOR
+		const TWeakObjectPtr<UControlRig> WeakThis = this;
+		DynamicHierarchy->OnUndoRedo().AddStatic(&UControlRig::OnHierarchyTransformUndoRedoWeak, WeakThis);
+#endif
+
+	if(!HasAnyFlags(RF_ClassDefaultObject))
+	{
+		if (ensure(InCDO))
+		{
+			ensure(VM == nullptr || VM == InCDO->GetVM());
+			if (VM == nullptr) // this is needed for some EngineTests, on a normal setup, the VM is set to the CDO VM already
+			{
+				VM = InCDO->GetVM();
+			}
+
+			DynamicHierarchy->CopyHierarchy(CastChecked<UControlRig>(InCDO)->GetHierarchy());
+		}
+	}
+	else // we are the CDO
+	{
+		check(InCDO == nullptr);
+
+		if (VM == nullptr)
+		{
+			VM = NewObject<URigVM>(this, TEXT("ControlRig_VM"), SubObjectFlags);
+		}
+
+		// for default objects we need to check if the CDO is rooted. specialized Control Rigs
+		// such as the FK control rig may not have a root since they are part of a C++ package.
+
+		// since the sub objects are created after the constructor
+		// GC won't consider them part of the CDO, even if they have the sub object flags
+		// so even if CDO is rooted and references these sub objects, 
+		// it is not enough to keep them alive.
+		// Hence, we have to add them to root here.
+		if(GetClass()->IsNative())
+		{
+			VM->AddToRoot();
+			DynamicHierarchy->AddToRoot();
+		}
+	}
+
+	if(UControlRig* CDOControlRig = Cast<UControlRig>(InCDO))
+	{
+		ElementKeyRedirector = FRigElementKeyRedirector(CDOControlRig->ElementKeyRedirector, DynamicHierarchy);
+	}
+
+	RequestInit();
+}
+
+UTransformableControlHandle* UControlRig::CreateTransformableControlHandle(
+	const FName& InControlName) const
+{
+	auto IsConstrainable = [this](const FName& InControlName)
+	{
+		const FRigControlElement* ControlElement = FindControl(InControlName);
+		if (!ControlElement)
+		{
+			return false;
+		}
+		
+		const FRigControlSettings& ControlSettings = ControlElement->Settings;
+		if (ControlSettings.ControlType == ERigControlType::Bool ||
+			ControlSettings.ControlType == ERigControlType::Float ||
+			ControlSettings.ControlType == ERigControlType::ScaleFloat ||
+			ControlSettings.ControlType == ERigControlType::Integer)
+		{
+			return false;
+		}
+
+		return true;
+	};
+
+	if (!IsConstrainable(InControlName))
+	{
+		return nullptr;
+	}
+	
+	UTransformableControlHandle* CtrlHandle = NewObject<UTransformableControlHandle>(GetTransientPackage(), NAME_None, RF_Transactional);
+	check(CtrlHandle);
+	CtrlHandle->ControlRig = this;
+	CtrlHandle->ControlName = InControlName;
+	CtrlHandle->RegisterDelegates();
+	return CtrlHandle;
+}
+
+void UControlRig::OnHierarchyTransformUndoRedo(URigHierarchy* InHierarchy, const FRigElementKey& InKey, ERigTransformType::Type InTransformType, const FTransform& InTransform, bool bIsUndo)
+{
+	if(InKey.Type == ERigElementType::Control)
+	{
+		if(FRigControlElement* ControlElement = InHierarchy->Find<FRigControlElement>(InKey))
+		{
+			ControlModified().Broadcast(this, ControlElement, FRigControlModifiedContext(EControlRigSetKey::Never));
+		}
+	}
+}
+
+UControlRig::FPoseScope::FPoseScope(UControlRig* InControlRig, ERigElementType InFilter, const TArray<FRigElementKey>& InElements, const ERigTransformType::Type InTransformType)
+: ControlRig(InControlRig)
+, Filter(InFilter)
+, TransformType(InTransformType)
+{
+	check(InControlRig);
+	const TArrayView<const FRigElementKey> ElementView(InElements.GetData(), InElements.Num());
+	CachedPose = InControlRig->GetHierarchy()->GetPose(IsInitial(InTransformType), InFilter, ElementView);
+}
+
+UControlRig::FPoseScope::~FPoseScope()
+{
+	check(ControlRig);
+
+	ControlRig->GetHierarchy()->SetPose(CachedPose, TransformType);
+}
+
+#if WITH_EDITOR
+
+UControlRig::FTransientControlScope::FTransientControlScope(TObjectPtr<URigHierarchy> InHierarchy)
+	:Hierarchy(InHierarchy)
+{
+	for (FRigControlElement* Control : Hierarchy->GetTransientControls())
+	{
+		FTransientControlInfo Info;
+		Info.Name = Control->GetFName();
+		Info.Parent = Hierarchy->GetFirstParent(Control->GetKey());
+		Info.Settings = Control->Settings;
+		// preserve whatever value that was produced by this transient control at the moment
+		Info.Value = Hierarchy->GetControlValue(Control->GetKey(),ERigControlValueType::Current);
+		Info.OffsetTransform = Hierarchy->GetControlOffsetTransform(Control, ERigTransformType::CurrentLocal);
+		Info.ShapeTransform = Hierarchy->GetControlShapeTransform(Control, ERigTransformType::CurrentLocal);
+				
+		SavedTransientControls.Add(Info);
+	}
+}
+
+UControlRig::FTransientControlScope::~FTransientControlScope()
+{
+	if (URigHierarchyController* Controller = Hierarchy->GetController())
+	{
+		for (const FTransientControlInfo& Info : SavedTransientControls)
+		{
+			Controller->AddControl(
+				Info.Name,
+				Info.Parent,
+				Info.Settings,
+				Info.Value,
+				Info.OffsetTransform,
+				Info.ShapeTransform,
+				false,
+				false
+			);
+		}
+	}
+}
+
+#endif
+ 
+#undef LOCTEXT_NAMESPACE
+
diff --git a/Engine/Plugins/Animation/ControlRig/Source/ControlRig/Private/Constraints/ControlRigTransformableHandle.cpp b/Engine/Plugins/Animation/ControlRig/Source/ControlRig/Private/Constraints/ControlRigTransformableHandle.cpp
--- a/Engine/Plugins/Animation/ControlRig/Source/ControlRig/Private/Constraints/ControlRigTransformableHandle.cpp
+++ b/Engine/Plugins/Animation/ControlRig/Source/ControlRig/Private/Constraints/ControlRigTransformableHandle.cpp
@@ -5,7 +5,7 @@
 
 #include "Components/SkeletalMeshComponent.h"
 #include "ControlRigComponent.h"
-#include "ControlRig.h"
+#include "BaseControlRig.h"
 #include "IControlRigObjectBinding.h"
 #include "Rigs/RigHierarchyElements.h"
 #include "Sequencer/MovieSceneControlRigParameterSection.h"
diff --git a/Engine/Plugins/Animation/ControlRig/Source/ControlRig/Private/ControlRig.cpp b/Engine/Plugins/Animation/ControlRig/Source/ControlRig/Private/ControlRig.cpp
deleted file mode 100644
diff --git a/Engine/Plugins/Animation/ControlRig/Source/ControlRig/Private/ControlRigBlueprintGeneratedClass.cpp b/Engine/Plugins/Animation/ControlRig/Source/ControlRig/Private/ControlRigBlueprintGeneratedClass.cpp
--- a/Engine/Plugins/Animation/ControlRig/Source/ControlRig/Private/ControlRigBlueprintGeneratedClass.cpp
+++ b/Engine/Plugins/Animation/ControlRig/Source/ControlRig/Private/ControlRigBlueprintGeneratedClass.cpp
@@ -3,7 +3,7 @@
 #include "ControlRigBlueprintGeneratedClass.h"
 #include "Units/Control/RigUnit_Control.h"
 #include "ControlRigObjectVersion.h"
-#include "ControlRig.h"
+#include "BaseControlRig.h"
 #include "UObject/Package.h"
 
 #include UE_INLINE_GENERATED_CPP_BY_NAME(ControlRigBlueprintGeneratedClass)
diff --git a/Engine/Plugins/Animation/ControlRig/Source/ControlRig/Private/ControlRigControlActor.cpp b/Engine/Plugins/Animation/ControlRig/Source/ControlRig/Private/ControlRigControlActor.cpp
--- a/Engine/Plugins/Animation/ControlRig/Source/ControlRig/Private/ControlRigControlActor.cpp
+++ b/Engine/Plugins/Animation/ControlRig/Source/ControlRig/Private/ControlRigControlActor.cpp
@@ -5,6 +5,7 @@
 #include "ControlRigGizmoLibrary.h"
 #include "IControlRigObjectBinding.h"
 #include "Sequencer/MovieSceneControlRigParameterTrack.h"
+#include "BaseControlRig.h"
 
 #include UE_INLINE_GENERATED_CPP_BY_NAME(ControlRigControlActor)
 
diff --git a/Engine/Plugins/Animation/ControlRig/Source/ControlRig/Private/ControlRigGizmoActor.cpp b/Engine/Plugins/Animation/ControlRig/Source/ControlRig/Private/ControlRigGizmoActor.cpp
--- a/Engine/Plugins/Animation/ControlRig/Source/ControlRig/Private/ControlRigGizmoActor.cpp
+++ b/Engine/Plugins/Animation/ControlRig/Source/ControlRig/Private/ControlRigGizmoActor.cpp
@@ -1,7 +1,7 @@
 // Copyright Epic Games, Inc. All Rights Reserved.
 
 #include "ControlRigGizmoActor.h"
-#include "ControlRig.h"
+#include "BaseControlRig.h"
 #include "Engine/World.h"
 #include "Engine/StaticMesh.h"
 #include "Components/StaticMeshComponent.h"
diff --git a/Engine/Plugins/Animation/ControlRig/Source/ControlRig/Private/ControlRigModule.cpp b/Engine/Plugins/Animation/ControlRig/Source/ControlRig/Private/ControlRigModule.cpp
--- a/Engine/Plugins/Animation/ControlRig/Source/ControlRig/Private/ControlRigModule.cpp
+++ b/Engine/Plugins/Animation/ControlRig/Source/ControlRig/Private/ControlRigModule.cpp
@@ -9,7 +9,7 @@
 #include "ControlRigObjectVersion.h"
 #include "Constraints/ControlRigTransformableHandle.h"
 #include "UObject/DevObjectVersion.h"
-#include "ControlRig.h"
+#include "BaseControlRig.h"
 
 // Unique Control Rig Object version id
 const FGuid FControlRigObjectVersion::GUID(0xA7820CFB, 0x20A74359, 0x8C542C14, 0x9623CF50);
diff --git a/Engine/Plugins/Animation/ControlRig/Source/ControlRig/Private/ControlRigTestData.cpp b/Engine/Plugins/Animation/ControlRig/Source/ControlRig/Private/ControlRigTestData.cpp
--- a/Engine/Plugins/Animation/ControlRig/Source/ControlRig/Private/ControlRigTestData.cpp
+++ b/Engine/Plugins/Animation/ControlRig/Source/ControlRig/Private/ControlRigTestData.cpp
@@ -1,7 +1,7 @@
 // Copyright Epic Games, Inc. All Rights Reserved.
 
 #include "ControlRigTestData.h"
-#include "ControlRig.h"
+#include "BaseControlRig.h"
 #include "HAL/PlatformTime.h"
 #if WITH_EDITOR
 #include "AssetToolsModule.h"
diff --git a/Engine/Plugins/Animation/ControlRig/Source/ControlRig/Private/Rigs/RigCurveContainer.cpp b/Engine/Plugins/Animation/ControlRig/Source/ControlRig/Private/Rigs/RigCurveContainer.cpp
--- a/Engine/Plugins/Animation/ControlRig/Source/ControlRig/Private/Rigs/RigCurveContainer.cpp
+++ b/Engine/Plugins/Animation/ControlRig/Source/ControlRig/Private/Rigs/RigCurveContainer.cpp
@@ -2,7 +2,6 @@
 
 #include "Rigs/RigCurveContainer.h"
 #include "Rigs/RigHierarchyContainer.h"
-#include "ControlRig.h"
 #include "HelperUtil.h"
 #include "Animation/Skeleton.h"
 
diff --git a/Engine/Plugins/Animation/ControlRig/Source/ControlRig/Private/Rigs/RigHierarchy.cpp b/Engine/Plugins/Animation/ControlRig/Source/ControlRig/Private/Rigs/RigHierarchy.cpp
--- a/Engine/Plugins/Animation/ControlRig/Source/ControlRig/Private/Rigs/RigHierarchy.cpp
+++ b/Engine/Plugins/Animation/ControlRig/Source/ControlRig/Private/Rigs/RigHierarchy.cpp
@@ -1,7 +1,7 @@
 // Copyright Epic Games, Inc. All Rights Reserved.
 
 #include "Rigs/RigHierarchy.h"
-#include "ControlRig.h"
+#include "BaseControlRig.h"
 
 #include "Rigs/RigHierarchyElements.h"
 #include "Rigs/RigHierarchyController.h"
diff --git a/Engine/Plugins/Animation/ControlRig/Source/ControlRig/Private/Rigs/RigHierarchyContainer.cpp b/Engine/Plugins/Animation/ControlRig/Source/ControlRig/Private/Rigs/RigHierarchyContainer.cpp
--- a/Engine/Plugins/Animation/ControlRig/Source/ControlRig/Private/Rigs/RigHierarchyContainer.cpp
+++ b/Engine/Plugins/Animation/ControlRig/Source/ControlRig/Private/Rigs/RigHierarchyContainer.cpp
@@ -1,7 +1,7 @@
 // Copyright Epic Games, Inc. All Rights Reserved.
 
 #include "Rigs/RigHierarchyContainer.h"
-#include "ControlRig.h"
+#include "BaseControlRig.h"
 #include "HelperUtil.h"
 #include "UObject/PropertyPortFlags.h"
 
diff --git a/Engine/Plugins/Animation/ControlRig/Source/ControlRig/Private/Rigs/RigHierarchyController.cpp b/Engine/Plugins/Animation/ControlRig/Source/ControlRig/Private/Rigs/RigHierarchyController.cpp
--- a/Engine/Plugins/Animation/ControlRig/Source/ControlRig/Private/Rigs/RigHierarchyController.cpp
+++ b/Engine/Plugins/Animation/ControlRig/Source/ControlRig/Private/Rigs/RigHierarchyController.cpp
@@ -1,7 +1,7 @@
 // Copyright Epic Games, Inc. All Rights Reserved.
 
 #include "Rigs/RigHierarchyController.h"
-#include "ControlRig.h"
+#include "BaseControlRig.h"
 #include "AnimationCoreLibrary.h"
 #include "UObject/Package.h"
 
@@ -14,8 +14,6 @@
 #include "RigVMPythonUtils.h"
 #endif
 
-#include "ControlRig.h"
-
 #include UE_INLINE_GENERATED_CPP_BY_NAME(RigHierarchyController)
 
 ////////////////////////////////////////////////////////////////////////////////
diff --git a/Engine/Plugins/Animation/ControlRig/Source/ControlRig/Private/Rigs/RigSpaceHierarchy.cpp b/Engine/Plugins/Animation/ControlRig/Source/ControlRig/Private/Rigs/RigSpaceHierarchy.cpp
--- a/Engine/Plugins/Animation/ControlRig/Source/ControlRig/Private/Rigs/RigSpaceHierarchy.cpp
+++ b/Engine/Plugins/Animation/ControlRig/Source/ControlRig/Private/Rigs/RigSpaceHierarchy.cpp
@@ -2,7 +2,6 @@
 
 #include "Rigs/RigSpaceHierarchy.h"
 #include "Rigs/RigHierarchyContainer.h"
-#include "ControlRig.h"
 #include "HelperUtil.h"
 
 #include UE_INLINE_GENERATED_CPP_BY_NAME(RigSpaceHierarchy)
diff --git a/Engine/Plugins/Animation/ControlRig/Source/ControlRig/Private/Sequencer/ControlRigLayerInstance.cpp b/Engine/Plugins/Animation/ControlRig/Source/ControlRig/Private/Sequencer/ControlRigLayerInstance.cpp
--- a/Engine/Plugins/Animation/ControlRig/Source/ControlRig/Private/Sequencer/ControlRigLayerInstance.cpp
+++ b/Engine/Plugins/Animation/ControlRig/Source/ControlRig/Private/Sequencer/ControlRigLayerInstance.cpp
@@ -7,7 +7,7 @@
 
 #include "Sequencer/ControlRigLayerInstance.h"
 #include "Sequencer/ControlRigLayerInstanceProxy.h"
-#include "ControlRig.h"
+#include "BaseControlRig.h"
 
 #include UE_INLINE_GENERATED_CPP_BY_NAME(ControlRigLayerInstance)
 
diff --git a/Engine/Plugins/Animation/ControlRig/Source/ControlRig/Private/Sequencer/ControlRigLayerInstanceProxy.cpp b/Engine/Plugins/Animation/ControlRig/Source/ControlRig/Private/Sequencer/ControlRigLayerInstanceProxy.cpp
--- a/Engine/Plugins/Animation/ControlRig/Source/ControlRig/Private/Sequencer/ControlRigLayerInstanceProxy.cpp
+++ b/Engine/Plugins/Animation/ControlRig/Source/ControlRig/Private/Sequencer/ControlRigLayerInstanceProxy.cpp
@@ -4,7 +4,7 @@
 #include "AnimNode_ControlRig_ExternalSource.h"
 #include "Sequencer/ControlRigLayerInstance.h"
 #include "AnimSequencerInstance.h"
-#include "ControlRig.h"
+#include "BaseControlRig.h"
 #include "Sequencer/MovieSceneControlRigParameterTrack.h"
 
 #include UE_INLINE_GENERATED_CPP_BY_NAME(ControlRigLayerInstanceProxy)
diff --git a/Engine/Plugins/Animation/ControlRig/Source/ControlRig/Private/Sequencer/ControlRigSequenceObjectReference.cpp b/Engine/Plugins/Animation/ControlRig/Source/ControlRig/Private/Sequencer/ControlRigSequenceObjectReference.cpp
--- a/Engine/Plugins/Animation/ControlRig/Source/ControlRig/Private/Sequencer/ControlRigSequenceObjectReference.cpp
+++ b/Engine/Plugins/Animation/ControlRig/Source/ControlRig/Private/Sequencer/ControlRigSequenceObjectReference.cpp
@@ -1,7 +1,7 @@
 // Copyright Epic Games, Inc. All Rights Reserved.
 
 #include "Sequencer/ControlRigSequenceObjectReference.h"
-#include "ControlRig.h"
+#include "BaseControlRig.h"
 
 #include UE_INLINE_GENERATED_CPP_BY_NAME(ControlRigSequenceObjectReference)
 
diff --git a/Engine/Plugins/Animation/ControlRig/Source/ControlRig/Private/Sequencer/ControlRigSequencerHelpers.cpp b/Engine/Plugins/Animation/ControlRig/Source/ControlRig/Private/Sequencer/ControlRigSequencerHelpers.cpp
--- a/Engine/Plugins/Animation/ControlRig/Source/ControlRig/Private/Sequencer/ControlRigSequencerHelpers.cpp
+++ b/Engine/Plugins/Animation/ControlRig/Source/ControlRig/Private/Sequencer/ControlRigSequencerHelpers.cpp
@@ -1,6 +1,6 @@
 // Copyright Epic Games, Inc. All Rights Reserved.
 #include "Sequencer/ControlRigSequencerHelpers.h"
-#include "ControlRig.h"
+#include "BaseControlRig.h"
 #include "Rigs/RigHierarchyElements.h"
 #include "Channels/MovieSceneChannelProxy.h"
 #include "Sequencer/MovieSceneControlRigParameterSection.h"
diff --git a/Engine/Plugins/Animation/ControlRig/Source/ControlRig/Private/Sequencer/MovieSceneControlRigParameterTemplate.cpp b/Engine/Plugins/Animation/ControlRig/Source/ControlRig/Private/Sequencer/MovieSceneControlRigParameterTemplate.cpp
--- a/Engine/Plugins/Animation/ControlRig/Source/ControlRig/Private/Sequencer/MovieSceneControlRigParameterTemplate.cpp
+++ b/Engine/Plugins/Animation/ControlRig/Source/ControlRig/Private/Sequencer/MovieSceneControlRigParameterTemplate.cpp
@@ -4,7 +4,7 @@
 #include "Evaluation/Blending/MovieSceneMultiChannelBlending.h"
 #include "Sequencer/MovieSceneControlRigParameterTrack.h"
 #include "Sequencer/MovieSceneControlRigParameterSection.h"
-#include "ControlRig.h"
+#include "BaseControlRig.h"
 #include "Evaluation/MovieSceneEvaluation.h"
 #include "IMovieScenePlayer.h"
 #include "MovieSceneCommonHelpers.h"
diff --git a/Engine/Plugins/Animation/ControlRig/Source/ControlRig/Private/Tools/ControlRigPose.cpp b/Engine/Plugins/Animation/ControlRig/Source/ControlRig/Private/Tools/ControlRigPose.cpp
--- a/Engine/Plugins/Animation/ControlRig/Source/ControlRig/Private/Tools/ControlRigPose.cpp
+++ b/Engine/Plugins/Animation/ControlRig/Source/ControlRig/Private/Tools/ControlRigPose.cpp
@@ -2,6 +2,7 @@
 #include "Tools/ControlRigPose.h"
 #include "Tools/ControlRigPoseProjectSettings.h"
 #include "IControlRigObjectBinding.h"
+#include "BaseControlRig.h"
 
 #include UE_INLINE_GENERATED_CPP_BY_NAME(ControlRigPose)
 
diff --git a/Engine/Plugins/Animation/ControlRig/Source/ControlRig/Private/Tools/ControlRigPoseMirrorTable.cpp b/Engine/Plugins/Animation/ControlRig/Source/ControlRig/Private/Tools/ControlRigPoseMirrorTable.cpp
--- a/Engine/Plugins/Animation/ControlRig/Source/ControlRig/Private/Tools/ControlRigPoseMirrorTable.cpp
+++ b/Engine/Plugins/Animation/ControlRig/Source/ControlRig/Private/Tools/ControlRigPoseMirrorTable.cpp
@@ -2,7 +2,7 @@
 
 #include "Tools/ControlRigPoseMirrorTable.h"
 #include "Tools/ControlRigPoseMirrorSettings.h"
-#include "ControlRig.h"
+#include "BaseControlRig.h"
 #include "Tools/ControlRigPose.h"
 #include "RigVMFunctions/Math/RigVMMathLibrary.h"
 
diff --git a/Engine/Plugins/Animation/ControlRig/Source/ControlRig/Private/Units/Execution/RigUnit_DynamicHierarchy.cpp b/Engine/Plugins/Animation/ControlRig/Source/ControlRig/Private/Units/Execution/RigUnit_DynamicHierarchy.cpp
--- a/Engine/Plugins/Animation/ControlRig/Source/ControlRig/Private/Units/Execution/RigUnit_DynamicHierarchy.cpp
+++ b/Engine/Plugins/Animation/ControlRig/Source/ControlRig/Private/Units/Execution/RigUnit_DynamicHierarchy.cpp
@@ -4,7 +4,7 @@
 #include "Engine/SkeletalMesh.h"
 #include "Rigs/RigHierarchyController.h"
 #include "Units/RigUnitContext.h"
-#include "ControlRig.h"
+#include "BaseControlRig.h"
 #include "Components/SkeletalMeshComponent.h"
 #include "Units/Execution/RigUnit_PrepareForExecution.h"
 
diff --git a/Engine/Plugins/Animation/ControlRig/Source/ControlRig/Private/Units/Execution/RigUnit_RigModules.cpp b/Engine/Plugins/Animation/ControlRig/Source/ControlRig/Private/Units/Execution/RigUnit_RigModules.cpp
--- a/Engine/Plugins/Animation/ControlRig/Source/ControlRig/Private/Units/Execution/RigUnit_RigModules.cpp
+++ b/Engine/Plugins/Animation/ControlRig/Source/ControlRig/Private/Units/Execution/RigUnit_RigModules.cpp
@@ -2,7 +2,7 @@
 
 #include "Units/Execution/RigUnit_RigModules.h"
 #include "Units/RigUnitContext.h"
-#include "ControlRig.h"
+#include "BaseControlRig.h"
 
 #include UE_INLINE_GENERATED_CPP_BY_NAME(RigUnit_RigModules)
 
diff --git a/Engine/Plugins/Animation/ControlRig/Source/ControlRig/Private/Units/RigUnitContext.cpp b/Engine/Plugins/Animation/ControlRig/Source/ControlRig/Private/Units/RigUnitContext.cpp
--- a/Engine/Plugins/Animation/ControlRig/Source/ControlRig/Private/Units/RigUnitContext.cpp
+++ b/Engine/Plugins/Animation/ControlRig/Source/ControlRig/Private/Units/RigUnitContext.cpp
@@ -1,7 +1,7 @@
 // Copyright Epic Games, Inc. All Rights Reserved.
 
 #include "Units/RigUnitContext.h"
-#include "ControlRig.h"
+#include "BaseControlRig.h"
 
 #include UE_INLINE_GENERATED_CPP_BY_NAME(RigUnitContext)
 
diff --git a/Engine/Plugins/Animation/ControlRig/Source/ControlRig/Public/AnimNode_ControlRig.h b/Engine/Plugins/Animation/ControlRig/Source/ControlRig/Public/AnimNode_ControlRig.h
--- a/Engine/Plugins/Animation/ControlRig/Source/ControlRig/Public/AnimNode_ControlRig.h
+++ b/Engine/Plugins/Animation/ControlRig/Source/ControlRig/Public/AnimNode_ControlRig.h
@@ -2,7 +2,7 @@
 
 #pragma once
 
-#include "ControlRig.h"
+#include "BaseControlRig.h"
 #include "Animation/InputScaleBias.h"
 #include "AnimNode_ControlRigBase.h"
 #include "Animation/AnimBulkCurves.h"
diff --git a/Engine/Plugins/Animation/ControlRig/Source/ControlRig/Public/AnimNode_ControlRigBase.h b/Engine/Plugins/Animation/ControlRig/Source/ControlRig/Public/AnimNode_ControlRigBase.h
--- a/Engine/Plugins/Animation/ControlRig/Source/ControlRig/Public/AnimNode_ControlRigBase.h
+++ b/Engine/Plugins/Animation/ControlRig/Source/ControlRig/Public/AnimNode_ControlRigBase.h
@@ -4,10 +4,10 @@
 
 #include "Animation/AnimBulkCurves.h"
 #include "Animation/AnimNode_CustomProperty.h"
-#include "ControlRig.h"
 #include "AnimNode_ControlRigBase.generated.h"
 
 class UNodeMappingContainer;
+class UControlRig;
 
 /** Struct defining the settings to override when driving a control rig */
 USTRUCT()
diff --git a/Engine/Plugins/Animation/ControlRig/Source/ControlRig/Public/BaseControlRig.h b/Engine/Plugins/Animation/ControlRig/Source/ControlRig/Public/BaseControlRig.h
new file mode 100644
--- /dev/null
+++ b/Engine/Plugins/Animation/ControlRig/Source/ControlRig/Public/BaseControlRig.h
@@ -0,0 +1,895 @@
+// Copyright Epic Games, Inc. All Rights Reserved.
+
+#pragma once
+#include "CoreMinimal.h"
+#include "UObject/SparseDelegate.h"
+#include "Engine/EngineBaseTypes.h"
+#include "Templates/SubclassOf.h"
+#include "ControlRigDefines.h"
+#include "ControlRigGizmoLibrary.h"
+#include "Rigs/RigHierarchy.h"
+#include "Units/RigUnitContext.h"
+#include "Animation/NodeMappingProviderInterface.h"
+#include "Units/RigUnit.h"
+#include "Units/Control/RigUnit_Control.h"
+#include "RigVMCore/RigVM.h"
+#include "RigVMHost.h"
+#include "Components/SceneComponent.h"
+#include "Animation/AnimInstance.h"
+#include "Animation/AnimInstanceProxy.h"
+#include "Animation/AttributesRuntime.h"
+#include "Rigs/RigModuleDefines.h"
+
+#if WITH_EDITOR
+#include "RigVMModel/RigVMPin.h"
+#include "RigVMModel/Nodes/RigVMUnitNode.h"
+#include "RigVMTypeUtils.h"
+#endif
+
+#if WITH_EDITOR
+#include "AnimPreviewInstance.h"
+#endif 
+
+#include "BaseControlRig.generated.h"
+
+class IControlRigObjectBinding;
+class UScriptStruct;
+class USkeletalMesh;
+class USkeletalMeshComponent;
+class AActor;
+class UTransformableControlHandle;
+
+struct FReferenceSkeleton;
+struct FRigUnit;
+struct FRigControl;
+
+CONTROLRIG_API DECLARE_LOG_CATEGORY_EXTERN(LogControlRig, Log, All);
+
+/** Runs logic for mapping input data to transforms (the "Rig") */
+UCLASS(Blueprintable, Abstract, editinlinenew)
+class CONTROLRIG_API UControlRig : public URigVMHost, public INodeMappingProviderInterface
+{
+	GENERATED_UCLASS_BODY()
+
+	friend class UControlRigComponent;
+	friend class SControlRigStackView;
+
+public:
+
+	/** Bindable event for external objects to contribute to / filter a control value */
+	DECLARE_EVENT_ThreeParams(UControlRig, FFilterControlEvent, UControlRig*, FRigControlElement*, FRigControlValue&);
+
+	/** Bindable event for external objects to be notified of Control changes */
+	DECLARE_EVENT_ThreeParams(UControlRig, FControlModifiedEvent, UControlRig*, FRigControlElement*, const FRigControlModifiedContext&);
+
+	/** Bindable event for external objects to be notified that a Control is Selected */
+	DECLARE_EVENT_ThreeParams(UControlRig, FControlSelectedEvent, UControlRig*, FRigControlElement*, bool);
+
+	/** Bindable event to manage undo / redo brackets in the client */
+	DECLARE_EVENT_TwoParams(UControlRig, FControlUndoBracketEvent, UControlRig*, bool /* bOpen */);
+	
+	// To support Blueprints/scripting, we need a different delegate type (a 'Dynamic' delegate) which supports looser style UFunction binding (using names).
+	DECLARE_DYNAMIC_MULTICAST_SPARSE_DELEGATE_ThreeParams(FOnControlSelectedBP, UControlRig, OnControlSelected_BP, UControlRig*, Rig, const FRigControlElement&, Control, bool, bSelected);
+
+	/** Bindable event to notify object binding change. */
+	DECLARE_EVENT_OneParam(UControlRig, FControlRigBoundEvent, UControlRig*);
+
+	static const FName OwnerComponent;
+
+	UFUNCTION(BlueprintCallable, Category = ControlRig)
+	static TArray<UControlRig*> FindControlRigs(UObject* Outer, TSubclassOf<UControlRig> OptionalClass);
+
+public:
+	virtual UWorld* GetWorld() const override;
+	virtual void Serialize(FArchive& Ar) override;
+	virtual void PostLoad() override;
+	virtual UScriptStruct* GetPublicContextStruct() const override { return FControlRigExecuteContext::StaticStruct(); }
+
+	// Returns the settings of the module this instance belongs to
+	const FRigModuleSettings& GetRigModuleSettings() const;
+
+	// Returns true if the rig is defined as a rig module
+	bool IsRigModule() const;
+
+	// Returns true if this rig is an instance module. Rigs may be a module but not instance
+	// when being interacted with the asset editor
+	bool IsRigModuleInstance() const;
+
+	// Returns the parent rig hosting this module instance
+	UControlRig* GetParentRig() const;
+
+	// Returns the namespace of this module (for example ArmModule::)
+	const FString& GetRigModuleNameSpace() const;
+
+	// Returns the redirector from key to key for this rig
+	FRigElementKeyRedirector& GetElementKeyRedirector();
+
+	/** Creates a transformable control handle for the specified control to be used by the constraints system. Should use the UObject from 
+	ConstraintsScriptingLibrary::GetManager(UWorld* InWorld)*/
+	UFUNCTION(BlueprintCallable, Category = "Control Rig | Constraints")
+	UTransformableControlHandle* CreateTransformableControlHandle(const FName& ControlName) const;
+
+
+#if WITH_EDITOR
+	/** Get the category of this ControlRig (for display in menus) */
+	virtual FText GetCategory() const;
+
+	/** Get the tooltip text to display for this node (displayed in graphs and from context menus) */
+	virtual FText GetToolTipText() const;
+#endif
+
+	/** Initialize things for the ControlRig */
+	virtual void Initialize(bool bInitRigUnits = true) override;
+
+	/** Initialize the VM */
+	virtual bool InitializeVM(const FName& InEventName) override;
+
+	/** Evaluates the ControlRig */
+	virtual void Evaluate_AnyThread() override;
+
+	/** Ticks animation of the skeletal mesh component bound to this control rig */
+	bool EvaluateSkeletalMeshComponent(double InDeltaTime);
+
+	/** Removes any stored additive control values */
+	void ResetControlValues();
+
+	/** Resets the stored pose coming from the anim sequence.
+	 * This usually indicates a new pose should be stored. */
+	void ClearPoseBeforeBackwardsSolve();
+
+	/* For additive rigs, will set control values by inverting the pose found after the backwards solve */
+	void InvertInputPose(const TArray<FRigElementKey>& InElements = TArray<FRigElementKey>(), EControlRigSetKey InSetKey = EControlRigSetKey::Never);
+
+	/** Setup bindings to a runtime object (or clear by passing in nullptr). */
+	void SetObjectBinding(TSharedPtr<IControlRigObjectBinding> InObjectBinding)
+	{
+		ObjectBinding = InObjectBinding;
+		OnControlRigBound.Broadcast(this);
+	}
+
+	TSharedPtr<IControlRigObjectBinding> GetObjectBinding() const
+	{
+		return ObjectBinding;
+	}
+
+	/** Find the actor the rig is bound to, if any */
+	UFUNCTION(BlueprintPure, Category = "Control Rig")
+	AActor* GetHostingActor() const;
+
+	UFUNCTION(BlueprintPure, Category = "Control Rig")
+	URigHierarchy* GetHierarchy()
+	{
+		return DynamicHierarchy;
+	}
+	
+	URigHierarchy* GetHierarchy() const
+	{
+		return DynamicHierarchy;
+	}
+
+#if WITH_EDITOR
+
+	// called after post reinstance when compilng blueprint by Sequencer
+	void PostReinstanceCallback(const UControlRig* Old);
+
+#endif // WITH_EDITOR
+	
+	// BEGIN UObject interface
+	virtual void BeginDestroy() override;
+	// END UObject interface
+
+	UPROPERTY(transient)
+	ERigExecutionType ExecutionType;
+
+	UPROPERTY()
+	FRigHierarchySettings HierarchySettings;
+
+	virtual bool Execute(const FName& InEventName) override;
+	virtual bool Execute_Internal(const FName& InEventName) override;
+	virtual void RequestInit() override;
+
+	bool AllConnectorsAreResolved(FString* OutFailureReason = nullptr, FRigElementKey* OutConnector = nullptr) const;
+
+	/** Requests to perform construction during the next execution */
+	UFUNCTION(BlueprintCallable, Category = "Control Rig")
+	void RequestConstruction();
+
+	bool IsConstructionRequired() const;
+
+	/** Contains a backwards solve event */
+	UFUNCTION(BlueprintCallable, Category = "Control Rig")
+	bool SupportsBackwardsSolve() const;
+
+	virtual void AdaptEventQueueForEvaluate(TArray<FName>& InOutEventQueueToRun) override;
+
+	/** INodeMappingInterface implementation */
+	virtual void GetMappableNodeData(TArray<FName>& OutNames, TArray<FNodeItem>& OutNodeItems) const override;
+
+	/** Data Source Registry Getter */
+	UAnimationDataSourceRegistry* GetDataSourceRegistry();
+
+	virtual TArray<FRigControlElement*> AvailableControls() const;
+	virtual FRigControlElement* FindControl(const FName& InControlName) const;
+	virtual bool ShouldApplyLimits() const { return !IsConstructionModeEnabled(); }
+	virtual bool IsConstructionModeEnabled() const;
+	virtual FTransform SetupControlFromGlobalTransform(const FName& InControlName, const FTransform& InGlobalTransform);
+	virtual FTransform GetControlGlobalTransform(const FName& InControlName) const;
+
+	// Sets the relative value of a Control
+	template<class T>
+	void SetControlValue(const FName& InControlName, T InValue, bool bNotify = true,
+		const FRigControlModifiedContext& Context = FRigControlModifiedContext(), bool bSetupUndo = true, bool bPrintPythonCommnds = false, bool bFixEulerFlips = false)
+	{
+		SetControlValueImpl(InControlName, FRigControlValue::Make<T>(InValue), bNotify, Context, bSetupUndo, bPrintPythonCommnds, bFixEulerFlips);
+	}
+
+	// Returns the value of a Control
+	FRigControlValue GetControlValue(const FName& InControlName)
+	{
+		const FRigElementKey Key(InControlName, ERigElementType::Control);
+		if (FRigBaseElement* Element = DynamicHierarchy->Find(Key))
+		{
+			if (FRigControlElement* ControlElement = Cast<FRigControlElement>(Element))
+			{
+				return GetControlValue(ControlElement, ERigControlValueType::Current);
+			}
+		}
+		return DynamicHierarchy->GetControlValue(Key);
+	}
+
+	FRigControlValue GetControlValue(FRigControlElement* InControl, const ERigControlValueType& InValueType);
+
+	// Sets the relative value of a Control
+	virtual void SetControlValueImpl(const FName& InControlName, const FRigControlValue& InValue, bool bNotify = true,
+		const FRigControlModifiedContext& Context = FRigControlModifiedContext(), bool bSetupUndo = true, bool bPrintPythonCommnds = false, bool bFixEulerFlips = false);
+
+	FTransform GetInitialLocalTransform(const FRigElementKey &InKey)
+	{
+		if (bIsAdditive)
+		{
+			// The initial value of all additive controls is always Identity
+			return FTransform::Identity;
+		}
+		return GetHierarchy()->GetInitialLocalTransform(InKey);
+	}
+
+	bool SetControlGlobalTransform(const FName& InControlName, const FTransform& InGlobalTransform, bool bNotify = true, const FRigControlModifiedContext& Context = FRigControlModifiedContext(), bool bSetupUndo = true, bool bPrintPythonCommands = false, bool bFixEulerFlips = false);
+
+	virtual FRigControlValue GetControlValueFromGlobalTransform(const FName& InControlName, const FTransform& InGlobalTransform, ERigTransformType::Type InTransformType);
+
+	virtual void SetControlLocalTransform(const FName& InControlName, const FTransform& InLocalTransform, bool bNotify = true, const FRigControlModifiedContext& Context = FRigControlModifiedContext(), bool bSetupUndo = true, bool bFixEulerFlips = false);
+	virtual FTransform GetControlLocalTransform(const FName& InControlName) ;
+
+	virtual const TArray<TSoftObjectPtr<UControlRigShapeLibrary>>& GetShapeLibraries() const;
+	virtual void CreateRigControlsForCurveContainer();
+	virtual void GetControlsInOrder(TArray<FRigControlElement*>& SortedControls) const;
+
+	UFUNCTION(BlueprintCallable, Category = "Control Rig")
+	virtual void SelectControl(const FName& InControlName, bool bSelect = true);
+	UFUNCTION(BlueprintCallable, Category = "Control Rig")
+	virtual bool ClearControlSelection();
+	UFUNCTION(BlueprintPure, Category = "Control Rig")
+	virtual TArray<FName> CurrentControlSelection() const;
+	UFUNCTION(BlueprintPure, Category = "Control Rig")
+	virtual bool IsControlSelected(const FName& InControlName)const;
+
+	// Returns true if this manipulatable subject is currently
+	// available for manipulation / is enabled.
+	virtual bool ManipulationEnabled() const
+	{
+		return bManipulationEnabled;
+	}
+
+	// Sets the manipulatable subject to enabled or disabled
+	virtual bool SetManipulationEnabled(bool Enabled = true)
+	{
+		if (bManipulationEnabled == Enabled)
+		{
+			return false;
+		}
+		bManipulationEnabled = Enabled;
+		return true;
+	}
+
+	// Returns a event that can be used to subscribe to
+	// filtering control data when needed
+	FFilterControlEvent& ControlFilter() { return OnFilterControl; }
+
+	// Returns a event that can be used to subscribe to
+	// change notifications coming from the manipulated subject.
+	FControlModifiedEvent& ControlModified() { return OnControlModified; }
+
+	// Returns a event that can be used to subscribe to
+	// selection changes coming from the manipulated subject.
+	FControlSelectedEvent& ControlSelected() { return OnControlSelected; }
+
+	// Returns an event that can be used to subscribe to
+	// Undo Bracket requests such as Open and Close.
+	FControlUndoBracketEvent & ControlUndoBracket() { return OnControlUndoBracket; }
+	
+	FControlRigBoundEvent& ControlRigBound() { return OnControlRigBound; };
+
+	bool IsCurveControl(const FRigControlElement* InControlElement) const;
+
+	DECLARE_EVENT_TwoParams(UControlRig, FControlRigExecuteEvent, class UControlRig*, const FName&);
+#if WITH_EDITOR
+	FControlRigExecuteEvent& OnPreConstructionForUI_AnyThread() { return PreConstructionForUIEvent; }
+#endif
+	FControlRigExecuteEvent& OnPreConstruction_AnyThread() { return PreConstructionEvent; }
+	FControlRigExecuteEvent& OnPostConstruction_AnyThread() { return PostConstructionEvent; }
+	FControlRigExecuteEvent& OnPreForwardsSolve_AnyThread() { return PreForwardsSolveEvent; }
+	FControlRigExecuteEvent& OnPostForwardsSolve_AnyThread() { return PostForwardsSolveEvent; }
+	FControlRigExecuteEvent& OnPreAdditiveValuesApplication_AnyThread() { return PreAdditiveValuesApplicationEvent; }
+	FRigEventDelegate& OnRigEvent_AnyThread() { return RigEventDelegate; }
+
+	// Setup the initial transform / ref pose of the bones based upon an anim instance
+	// This uses the current refpose instead of the RefSkeleton pose.
+	virtual void SetBoneInitialTransformsFromAnimInstance(UAnimInstance* InAnimInstance);
+
+	// Setup the initial transform / ref pose of the bones based upon an anim instance proxy
+	// This uses the current refpose instead of the RefSkeleton pose.
+	virtual void SetBoneInitialTransformsFromAnimInstanceProxy(const FAnimInstanceProxy* InAnimInstanceProxy);
+
+	// Setup the initial transform / ref pose of the bones based upon skeletal mesh component (ref skeleton)
+	// This uses the RefSkeleton pose instead of the current refpose (or vice versae is bUseAnimInstance == true)
+	virtual void SetBoneInitialTransformsFromSkeletalMeshComponent(USkeletalMeshComponent* InSkelMeshComp, bool bUseAnimInstance = false);
+
+	// Setup the initial transforms / ref pose of the bones based on a skeletal mesh
+	// This uses the RefSkeleton pose instead of the current refpose.
+	virtual void SetBoneInitialTransformsFromSkeletalMesh(USkeletalMesh* InSkeletalMesh);
+
+	// Setup the initial transforms / ref pose of the bones based on a reference skeleton
+	// This uses the RefSkeleton pose instead of the current refpose.
+	virtual void SetBoneInitialTransformsFromRefSkeleton(const FReferenceSkeleton& InReferenceSkeleton);
+
+private:
+
+	void SetBoneInitialTransformsFromCompactPose(FCompactPose* InCompactPose);
+
+public:
+	
+	const FRigControlElementCustomization* GetControlCustomization(const FRigElementKey& InControl) const;
+	void SetControlCustomization(const FRigElementKey& InControl, const FRigControlElementCustomization& InCustomization);
+
+	virtual void PostInitInstanceIfRequired() override;
+#if WITH_EDITORONLY_DATA
+	static void DeclareConstructClasses(TArray<FTopLevelAssetPath>& OutConstructClasses, const UClass* SpecificSubclass);
+#endif
+
+	virtual USceneComponent* GetOwningSceneComponent() override;
+
+protected:
+
+	virtual void PostInitInstance(URigVMHost* InCDO) override;
+
+	UPROPERTY()
+	TMap<FRigElementKey, FRigControlElementCustomization> ControlCustomizations;
+
+	UPROPERTY()
+	TObjectPtr<URigHierarchy> DynamicHierarchy;
+
+#if WITH_EDITORONLY_DATA
+	UPROPERTY()
+	TSoftObjectPtr<UControlRigShapeLibrary> GizmoLibrary_DEPRECATED;
+#endif
+
+	UPROPERTY()
+	TArray<TSoftObjectPtr<UControlRigShapeLibrary>> ShapeLibraries;
+
+	UPROPERTY(transient)
+	TMap<FString, FString> ShapeLibraryNameMap;
+
+	/** Runtime object binding */
+	TSharedPtr<IControlRigObjectBinding> ObjectBinding;
+
+#if WITH_EDITORONLY_DATA
+	// you either go Input or Output, currently if you put it in both place, Output will override
+	UPROPERTY()
+	TMap<FName, FCachedPropertyPath> InputProperties_DEPRECATED;
+
+	UPROPERTY()
+	TMap<FName, FCachedPropertyPath> OutputProperties_DEPRECATED;
+#endif
+
+private:
+	
+	void HandleOnControlModified(UControlRig* Subject, FRigControlElement* Control, const FRigControlModifiedContext& Context);
+
+public:
+	
+	class CONTROLRIG_API FAnimAttributeContainerPtrScope
+	{
+	public:
+		FAnimAttributeContainerPtrScope(UControlRig* InControlRig, UE::Anim::FStackAttributeContainer& InExternalContainer);
+		~FAnimAttributeContainerPtrScope();
+
+		UControlRig* ControlRig;
+	};
+	
+private:
+	UPROPERTY(Transient)
+	FRigVMExtendedExecuteContext RigVMExtendedExecuteContext;
+
+	UE::Anim::FStackAttributeContainer* ExternalAnimAttributeContainer;
+
+#if WITH_EDITOR
+	void SetEnableAnimAttributeTrace(bool bInEnable)
+	{
+		bEnableAnimAttributeTrace = bInEnable;
+	};
+	
+	bool bEnableAnimAttributeTrace;
+	
+	UE::Anim::FHeapAttributeContainer InputAnimAttributeSnapshot;
+	UE::Anim::FHeapAttributeContainer OutputAnimAttributeSnapshot;
+#endif
+	
+	/** The registry to access data source */
+	UPROPERTY(Transient)
+	TObjectPtr<UAnimationDataSourceRegistry> DataSourceRegistry;
+
+	/** Broadcasts a notification when launching the construction event */
+	FControlRigExecuteEvent PreConstructionForUIEvent;
+
+	/** Broadcasts a notification just before the controlrig is setup. */
+	FControlRigExecuteEvent PreConstructionEvent;
+
+	/** Broadcasts a notification whenever the controlrig has been setup. */
+	FControlRigExecuteEvent PostConstructionEvent;
+
+	/** Broadcasts a notification before a forward solve has been initiated */
+	FControlRigExecuteEvent PreForwardsSolveEvent;
+	
+	/** Broadcasts a notification after a forward solve has been initiated */
+	FControlRigExecuteEvent PostForwardsSolveEvent;
+
+	/** Broadcasts a notification before additive controls have been applied */
+	FControlRigExecuteEvent PreAdditiveValuesApplicationEvent;
+
+	/** Handle changes within the hierarchy */
+	void HandleHierarchyModified(ERigHierarchyNotification InNotification, URigHierarchy* InHierarchy, const FRigBaseElement* InElement);
+
+#if WITH_EDITOR
+	/** Add a transient / temporary control used to interact with a node */
+	FName AddTransientControl(const URigVMUnitNode* InNode, const FRigDirectManipulationTarget& InTarget);
+
+	/** Sets the value of a transient control based on a node */
+	bool SetTransientControlValue(const URigVMUnitNode* InNode, TSharedPtr<FRigDirectManipulationInfo> InInfo);
+
+	/** Remove a transient / temporary control used to interact with a node */
+	FName RemoveTransientControl(const URigVMUnitNode* InNode, const FRigDirectManipulationTarget& InTarget);
+
+	FName AddTransientControl(const FRigElementKey& InElement);
+
+	/** Sets the value of a transient control based on a bone */
+	bool SetTransientControlValue(const FRigElementKey& InElement);
+
+	/** Remove a transient / temporary control used to interact with a bone */
+	FName RemoveTransientControl(const FRigElementKey& InElement);
+
+	static FName GetNameForTransientControl(const FRigElementKey& InElement);
+	FName GetNameForTransientControl(const URigVMUnitNode* InNode, const FRigDirectManipulationTarget& InTarget) const;
+	static FString GetNodeNameFromTransientControl(const FRigElementKey& InKey);
+	static FString GetTargetFromTransientControl(const FRigElementKey& InKey);
+	TSharedPtr<FRigDirectManipulationInfo> GetRigUnitManipulationInfoForTransientControl(const FRigElementKey& InKey);
+	
+	static FRigElementKey GetElementKeyFromTransientControl(const FRigElementKey& InKey);
+	bool CanAddTransientControl(const URigVMUnitNode* InNode, const FRigDirectManipulationTarget& InTarget, FString* OutFailureReason);
+
+	/** Removes all  transient / temporary control used to interact with pins */
+	void ClearTransientControls();
+
+	UAnimPreviewInstance* PreviewInstance;
+
+	// this is needed because PreviewInstance->ModifyBone(...) cannot modify user created bones,
+	TMap<FName, FTransform> TransformOverrideForUserCreatedBones;
+	
+public:
+	
+	void ApplyTransformOverrideForUserCreatedBones();
+	void ApplySelectionPoseForConstructionMode(const FName& InEventName);
+	
+#endif
+
+private: 
+
+	void HandleHierarchyEvent(URigHierarchy* InHierarchy, const FRigEventContext& InEvent);
+	FRigEventDelegate RigEventDelegate;
+
+	void OnAddShapeLibrary(const FControlRigExecuteContext* InContext, const FString& InLibraryName, UControlRigShapeLibrary* InShapeLibrary, bool bReplaceExisting, bool bLogResults);
+	bool OnShapeExists(const FName& InShapeName) const;
+	virtual void InitializeFromCDO() override;
+
+
+	UPROPERTY()
+	FRigInfluenceMapPerEvent Influences;
+
+	const FRigInfluenceMap* FindInfluenceMap(const FName& InEventName);
+
+#if WITH_EDITORONLY_DATA
+	UPROPERTY()
+	TObjectPtr<UControlRig> InteractionRig_DEPRECATED;
+
+	UPROPERTY()
+	TSubclassOf<UControlRig> InteractionRigClass_DEPRECATED;
+#endif
+
+	FRigElementKeyRedirector ElementKeyRedirector;
+
+public:
+
+	UE_DEPRECATED(5.4, "InteractionRig is no longer used")
+	UFUNCTION(BlueprintGetter, meta = (DeprecatedFunction, DeprecationMessage = "InteractionRig is no longer used"))
+	UControlRig* GetInteractionRig() const 
+	{
+#if WITH_EDITORONLY_DATA
+		return InteractionRig_DEPRECATED;
+#endif
+		return nullptr;
+	}
+
+	UE_DEPRECATED(5.4, "InteractionRig is no longer used")
+	UFUNCTION(BlueprintSetter, meta = (DeprecatedFunction, DeprecationMessage = "InteractionRig is no longer used"))
+	void SetInteractionRig(UControlRig* InInteractionRig) {}
+
+	UE_DEPRECATED(5.4, "InteractionRig is no longer used")
+	UFUNCTION(BlueprintGetter, meta = (DeprecatedFunction, DeprecationMessage = "InteractionRig is no longer used"))
+	TSubclassOf<UControlRig> GetInteractionRigClass() const 
+	{ 
+#if WITH_EDITORONLY_DATA
+		return InteractionRigClass_DEPRECATED;
+#endif
+		return nullptr;
+	}
+
+	UE_DEPRECATED(5.4, "InteractionRig is no longer used")
+	UFUNCTION(BlueprintSetter, meta = (DeprecatedFunction, DeprecationMessage = "InteractionRig is no longer used"))
+	void SetInteractionRigClass(TSubclassOf<UControlRig> InInteractionRigClass) {}
+
+	// UObject interface
+#if WITH_EDITOR
+	virtual void PreEditChange(FProperty* PropertyAboutToChange) override;
+	virtual void PostEditChangeProperty(FPropertyChangedEvent& PropertyChangedEvent) override;
+#endif
+
+	float GetDebugBoneRadiusMultiplier() const { return DebugBoneRadiusMultiplier; }
+	static FRigUnit* GetRigUnitInstanceFromScope(TSharedPtr<FStructOnScope> InScope);
+
+public:
+	//~ Begin IInterface_AssetUserData Interface
+	virtual const TArray<UAssetUserData*>* GetAssetUserDataArray() const override;
+	//~ End IInterface_AssetUserData Interface
+protected:
+	mutable TArray<TObjectPtr<UAssetUserData>> CombinedAssetUserData;
+
+	UPROPERTY(Transient, DuplicateTransient)
+	mutable TMap<FName, TObjectPtr<UDataAssetLink>> ExternalVariableDataAssetLinks;
+
+	DECLARE_DELEGATE_RetVal(TArray<TObjectPtr<UAssetUserData>>, FGetExternalAssetUserData);
+	FGetExternalAssetUserData GetExternalAssetUserDataDelegate;
+
+private:
+
+	void CopyPoseFromOtherRig(UControlRig* Subject);
+
+protected:
+	bool bCopyHierarchyBeforeConstruction;
+	bool bResetInitialTransformsBeforeConstruction;
+	bool bManipulationEnabled;
+
+	int32 PreConstructionBracket;
+	int32 PostConstructionBracket;
+	int32 PreForwardsSolveBracket;
+	int32 PostForwardsSolveBracket;
+	int32 PreAdditiveValuesApplicationBracket;
+	int32 InteractionBracket;
+	int32 InterRigSyncBracket;
+	int32 ControlUndoBracketIndex;
+	uint8 InteractionType;
+	TArray<FRigElementKey> ElementsBeingInteracted;
+#if WITH_EDITOR
+	TArray<TSharedPtr<FRigDirectManipulationInfo>> RigUnitManipulationInfos;
+#endif
+	bool bInteractionJustBegan;
+
+	TWeakObjectPtr<USceneComponent> OuterSceneComponent;
+
+	bool IsRunningPreConstruction() const
+	{
+		return PreConstructionBracket > 0;
+	}
+
+	bool IsRunningPostConstruction() const
+	{
+		return PostConstructionBracket > 0;
+	}
+
+	bool IsInteracting() const
+	{
+		return InteractionBracket > 0;
+	}
+
+	uint8 GetInteractionType() const
+	{
+		return InteractionType;
+	}
+
+	bool IsSyncingWithOtherRig() const
+	{
+		return InterRigSyncBracket > 0;
+	}
+
+#if WITH_EDITOR
+	static void OnHierarchyTransformUndoRedoWeak(URigHierarchy* InHierarchy, const FRigElementKey& InKey, ERigTransformType::Type InTransformType, const FTransform& InTransform, bool bIsUndo, TWeakObjectPtr<UControlRig> WeakThis)
+	{
+		if(WeakThis.IsValid() && InHierarchy != nullptr)
+		{
+			WeakThis->OnHierarchyTransformUndoRedo(InHierarchy, InKey, InTransformType, InTransform, bIsUndo);
+		}
+	}
+#endif
+	
+	void OnHierarchyTransformUndoRedo(URigHierarchy* InHierarchy, const FRigElementKey& InKey, ERigTransformType::Type InTransformType, const FTransform& InTransform, bool bIsUndo);
+
+	FFilterControlEvent OnFilterControl;
+	FControlModifiedEvent OnControlModified;
+	FControlSelectedEvent OnControlSelected;
+	FControlUndoBracketEvent OnControlUndoBracket;
+	FControlRigBoundEvent OnControlRigBound;
+
+	UPROPERTY(BlueprintAssignable, Category = ControlRig, meta=(DisplayName="OnControlSelected"))
+	FOnControlSelectedBP OnControlSelected_BP;
+
+	TArray<FRigElementKey> QueuedModifiedControls;
+
+private:
+
+#if WITH_EDITORONLY_DATA
+
+	/** Whether controls are visible */
+	UPROPERTY(transient)
+	bool bControlsVisible;
+
+#endif
+
+	/** An additive control rig runs a backwards solve before applying additive control values
+	 * and running the forward solve
+	 */
+	UPROPERTY()
+	bool bIsAdditive = false;
+
+	struct FRigSetControlValueInfo
+	{
+		FRigControlValue Value;
+		bool bNotify;
+		FRigControlModifiedContext Context;
+		bool bSetupUndo;
+		bool bPrintPythonCommnds;
+		bool bFixEulerFlips;
+	};
+	FRigPose PoseBeforeBackwardsSolve;
+	FRigPose ControlsAfterBackwardsSolve;
+	TMap<FRigElementKey, FRigSetControlValueInfo> ControlValues; // Additive values in local space (to add after backwards solve)
+
+	float DebugBoneRadiusMultiplier;
+	
+public:
+	
+#if WITH_EDITOR	
+
+	void ToggleControlsVisible() { bControlsVisible = !bControlsVisible; }
+	void SetControlsVisible(const bool bIsVisible) { bControlsVisible = bIsVisible; }
+	bool GetControlsVisible()const { return bControlsVisible;}
+
+#endif
+	
+	bool IsAdditive() const { return bIsAdditive; }
+	void SetIsAdditive(const bool bInIsAdditive)
+	{
+		bIsAdditive = bInIsAdditive;
+	}
+
+private:
+
+	// Class used to temporarily cache current pose of the hierarchy
+	// restore it on destruction, similar to UControlRigBlueprint::FControlValueScope
+	class FPoseScope
+	{
+	public:
+		FPoseScope(UControlRig* InControlRig, ERigElementType InFilter = ERigElementType::All,
+			const TArray<FRigElementKey>& InElements = TArray<FRigElementKey>(), const ERigTransformType::Type InTransformType = ERigTransformType::CurrentLocal);
+		~FPoseScope();
+
+	private:
+
+		UControlRig* ControlRig;
+		ERigElementType Filter;
+		FRigPose CachedPose;
+		ERigTransformType::Type TransformType;
+	};
+
+	UPROPERTY()
+	FRigModuleSettings RigModuleSettings;
+
+	UPROPERTY(transient)
+	mutable FString RigModuleNameSpace;
+
+#if WITH_EDITOR
+
+	// Class used to temporarily cache current transient controls
+	// restore them after a CopyHierarchy call
+	class FTransientControlScope
+	{
+	public:
+		FTransientControlScope(TObjectPtr<URigHierarchy> InHierarchy);
+		~FTransientControlScope();
+	
+	private:
+		// used to match URigHierarchyController::AddControl(...)
+		struct FTransientControlInfo
+		{
+			FName Name;
+			// transient control should only have 1 parent, with weight = 1.0
+			FRigElementKey Parent;
+			FRigControlSettings Settings;
+			FRigControlValue Value;
+			FTransform OffsetTransform;
+			FTransform ShapeTransform;
+		};
+		
+		TArray<FTransientControlInfo> SavedTransientControls;
+		TObjectPtr<URigHierarchy> Hierarchy;
+	};
+
+public:
+	// Class used to temporarily cache current pose of transient controls
+	// restore them after a ResetPoseToInitial call,
+	// which allows user to move bones in construction mode
+	class FTransientControlPoseScope
+	{
+	public:
+		FTransientControlPoseScope(TObjectPtr<UControlRig> InControlRig)
+		{
+			ControlRig = InControlRig;
+
+			TArray<FRigControlElement*> TransientControls = ControlRig->GetHierarchy()->GetTransientControls();
+			TArray<FRigElementKey> Keys;
+			for(FRigControlElement* TransientControl : TransientControls)
+			{
+				Keys.Add(TransientControl->GetKey());
+			}
+
+			if(Keys.Num() > 0)
+			{
+				CachedPose = ControlRig->GetHierarchy()->GetPose(false, ERigElementType::Control, TArrayView<FRigElementKey>(Keys));
+			}
+		}
+		~FTransientControlPoseScope()
+		{
+			check(ControlRig);
+
+			if(CachedPose.Num() > 0)
+			{
+				ControlRig->GetHierarchy()->SetPose(CachedPose);
+			}
+		}
+	
+	private:
+		
+		UControlRig* ControlRig;
+		FRigPose CachedPose;	
+	};	
+
+	bool bRecordSelectionPoseForConstructionMode;
+	TMap<FRigElementKey, FTransform> SelectionPoseForConstructionMode;
+	bool bIsClearingTransientControls;
+
+	FRigPose InputPoseOnDebuggedRig;
+	
+#endif
+	
+private:
+
+	friend class FControlRigBlueprintCompilerContext;
+	friend struct FRigHierarchyRef;
+	friend class FControlRigEditor;
+	friend class SRigCurveContainer;
+	friend class SRigHierarchy;
+	friend class SControlRigAnimAttributeView;
+	friend class UEngineTestControlRig;
+ 	friend class FControlRigEditMode;
+	friend class UControlRigBlueprint;
+	friend class UControlRigComponent;
+	friend class UControlRigBlueprintGeneratedClass;
+	friend class FControlRigInteractionScope;
+	friend class UControlRigValidator;
+	friend struct FAnimNode_ControlRig;
+	friend struct FAnimNode_ControlRigBase;
+	friend class URigHierarchy;
+	friend class UFKControlRig;
+	friend class UControlRigGraph;
+	friend class AControlRigControlActor;
+	friend class AControlRigShapeActor;
+	friend class FRigTransformElementDetails;
+	friend class FControlRigEditorModule;
+};
+
+class CONTROLRIG_API FControlRigBracketScope
+{
+public:
+
+	FControlRigBracketScope(int32& InBracket)
+		: Bracket(InBracket)
+	{
+		Bracket++;
+	}
+
+	~FControlRigBracketScope()
+	{
+		Bracket--;
+	}
+
+private:
+
+	int32& Bracket;
+};
+
+class CONTROLRIG_API FControlRigInteractionScope
+{
+public:
+
+	FControlRigInteractionScope(UControlRig* InControlRig)
+		: ControlRig(InControlRig)
+		, InteractionBracketScope(InControlRig->InteractionBracket)
+		, SyncBracketScope(InControlRig->InterRigSyncBracket)
+	{
+		InControlRig->GetHierarchy()->StartInteraction();
+	}
+
+	FControlRigInteractionScope(
+		UControlRig* InControlRig,
+		const FRigElementKey& InKey,
+		EControlRigInteractionType InInteractionType = EControlRigInteractionType::All
+	)
+		: ControlRig(InControlRig)
+		, InteractionBracketScope(InControlRig->InteractionBracket)
+		, SyncBracketScope(InControlRig->InterRigSyncBracket)
+	{
+		InControlRig->ElementsBeingInteracted = { InKey };
+		InControlRig->InteractionType = (uint8)InInteractionType;
+		InControlRig->bInteractionJustBegan = true;
+		InControlRig->GetHierarchy()->StartInteraction();
+	}
+
+	FControlRigInteractionScope(
+		UControlRig* InControlRig,
+		const TArray<FRigElementKey>& InKeys,
+		EControlRigInteractionType InInteractionType = EControlRigInteractionType::All
+	)
+		: ControlRig(InControlRig)
+		, InteractionBracketScope(InControlRig->InteractionBracket)
+	, SyncBracketScope(InControlRig->InterRigSyncBracket)
+	{
+		InControlRig->ElementsBeingInteracted = InKeys;
+		InControlRig->InteractionType = (uint8)InInteractionType;
+		InControlRig->bInteractionJustBegan = true;
+		InControlRig->GetHierarchy()->StartInteraction();
+	}
+
+	~FControlRigInteractionScope()
+	{
+		if(ensure(ControlRig.IsValid()))
+		{
+			ControlRig->GetHierarchy()->EndInteraction();
+			ControlRig->InteractionType = (uint8)EControlRigInteractionType::None;
+			ControlRig->bInteractionJustBegan = false;
+			ControlRig->ElementsBeingInteracted.Reset();
+		}
+	}
+
+private:
+
+	TWeakObjectPtr<UControlRig> ControlRig;
+	FControlRigBracketScope InteractionBracketScope;
+	FControlRigBracketScope SyncBracketScope;
+};
diff --git a/Engine/Plugins/Animation/ControlRig/Source/ControlRig/Public/Constraints/ControlRigTransformableHandle.h b/Engine/Plugins/Animation/ControlRig/Source/ControlRig/Public/Constraints/ControlRigTransformableHandle.h
--- a/Engine/Plugins/Animation/ControlRig/Source/ControlRig/Public/Constraints/ControlRigTransformableHandle.h
+++ b/Engine/Plugins/Animation/ControlRig/Source/ControlRig/Public/Constraints/ControlRigTransformableHandle.h
@@ -4,7 +4,6 @@
 
 #include "TransformableHandle.h"
 #include "Rigs/RigHierarchyDefines.h"
-#include "ControlRig.h"
 
 #include "ControlRigTransformableHandle.generated.h"
 
diff --git a/Engine/Plugins/Animation/ControlRig/Source/ControlRig/Public/ControlRig.h b/Engine/Plugins/Animation/ControlRig/Source/ControlRig/Public/ControlRig.h
deleted file mode 100644
diff --git a/Engine/Plugins/Animation/ControlRig/Source/ControlRig/Public/ControlRigComponent.h b/Engine/Plugins/Animation/ControlRig/Source/ControlRig/Public/ControlRigComponent.h
--- a/Engine/Plugins/Animation/ControlRig/Source/ControlRig/Public/ControlRigComponent.h
+++ b/Engine/Plugins/Animation/ControlRig/Source/ControlRig/Public/ControlRigComponent.h
@@ -3,7 +3,7 @@
 #pragma once
 
 #include "Components/PrimitiveComponent.h"
-#include "ControlRig.h"
+#include "BaseControlRig.h"
 #include "ControlRigAnimInstance.h"
 #include "Components/SkeletalMeshComponent.h"
 #include "PrimitiveSceneProxy.h"
@@ -15,7 +15,6 @@
 
 #include "ControlRigComponent.generated.h"
 
-class UControlRig;
 class UControlRigComponent;
 
 /** Enum for controlling which space a transform is applied in. */
diff --git a/Engine/Plugins/Animation/ControlRig/Source/ControlRig/Public/ControlRigControlActor.h b/Engine/Plugins/Animation/ControlRig/Source/ControlRig/Public/ControlRigControlActor.h
--- a/Engine/Plugins/Animation/ControlRig/Source/ControlRig/Public/ControlRigControlActor.h
+++ b/Engine/Plugins/Animation/ControlRig/Source/ControlRig/Public/ControlRigControlActor.h
@@ -8,7 +8,7 @@
 #include "Engine/StaticMesh.h"
 #include "Materials/Material.h"
 #include "Materials/MaterialInstanceDynamic.h"
-#include "ControlRig.h"
+#include "BaseControlRig.h"
 #include "Rigs/RigHierarchyContainer.h"
 #include "Rigs/RigControlHierarchy.h"
 
diff --git a/Engine/Plugins/Animation/ControlRig/Source/ControlRig/Public/ControlRigValidationPass.h b/Engine/Plugins/Animation/ControlRig/Source/ControlRig/Public/ControlRigValidationPass.h
--- a/Engine/Plugins/Animation/ControlRig/Source/ControlRig/Public/ControlRigValidationPass.h
+++ b/Engine/Plugins/Animation/ControlRig/Source/ControlRig/Public/ControlRigValidationPass.h
@@ -1,7 +1,7 @@
 // Copyright Epic Games, Inc. All Rights Reserved.
 
 #pragma once
-#include "ControlRig.h"
+#include "BaseControlRig.h"
 #include "Logging/TokenizedMessage.h"
 
 #include "ControlRigValidationPass.generated.h"
diff --git a/Engine/Plugins/Animation/ControlRig/Source/ControlRig/Public/Rigs/AdditiveControlRig.h b/Engine/Plugins/Animation/ControlRig/Source/ControlRig/Public/Rigs/AdditiveControlRig.h
--- a/Engine/Plugins/Animation/ControlRig/Source/ControlRig/Public/Rigs/AdditiveControlRig.h
+++ b/Engine/Plugins/Animation/ControlRig/Source/ControlRig/Public/Rigs/AdditiveControlRig.h
@@ -1,7 +1,7 @@
 // Copyright Epic Games, Inc. All Rights Reserved.
 
 #pragma once
-#include "ControlRig.h"
+#include "BaseControlRig.h"
 #include "Units/Hierarchy/RigUnit_AddBoneTransform.h"
 #include "AdditiveControlRig.generated.h"
 
diff --git a/Engine/Plugins/Animation/ControlRig/Source/ControlRig/Public/Rigs/FKControlRig.h b/Engine/Plugins/Animation/ControlRig/Source/ControlRig/Public/Rigs/FKControlRig.h
--- a/Engine/Plugins/Animation/ControlRig/Source/ControlRig/Public/Rigs/FKControlRig.h
+++ b/Engine/Plugins/Animation/ControlRig/Source/ControlRig/Public/Rigs/FKControlRig.h
@@ -1,7 +1,7 @@
 // Copyright Epic Games, Inc. All Rights Reserved.
 
 #pragma once
-#include "ControlRig.h"
+#include "BaseControlRig.h"
 #include "Rigs/RigHierarchyController.h"
 #include "Units/Hierarchy/RigUnit_SetTransform.h"
 #include "Units/Hierarchy/RigUnit_SetCurveValue.h"
diff --git a/Engine/Plugins/Animation/ControlRig/Source/ControlRig/Public/Sequencer/ControlRigSequenceObjectReference.h b/Engine/Plugins/Animation/ControlRig/Source/ControlRig/Public/Sequencer/ControlRigSequenceObjectReference.h
--- a/Engine/Plugins/Animation/ControlRig/Source/ControlRig/Public/Sequencer/ControlRigSequenceObjectReference.h
+++ b/Engine/Plugins/Animation/ControlRig/Source/ControlRig/Public/Sequencer/ControlRigSequenceObjectReference.h
@@ -3,7 +3,7 @@
 #pragma once
 
 #include "UObject/LazyObjectPtr.h"
-#include "ControlRig.h"
+#include "BaseControlRig.h"
 #include "ControlRigSequenceObjectReference.generated.h"
 
 
diff --git a/Engine/Plugins/Animation/ControlRig/Source/ControlRig/Public/Sequencer/MovieSceneControlRigParameterSection.h b/Engine/Plugins/Animation/ControlRig/Source/ControlRig/Public/Sequencer/MovieSceneControlRigParameterSection.h
--- a/Engine/Plugins/Animation/ControlRig/Source/ControlRig/Public/Sequencer/MovieSceneControlRigParameterSection.h
+++ b/Engine/Plugins/Animation/ControlRig/Source/ControlRig/Public/Sequencer/MovieSceneControlRigParameterSection.h
@@ -9,7 +9,7 @@
 #include "UObject/ObjectMacros.h"
 #include "Channels/MovieSceneFloatChannel.h"
 #include "Sections/MovieSceneSubSection.h"
-#include "ControlRig.h"
+#include "BaseControlRig.h"
 #include "MovieSceneSequencePlayer.h"
 #include "Animation/AnimData/BoneMaskFilter.h"
 #include "MovieSceneObjectBindingID.h"
diff --git a/Engine/Plugins/Animation/ControlRig/Source/ControlRig/Public/Sequencer/MovieSceneControlRigParameterTrack.h b/Engine/Plugins/Animation/ControlRig/Source/ControlRig/Public/Sequencer/MovieSceneControlRigParameterTrack.h
--- a/Engine/Plugins/Animation/ControlRig/Source/ControlRig/Public/Sequencer/MovieSceneControlRigParameterTrack.h
+++ b/Engine/Plugins/Animation/ControlRig/Source/ControlRig/Public/Sequencer/MovieSceneControlRigParameterTrack.h
@@ -6,7 +6,7 @@
 #include "UObject/ObjectMacros.h"
 #include "Sections/MovieSceneParameterSection.h"
 #include "MovieSceneNameableTrack.h"
-#include "ControlRig.h"
+#include "BaseControlRig.h"
 #include "Compilation/IMovieSceneTrackTemplateProducer.h"
 #include "INodeAndChannelMappings.h"
 #include "MovieSceneControlRigParameterSection.h"
diff --git a/Engine/Plugins/Animation/ControlRig/Source/ControlRig/Public/Tools/ControlRigPose.h b/Engine/Plugins/Animation/ControlRig/Source/ControlRig/Public/Tools/ControlRigPose.h
--- a/Engine/Plugins/Animation/ControlRig/Source/ControlRig/Public/Tools/ControlRigPose.h
+++ b/Engine/Plugins/Animation/ControlRig/Source/ControlRig/Public/Tools/ControlRigPose.h
@@ -7,13 +7,13 @@
 #include "CoreMinimal.h"
 #include "UObject/Object.h"
 #include "TransformNoScale.h"
-#include "ControlRig.h"
 #include "Engine/SkeletalMesh.h"
 #include "ControlRigToolAsset.h"
 #include "Tools/ControlRigPoseMirrorTable.h"
 #include "Rigs/RigControlHierarchy.h"
 #include "ControlRigPose.generated.h"
 
+class UControlRig;
 /**
 * The Data Stored For Each Control in A Pose.
 */
diff --git a/Engine/Plugins/Animation/ControlRig/Source/ControlRig/Public/Tools/ControlRigPoseMirrorSettings.h b/Engine/Plugins/Animation/ControlRig/Source/ControlRig/Public/Tools/ControlRigPoseMirrorSettings.h
--- a/Engine/Plugins/Animation/ControlRig/Source/ControlRig/Public/Tools/ControlRigPoseMirrorSettings.h
+++ b/Engine/Plugins/Animation/ControlRig/Source/ControlRig/Public/Tools/ControlRigPoseMirrorSettings.h
@@ -6,7 +6,6 @@
 #include "CoreMinimal.h"
 #include "UObject/ObjectMacros.h"
 #include "UObject/Object.h"
-#include "ControlRig.h"
 #include "ControlRigPoseMirrorSettings.generated.h"
 
 
diff --git a/Engine/Plugins/Animation/ControlRig/Source/ControlRig/Public/Tools/ControlRigToolAsset.h b/Engine/Plugins/Animation/ControlRig/Source/ControlRig/Public/Tools/ControlRigToolAsset.h
--- a/Engine/Plugins/Animation/ControlRig/Source/ControlRig/Public/Tools/ControlRigToolAsset.h
+++ b/Engine/Plugins/Animation/ControlRig/Source/ControlRig/Public/Tools/ControlRigToolAsset.h
@@ -4,7 +4,7 @@
 #include "CoreMinimal.h"
 #include "Misc/PackageName.h"
 #include "UObject/Object.h"
-#include "ControlRig.h"
+#include "BaseControlRig.h"
 #include "Tools/ControlRigPoseProjectSettings.h"
 #include "AssetRegistry/IAssetRegistry.h"
 #include "AssetRegistry/AssetRegistryModule.h"
diff --git a/Engine/Plugins/Animation/ControlRig/Source/ControlRig/Public/Units/RigDispatchFactory.h b/Engine/Plugins/Animation/ControlRig/Source/ControlRig/Public/Units/RigDispatchFactory.h
--- a/Engine/Plugins/Animation/ControlRig/Source/ControlRig/Public/Units/RigDispatchFactory.h
+++ b/Engine/Plugins/Animation/ControlRig/Source/ControlRig/Public/Units/RigDispatchFactory.h
@@ -5,7 +5,6 @@
 #include "ControlRigDefines.h"
 #include "RigVMCore/RigVMDispatchFactory.h"
 #include "RigUnitContext.h"
-#include "ControlRig.h"
 #include "RigDispatchFactory.generated.h"
 
 /** Base class for all rig dispatch factories */
diff --git a/Engine/Plugins/Animation/ControlRig/Source/ControlRigDeveloper/Private/ControlRigBlueprint.cpp b/Engine/Plugins/Animation/ControlRig/Source/ControlRigDeveloper/Private/ControlRigBlueprint.cpp
--- a/Engine/Plugins/Animation/ControlRig/Source/ControlRigDeveloper/Private/ControlRigBlueprint.cpp
+++ b/Engine/Plugins/Animation/ControlRig/Source/ControlRigDeveloper/Private/ControlRigBlueprint.cpp
@@ -7,7 +7,7 @@
 #include "EdGraphNode_Comment.h"
 #include "Engine/SkeletalMesh.h"
 #include "BlueprintActionDatabaseRegistrar.h"
-#include "ControlRig.h"
+#include "BaseControlRig.h"
 #include "Graph/ControlRigGraph.h"
 #include "Graph/ControlRigGraphSchema.h"
 #include "UObject/ObjectSaveContext.h"
diff --git a/Engine/Plugins/Animation/ControlRig/Source/ControlRigDeveloper/Private/ControlRigBlueprintCompiler.cpp b/Engine/Plugins/Animation/ControlRig/Source/ControlRigDeveloper/Private/ControlRigBlueprintCompiler.cpp
--- a/Engine/Plugins/Animation/ControlRig/Source/ControlRigDeveloper/Private/ControlRigBlueprintCompiler.cpp
+++ b/Engine/Plugins/Animation/ControlRig/Source/ControlRigDeveloper/Private/ControlRigBlueprintCompiler.cpp
@@ -1,7 +1,7 @@
 // Copyright Epic Games, Inc. All Rights Reserved.
 
 #include "ControlRigBlueprintCompiler.h"
-#include "ControlRig.h"
+#include "BaseControlRig.h"
 #include "ControlRigBlueprint.h"
 #include "Kismet2/KismetReinstanceUtilities.h"
 
diff --git a/Engine/Plugins/Animation/ControlRig/Source/ControlRigDeveloper/Private/ControlRigVisualGraphUtils.cpp b/Engine/Plugins/Animation/ControlRig/Source/ControlRigDeveloper/Private/ControlRigVisualGraphUtils.cpp
--- a/Engine/Plugins/Animation/ControlRig/Source/ControlRigDeveloper/Private/ControlRigVisualGraphUtils.cpp
+++ b/Engine/Plugins/Animation/ControlRig/Source/ControlRigDeveloper/Private/ControlRigVisualGraphUtils.cpp
@@ -2,7 +2,7 @@
 
 #include "ControlRigVisualGraphUtils.h"
 
-#include "ControlRig.h"
+#include "BaseControlRig.h"
 #include "VisualGraphUtilsModule.h"
 
 #if WITH_EDITOR
diff --git a/Engine/Plugins/Animation/ControlRig/Source/ControlRigDeveloper/Private/Graph/ControlRigGraph.cpp b/Engine/Plugins/Animation/ControlRig/Source/ControlRigDeveloper/Private/Graph/ControlRigGraph.cpp
--- a/Engine/Plugins/Animation/ControlRig/Source/ControlRigDeveloper/Private/Graph/ControlRigGraph.cpp
+++ b/Engine/Plugins/Animation/ControlRig/Source/ControlRigDeveloper/Private/Graph/ControlRigGraph.cpp
@@ -2,7 +2,7 @@
 
 #include "Graph/ControlRigGraph.h"
 #include "ControlRigBlueprint.h"
-#include "ControlRig.h"
+#include "BaseControlRig.h"
 #include "RigVMModel/RigVMGraph.h"
 #include "Units/RigUnit.h"
 
diff --git a/Engine/Plugins/Animation/ControlRig/Source/ControlRigEditor/Private/ControlRigAnimGraphDetails.cpp b/Engine/Plugins/Animation/ControlRig/Source/ControlRigEditor/Private/ControlRigAnimGraphDetails.cpp
--- a/Engine/Plugins/Animation/ControlRig/Source/ControlRigEditor/Private/ControlRigAnimGraphDetails.cpp
+++ b/Engine/Plugins/Animation/ControlRig/Source/ControlRigEditor/Private/ControlRigAnimGraphDetails.cpp
@@ -2,7 +2,7 @@
 
 #include "ControlRigAnimGraphDetails.h"
 
-#include "ControlRig.h"
+#include "BaseControlRig.h"
 #include "Widgets/SWidget.h"
 #include "DetailLayoutBuilder.h"
 #include "Styling/AppStyle.h"
diff --git a/Engine/Plugins/Animation/ControlRig/Source/ControlRigEditor/Private/ControlRigBlueprintFactory.h b/Engine/Plugins/Animation/ControlRig/Source/ControlRigEditor/Private/ControlRigBlueprintFactory.h
--- a/Engine/Plugins/Animation/ControlRig/Source/ControlRigEditor/Private/ControlRigBlueprintFactory.h
+++ b/Engine/Plugins/Animation/ControlRig/Source/ControlRigEditor/Private/ControlRigBlueprintFactory.h
@@ -7,10 +7,11 @@
 #include "Templates/SubclassOf.h"
 #include "Engine/Blueprint.h"
 #include "Factories/Factory.h"
-#include "ControlRig.h"
 #include "ControlRigBlueprint.h"
 #include "ControlRigBlueprintFactory.generated.h"
 
+class UControlRig;
+
 UCLASS(HideCategories=Object)
 class CONTROLRIGEDITOR_API UControlRigBlueprintFactory : public UFactory
 {
diff --git a/Engine/Plugins/Animation/ControlRig/Source/ControlRigEditor/Private/ControlRigCompilerDetails.cpp b/Engine/Plugins/Animation/ControlRig/Source/ControlRigEditor/Private/ControlRigCompilerDetails.cpp
--- a/Engine/Plugins/Animation/ControlRig/Source/ControlRigEditor/Private/ControlRigCompilerDetails.cpp
+++ b/Engine/Plugins/Animation/ControlRig/Source/ControlRigEditor/Private/ControlRigCompilerDetails.cpp
@@ -9,7 +9,7 @@
 #include "Widgets/Input/SButton.h"
 #include "Widgets/Text/STextBlock.h"
 #include "HAL/PlatformApplicationMisc.h"
-#include "ControlRig.h"
+#include "BaseControlRig.h"
 #include "IPropertyUtilities.h"
 #include "IPythonScriptPlugin.h"
 #include "RigVMPythonUtils.h"
diff --git a/Engine/Plugins/Animation/ControlRig/Source/ControlRigEditor/Private/ControlRigEditorModule.cpp b/Engine/Plugins/Animation/ControlRig/Source/ControlRigEditor/Private/ControlRigEditorModule.cpp
--- a/Engine/Plugins/Animation/ControlRig/Source/ControlRigEditor/Private/ControlRigEditorModule.cpp
+++ b/Engine/Plugins/Animation/ControlRig/Source/ControlRigEditor/Private/ControlRigEditorModule.cpp
@@ -4,7 +4,7 @@
 #include "BlueprintActionDatabaseRegistrar.h"
 #include "BlueprintNodeSpawner.h"
 #include "PropertyEditorModule.h"
-#include "ControlRig.h"
+#include "BaseControlRig.h"
 #include "ControlRigComponent.h"
 #include "GraphEditorActions.h"
 #include "Kismet2/KismetEditorUtilities.h"
diff --git a/Engine/Plugins/Animation/ControlRig/Source/ControlRigEditor/Private/ControlRigElementDetails.h b/Engine/Plugins/Animation/ControlRig/Source/ControlRigEditor/Private/ControlRigElementDetails.h
--- a/Engine/Plugins/Animation/ControlRig/Source/ControlRigEditor/Private/ControlRigElementDetails.h
+++ b/Engine/Plugins/Animation/ControlRig/Source/ControlRigEditor/Private/ControlRigElementDetails.h
@@ -7,7 +7,7 @@
 #include "IPropertyTypeCustomization.h"
 #include "Rigs/RigHierarchyDefines.h"
 #include "Rigs/RigHierarchy.h"
-#include "ControlRig.h"
+#include "BaseControlRig.h"
 #include "ControlRigBlueprint.h"
 #include "Editor/ControlRigWrapperObject.h"
 #include "Widgets/SRigVMGraphPinNameListValueWidget.h"
diff --git a/Engine/Plugins/Animation/ControlRig/Source/ControlRigEditor/Private/ControlRigInfluenceMapDetails.cpp b/Engine/Plugins/Animation/ControlRig/Source/ControlRigEditor/Private/ControlRigInfluenceMapDetails.cpp
--- a/Engine/Plugins/Animation/ControlRig/Source/ControlRigEditor/Private/ControlRigInfluenceMapDetails.cpp
+++ b/Engine/Plugins/Animation/ControlRig/Source/ControlRigEditor/Private/ControlRigInfluenceMapDetails.cpp
@@ -5,7 +5,7 @@
 #include "DetailLayoutBuilder.h"
 #include "DetailCategoryBuilder.h"
 #include "DetailWidgetRow.h"
-#include "ControlRig.h"
+#include "BaseControlRig.h"
 #include "Units/Execution/RigUnit_BeginExecution.h"
 #include "Units/Execution/RigUnit_PrepareForExecution.h"
 
diff --git a/Engine/Plugins/Animation/ControlRig/Source/ControlRigEditor/Private/ControlRigPythonLogDetails.cpp b/Engine/Plugins/Animation/ControlRig/Source/ControlRigEditor/Private/ControlRigPythonLogDetails.cpp
--- a/Engine/Plugins/Animation/ControlRig/Source/ControlRigEditor/Private/ControlRigPythonLogDetails.cpp
+++ b/Engine/Plugins/Animation/ControlRig/Source/ControlRigEditor/Private/ControlRigPythonLogDetails.cpp
@@ -9,7 +9,6 @@
 #include "Widgets/Input/SButton.h"
 #include "Widgets/Text/STextBlock.h"
 #include "HAL/PlatformApplicationMisc.h"
-#include "ControlRig.h"
 #include "IPropertyUtilities.h"
 #include "IPythonScriptPlugin.h"
 #include "RigVMPythonUtils.h"
diff --git a/Engine/Plugins/Animation/ControlRig/Source/ControlRigEditor/Private/ControlRigSequencerEditorLibrary.cpp b/Engine/Plugins/Animation/ControlRig/Source/ControlRigEditor/Private/ControlRigSequencerEditorLibrary.cpp
--- a/Engine/Plugins/Animation/ControlRig/Source/ControlRigEditor/Private/ControlRigSequencerEditorLibrary.cpp
+++ b/Engine/Plugins/Animation/ControlRig/Source/ControlRigEditor/Private/ControlRigSequencerEditorLibrary.cpp
@@ -17,7 +17,7 @@
 #include "MovieSceneToolHelpers.h"
 #include "Rigs/FKControlRig.h"
 #include "Units/Execution/RigUnit_InverseExecution.h"
-#include "ControlRig.h"
+#include "BaseControlRig.h"
 #include "EditMode/ControlRigEditMode.h"
 #include "EditorModeManager.h"
 #include "Engine/Selection.h"
diff --git a/Engine/Plugins/Animation/ControlRig/Source/ControlRigEditor/Private/ControlRigThumbnailRenderer.cpp b/Engine/Plugins/Animation/ControlRig/Source/ControlRigEditor/Private/ControlRigThumbnailRenderer.cpp
--- a/Engine/Plugins/Animation/ControlRig/Source/ControlRigEditor/Private/ControlRigThumbnailRenderer.cpp
+++ b/Engine/Plugins/Animation/ControlRig/Source/ControlRigEditor/Private/ControlRigThumbnailRenderer.cpp
@@ -10,7 +10,7 @@
 #include "Animation/SkeletalMeshActor.h"
 #include "Components/SkeletalMeshComponent.h"
 #include "Materials/MaterialInstanceDynamic.h"
-#include "ControlRig.h"
+#include "BaseControlRig.h"
 #include "CanvasTypes.h"
 
 #include UE_INLINE_GENERATED_CPP_BY_NAME(ControlRigThumbnailRenderer)
diff --git a/Engine/Plugins/Animation/ControlRig/Source/ControlRigEditor/Private/EditMode/ControlRigBaseDockableView.cpp b/Engine/Plugins/Animation/ControlRig/Source/ControlRigEditor/Private/EditMode/ControlRigBaseDockableView.cpp
--- a/Engine/Plugins/Animation/ControlRig/Source/ControlRigEditor/Private/EditMode/ControlRigBaseDockableView.cpp
+++ b/Engine/Plugins/Animation/ControlRig/Source/ControlRigEditor/Private/EditMode/ControlRigBaseDockableView.cpp
@@ -5,7 +5,7 @@
 #include "Widgets/Text/SInlineEditableTextBlock.h"
 #include "Styling/CoreStyle.h"
 #include "ScopedTransaction.h"
-#include "ControlRig.h"
+#include "BaseControlRig.h"
 #include "UnrealEdGlobals.h"
 #include "EditMode/ControlRigEditMode.h"
 #include "EditorModeManager.h"
diff --git a/Engine/Plugins/Animation/ControlRig/Source/ControlRigEditor/Private/EditMode/ControlRigControlsProxy.cpp b/Engine/Plugins/Animation/ControlRig/Source/ControlRigEditor/Private/EditMode/ControlRigControlsProxy.cpp
--- a/Engine/Plugins/Animation/ControlRig/Source/ControlRigEditor/Private/EditMode/ControlRigControlsProxy.cpp
+++ b/Engine/Plugins/Animation/ControlRig/Source/ControlRigEditor/Private/EditMode/ControlRigControlsProxy.cpp
@@ -3,7 +3,7 @@
 #include "EditMode/ControlRigControlsProxy.h"
 #include "EditorModeManager.h"
 #include "EditMode/ControlRigEditMode.h"
-#include "ControlRig.h"
+#include "BaseControlRig.h"
 #include "Rigs/RigHierarchy.h"
 
 #include "Components/SkeletalMeshComponent.h"
@@ -19,8 +19,6 @@
 #include "DetailCategoryBuilder.h"
 #include "DetailWidgetRow.h"
 #include "SEnumCombo.h"
-#include "ControlRig.h"
-
 
 void UControlRigControlsProxy::SetIsMultiple(bool bIsVal)
 { 
diff --git a/Engine/Plugins/Animation/ControlRig/Source/ControlRigEditor/Private/EditMode/ControlRigEditMode.cpp b/Engine/Plugins/Animation/ControlRig/Source/ControlRigEditor/Private/EditMode/ControlRigEditMode.cpp
--- a/Engine/Plugins/Animation/ControlRig/Source/ControlRigEditor/Private/EditMode/ControlRigEditMode.cpp
+++ b/Engine/Plugins/Animation/ControlRig/Source/ControlRigEditor/Private/EditMode/ControlRigEditMode.cpp
@@ -8,7 +8,7 @@
 #include "Toolkits/ToolkitManager.h"
 #include "EditMode/SControlRigEditModeTools.h"
 #include "Algo/Transform.h"
-#include "ControlRig.h"
+#include "BaseControlRig.h"
 #include "HitProxies.h"
 #include "EditMode/ControlRigEditModeSettings.h"
 #include "ISequencer.h"
diff --git a/Engine/Plugins/Animation/ControlRig/Source/ControlRigEditor/Private/EditMode/SControlPicker.cpp b/Engine/Plugins/Animation/ControlRig/Source/ControlRigEditor/Private/EditMode/SControlPicker.cpp
--- a/Engine/Plugins/Animation/ControlRig/Source/ControlRigEditor/Private/EditMode/SControlPicker.cpp
+++ b/Engine/Plugins/Animation/ControlRig/Source/ControlRigEditor/Private/EditMode/SControlPicker.cpp
@@ -13,7 +13,7 @@
 #include "Widgets/Input/SButton.h"
 #include "ScopedTransaction.h"
 #include "EditMode/ControlRigEditMode.h"
-#include "ControlRig.h"
+#include "BaseControlRig.h"
 #include "EditorModeManager.h"
 #include "EditMode/SEditorUserWidgetHost.h"
 #include "ControlRigBlueprint.h"
diff --git a/Engine/Plugins/Animation/ControlRig/Source/ControlRigEditor/Private/EditMode/SControlRigBaseListWidget.cpp b/Engine/Plugins/Animation/ControlRig/Source/ControlRigEditor/Private/EditMode/SControlRigBaseListWidget.cpp
--- a/Engine/Plugins/Animation/ControlRig/Source/ControlRigEditor/Private/EditMode/SControlRigBaseListWidget.cpp
+++ b/Engine/Plugins/Animation/ControlRig/Source/ControlRigEditor/Private/EditMode/SControlRigBaseListWidget.cpp
@@ -17,7 +17,7 @@
 
 #include "ScopedTransaction.h"
 
-#include "ControlRig.h"
+#include "BaseControlRig.h"
 #include "UnrealEdGlobals.h"
 #include "EditMode/ControlRigEditMode.h"
 #include "Tools/ControlRigPose.h"
diff --git a/Engine/Plugins/Animation/ControlRig/Source/ControlRigEditor/Private/EditMode/SControlRigControlViews.cpp b/Engine/Plugins/Animation/ControlRig/Source/ControlRigEditor/Private/EditMode/SControlRigControlViews.cpp
--- a/Engine/Plugins/Animation/ControlRig/Source/ControlRigEditor/Private/EditMode/SControlRigControlViews.cpp
+++ b/Engine/Plugins/Animation/ControlRig/Source/ControlRigEditor/Private/EditMode/SControlRigControlViews.cpp
@@ -2,7 +2,7 @@
 
 
 #include "EditMode/SControlRigControlViews.h"
-#include "ControlRig.h"
+#include "BaseControlRig.h"
 #include "Tools/ControlRigPose.h"
 #include "Widgets/Input/SEditableTextBox.h"
 #include "Styling/AppStyle.h"
diff --git a/Engine/Plugins/Animation/ControlRig/Source/ControlRigEditor/Private/EditMode/SControlRigDetails.cpp b/Engine/Plugins/Animation/ControlRig/Source/ControlRigEditor/Private/EditMode/SControlRigDetails.cpp
--- a/Engine/Plugins/Animation/ControlRig/Source/ControlRigEditor/Private/EditMode/SControlRigDetails.cpp
+++ b/Engine/Plugins/Animation/ControlRig/Source/ControlRigEditor/Private/EditMode/SControlRigDetails.cpp
@@ -8,7 +8,7 @@
 #include "Widgets/Text/SInlineEditableTextBlock.h"
 #include "Styling/CoreStyle.h"
 #include "ScopedTransaction.h"
-#include "ControlRig.h"
+#include "BaseControlRig.h"
 #include "UnrealEdGlobals.h"
 #include "EditMode/ControlRigEditMode.h"
 #include "EditorModeManager.h"
diff --git a/Engine/Plugins/Animation/ControlRig/Source/ControlRigEditor/Private/EditMode/SControlRigDetails.h b/Engine/Plugins/Animation/ControlRig/Source/ControlRigEditor/Private/EditMode/SControlRigDetails.h
--- a/Engine/Plugins/Animation/ControlRig/Source/ControlRigEditor/Private/EditMode/SControlRigDetails.h
+++ b/Engine/Plugins/Animation/ControlRig/Source/ControlRigEditor/Private/EditMode/SControlRigDetails.h
@@ -11,7 +11,6 @@
 #include "Widgets/SCompoundWidget.h"
 #include "Misc/FrameNumber.h"
 #include "IDetailsView.h"
-#include "ControlRig.h"
 #include "Rigs/RigHierarchy.h"
 #include "IDetailKeyframeHandler.h"
 #include "IDetailCustomization.h"
@@ -20,6 +19,7 @@
 #include "DetailWidgetRow.h"
 
 class ISequencer;
+class UControlRig;
 
 class FControlRigEditModeGenericDetails : public IDetailCustomization
 {
diff --git a/Engine/Plugins/Animation/ControlRig/Source/ControlRigEditor/Private/EditMode/SControlRigEditModeTools.cpp b/Engine/Plugins/Animation/ControlRig/Source/ControlRigEditor/Private/EditMode/SControlRigEditModeTools.cpp
--- a/Engine/Plugins/Animation/ControlRig/Source/ControlRigEditor/Private/EditMode/SControlRigEditModeTools.cpp
+++ b/Engine/Plugins/Animation/ControlRig/Source/ControlRigEditor/Private/EditMode/SControlRigEditModeTools.cpp
@@ -6,7 +6,7 @@
 #include "IDetailsView.h"
 #include "ISequencer.h"
 #include "PropertyHandle.h"
-#include "ControlRig.h"
+#include "BaseControlRig.h"
 #include "EditMode/ControlRigEditModeSettings.h"
 #include "IDetailRootObjectCustomization.h"
 #include "Modules/ModuleManager.h"
diff --git a/Engine/Plugins/Animation/ControlRig/Source/ControlRigEditor/Private/EditMode/SControlRigOutliner.cpp b/Engine/Plugins/Animation/ControlRig/Source/ControlRigEditor/Private/EditMode/SControlRigOutliner.cpp
--- a/Engine/Plugins/Animation/ControlRig/Source/ControlRigEditor/Private/EditMode/SControlRigOutliner.cpp
+++ b/Engine/Plugins/Animation/ControlRig/Source/ControlRigEditor/Private/EditMode/SControlRigOutliner.cpp
@@ -12,7 +12,7 @@
 #include "Styling/AppStyle.h"
 #include "Styling/CoreStyle.h"
 #include "ScopedTransaction.h"
-#include "ControlRig.h"
+#include "BaseControlRig.h"
 #include "UnrealEdGlobals.h"
 #include "EditMode/ControlRigEditMode.h"
 #include "EditorModeManager.h"
diff --git a/Engine/Plugins/Animation/ControlRig/Source/ControlRigEditor/Private/EditMode/SControlRigOutliner.h b/Engine/Plugins/Animation/ControlRig/Source/ControlRigEditor/Private/EditMode/SControlRigOutliner.h
--- a/Engine/Plugins/Animation/ControlRig/Source/ControlRigEditor/Private/EditMode/SControlRigOutliner.h
+++ b/Engine/Plugins/Animation/ControlRig/Source/ControlRigEditor/Private/EditMode/SControlRigOutliner.h
@@ -8,7 +8,7 @@
 #include "EditMode/ControlRigBaseDockableView.h"
 #include "Widgets/SWidget.h"
 #include "Widgets/SCompoundWidget.h"
-#include "ControlRig.h"
+#include "BaseControlRig.h"
 #include "Rigs/RigHierarchy.h"
 #include "Editor/SRigHierarchyTreeView.h"
 #include "Widgets/SBoxPanel.h"
@@ -17,7 +17,6 @@
 class ISequencer;
 class SExpandableArea;
 class SSearchableRigHierarchyTreeView;
-class UControlRig;
 
 class SMultiRigHierarchyTreeView;
 class SMultiRigHierarchyItem;
diff --git a/Engine/Plugins/Animation/ControlRig/Source/ControlRigEditor/Private/EditMode/SControlRigSnapper.cpp b/Engine/Plugins/Animation/ControlRig/Source/ControlRigEditor/Private/EditMode/SControlRigSnapper.cpp
--- a/Engine/Plugins/Animation/ControlRig/Source/ControlRigEditor/Private/EditMode/SControlRigSnapper.cpp
+++ b/Engine/Plugins/Animation/ControlRig/Source/ControlRigEditor/Private/EditMode/SControlRigSnapper.cpp
@@ -10,7 +10,7 @@
 #include "Styling/AppStyle.h"
 #include "Styling/CoreStyle.h"
 #include "ScopedTransaction.h"
-#include "ControlRig.h"
+#include "BaseControlRig.h"
 #include "UnrealEdGlobals.h"
 #include "EditMode/ControlRigEditMode.h"
 #include "Tools/ControlRigPose.h"
diff --git a/Engine/Plugins/Animation/ControlRig/Source/ControlRigEditor/Private/EditMode/SControlRigSpacePicker.cpp b/Engine/Plugins/Animation/ControlRig/Source/ControlRigEditor/Private/EditMode/SControlRigSpacePicker.cpp
--- a/Engine/Plugins/Animation/ControlRig/Source/ControlRigEditor/Private/EditMode/SControlRigSpacePicker.cpp
+++ b/Engine/Plugins/Animation/ControlRig/Source/ControlRigEditor/Private/EditMode/SControlRigSpacePicker.cpp
@@ -13,7 +13,7 @@
 #include "Styling/AppStyle.h"
 #include "ISequencer.h"
 #include "ScopedTransaction.h"
-#include "ControlRig.h"
+#include "BaseControlRig.h"
 #include "EditMode/ControlRigEditMode.h"
 #include "MovieSceneSequence.h"
 #include "MovieScene.h"
diff --git a/Engine/Plugins/Animation/ControlRig/Source/ControlRigEditor/Private/EditMode/SControlRigSpacePicker.h b/Engine/Plugins/Animation/ControlRig/Source/ControlRigEditor/Private/EditMode/SControlRigSpacePicker.h
--- a/Engine/Plugins/Animation/ControlRig/Source/ControlRigEditor/Private/EditMode/SControlRigSpacePicker.h
+++ b/Engine/Plugins/Animation/ControlRig/Source/ControlRigEditor/Private/EditMode/SControlRigSpacePicker.h
@@ -10,13 +10,13 @@
 #include "Widgets/SWidget.h"
 #include "Widgets/SCompoundWidget.h"
 #include "Misc/FrameNumber.h"
-#include "ControlRig.h"
 #include "Editor/SRigSpacePickerWidget.h"
 
 
 class ISequencer;
 class SExpandableArea;
 class SSearchableRigHierarchyTreeView;
+class UControlRig;
 
 class SControlRigSpacePicker : public SCompoundWidget, public FControlRigBaseDockableView
 {
diff --git a/Engine/Plugins/Animation/ControlRig/Source/ControlRigEditor/Private/EditMode/SControlRigTweenWidget.cpp b/Engine/Plugins/Animation/ControlRig/Source/ControlRigEditor/Private/EditMode/SControlRigTweenWidget.cpp
--- a/Engine/Plugins/Animation/ControlRig/Source/ControlRigEditor/Private/EditMode/SControlRigTweenWidget.cpp
+++ b/Engine/Plugins/Animation/ControlRig/Source/ControlRigEditor/Private/EditMode/SControlRigTweenWidget.cpp
@@ -12,7 +12,6 @@
 #include "Styling/AppStyle.h"
 #include "Styling/CoreStyle.h"
 #include "ScopedTransaction.h"
-#include "ControlRig.h"
 #include "UnrealEdGlobals.h"
 #include "Tools/ControlRigPose.h"
 #include "ISequencer.h"
diff --git a/Engine/Plugins/Animation/ControlRig/Source/ControlRigEditor/Private/EditMode/SCurveControlContainer.cpp b/Engine/Plugins/Animation/ControlRig/Source/ControlRigEditor/Private/EditMode/SCurveControlContainer.cpp
--- a/Engine/Plugins/Animation/ControlRig/Source/ControlRigEditor/Private/EditMode/SCurveControlContainer.cpp
+++ b/Engine/Plugins/Animation/ControlRig/Source/ControlRigEditor/Private/EditMode/SCurveControlContainer.cpp
@@ -12,7 +12,7 @@
 #include "Editor/ControlRigEditor.h"
 #include "Engine/SkeletalMesh.h"
 #include "Animation/Skeleton.h"
-#include "ControlRig.h"
+#include "BaseControlRig.h"
 #include "ScopedTransaction.h"
 #include "Kismet2/BlueprintEditorUtils.h"
 
diff --git a/Engine/Plugins/Animation/ControlRig/Source/ControlRigEditor/Private/Editor/ControlRigEditor.cpp b/Engine/Plugins/Animation/ControlRig/Source/ControlRigEditor/Private/Editor/ControlRigEditor.cpp
--- a/Engine/Plugins/Animation/ControlRig/Source/ControlRigEditor/Private/Editor/ControlRigEditor.cpp
+++ b/Engine/Plugins/Animation/ControlRig/Source/ControlRigEditor/Private/Editor/ControlRigEditor.cpp
@@ -27,7 +27,7 @@
 #include "Animation/DebugSkelMeshComponent.h"
 #include "IPersonaPreviewScene.h"
 #include "Animation/AnimData/BoneMaskFilter.h"
-#include "ControlRig.h"
+#include "BaseControlRig.h"
 #include "Editor/ControlRigSkeletalMeshComponent.h"
 #include "ControlRigObjectBinding.h"
 #include "RigVMBlueprintUtils.h"
diff --git a/Engine/Plugins/Animation/ControlRig/Source/ControlRigEditor/Private/Editor/ControlRigEditorEditMode.cpp b/Engine/Plugins/Animation/ControlRig/Source/ControlRigEditor/Private/Editor/ControlRigEditorEditMode.cpp
--- a/Engine/Plugins/Animation/ControlRig/Source/ControlRigEditor/Private/Editor/ControlRigEditorEditMode.cpp
+++ b/Engine/Plugins/Animation/ControlRig/Source/ControlRigEditor/Private/Editor/ControlRigEditorEditMode.cpp
@@ -4,7 +4,7 @@
 #include "AnimationEditorViewportClient.h"
 #include "IPersonaPreviewScene.h"
 #include "AssetEditorModeManager.h"
-#include "ControlRig.h"
+#include "BaseControlRig.h"
 #include "ControlRigBlueprint.h"
 #include "ControlRigGizmoActor.h"
 #include "Preferences/PersonaOptions.h"
diff --git a/Engine/Plugins/Animation/ControlRig/Source/ControlRigEditor/Private/Editor/ControlRigSkeletalMeshComponent.cpp b/Engine/Plugins/Animation/ControlRig/Source/ControlRigEditor/Private/Editor/ControlRigSkeletalMeshComponent.cpp
--- a/Engine/Plugins/Animation/ControlRig/Source/ControlRigEditor/Private/Editor/ControlRigSkeletalMeshComponent.cpp
+++ b/Engine/Plugins/Animation/ControlRig/Source/ControlRigEditor/Private/Editor/ControlRigSkeletalMeshComponent.cpp
@@ -3,7 +3,7 @@
 #include "Editor/ControlRigSkeletalMeshComponent.h"
 #include "Sequencer/ControlRigLayerInstance.h" 
 #include "SkeletalDebugRendering.h"
-#include "ControlRig.h"
+#include "BaseControlRig.h"
 #include "AnimPreviewInstance.h"
 
 #include UE_INLINE_GENERATED_CPP_BY_NAME(ControlRigSkeletalMeshComponent)
diff --git a/Engine/Plugins/Animation/ControlRig/Source/ControlRigEditor/Private/Editor/ControlRigSkeletalMeshComponent.h b/Engine/Plugins/Animation/ControlRig/Source/ControlRigEditor/Private/Editor/ControlRigSkeletalMeshComponent.h
--- a/Engine/Plugins/Animation/ControlRig/Source/ControlRigEditor/Private/Editor/ControlRigSkeletalMeshComponent.h
+++ b/Engine/Plugins/Animation/ControlRig/Source/ControlRigEditor/Private/Editor/ControlRigSkeletalMeshComponent.h
@@ -2,7 +2,7 @@
 
 #pragma once
 
-#include "ControlRig.h"
+#include "BaseControlRig.h"
 #include "Animation/DebugSkelMeshComponent.h"
 #include "ControlRigSkeletalMeshComponent.generated.h"
 
diff --git a/Engine/Plugins/Animation/ControlRig/Source/ControlRigEditor/Private/Editor/SControlRigGizmoNameList.cpp b/Engine/Plugins/Animation/ControlRig/Source/ControlRigEditor/Private/Editor/SControlRigGizmoNameList.cpp
--- a/Engine/Plugins/Animation/ControlRig/Source/ControlRigEditor/Private/Editor/SControlRigGizmoNameList.cpp
+++ b/Engine/Plugins/Animation/ControlRig/Source/ControlRigEditor/Private/Editor/SControlRigGizmoNameList.cpp
@@ -8,7 +8,6 @@
 #include "ScopedTransaction.h"
 #include "DetailLayoutBuilder.h"
 #include "ControlRigBlueprint.h"
-#include "ControlRig.h"
 
 void SControlRigShapeNameList::Construct(const FArguments& InArgs, FRigControlElement* ControlElement, UControlRigBlueprint* InBlueprint)
 {
diff --git a/Engine/Plugins/Animation/ControlRig/Source/ControlRigEditor/Private/Editor/SRigHierarchy.cpp b/Engine/Plugins/Animation/ControlRig/Source/ControlRigEditor/Private/Editor/SRigHierarchy.cpp
--- a/Engine/Plugins/Animation/ControlRig/Source/ControlRigEditor/Private/Editor/SRigHierarchy.cpp
+++ b/Engine/Plugins/Animation/ControlRig/Source/ControlRigEditor/Private/Editor/SRigHierarchy.cpp
@@ -28,7 +28,7 @@
 #include "Editor/EditorEngine.h"
 #include "HelperUtil.h"
 #include "Widgets/Text/SInlineEditableTextBlock.h"
-#include "ControlRig.h"
+#include "BaseControlRig.h"
 #include "HAL/PlatformApplicationMisc.h"
 #include "HAL/PlatformTime.h"
 #include "Dialogs/Dialogs.h"
diff --git a/Engine/Plugins/Animation/ControlRig/Source/ControlRigEditor/Private/Editor/SRigHierarchyTreeView.cpp b/Engine/Plugins/Animation/ControlRig/Source/ControlRigEditor/Private/Editor/SRigHierarchyTreeView.cpp
--- a/Engine/Plugins/Animation/ControlRig/Source/ControlRigEditor/Private/Editor/SRigHierarchyTreeView.cpp
+++ b/Engine/Plugins/Animation/ControlRig/Source/ControlRigEditor/Private/Editor/SRigHierarchyTreeView.cpp
@@ -12,7 +12,7 @@
 #include "Editor/EditorEngine.h"
 #include "HelperUtil.h"
 #include "Widgets/Text/SInlineEditableTextBlock.h"
-#include "ControlRig.h"
+#include "BaseControlRig.h"
 #include "ControlRigEditorStyle.h"
 #include "Editor/SRigHierarchy.h"
 #include "Settings/ControlRigSettings.h"
diff --git a/Engine/Plugins/Animation/ControlRig/Source/ControlRigEditor/Private/Editor/SRigSpacePickerWidget.cpp b/Engine/Plugins/Animation/ControlRig/Source/ControlRigEditor/Private/Editor/SRigSpacePickerWidget.cpp
--- a/Engine/Plugins/Animation/ControlRig/Source/ControlRigEditor/Private/Editor/SRigSpacePickerWidget.cpp
+++ b/Engine/Plugins/Animation/ControlRig/Source/ControlRigEditor/Private/Editor/SRigSpacePickerWidget.cpp
@@ -18,7 +18,7 @@
 #include "Framework/Application/SlateUser.h"
 #include "Framework/Notifications/NotificationManager.h"
 #include "Modules/ModuleManager.h"
-#include "ControlRig.h"
+#include "BaseControlRig.h"
 #include "ControlRigBlueprint.h"
 #include "RigVMBlueprintGeneratedClass.h"
 
diff --git a/Engine/Plugins/Animation/ControlRig/Source/ControlRigEditor/Private/Editor/SRigSpacePickerWidget.h b/Engine/Plugins/Animation/ControlRig/Source/ControlRigEditor/Private/Editor/SRigSpacePickerWidget.h
--- a/Engine/Plugins/Animation/ControlRig/Source/ControlRigEditor/Private/Editor/SRigSpacePickerWidget.h
+++ b/Engine/Plugins/Animation/ControlRig/Source/ControlRigEditor/Private/Editor/SRigSpacePickerWidget.h
@@ -3,7 +3,6 @@
 #pragma once
 
 #include "CoreMinimal.h"
-#include "ControlRig.h"
 #include "Rigs/RigHierarchy.h"
 #include "Editor/SRigHierarchy.h"
 #include "Widgets/Layout/SBox.h"
diff --git a/Engine/Plugins/Animation/ControlRig/Source/ControlRigEditor/Private/Graph/ControlRigGraphPanelPinFactory.cpp b/Engine/Plugins/Animation/ControlRig/Source/ControlRigEditor/Private/Graph/ControlRigGraphPanelPinFactory.cpp
--- a/Engine/Plugins/Animation/ControlRig/Source/ControlRigEditor/Private/Graph/ControlRigGraphPanelPinFactory.cpp
+++ b/Engine/Plugins/Animation/ControlRig/Source/ControlRigEditor/Private/Graph/ControlRigGraphPanelPinFactory.cpp
@@ -12,7 +12,7 @@
 #include "Widgets/SRigVMGraphPinUserDataPath.h"
 #include "KismetPins/SGraphPinExec.h"
 #include "SGraphPinComboBox.h"
-#include "ControlRig.h"
+#include "BaseControlRig.h"
 #include "NodeFactory.h"
 #include "EdGraphSchema_K2.h"
 #include "Curves/CurveFloat.h"
diff --git a/Engine/Plugins/Animation/ControlRig/Source/ControlRigEditor/Private/Sequencer/ControlRigParameterTrackEditor.cpp b/Engine/Plugins/Animation/ControlRig/Source/ControlRigEditor/Private/Sequencer/ControlRigParameterTrackEditor.cpp
--- a/Engine/Plugins/Animation/ControlRig/Source/ControlRigEditor/Private/Sequencer/ControlRigParameterTrackEditor.cpp
+++ b/Engine/Plugins/Animation/ControlRig/Source/ControlRigEditor/Private/Sequencer/ControlRigParameterTrackEditor.cpp
@@ -42,7 +42,7 @@
 #include "Engine/SCS_Node.h"
 #include "Engine/SimpleConstructionScript.h"
 #include "Engine/Blueprint.h"
-#include "ControlRig.h"
+#include "BaseControlRig.h"
 #include "EditMode/ControlRigEditMode.h"
 #include "EditorModeManager.h"
 #include "Engine/Selection.h"
diff --git a/Engine/Plugins/Animation/ControlRig/Source/ControlRigEditor/Private/Sequencer/ControlRigSequencerFilter.cpp b/Engine/Plugins/Animation/ControlRig/Source/ControlRigEditor/Private/Sequencer/ControlRigSequencerFilter.cpp
--- a/Engine/Plugins/Animation/ControlRig/Source/ControlRigEditor/Private/Sequencer/ControlRigSequencerFilter.cpp
+++ b/Engine/Plugins/Animation/ControlRig/Source/ControlRigEditor/Private/Sequencer/ControlRigSequencerFilter.cpp
@@ -1,7 +1,7 @@
 // Copyright Epic Games, Inc. All Rights Reserved.
 
 #include "Sequencer/ControlRigSequencerFilter.h"
-#include "ControlRig.h"
+#include "BaseControlRig.h"
 #include "Editor/ControlRigSkeletalMeshComponent.h"
 #include "Sequencer/MovieSceneControlRigParameterTrack.h"
 #include "Styling/AppStyle.h"
diff --git a/Engine/Plugins/Animation/ControlRig/Source/ControlRigEditor/Private/Sequencer/ControlRigSpaceChannelEditors.cpp b/Engine/Plugins/Animation/ControlRig/Source/ControlRigEditor/Private/Sequencer/ControlRigSpaceChannelEditors.cpp
--- a/Engine/Plugins/Animation/ControlRig/Source/ControlRigEditor/Private/Sequencer/ControlRigSpaceChannelEditors.cpp
+++ b/Engine/Plugins/Animation/ControlRig/Source/ControlRigEditor/Private/Sequencer/ControlRigSpaceChannelEditors.cpp
@@ -48,7 +48,7 @@
 #include "Sequencer/MovieSceneControlRigSpaceChannel.h"
 #include "Channels/MovieSceneChannelProxy.h"
 #include "MovieSceneObjectBindingID.h"
-#include "ControlRig.h"
+#include "BaseControlRig.h"
 #include "IControlRigObjectBinding.h"
 #include "ControlRigSpaceChannelCurveModel.h"
 #include "ScopedTransaction.h"
diff --git a/Engine/Plugins/Animation/ControlRig/Source/ControlRigEditor/Private/Tools/AssetDefinition_ControlRigPose.cpp b/Engine/Plugins/Animation/ControlRig/Source/ControlRigEditor/Private/Tools/AssetDefinition_ControlRigPose.cpp
--- a/Engine/Plugins/Animation/ControlRig/Source/ControlRigEditor/Private/Tools/AssetDefinition_ControlRigPose.cpp
+++ b/Engine/Plugins/Animation/ControlRig/Source/ControlRigEditor/Private/Tools/AssetDefinition_ControlRigPose.cpp
@@ -4,7 +4,7 @@
 
 #include "ContentBrowserModule.h"
 #include "ContentBrowserMenuContexts.h"
-#include "ControlRig.h"
+#include "BaseControlRig.h"
 #include "EditorModeManager.h"
 #include "EditMode/ControlRigEditMode.h"
 #include "EditMode/SControlRigRenamePoseControls.h"
diff --git a/Engine/Plugins/Animation/ControlRig/Source/ControlRigEditor/Private/Tools/BakingHelper.cpp b/Engine/Plugins/Animation/ControlRig/Source/ControlRigEditor/Private/Tools/BakingHelper.cpp
--- a/Engine/Plugins/Animation/ControlRig/Source/ControlRigEditor/Private/Tools/BakingHelper.cpp
+++ b/Engine/Plugins/Animation/ControlRig/Source/ControlRigEditor/Private/Tools/BakingHelper.cpp
@@ -3,7 +3,6 @@
 
 #include "Tools/BakingHelper.h"
 
-#include "ControlRig.h"
 #include "ILevelSequenceEditorToolkit.h"
 #include "ISequencer.h"
 #include "LevelSequence.h"
diff --git a/Engine/Plugins/Animation/ControlRig/Source/ControlRigEditor/Private/Tools/ControlRigSnapper.cpp b/Engine/Plugins/Animation/ControlRig/Source/ControlRigEditor/Private/Tools/ControlRigSnapper.cpp
--- a/Engine/Plugins/Animation/ControlRig/Source/ControlRigEditor/Private/Tools/ControlRigSnapper.cpp
+++ b/Engine/Plugins/Animation/ControlRig/Source/ControlRigEditor/Private/Tools/ControlRigSnapper.cpp
@@ -4,7 +4,7 @@
 #include "Tools/ControlRigTweener.h" //remove
 
 #include "Channels/MovieSceneFloatChannel.h"
-#include "ControlRig.h"
+#include "BaseControlRig.h"
 #include "ISequencer.h"
 #include "MovieSceneSequence.h"
 #include "MovieScene.h"
diff --git a/Engine/Plugins/Animation/ControlRig/Source/ControlRigEditor/Private/Tools/ControlRigTweener.cpp b/Engine/Plugins/Animation/ControlRig/Source/ControlRigEditor/Private/Tools/ControlRigTweener.cpp
--- a/Engine/Plugins/Animation/ControlRig/Source/ControlRigEditor/Private/Tools/ControlRigTweener.cpp
+++ b/Engine/Plugins/Animation/ControlRig/Source/ControlRigEditor/Private/Tools/ControlRigTweener.cpp
@@ -3,7 +3,7 @@
 #include "Tools/ControlRigTweener.h"
 #include "Channels/MovieSceneFloatChannel.h"
 #include "Channels/MovieSceneDoubleChannel.h"
-#include "ControlRig.h"
+#include "BaseControlRig.h"
 #include "ISequencer.h"
 #include "MovieSceneSequence.h"
 #include "MovieScene.h"
diff --git a/Engine/Plugins/Animation/ControlRig/Source/ControlRigEditor/Public/ControlRigSequencerEditorLibrary.h b/Engine/Plugins/Animation/ControlRig/Source/ControlRigEditor/Public/ControlRigSequencerEditorLibrary.h
--- a/Engine/Plugins/Animation/ControlRig/Source/ControlRigEditor/Public/ControlRigSequencerEditorLibrary.h
+++ b/Engine/Plugins/Animation/ControlRig/Source/ControlRigEditor/Public/ControlRigSequencerEditorLibrary.h
@@ -10,7 +10,6 @@
 #include "CoreMinimal.h"
 #include "UObject/ObjectMacros.h"
 #include "MovieSceneBindingProxy.h"
-#include "ControlRig.h"
 #include "Tools/ControlRigSnapper.h"
 #include "TransformNoScale.h"
 #include "EulerTransform.h"
@@ -25,6 +24,7 @@ class UTickableConstraint;
 class UTickableTransformConstraint;
 class UTransformableHandle;
 struct FBakingAnimationKeySettings;
+class UControlRig;
 
 USTRUCT(BlueprintType)
 struct FControlRigSequencerBindingProxy
diff --git a/Engine/Plugins/Animation/ControlRig/Source/ControlRigEditor/Public/Tools/ControlRigSnapSettings.h b/Engine/Plugins/Animation/ControlRig/Source/ControlRigEditor/Public/Tools/ControlRigSnapSettings.h
--- a/Engine/Plugins/Animation/ControlRig/Source/ControlRigEditor/Public/Tools/ControlRigSnapSettings.h
+++ b/Engine/Plugins/Animation/ControlRig/Source/ControlRigEditor/Public/Tools/ControlRigSnapSettings.h
@@ -6,7 +6,6 @@
 #include "CoreMinimal.h"
 #include "UObject/ObjectMacros.h"
 #include "UObject/Object.h"
-#include "ControlRig.h"
 #include "BakingAnimationKeySettings.h"
 #include "ControlRigSnapSettings.generated.h"
 
diff --git a/Engine/Plugins/Animation/ControlRigSpline/Source/ControlRigSpline/Private/ControlRigSplineUnits.cpp b/Engine/Plugins/Animation/ControlRigSpline/Source/ControlRigSpline/Private/ControlRigSplineUnits.cpp
--- a/Engine/Plugins/Animation/ControlRigSpline/Source/ControlRigSpline/Private/ControlRigSplineUnits.cpp
+++ b/Engine/Plugins/Animation/ControlRigSpline/Source/ControlRigSpline/Private/ControlRigSplineUnits.cpp
@@ -4,6 +4,7 @@
 #include "Units/RigUnitContext.h"
 #include "Features/IModularFeatures.h"
 #include "Algo/BinarySearch.h"
+#include "BaseControlRig.h"
 #include "Units/Highlevel/Hierarchy/RigUnit_FitChainToCurve.h"
 
 #include UE_INLINE_GENERATED_CPP_BY_NAME(ControlRigSplineUnits)
diff --git a/Engine/Plugins/Animation/ControlRigSpline/Source/ControlRigSpline/Public/ControlRigSplineUnits.h b/Engine/Plugins/Animation/ControlRigSpline/Source/ControlRigSpline/Public/ControlRigSplineUnits.h
--- a/Engine/Plugins/Animation/ControlRigSpline/Source/ControlRigSpline/Public/ControlRigSplineUnits.h
+++ b/Engine/Plugins/Animation/ControlRigSpline/Source/ControlRigSpline/Public/ControlRigSplineUnits.h
@@ -3,7 +3,6 @@
 #include "CoreMinimal.h"
 #include "ControlRigDefines.h"
 #include "Units/RigUnit.h"
-#include "ControlRig.h"
 #include "ControlRigSplineTypes.h"
 #include "Units/Highlevel/Hierarchy/RigUnit_FitChainToCurve.h"
 #include "ControlRigSplineUnits.generated.h"
diff --git a/Engine/Plugins/Animation/RigLogic/Source/RigLogicModule/Private/RigUnit_RigLogic.cpp b/Engine/Plugins/Animation/RigLogic/Source/RigLogicModule/Private/RigUnit_RigLogic.cpp
--- a/Engine/Plugins/Animation/RigLogic/Source/RigLogicModule/Private/RigUnit_RigLogic.cpp
+++ b/Engine/Plugins/Animation/RigLogic/Source/RigLogicModule/Private/RigUnit_RigLogic.cpp
@@ -2,7 +2,7 @@
 
 #include "RigUnit_RigLogic.h"
 
-#include "ControlRig.h"
+#include "BaseControlRig.h"
 #include "DNAReader.h"
 #include "Engine/SkeletalMesh.h"
 #include "RigInstance.h"
diff --git a/Engine/Plugins/Animation/RigLogic/Source/RigLogicModule/Private/RigUnit_RigLogic_Test.cpp b/Engine/Plugins/Animation/RigLogic/Source/RigLogicModule/Private/RigUnit_RigLogic_Test.cpp
--- a/Engine/Plugins/Animation/RigLogic/Source/RigLogicModule/Private/RigUnit_RigLogic_Test.cpp
+++ b/Engine/Plugins/Animation/RigLogic/Source/RigLogicModule/Private/RigUnit_RigLogic_Test.cpp
@@ -9,7 +9,6 @@
 #include "SharedRigRuntimeContext.h"
 #include "Components/SkeletalMeshComponent.h"
 #include "Units/RigUnitContext.h"
-#include "ControlRig.h"
 #include "Math/TransformNonVectorized.h"
 #include "DNAUtils.h"
 
diff --git a/Engine/Plugins/Animation/RigLogic/Source/RigLogicModule/Public/RigUnit_RigLogic.h b/Engine/Plugins/Animation/RigLogic/Source/RigLogicModule/Public/RigUnit_RigLogic.h
--- a/Engine/Plugins/Animation/RigLogic/Source/RigLogicModule/Public/RigUnit_RigLogic.h
+++ b/Engine/Plugins/Animation/RigLogic/Source/RigLogicModule/Public/RigUnit_RigLogic.h
@@ -5,7 +5,6 @@
 #include "RigInstance.h"
 #include "RigLogic.h"
 
-#include "ControlRig.h"
 #include "DNAAsset.h"
 #include "Materials/Material.h"
 #include "Units/RigUnit.h"
diff --git a/Engine/Plugins/Editor/SequencerAnimTools/Source/SequencerAnimTools/Private/SequencerAnimEditPivotTool.cpp b/Engine/Plugins/Editor/SequencerAnimTools/Source/SequencerAnimTools/Private/SequencerAnimEditPivotTool.cpp
--- a/Engine/Plugins/Editor/SequencerAnimTools/Source/SequencerAnimTools/Private/SequencerAnimEditPivotTool.cpp
+++ b/Engine/Plugins/Editor/SequencerAnimTools/Source/SequencerAnimTools/Private/SequencerAnimEditPivotTool.cpp
@@ -17,7 +17,7 @@
 #include "EditorModeManager.h"
 #include "Misc/QualifiedFrameTime.h"
 #include "Modules/ModuleManager.h"
-#include "ControlRig.h"
+#include "BaseControlRig.h"
 #include "ControlRigSequencerEditorLibrary.h"
 #include "LevelSequence.h"
 #include "LevelSequenceEditorBlueprintLibrary.h"
diff --git a/Engine/Plugins/Editor/SequencerAnimTools/Source/SequencerAnimTools/Public/SequencerAnimEditPivotTool.h b/Engine/Plugins/Editor/SequencerAnimTools/Source/SequencerAnimTools/Public/SequencerAnimEditPivotTool.h
--- a/Engine/Plugins/Editor/SequencerAnimTools/Source/SequencerAnimTools/Public/SequencerAnimEditPivotTool.h
+++ b/Engine/Plugins/Editor/SequencerAnimTools/Source/SequencerAnimTools/Public/SequencerAnimEditPivotTool.h
@@ -250,6 +250,6 @@ public:
 
 #if UE_ENABLE_INCLUDE_ORDER_DEPRECATED_IN_5_2
 #include "BaseGizmos/TransformProxy.h"
-#include "ControlRig.h"
+#include "BaseControlRig.h"
 #include "CoreMinimal.h"
 #endif
diff --git a/Engine/Plugins/Importers/USDImporter/Source/USDUtilities/Private/USDSkeletalDataConversion.cpp b/Engine/Plugins/Importers/USDImporter/Source/USDUtilities/Private/USDSkeletalDataConversion.cpp
--- a/Engine/Plugins/Importers/USDImporter/Source/USDUtilities/Private/USDSkeletalDataConversion.cpp
+++ b/Engine/Plugins/Importers/USDImporter/Source/USDUtilities/Private/USDSkeletalDataConversion.cpp
@@ -28,7 +28,7 @@
 #include "AnimEncoding.h"
 #include "Async/ParallelFor.h"
 #include "BoneWeights.h"
-#include "ControlRig.h"
+#include "BaseControlRig.h"
 #include "Evaluation/MovieSceneSequenceTransform.h"
 #include "IMovieScenePlayer.h"
 #include "Interfaces/ITargetPlatform.h"
