Control Rig: Fix compilation issue with RigUnit_UserData
#rnx

diff --git a/Engine/Plugins/Animation/ControlRig/Source/ControlRig/Private/Units/Hierarchy/RigUnit_Metadata.h b/Engine/Plugins/Animation/ControlRig/Source/ControlRig/Private/Units/Hierarchy/RigUnit_Metadata.h
--- a/Engine/Plugins/Animation/ControlRig/Source/ControlRig/Private/Units/Hierarchy/RigUnit_Metadata.h
+++ b/Engine/Plugins/Animation/ControlRig/Source/ControlRig/Private/Units/Hierarchy/RigUnit_Metadata.h
@@ -3,6 +3,7 @@
 #pragma once
 
 #include "Units/RigDispatchFactory.h"
+#include "Units/RigUnit.h"
 #include "RigUnit_Metadata.generated.h"
 
 USTRUCT(meta=(Abstract, Category="Hierarchy", NodeColor="0.462745, 1,0, 0.329412", Varying))
