Control Rig: Fix compilation issue with RigUnit_UserData again
#rnx

diff --git a/Engine/Plugins/Animation/ControlRig/Source/ControlRig/Private/Units/Core/RigUnit_UserData.h b/Engine/Plugins/Animation/ControlRig/Source/ControlRig/Private/Units/Core/RigUnit_UserData.h
--- a/Engine/Plugins/Animation/ControlRig/Source/ControlRig/Private/Units/Core/RigUnit_UserData.h
+++ b/Engine/Plugins/Animation/ControlRig/Source/ControlRig/Private/Units/Core/RigUnit_UserData.h
@@ -3,6 +3,7 @@
 #pragma once
 
 #include "Units/RigDispatchFactory.h"
+#include "Units/RigUnit.h"
 #include "RigUnit_UserData.generated.h"
 
 /*
