Control Rig: Perform Lazy caching of channel proxies to avoid constraint channels missing on load. Also rmeoved explicit MarkAsChanged in the reconstruct, which may fix non-determinism issues. Now we do it when a constraint is added since we need to update the channel proxy in that case.

diff --git a/Engine/Plugins/Animation/ControlRig/Source/ControlRig/Private/Sequencer/MovieSceneControlRigParameterSection.cpp b/Engine/Plugins/Animation/ControlRig/Source/ControlRig/Private/Sequencer/MovieSceneControlRigParameterSection.cpp
--- a/Engine/Plugins/Animation/ControlRig/Source/ControlRig/Private/Sequencer/MovieSceneControlRigParameterSection.cpp
+++ b/Engine/Plugins/Animation/ControlRig/Source/ControlRig/Private/Sequencer/MovieSceneControlRigParameterSection.cpp
@@ -1589,6 +1589,11 @@ void UMovieSceneControlRigParameterSection::StoreLastControlsUsedToReconstruct(c
 }
 
 void UMovieSceneControlRigParameterSection::ReconstructChannelProxy()
+{
+	ChannelProxy.Reset();
+}
+
+EMovieSceneChannelProxyType UMovieSceneControlRigParameterSection::CacheChannelProxy()
 {
 	FMovieSceneChannelProxyData Channels;
 	ControlChannelMap.Empty();
@@ -1605,7 +1610,6 @@ void UMovieSceneControlRigParameterSection::ReconstructChannelProxy()
 			OnArray.Init(true, ControlRig->AvailableControls().Num());
 			SetControlsMask(OnArray);
 		}
-
 		int32 ControlIndex = 0; 
 		int32 MaskIndex = 0;
 		int32 TotalIndex = 0; 
@@ -1623,20 +1627,17 @@ void UMovieSceneControlRigParameterSection::ReconstructChannelProxy()
 		const FName SpaceName = FName(TEXT("Space"));
 
 		// begin constraints
-		const FConstraintsManagerController& Controller = FConstraintsManagerController::Get(ControlRig->GetWorld());
-		auto GetConstraints = [&Controller, this](const FName& InControlName)
+	// begin constraints
+		auto AddConstrainChannels = [this, &ConstraintsChannelIndex, &TotalIndex, &Channels](
+			const FName& InControlName, const FText& InGroup, const bool bEnabled)
 		{
+			
+			const FConstraintsManagerController& Controller = FConstraintsManagerController::Get(ControlRig->GetWorld());
+
 			static constexpr bool bSorted = true;
 			const uint32 ControlHash = UTransformableControlHandle::ComputeHash(ControlRig.Get(), InControlName);
-			return Controller.GetParentConstraints(ControlHash, bSorted);
-		};
-
-
-		auto AddConstrainChannels = [this, GetConstraints, &ConstraintsChannelIndex, &TotalIndex, &Channels](
-			const FName& InControlName, const FText& InGroup, const bool bEnabled)
-		{
-			TArray<TWeakObjectPtr<UTickableConstraint>> Constraints = GetConstraints(InControlName);
-			for (const TWeakObjectPtr<UTickableConstraint>& Constraint: Constraints)
+			TArray<TWeakObjectPtr<UTickableConstraint>> Constraints = Controller.GetParentConstraints(ControlHash, bSorted);
+			for (const TWeakObjectPtr<UTickableConstraint>& Constraint : Constraints)
 			{
 				if (Constraint.IsValid())
 				{
@@ -2285,8 +2286,8 @@ void UMovieSceneControlRigParameterSection::ReconstructChannelProxy()
 
 
 	ChannelProxy = MakeShared<FMovieSceneChannelProxy>(MoveTemp(Channels));
-
-	MarkAsChanged();
+	
+	return EMovieSceneChannelProxyType::Dynamic;
 }
 
 FMovieSceneInterrogationKey UMovieSceneControlRigParameterSection::GetFloatInterrogationKey()
diff --git a/Engine/Plugins/Animation/ControlRig/Source/ControlRig/Private/Sequencer/MovieSceneControlRigParameterTemplate.cpp b/Engine/Plugins/Animation/ControlRig/Source/ControlRig/Private/Sequencer/MovieSceneControlRigParameterTemplate.cpp
--- a/Engine/Plugins/Animation/ControlRig/Source/ControlRig/Private/Sequencer/MovieSceneControlRigParameterTemplate.cpp
+++ b/Engine/Plugins/Animation/ControlRig/Source/ControlRig/Private/Sequencer/MovieSceneControlRigParameterTemplate.cpp
@@ -1133,6 +1133,9 @@ static UTickableConstraint* CreateConstraintIfNeeded(const FConstraintsManagerCo
 		if (Controller.GetConstraint(Constraint->ConstraintID) == nullptr)
 		{
 			Controller.AddConstraint(ConstraintValue.Constraint.Get());
+			//need to reconstuct channels here.. note this is now lazy and so will recreate it next time view requests it
+			Section->ReconstructChannelProxy();
+			Section->MarkAsChanged();
 		}
 	}
 
diff --git a/Engine/Plugins/Animation/ControlRig/Source/ControlRig/Public/Sequencer/MovieSceneControlRigParameterSection.h b/Engine/Plugins/Animation/ControlRig/Source/ControlRig/Public/Sequencer/MovieSceneControlRigParameterSection.h
--- a/Engine/Plugins/Animation/ControlRig/Source/ControlRig/Public/Sequencer/MovieSceneControlRigParameterSection.h
+++ b/Engine/Plugins/Animation/ControlRig/Source/ControlRig/Public/Sequencer/MovieSceneControlRigParameterSection.h
@@ -268,6 +268,8 @@ public:
 	//UMovieSceneSection virtuals
 	virtual void SetBlendType(EMovieSceneBlendType InBlendType) override;
 	virtual UObject* GetImplicitObjectOwner() override;
+	virtual EMovieSceneChannelProxyType CacheChannelProxy() override;
+
 	// IMovieSceneConstrainedSection overrides
 	/*
 	* Whether it has that channel
