Control Rig Sequencer: Rename additive control rigs to layered control rigs
#rnx

diff --git a/Engine/Plugins/Animation/ControlRig/Source/ControlRig/Private/Sequencer/MovieSceneControlRigParameterTrack.cpp b/Engine/Plugins/Animation/ControlRig/Source/ControlRig/Private/Sequencer/MovieSceneControlRigParameterTrack.cpp
--- a/Engine/Plugins/Animation/ControlRig/Source/ControlRig/Private/Sequencer/MovieSceneControlRigParameterTrack.cpp
+++ b/Engine/Plugins/Animation/ControlRig/Source/ControlRig/Private/Sequencer/MovieSceneControlRigParameterTrack.cpp
@@ -18,7 +18,7 @@
 #define LOCTEXT_NAMESPACE "MovieSceneParameterControlRigTrack"
 
 FColor UMovieSceneControlRigParameterTrack::AbsoluteRigTrackColor = FColor(65, 89, 194, 65);
-FColor UMovieSceneControlRigParameterTrack::AdditiveRigTrackColor = FColor(173, 151, 114);
+FColor UMovieSceneControlRigParameterTrack::LayeredRigTrackColor = FColor(173, 151, 114);
 
 FControlRotationOrder::FControlRotationOrder()
 	:RotationOrder(EEulerRotationOrder::YZX),
diff --git a/Engine/Plugins/Animation/ControlRig/Source/ControlRig/Public/Sequencer/MovieSceneControlRigParameterTrack.h b/Engine/Plugins/Animation/ControlRig/Source/ControlRig/Public/Sequencer/MovieSceneControlRigParameterTrack.h
--- a/Engine/Plugins/Animation/ControlRig/Source/ControlRig/Public/Sequencer/MovieSceneControlRigParameterTrack.h
+++ b/Engine/Plugins/Animation/ControlRig/Source/ControlRig/Public/Sequencer/MovieSceneControlRigParameterTrack.h
@@ -205,7 +205,7 @@ private:
 
 public:
 	static CONTROLRIG_API FColor AbsoluteRigTrackColor;
-	static CONTROLRIG_API FColor AdditiveRigTrackColor;
+	static CONTROLRIG_API FColor LayeredRigTrackColor;
 };
 
 
diff --git a/Engine/Plugins/Animation/ControlRig/Source/ControlRigEditor/Private/ControlRigSequencerEditorLibrary.cpp b/Engine/Plugins/Animation/ControlRig/Source/ControlRigEditor/Private/ControlRigSequencerEditorLibrary.cpp
--- a/Engine/Plugins/Animation/ControlRig/Source/ControlRigEditor/Private/ControlRigSequencerEditorLibrary.cpp
+++ b/Engine/Plugins/Animation/ControlRig/Source/ControlRigEditor/Private/ControlRigSequencerEditorLibrary.cpp
@@ -270,10 +270,10 @@ static UMovieSceneControlRigParameterTrack* AddControlRig(ULevelSequence* LevelS
 
 		if (bIsAdditiveControlRig)
 		{
-			const FString AdditiveObjectName = ObjectName + TEXT(" (Additive)");
+			const FString AdditiveObjectName = ObjectName + TEXT(" (Layered)");
 			Track->SetTrackName(FName(*ObjectName));
 			Track->SetDisplayName(FText::FromString(AdditiveObjectName));
-			Track->SetColorTint(UMovieSceneControlRigParameterTrack::AdditiveRigTrackColor);
+			Track->SetColorTint(UMovieSceneControlRigParameterTrack::LayeredRigTrackColor);
 		}
 		else
 		{
diff --git a/Engine/Plugins/Animation/ControlRig/Source/ControlRigEditor/Private/Sequencer/ControlRigParameterTrackEditor.cpp b/Engine/Plugins/Animation/ControlRig/Source/ControlRigEditor/Private/Sequencer/ControlRigParameterTrackEditor.cpp
--- a/Engine/Plugins/Animation/ControlRig/Source/ControlRigEditor/Private/Sequencer/ControlRigParameterTrackEditor.cpp
+++ b/Engine/Plugins/Animation/ControlRig/Source/ControlRigEditor/Private/Sequencer/ControlRigParameterTrackEditor.cpp
@@ -222,7 +222,7 @@ FControlRigParameterTrackEditor::FControlRigParameterTrackEditor(TSharedRef<ISeq
 	, bCurveDisplayTickIsPending(false)
 	, bIsDoingSelection(false)
 	, bSkipNextSelectionFromTimer(false)
-	, bIsAdditiveControlRig(false)
+	, bIsLayeredControlRig(false)
 	, bFilterAssetBySkeleton(true)
 	, bFilterAssetByAnimatableControls(false)
 	, ControlUndoBracket(0)
@@ -1229,8 +1229,8 @@ void FControlRigParameterTrackEditor::HandleAddTrackSubMenu(FMenuBuilder& MenuBu
 	if (CVarEnableAdditiveControlRigs->GetBool())
 	{
 		MenuBuilder.AddMenuEntry(
-		LOCTEXT("IsAdditiveControlRig", "Additive"),
-		LOCTEXT("IsAdditiveControlRigTooltip", "Add an additive control rig"),
+		LOCTEXT("IsLayeredControlRig", "Layered"),
+		LOCTEXT("IsLayeredControlRigTooltip", "Add an layered control rig"),
 		FSlateIcon(),
 		FUIAction(
 			FExecuteAction::CreateSP(this, &FControlRigParameterTrackEditor::ToggleIsAdditiveControlRig),
@@ -1274,12 +1274,12 @@ void FControlRigParameterTrackEditor::HandleAddTrackSubMenu(FMenuBuilder& MenuBu
 
 void FControlRigParameterTrackEditor::ToggleIsAdditiveControlRig()
 {
-	bIsAdditiveControlRig = bIsAdditiveControlRig ? false : true;
+	bIsLayeredControlRig = bIsLayeredControlRig ? false : true;
 }
 
 bool FControlRigParameterTrackEditor::IsToggleIsAdditiveControlRig()
 {
-	return bIsAdditiveControlRig;
+	return bIsLayeredControlRig;
 }
 
 void FControlRigParameterTrackEditor::ToggleFilterAssetBySkeleton()
@@ -1350,7 +1350,7 @@ void FControlRigParameterTrackEditor::HandleAddControlRigSubMenu(FMenuBuilder& M
 		Options.bShowUnloadedBlueprints = true;
 		Options.NameTypeToDisplay = EClassViewerNameTypeToDisplay::DisplayName;
 
-		const bool bCheckInversion = bIsAdditiveControlRig;
+		const bool bCheckInversion = bIsLayeredControlRig;
 		TSharedPtr<FControlRigClassFilter> ClassFilter = MakeShareable(new FControlRigClassFilter(bFilterAssetBySkeleton, bFilterAssetByAnimatableControls, bCheckInversion, Skeleton));
 		Options.ClassFilters.Add(ClassFilter.ToSharedRef());
 		Options.bShowNoneOption = false;
@@ -1435,7 +1435,7 @@ void FControlRigParameterTrackEditor::AddControlRig(const UClass* InClass, UObje
 	ULevelSequence* LevelSequence = Cast<ULevelSequence>(GetSequencer()->GetFocusedMovieSceneSequence());
 	UMovieSceneSequence* Sequence = GetSequencer()->GetFocusedMovieSceneSequence();
 	FMovieSceneBindingProxy BindingProxy(ObjectBinding, Sequence);
-	if (UMovieSceneTrack* Track = UControlRigSequencerEditorLibrary::FindOrCreateControlRigTrack(World, LevelSequence, InClass, BindingProxy, bIsAdditiveControlRig))
+	if (UMovieSceneTrack* Track = UControlRigSequencerEditorLibrary::FindOrCreateControlRigTrack(World, LevelSequence, InClass, BindingProxy, bIsLayeredControlRig))
 	{
 		BindControlRig(CastChecked<UMovieSceneControlRigParameterTrack>(Track)->GetControlRig());
 	}
@@ -3815,18 +3815,18 @@ void FControlRigParameterTrackEditor::BuildTrackContextMenu(FMenuBuilder& MenuBu
 
 		MenuBuilder.AddMenuSeparator();
 	}
-	else if (UControlRig* AdditiveRig = Cast<UControlRig>(Track->GetControlRig()))
+	else if (UControlRig* LayeredRig = Cast<UControlRig>(Track->GetControlRig()))
 	{
-		if (AdditiveRig->IsAdditive())
+		if (LayeredRig->IsAdditive())
 		{
-			MenuBuilder.BeginSection("Additive Control Rig", LOCTEXT("AdditiveControlRig", "Additive Control Rig"));
+			MenuBuilder.BeginSection("Layered Control Rig", LOCTEXT("LayeredControlRig", "Layered Control Rig"));
 			{
 				MenuBuilder.AddMenuEntry(
 					LOCTEXT("Bake Inverted Pose", "Bake Inverted Pose"),
 					LOCTEXT("BakeInvertedPoseToolTip", "Bake inversion of the input pose into the rig"),
 					FSlateIcon(),
 					FUIAction(
-						FExecuteAction::CreateRaw(this, &FControlRigParameterTrackEditor::BakeInvertedPose, AdditiveRig, Track)));
+						FExecuteAction::CreateRaw(this, &FControlRigParameterTrackEditor::BakeInvertedPose, LayeredRig, Track)));
 			}
 			MenuBuilder.EndSection();
 			MenuBuilder.AddMenuSeparator();
@@ -3906,14 +3906,14 @@ void FControlRigParameterTrackEditor::ToggleFKControlRig(UMovieSceneControlRigPa
 	FKControlRig->ToggleApplyMode();
 	if (FKControlRig->GetApplyMode() == EControlRigFKRigExecuteMode::Additive)
 	{
-		const FString AdditiveObjectName = Track->GetTrackName().ToString() + TEXT(" (Additive)");
-		Track->SetDisplayName(FText::FromString(AdditiveObjectName));
-		Track->SetColorTint(UMovieSceneControlRigParameterTrack::AdditiveRigTrackColor);
+		const FString LayeredObjectName = Track->GetTrackName().ToString() + TEXT(" (Layered)");
+		Track->SetDisplayName(FText::FromString(LayeredObjectName));
+		Track->SetColorTint(UMovieSceneControlRigParameterTrack::LayeredRigTrackColor);
 	}
 	else
 	{
 		Track->SetDisplayName(FText::FromName(Track->GetTrackName()));
-		Track->SetColorTint(UMovieSceneControlRigParameterTrack::AdditiveRigTrackColor);
+		Track->SetColorTint(UMovieSceneControlRigParameterTrack::LayeredRigTrackColor);
 		Track->SetColorTint(UMovieSceneControlRigParameterTrack::AbsoluteRigTrackColor);
 	}
 	for (UMovieSceneSection* Section : Track->GetAllSections())
diff --git a/Engine/Plugins/Animation/ControlRig/Source/ControlRigEditor/Private/Sequencer/ControlRigParameterTrackEditor.h b/Engine/Plugins/Animation/ControlRig/Source/ControlRigEditor/Private/Sequencer/ControlRigParameterTrackEditor.h
--- a/Engine/Plugins/Animation/ControlRig/Source/ControlRigEditor/Private/Sequencer/ControlRigParameterTrackEditor.h
+++ b/Engine/Plugins/Animation/ControlRig/Source/ControlRigEditor/Private/Sequencer/ControlRigParameterTrackEditor.h
@@ -267,8 +267,8 @@ private:
 	/** A flag to determine if the next update coming from the timer should be skipped */
 	bool bSkipNextSelectionFromTimer;
 
-	/** Whether or not this rig will be used as an additive control rig */
-	bool bIsAdditiveControlRig;
+	/** Whether or not this rig will be used as an layered control rig */
+	bool bIsLayeredControlRig;
 
 	/** Whether or not we should check Skeleton when filtering*/
 	bool bFilterAssetBySkeleton;
