Control Rig: Additive Control rig - CTRL-G and CTRL+SHIFT+G is not keying the controls
#rnx

diff --git a/Engine/Plugins/Animation/ControlRig/Source/ControlRig/Private/BaseControlRig.cpp b/Engine/Plugins/Animation/ControlRig/Source/ControlRig/Private/BaseControlRig.cpp
--- a/Engine/Plugins/Animation/ControlRig/Source/ControlRig/Private/BaseControlRig.cpp
+++ b/Engine/Plugins/Animation/ControlRig/Source/ControlRig/Private/BaseControlRig.cpp
@@ -458,21 +458,30 @@ void UControlRig::ClearPoseBeforeBackwardsSolve()
 	PoseBeforeBackwardsSolve.Reset();
 }
 
-void UControlRig::InvertInputPose(const TArray<FRigElementKey>& InElements, EControlRigSetKey InSetKey)
+TArray<FRigControlElement*> UControlRig::InvertInputPose(const TArray<FRigElementKey>& InElements, EControlRigSetKey InSetKey)
 {
+	TArray<FRigControlElement*> ModifiedElements;
+	ModifiedElements.Reserve(ControlsAfterBackwardsSolve.Num());
 	for (const FRigPoseElement& PoseElement : ControlsAfterBackwardsSolve)
 	{
 		if (FRigControlElement* ControlElement = Cast<FRigControlElement>(DynamicHierarchy->Get(PoseElement.Index)))
 		{
+			if (IsAdditive() && ControlElement->Settings.ControlType == ERigControlType::Bool)
+			{
+				continue;
+			}
+			
 			if (InElements.IsEmpty() || InElements.Contains(ControlElement->GetKey()))
 			{
 				FRigControlValue Value;
 				Value.SetFromTransform(PoseElement.LocalTransform.Inverse(), ControlElement->Settings.ControlType, ControlElement->Settings.PrimaryAxis);
 				FRigSetControlValueInfo Info = {Value, true, FRigControlModifiedContext(InSetKey), false, false, false};
 				ControlValues.Add(ControlElement->GetKey(), Info);
+				ModifiedElements.Add(ControlElement);
 			}
 		}
 	}
+	return ModifiedElements;
 }
 
 void UControlRig::InitializeFromCDO()
diff --git a/Engine/Plugins/Animation/ControlRig/Source/ControlRig/Public/BaseControlRig.h b/Engine/Plugins/Animation/ControlRig/Source/ControlRig/Public/BaseControlRig.h
--- a/Engine/Plugins/Animation/ControlRig/Source/ControlRig/Public/BaseControlRig.h
+++ b/Engine/Plugins/Animation/ControlRig/Source/ControlRig/Public/BaseControlRig.h
@@ -138,7 +138,8 @@ public:
 	void ClearPoseBeforeBackwardsSolve();
 
 	/* For additive rigs, will set control values by inverting the pose found after the backwards solve */
-	void InvertInputPose(const TArray<FRigElementKey>& InElements = TArray<FRigElementKey>(), EControlRigSetKey InSetKey = EControlRigSetKey::Never);
+	/* Returns the array of control elements that were modified*/
+	TArray<FRigControlElement*> InvertInputPose(const TArray<FRigElementKey>& InElements = TArray<FRigElementKey>(), EControlRigSetKey InSetKey = EControlRigSetKey::Never);
 
 	/** Setup bindings to a runtime object (or clear by passing in nullptr). */
 	void SetObjectBinding(TSharedPtr<IControlRigObjectBinding> InObjectBinding)
diff --git a/Engine/Plugins/Animation/ControlRig/Source/ControlRigEditor/Private/EditMode/ControlRigEditMode.cpp b/Engine/Plugins/Animation/ControlRig/Source/ControlRigEditor/Private/EditMode/ControlRigEditMode.cpp
--- a/Engine/Plugins/Animation/ControlRig/Source/ControlRigEditor/Private/EditMode/ControlRigEditMode.cpp
+++ b/Engine/Plugins/Animation/ControlRig/Source/ControlRigEditor/Private/EditMode/ControlRigEditMode.cpp
@@ -3150,8 +3150,13 @@ void FControlRigEditMode::InvertInputPose(bool bSelectionOnly)
 			});
 		}
 
-		ControlRig->InvertInputPose(SelectedRigElements, EControlRigSetKey::Never);
+		const TArray<FRigControlElement*> ModifiedElements = ControlRig->InvertInputPose(SelectedRigElements, EControlRigSetKey::Never);
 		ControlRig->Evaluate_AnyThread();
+
+		for (FRigControlElement* ControlElement : ModifiedElements)
+		{
+			ControlRig->ControlModified().Broadcast(ControlRig, ControlElement, EControlRigSetKey::DoNotCare);
+		}
 	}
 }
 
