Control Rig: Fix skeletal mesh component has invalid AnimInstance after undoing an add CR track transaction
#rnx

diff --git a/Engine/Plugins/Animation/ControlRig/Source/ControlRig/Private/Sequencer/MovieSceneControlRigParameterTemplate.cpp b/Engine/Plugins/Animation/ControlRig/Source/ControlRig/Private/Sequencer/MovieSceneControlRigParameterTemplate.cpp
--- a/Engine/Plugins/Animation/ControlRig/Source/ControlRig/Private/Sequencer/MovieSceneControlRigParameterTemplate.cpp
+++ b/Engine/Plugins/Animation/ControlRig/Source/ControlRig/Private/Sequencer/MovieSceneControlRigParameterTemplate.cpp
@@ -690,6 +690,54 @@ void FControlRigBindingHelper::UnBindFromSequencerInstance(UControlRig* ControlR
 	}
 }
 
+struct FControlRigSkeletalMeshComponentBindingTokenProducer : IMovieScenePreAnimatedTokenProducer
+{
+	FControlRigSkeletalMeshComponentBindingTokenProducer(FMovieSceneSequenceIDRef InSequenceID, UControlRig* InControlRig)
+		: SequenceID(InSequenceID), ControlRigUniqueID(InControlRig->GetUniqueID())
+	{}
+
+	static FMovieSceneAnimTypeID GetAnimTypeID()
+	{
+		return TMovieSceneAnimTypeID<FControlRigSkeletalMeshComponentBindingTokenProducer>();
+	}
+
+	virtual IMovieScenePreAnimatedTokenPtr CacheExistingState(UObject& Object) const
+	{
+		struct FToken : IMovieScenePreAnimatedToken
+		{
+			FToken(FMovieSceneSequenceIDRef InSequenceID, uint32 InControlRigUniqueID)
+				: SequenceID(InSequenceID), ControlRigUniqueID(InControlRigUniqueID)
+			{
+
+			}
+			
+			virtual void RestoreState(UObject& InObject, const UE::MovieScene::FRestoreStateParams& Params) override
+			{
+				if (USkeletalMeshComponent* SkeletalMeshComponent = Cast<USkeletalMeshComponent>(&InObject))
+				{
+					if (UControlRigLayerInstance* ControlRigLayerInstance = Cast<UControlRigLayerInstance>(SkeletalMeshComponent->GetAnimInstance()))
+					{
+						ControlRigLayerInstance->RemoveControlRigTrack(ControlRigUniqueID);
+						if (!ControlRigLayerInstance->GetFirstAvailableControlRig())
+						{
+							FAnimCustomInstanceHelper::UnbindFromSkeletalMeshComponent<UControlRigLayerInstance>(SkeletalMeshComponent);
+						}
+					}
+				}
+			}
+
+			FMovieSceneSequenceID SequenceID;
+			uint32 ControlRigUniqueID;
+		};
+
+
+		FToken Token(SequenceID, ControlRigUniqueID);
+		return MoveTemp(Token);
+	}
+
+	FMovieSceneSequenceID SequenceID;
+	uint32 ControlRigUniqueID;
+};
 
 struct FControlRigParameterPreAnimatedTokenProducer : IMovieScenePreAnimatedTokenProducer
 {
@@ -1216,6 +1264,10 @@ struct FControlRigParameterExecutionToken : IMovieSceneExecutionToken
 				
 				// ensure that pre animated state is saved, must be done before bind
 				Player.SavePreAnimatedState(*ControlRig, FMovieSceneControlRigParameterTemplate::GetAnimTypeID(), FControlRigParameterPreAnimatedTokenProducer(Operand.SequenceID));
+				if (USkeletalMeshComponent* SkeletalMeshComponent = Cast<USkeletalMeshComponent>(FControlRigObjectBinding::GetBindableObject(BoundObject)))
+				{
+					Player.SavePreAnimatedState(*SkeletalMeshComponent, FControlRigSkeletalMeshComponentBindingTokenProducer::GetAnimTypeID(), FControlRigSkeletalMeshComponentBindingTokenProducer(Operand.SequenceID, ControlRig));
+				}
 
 				FControlRigBindingHelper::BindToSequencerInstance(ControlRig);
 
