Control Rig: Rename ControlRig to BaseControlRig and created a specialized ControlRig

diff --git a/Engine/Plugins/Animation/AnimationData/Source/AnimationData/Private/AnimSequencerController.cpp b/Engine/Plugins/Animation/AnimationData/Source/AnimationData/Private/AnimSequencerController.cpp
--- a/Engine/Plugins/Animation/AnimationData/Source/AnimationData/Private/AnimSequencerController.cpp
+++ b/Engine/Plugins/Animation/AnimationData/Source/AnimationData/Private/AnimSequencerController.cpp
@@ -1305,7 +1305,7 @@ bool UAnimSequencerController::SetCurveAttributes(const FAnimationCurveIdentifie
 
 		if (UMovieSceneControlRigParameterSection* Section = Model->GetFKControlRigSection())
 		{
-			UControlRig* ControlRig = Section->GetControlRig();
+			UBaseControlRig* ControlRig = Section->GetControlRig();
 
 			if (UFKControlRig* FKRig = Cast<UFKControlRig>(ControlRig))
 			{
@@ -1543,7 +1543,7 @@ bool UAnimSequencerController::SetBoneTrackKeys(FName BoneName, const TArray<FVe
 		{
 			if (UMovieSceneControlRigParameterSection* Section = Model->GetFKControlRigSection())
 			{
-				if (UControlRig* ControlRig = Model->GetControlRig())
+				if (UBaseControlRig* ControlRig = Model->GetControlRig())
 				{
 					if (URigHierarchy* Hierarchy = ControlRig->GetHierarchy())
 					{
@@ -2150,7 +2150,7 @@ void UAnimSequencerController::RemoveUnusedControlsAndCurves() const
 {
 	if (UMovieSceneControlRigParameterSection* Section = Model->GetFKControlRigSection())
 	{
-		if (UControlRig* ControlRig = Model->GetControlRig())
+		if (UBaseControlRig* ControlRig = Model->GetControlRig())
 		{
 			if (URigHierarchy* Hierarchy = ControlRig->GetHierarchy())
 			{
@@ -2609,7 +2609,7 @@ bool UAnimSequencerController::SetBoneCurveKeys(const FName& BoneName, const TAr
 	check(UAnimationSequencerDataModel::UseDirectFKControlRigMode == 1);
 	if (UMovieSceneControlRigParameterSection* Section = Model->GetFKControlRigSection())
 	{
-		UControlRig* ControlRig = Section->GetControlRig();
+		UBaseControlRig* ControlRig = Section->GetControlRig();
 		if (UFKControlRig* FKRig = Cast<UFKControlRig>(ControlRig))
 		{
 			const URigHierarchy* Hierarchy = FKRig->GetHierarchy();
@@ -2801,7 +2801,7 @@ bool UAnimSequencerController::UpdateBoneCurveKeys(const FName& BoneName, const
 {
 	if (UMovieSceneControlRigParameterSection* Section = Model->GetFKControlRigSection())
 	{		
-		UControlRig* ControlRig = Section->GetControlRig();
+		UBaseControlRig* ControlRig = Section->GetControlRig();
 		if (UFKControlRig* FKRig = Cast<UFKControlRig>(ControlRig))
 		{
 			if(URigHierarchy* Hierarchy = FKRig->GetHierarchy())
@@ -2896,7 +2896,7 @@ bool UAnimSequencerController::RemoveBoneCurveKey(const FName& BoneName, float T
     	const FAnimationCurveIdentifier CurveId = UAnimationCurveIdentifierExtensions::GetCurveIdentifier(Skeleton, BoneName, ERawCurveTrackTypes::RCT_Transform);
 		if (const UMovieSceneControlRigParameterSection* Section = Model->GetFKControlRigSection())
 		{
-			UControlRig* ControlRig = Section->GetControlRig();
+			UBaseControlRig* ControlRig = Section->GetControlRig();
 	
 			if (UFKControlRig* FKRig = Cast<UFKControlRig>(ControlRig))
 			{
@@ -2970,7 +2970,7 @@ bool UAnimSequencerController::AddCurveControl(const FName& CurveName) const
 {
 	if (UMovieSceneControlRigParameterSection* Section = Model->GetFKControlRigSection())
 	{
-		UControlRig* ControlRig = Section->GetControlRig();
+		UBaseControlRig* ControlRig = Section->GetControlRig();
 		if (UFKControlRig* FKRig = Cast<UFKControlRig>(ControlRig))
 		{
 			if(URigHierarchy* Hierarchy = FKRig->GetHierarchy())
@@ -3053,7 +3053,7 @@ bool UAnimSequencerController::RenameCurveControl(const FName& CurveName, const
 	// Need to rename curve and control element and curve				
 	if (UMovieSceneControlRigParameterSection* Section = Model->GetFKControlRigSection())
 	{
-		if (UControlRig* ControlRig = Section->GetControlRig())
+		if (UBaseControlRig* ControlRig = Section->GetControlRig())
 		{
 			if (URigHierarchy* RigHierarchy = ControlRig->GetHierarchy())
 			{
@@ -3140,7 +3140,7 @@ bool UAnimSequencerController::RemoveCurveControl(const FName& CurveName) const
 {
 	if (UMovieSceneControlRigParameterSection* Section = Model->GetFKControlRigSection())
 	{
-		UControlRig* ControlRig = Section->GetControlRig();
+		UBaseControlRig* ControlRig = Section->GetControlRig();
 
 		if (UFKControlRig* FKRig = Cast<UFKControlRig>(ControlRig))
 		{
@@ -3204,7 +3204,7 @@ bool UAnimSequencerController::SetCurveControlKeys(const FName& CurveName, const
 {
 	if (UMovieSceneControlRigParameterSection* Section = Model->GetFKControlRigSection())
 	{
-		UControlRig* ControlRig = Section->GetControlRig();
+		UBaseControlRig* ControlRig = Section->GetControlRig();
 
 		if (UFKControlRig* FKRig = Cast<UFKControlRig>(ControlRig))
 		{
@@ -3292,7 +3292,7 @@ bool UAnimSequencerController::SetCurveControlKey(const FName& CurveName, const
 {
 	if (UMovieSceneControlRigParameterSection* Section = Model->GetFKControlRigSection())
 	{
-		UControlRig* ControlRig = Section->GetControlRig();
+		UBaseControlRig* ControlRig = Section->GetControlRig();
 
 		if (UFKControlRig* FKRig = Cast<UFKControlRig>(ControlRig))
 		{
@@ -3381,7 +3381,7 @@ bool UAnimSequencerController::RemoveCurveControlKey(const FName& CurveName, flo
 {
 	if (UMovieSceneControlRigParameterSection* Section = Model->GetFKControlRigSection())
 	{
-		UControlRig* ControlRig = Section->GetControlRig();
+		UBaseControlRig* ControlRig = Section->GetControlRig();
 
 		if (UFKControlRig* FKRig = Cast<UFKControlRig>(ControlRig))
 		{
diff --git a/Engine/Plugins/Animation/AnimationData/Source/AnimationData/Private/AnimSequencerDataModel.cpp b/Engine/Plugins/Animation/AnimationData/Source/AnimationData/Private/AnimSequencerDataModel.cpp
--- a/Engine/Plugins/Animation/AnimationData/Source/AnimationData/Private/AnimSequencerDataModel.cpp
+++ b/Engine/Plugins/Animation/AnimationData/Source/AnimationData/Private/AnimSequencerDataModel.cpp
@@ -144,7 +144,7 @@ void UAnimationSequencerDataModel::InitializeFKControlRig(UFKControlRig* FKContr
 	}
 }
 
-UControlRig* UAnimationSequencerDataModel::GetControlRig() const
+UBaseControlRig* UAnimationSequencerDataModel::GetControlRig() const
 {
 	if(const UMovieSceneControlRigParameterTrack* Track = GetControlRigTrack())
 	{
@@ -196,7 +196,7 @@ void UAnimationSequencerDataModel::PostLoad()
 void UAnimationSequencerDataModel::DeclareConstructClasses(TArray<FTopLevelAssetPath>& OutConstructClasses, const UClass* SpecificSubclass)
 {
 	Super::DeclareConstructClasses(OutConstructClasses, SpecificSubclass);
-	OutConstructClasses.Add(FTopLevelAssetPath(UControlRig::StaticClass()));
+	OutConstructClasses.Add(FTopLevelAssetPath(UBaseControlRig::StaticClass()));
 }
 #endif
 
@@ -751,7 +751,7 @@ void UAnimationSequencerDataModel::Evaluate(FAnimationPoseData& InOutPoseData, c
 		EvaluateTrack(Track, EvaluationContext);
 
 		// Generate/populate the output animation pose data
-		UControlRig* ControlRig = Track->GetControlRig();
+		UBaseControlRig* ControlRig = Track->GetControlRig();
 		GeneratePoseData(ControlRig, InOutPoseData, EvaluationContext);
 	}
 }
@@ -848,7 +848,7 @@ UMovieSceneControlRigParameterSection* UAnimationSequencerDataModel::GetFKContro
 			{
 				if (UMovieSceneControlRigParameterSection* Section = Cast<UMovieSceneControlRigParameterSection>(TrackSection))
 				{
-					if (const UControlRig* ControlRig = Section->GetControlRig())
+					if (const UBaseControlRig* ControlRig = Section->GetControlRig())
 					{
 						if (ControlRig->IsA<UFKControlRig>())
 						{
@@ -874,7 +874,7 @@ void UAnimationSequencerDataModel::GenerateLegacyCurveData()
 		{
 			if (const UMovieSceneControlRigParameterSection* Section = Cast<UMovieSceneControlRigParameterSection>(TrackSection))
 			{
-				if (const UControlRig* ControlRig = Section->GetControlRig())
+				if (const UBaseControlRig* ControlRig = Section->GetControlRig())
 				{
 						if(URigHierarchy* Hierarchy = ControlRig->GetHierarchy())
 						{
@@ -974,7 +974,7 @@ void UAnimationSequencerDataModel::ValidateControlRigData() const
 	const UMovieSceneControlRigParameterSection* Section = GetFKControlRigSection();
 	checkf(Section, TEXT("Unable to find Control Rig Section"));
 
-	UControlRig* ControlRig = Section->GetControlRig();
+	UBaseControlRig* ControlRig = Section->GetControlRig();
 	checkf(ControlRig, TEXT("Unable to find Control Rig instance for Section"));
 
 	checkf(ControlRig->IsA<UFKControlRig>(), TEXT("Invalid class for Control Rig expected UFKControlRig"));
@@ -1018,7 +1018,7 @@ void UAnimationSequencerDataModel::ValidateLegacyAgainstControlRigData() const
 {
 	UMovieSceneControlRigParameterSection* Section = GetFKControlRigSection();
 
-	UControlRig* ControlRig = Section->GetControlRig();
+	UBaseControlRig* ControlRig = Section->GetControlRig();
 	URigHierarchy* Hierarchy = ControlRig->GetHierarchy();
 
 	// Validate bone tracks against controls
@@ -1083,7 +1083,7 @@ void UAnimationSequencerDataModel::IterateTransformControlCurve(const FName& Bon
 	ValidateControlRigData();
 
 	const UMovieSceneControlRigParameterSection* Section = GetFKControlRigSection();
-	UControlRig* ControlRig = Section->GetControlRig();
+	UBaseControlRig* ControlRig = Section->GetControlRig();
 	const URigHierarchy* Hierarchy = ControlRig->GetHierarchy();
 	
 	const FRigElementKey BoneControlKey(UFKControlRig::GetControlName(BoneName, ERigElementType::Bone), ERigElementType::Control);
@@ -1164,7 +1164,7 @@ UObject* UAnimationSequencerDataModel::GetParentObject(UObject* MovieSceneBlends
 	return GetOuter();
 }
 
-void UAnimationSequencerDataModel::GeneratePoseData(UControlRig* ControlRig, FAnimationPoseData& InOutPoseData, const UE::Anim::DataModel::FEvaluationContext& EvaluationContext) const
+void UAnimationSequencerDataModel::GeneratePoseData(UBaseControlRig* ControlRig, FAnimationPoseData& InOutPoseData, const UE::Anim::DataModel::FEvaluationContext& EvaluationContext) const
 {
 	QUICK_SCOPE_CYCLE_COUNTER(STAT_GeneratePoseData);
 	
@@ -1315,7 +1315,7 @@ void UAnimationSequencerDataModel::EvaluateTrack(UMovieSceneControlRigParameterT
 		bWasDoNotKey = FKRigSection->GetDoNotKey();
 		FKRigSection->SetDoNotKey(true);
 
-		UControlRig* ControlRig = FKRigSection->GetControlRig();
+		UBaseControlRig* ControlRig = FKRigSection->GetControlRig();
 		check(ControlRig);
 
 		// Reset to ref-pose
@@ -1562,7 +1562,7 @@ void UAnimationSequencerDataModel::IterateBoneKeys(const FName& BoneName, TFunct
 	ValidateControlRigData();
 
 	const UMovieSceneControlRigParameterSection* Section = GetFKControlRigSection();
-	UControlRig* ControlRig = Section->GetControlRig();
+	UBaseControlRig* ControlRig = Section->GetControlRig();
 	const URigHierarchy* Hierarchy = ControlRig->GetHierarchy();
 	
 	const FRigElementKey BoneControlKey(UFKControlRig::GetControlName(BoneName, ERigElementType::Bone), ERigElementType::Control);
diff --git a/Engine/Plugins/Animation/AnimationData/Source/AnimationData/Private/AnimSequencerDataModel.h b/Engine/Plugins/Animation/AnimationData/Source/AnimationData/Private/AnimSequencerDataModel.h
--- a/Engine/Plugins/Animation/AnimationData/Source/AnimationData/Private/AnimSequencerDataModel.h
+++ b/Engine/Plugins/Animation/AnimationData/Source/AnimationData/Private/AnimSequencerDataModel.h
@@ -9,7 +9,7 @@
 
 #include "AnimSequencerDataModel.generated.h"
 
-class UControlRig;
+class UBaseControlRig;
 class UMovieSceneControlRigParameterTrack;
 class UMovieSceneControlRigParameterSection;
 class UFKControlRig;
@@ -146,14 +146,14 @@ protected:
 	void ValidateControlRigData() const;
 	void ValidateLegacyAgainstControlRigData() const;
 	
-	void GeneratePoseData(UControlRig* ControlRig, FAnimationPoseData& InOutPoseData, const UE::Anim::DataModel::FEvaluationContext& EvaluationContext) const;
+	void GeneratePoseData(UBaseControlRig* ControlRig, FAnimationPoseData& InOutPoseData, const UE::Anim::DataModel::FEvaluationContext& EvaluationContext) const;
 	void EvaluateTrack(UMovieSceneControlRigParameterTrack* CR_Track, const UE::Anim::DataModel::FEvaluationContext& EvaluationContext) const;
 
 	UMovieSceneControlRigParameterTrack* GetControlRigTrack() const;
 	UMovieSceneControlRigParameterSection* GetFKControlRigSection() const;
 	USkeleton* GetSkeleton() const;
 	void InitializeFKControlRig(UFKControlRig* FKControlRig, USkeleton* Skeleton) const;
-	UControlRig* GetControlRig() const;
+	UBaseControlRig* GetControlRig() const;
 	
 	void IterateTransformControlCurve(const FName& BoneName, TFunction<void(const FTransform&, const FFrameNumber&)> IterationFunction, const TArray<FFrameNumber>* InFrameNumbers = nullptr) const;
 	void GenerateTransformKeysForControl(const FName& BoneName, TArray<FTransform>& InOutTransforms, TArray<FFrameNumber>& InOutFrameNumbers) const;
diff --git a/Engine/Plugins/Animation/ControlRig/Source/ControlRig/Private/AnimNode_ControlRig.cpp b/Engine/Plugins/Animation/ControlRig/Source/ControlRig/Private/AnimNode_ControlRig.cpp
--- a/Engine/Plugins/Animation/ControlRig/Source/ControlRig/Private/AnimNode_ControlRig.cpp
+++ b/Engine/Plugins/Animation/ControlRig/Source/ControlRig/Private/AnimNode_ControlRig.cpp
@@ -49,7 +49,7 @@ void FAnimNode_ControlRig::OnInitializeAnimInstance(const FAnimInstanceProxy* In
 
 	if (ControlRigClass)
 	{
-		ControlRig = NewObject<UControlRig>(InAnimInstance->GetOwningComponent(), ControlRigClass);
+		ControlRig = NewObject<UBaseControlRig>(InAnimInstance->GetOwningComponent(), ControlRigClass);
 		ControlRig->Initialize(true);
 		ControlRig->RequestInit();
 		RefPoseSetterHash.Reset();
@@ -183,7 +183,7 @@ void FAnimNode_ControlRig::CacheBones_AnyThread(const FAnimationCacheBonesContex
 		};
 
 		URigHierarchy* Hierarchy = nullptr;
-		if(UControlRig* CurrentControlRig = GetControlRig())
+		if(UBaseControlRig* CurrentControlRig = GetControlRig())
 		{
 			Hierarchy = CurrentControlRig->GetHierarchy();
 		}
@@ -207,7 +207,7 @@ void FAnimNode_ControlRig::PostSerialize(const FArchive& Ar)
 	DECLARE_SCOPE_HIERARCHICAL_COUNTER_FUNC()
 }
 
-void FAnimNode_ControlRig::UpdateInput(UControlRig* InControlRig, const FPoseContext& InOutput)
+void FAnimNode_ControlRig::UpdateInput(UBaseControlRig* InControlRig, const FPoseContext& InOutput)
 {
 	DECLARE_SCOPE_HIERARCHICAL_COUNTER_FUNC()
 
@@ -232,7 +232,7 @@ void FAnimNode_ControlRig::UpdateInput(UControlRig* InControlRig, const FPoseCon
 	}
 }
 
-void FAnimNode_ControlRig::UpdateOutput(UControlRig* InControlRig, FPoseContext& InOutput)
+void FAnimNode_ControlRig::UpdateOutput(UBaseControlRig* InControlRig, FPoseContext& InOutput)
 {
 	DECLARE_SCOPE_HIERARCHICAL_COUNTER_FUNC()
 
@@ -317,7 +317,7 @@ void FAnimNode_ControlRig::SetIOMapping(bool bInput, const FName& SourceProperty
 	UClass* TargetClass = GetTargetClass();
 	if (TargetClass)
 	{
-		UControlRig* CDO = TargetClass->GetDefaultObject<UControlRig>();
+		UBaseControlRig* CDO = TargetClass->GetDefaultObject<UBaseControlRig>();
 		if (CDO)
 		{
 			TMap<FName, FName>& MappingData = (bInput) ? InputMapping : OutputMapping;
@@ -375,7 +375,7 @@ void FAnimNode_ControlRig::PropagateInputProperties(const UObject* InSourceInsta
 {
 	if (TargetInstance)
 	{
-		UControlRig* TargetControlRig = Cast<UControlRig>((UObject*)TargetInstance);
+		UBaseControlRig* TargetControlRig = Cast<UBaseControlRig>((UObject*)TargetInstance);
 		if(TargetControlRig == nullptr)
 		{
 			return;
diff --git a/Engine/Plugins/Animation/ControlRig/Source/ControlRig/Private/AnimNode_ControlRigBase.cpp b/Engine/Plugins/Animation/ControlRig/Source/ControlRig/Private/AnimNode_ControlRigBase.cpp
--- a/Engine/Plugins/Animation/ControlRig/Source/ControlRig/Private/AnimNode_ControlRigBase.cpp
+++ b/Engine/Plugins/Animation/ControlRig/Source/ControlRig/Private/AnimNode_ControlRigBase.cpp
@@ -52,7 +52,7 @@ void FAnimNode_ControlRigBase::OnInitializeAnimInstance(const FAnimInstanceProxy
 	WeakAnimInstanceObject = TWeakObjectPtr<const UAnimInstance>(InAnimInstance);
 
 	USkeletalMeshComponent* Component = InAnimInstance->GetOwningComponent();
-	UControlRig* ControlRig = GetControlRig();
+	UBaseControlRig* ControlRig = GetControlRig();
 	if (Component && Component->GetSkeletalMeshAsset() && ControlRig)
 	{
 #if WITH_EDITORONLY_DATA
@@ -66,7 +66,7 @@ void FAnimNode_ControlRigBase::OnInitializeAnimInstance(const FAnimInstanceProxy
 #endif
 
 		// register skeletalmesh component for now
-		ControlRig->GetDataSourceRegistry()->RegisterDataSource(UControlRig::OwnerComponent, InAnimInstance->GetOwningComponent());
+		ControlRig->GetDataSourceRegistry()->RegisterDataSource(UBaseControlRig::OwnerComponent, InAnimInstance->GetOwningComponent());
 		UpdateGetAssetUserDataDelegate(ControlRig);
 	}
 }
@@ -93,7 +93,7 @@ void FAnimNode_ControlRigBase::Update_AnyThread(const FAnimationUpdateContext& C
 
 	if (bExecute)
 	{
-		if (UControlRig* ControlRig = GetControlRig())
+		if (UBaseControlRig* ControlRig = GetControlRig())
 		{
 			// @TODO: fix this to be thread-safe
 			// Pre-update doesn't work for custom anim instances
@@ -110,7 +110,7 @@ bool FAnimNode_ControlRigBase::CanExecute()
 		return false;
 	}
 
-	if (UControlRig* ControlRig = GetControlRig())
+	if (UBaseControlRig* ControlRig = GetControlRig())
 	{
 		return ControlRig->CanExecute(); 
 	}
@@ -118,7 +118,7 @@ bool FAnimNode_ControlRigBase::CanExecute()
 	return false;
 }
 
-void FAnimNode_ControlRigBase::UpdateInput(UControlRig* ControlRig, const FPoseContext& InOutput)
+void FAnimNode_ControlRigBase::UpdateInput(UBaseControlRig* ControlRig, const FPoseContext& InOutput)
 {
 	SCOPE_CYCLE_COUNTER(STAT_ControlRig_UpdateInput);
 
@@ -161,7 +161,7 @@ void FAnimNode_ControlRigBase::UpdateInput(UControlRig* ControlRig, const FPoseC
 			{
 #if WITH_EDITOR
 				// make sure transient controls don't get reset
-				UControlRig::FTransientControlPoseScope PoseScope(ControlRig);
+				UBaseControlRig::FTransientControlPoseScope PoseScope(ControlRig);
 #endif 
 				Hierarchy->ResetPoseToInitial(ERigElementType::Bone);
 			}
@@ -259,7 +259,7 @@ void FAnimNode_ControlRigBase::UpdateInput(UControlRig* ControlRig, const FPoseC
 #endif
 }
 
-void FAnimNode_ControlRigBase::UpdateOutput(UControlRig* ControlRig, FPoseContext& InOutput)
+void FAnimNode_ControlRigBase::UpdateOutput(UBaseControlRig* ControlRig, FPoseContext& InOutput)
 {
 	SCOPE_CYCLE_COUNTER(STAT_ControlRig_UpdateOutput);
 
@@ -440,11 +440,11 @@ void FAnimNode_ControlRigBase::ExecuteControlRig(FPoseContext& InOutput)
 {
 	SCOPE_CYCLE_COUNTER(STAT_ControlRig_Evaluate);
 
-	if (UControlRig* ControlRig = GetControlRig())
+	if (UBaseControlRig* ControlRig = GetControlRig())
 	{
 		// temporarily give control rig access to the stack allocated attribute container
 		// control rig may have rig units that can add/get attributes to/from this container
-		UControlRig::FAnimAttributeContainerPtrScope AttributeScope(ControlRig, InOutput.CustomAttributes);
+		UBaseControlRig::FAnimAttributeContainerPtrScope AttributeScope(ControlRig, InOutput.CustomAttributes);
 		
 		// first update input to the system
 		UpdateInput(ControlRig, InOutput);
@@ -519,7 +519,7 @@ void FAnimNode_ControlRigBase::ExecuteControlRig(FPoseContext& InOutput)
 
 struct FControlRigControlScope
 {
-	FControlRigControlScope(UControlRig* InControlRig)
+	FControlRigControlScope(UBaseControlRig* InControlRig)
 		: ControlRig(InControlRig)
 	{
 		if (ControlRig.IsValid())
@@ -546,7 +546,7 @@ struct FControlRigControlScope
 	}
 
 	TMap<FRigElementKey, FRigControlValue> ControlValues;
-	TWeakObjectPtr<UControlRig> ControlRig;
+	TWeakObjectPtr<UBaseControlRig> ControlRig;
 };
 
 void FAnimNode_ControlRigBase::CacheBones_AnyThread(const FAnimationCacheBonesContext& Context)
@@ -556,7 +556,7 @@ void FAnimNode_ControlRigBase::CacheBones_AnyThread(const FAnimationCacheBonesCo
 	FAnimNode_CustomProperty::CacheBones_AnyThread(Context);
 	Source.CacheBones(Context);
 
-	if (UControlRig* ControlRig = GetControlRig())
+	if (UBaseControlRig* ControlRig = GetControlRig())
 	{
 		URigHierarchy* Hierarchy = ControlRig->GetHierarchy();
 		if(Hierarchy == nullptr)
@@ -635,7 +635,7 @@ void FAnimNode_ControlRigBase::CacheBones_AnyThread(const FAnimationCacheBonesCo
 				FBoneContainer& InRequiredBones,
 				const FReferenceSkeleton& InRefSkeleton,
 				const TArray<FBoneIndexType>& InRequiredBonesArray,
-				const UControlRig* InControlRig,
+				const UBaseControlRig* InControlRig,
 				TMap<FName, uint16>& OutMapping
 			) {
 				OutMapping.Reset();
@@ -761,7 +761,7 @@ void FAnimNode_ControlRigBase::CacheBones_AnyThread(const FAnimationCacheBonesCo
 
 UClass* FAnimNode_ControlRigBase::GetTargetClass() const
 {
-	if (UControlRig* ControlRig = GetControlRig())
+	if (UBaseControlRig* ControlRig = GetControlRig())
 	{
 		return ControlRig->GetClass();
 	}
@@ -769,7 +769,7 @@ UClass* FAnimNode_ControlRigBase::GetTargetClass() const
 	return nullptr;
 }
 
-void FAnimNode_ControlRigBase::QueueControlRigDrawInstructions(UControlRig* ControlRig, FAnimInstanceProxy* Proxy) const
+void FAnimNode_ControlRigBase::QueueControlRigDrawInstructions(UBaseControlRig* ControlRig, FAnimInstanceProxy* Proxy) const
 {
 	ensure(ControlRig);
 	ensure(Proxy);
@@ -825,7 +825,7 @@ void FAnimNode_ControlRigBase::QueueControlRigDrawInstructions(UControlRig* Cont
 	}
 }
 
-void FAnimNode_ControlRigBase::UpdateGetAssetUserDataDelegate(UControlRig* InControlRig) const
+void FAnimNode_ControlRigBase::UpdateGetAssetUserDataDelegate(UBaseControlRig* InControlRig) const
 {
 	if(GetAssetUserData().IsEmpty() || !WeakAnimInstanceObject.IsValid())
 	{
@@ -837,7 +837,7 @@ void FAnimNode_ControlRigBase::UpdateGetAssetUserDataDelegate(UControlRig* InCon
 	// since the delegate may go stale quickly. to guard against destroyed anim nodes
 	// we'll rely on the anim instance to provide an indication if the memory is still valid. 
 	TWeakObjectPtr<const UAnimInstance> LocalWeakAnimInstance = WeakAnimInstanceObject;
-	InControlRig->GetExternalAssetUserDataDelegate = UControlRig::FGetExternalAssetUserData::CreateLambda([InControlRig, LocalWeakAnimInstance, this]
+	InControlRig->GetExternalAssetUserDataDelegate = UBaseControlRig::FGetExternalAssetUserData::CreateLambda([InControlRig, LocalWeakAnimInstance, this]
 	{
 		if(LocalWeakAnimInstance.IsValid())
 		{
diff --git a/Engine/Plugins/Animation/ControlRig/Source/ControlRig/Private/AnimNode_ControlRig_ExternalSource.cpp b/Engine/Plugins/Animation/ControlRig/Source/ControlRig/Private/AnimNode_ControlRig_ExternalSource.cpp
--- a/Engine/Plugins/Animation/ControlRig/Source/ControlRig/Private/AnimNode_ControlRig_ExternalSource.cpp
+++ b/Engine/Plugins/Animation/ControlRig/Source/ControlRig/Private/AnimNode_ControlRig_ExternalSource.cpp
@@ -10,20 +10,20 @@ FAnimNode_ControlRig_ExternalSource::FAnimNode_ControlRig_ExternalSource()
 	bControlRigRequiresInitialization = false;
 }
 
-void FAnimNode_ControlRig_ExternalSource::SetControlRig(UControlRig* InControlRig)
+void FAnimNode_ControlRig_ExternalSource::SetControlRig(UBaseControlRig* InControlRig)
 {
 	ControlRig = InControlRig;
 	// requires initializing animation system
 }
 
-UControlRig* FAnimNode_ControlRig_ExternalSource::GetControlRig() const
+UBaseControlRig* FAnimNode_ControlRig_ExternalSource::GetControlRig() const
 {
 	return (ControlRig.IsValid()? ControlRig.Get() : nullptr);
 }
 
-TSubclassOf<UControlRig> FAnimNode_ControlRig_ExternalSource::GetControlRigClass() const
+TSubclassOf<UBaseControlRig> FAnimNode_ControlRig_ExternalSource::GetControlRigClass() const
 {
-	if(UControlRig* CR = GetControlRig())
+	if(UBaseControlRig* CR = GetControlRig())
 	{
 		return CR->GetClass();
 	}
diff --git a/Engine/Plugins/Animation/ControlRig/Source/ControlRig/Private/BaseControlRig.cpp b/Engine/Plugins/Animation/ControlRig/Source/ControlRig/Private/BaseControlRig.cpp
--- a/Engine/Plugins/Animation/ControlRig/Source/ControlRig/Private/BaseControlRig.cpp
+++ b/Engine/Plugins/Animation/ControlRig/Source/ControlRig/Private/BaseControlRig.cpp
@@ -47,7 +47,7 @@ DECLARE_STATS_GROUP(TEXT("ControlRig"), STATGROUP_ControlRig, STATCAT_Advanced);
 DECLARE_CYCLE_STAT_EXTERN(TEXT("Control Rig Execution"), STAT_RigExecution, STATGROUP_ControlRig, );
 DEFINE_STAT(STAT_RigExecution);
 
-const FName UControlRig::OwnerComponent("OwnerComponent");
+const FName UBaseControlRig::OwnerComponent("OwnerComponent");
 
 //CVar to specify if we should create a float control for each curve in the curve container
 //By default we don't but it may be useful to do so for debugging
@@ -66,7 +66,7 @@ static TAutoConsoleVariable<float> CVarControlRigEnableDrawInterfaceInGame(
 	ECVF_Default);
 
 
-UControlRig::UControlRig(const FObjectInitializer& ObjectInitializer)
+UBaseControlRig::UBaseControlRig(const FObjectInitializer& ObjectInitializer)
 	: Super(ObjectInitializer)
 #if WITH_EDITOR
 	, bEnableAnimAttributeTrace(false)
@@ -96,7 +96,7 @@ UControlRig::UControlRig(const FObjectInitializer& ObjectInitializer)
 	SetRigVMExtendedExecuteContext(&RigVMExtendedExecuteContext);
 }
 
-void UControlRig::BeginDestroy()
+void UBaseControlRig::BeginDestroy()
 {
 	Super::BeginDestroy();
 	SetRigVMExtendedExecuteContext(nullptr);
@@ -110,7 +110,7 @@ void UControlRig::BeginDestroy()
 #if WITH_EDITOR
 	if (!HasAnyFlags(RF_ClassDefaultObject))
 	{
-		if(UControlRig* CDO = GetClass()->GetDefaultObject<UControlRig>())
+		if(UBaseControlRig* CDO = GetClass()->GetDefaultObject<UBaseControlRig>())
 		{
 			if (!CDO->HasAnyFlags(RF_BeginDestroyed))
 			{
@@ -126,7 +126,7 @@ void UControlRig::BeginDestroy()
 	TRACE_OBJECT_LIFETIME_END(this);
 }
 
-UWorld* UControlRig::GetWorld() const
+UWorld* UBaseControlRig::GetWorld() const
 {
 	if (ObjectBinding.IsValid())
 	{
@@ -145,14 +145,14 @@ UWorld* UControlRig::GetWorld() const
 	return Super::GetWorld();
 }
 
-void UControlRig::Initialize(bool bRequestInit)
+void UBaseControlRig::Initialize(bool bRequestInit)
 {
 	TRACE_OBJECT_LIFETIME_BEGIN(this);
 
 	DECLARE_SCOPE_HIERARCHICAL_COUNTER_FUNC()
 	QUICK_SCOPE_CYCLE_COUNTER(STAT_ControlRig_Initialize);
 
-	Super::Initialize(bRequestInit);
+	InitializeVMs(bRequestInit);
 
 	// Create the data source registry here to avoid UObject creation from Non-Game Threads
 	GetDataSourceRegistry();
@@ -164,13 +164,13 @@ void UControlRig::Initialize(bool bRequestInit)
 	RequestConstruction();
 	
 	GetHierarchy()->OnModified().RemoveAll(this);
-	GetHierarchy()->OnModified().AddUObject(this, &UControlRig::HandleHierarchyModified);
+	GetHierarchy()->OnModified().AddUObject(this, &UBaseControlRig::HandleHierarchyModified);
 	GetHierarchy()->OnEventReceived().RemoveAll(this);
-	GetHierarchy()->OnEventReceived().AddUObject(this, &UControlRig::HandleHierarchyEvent);
+	GetHierarchy()->OnEventReceived().AddUObject(this, &UBaseControlRig::HandleHierarchyEvent);
 	GetHierarchy()->UpdateVisibilityOnProxyControls();
 }
 
-void UControlRig::OnAddShapeLibrary(const FControlRigExecuteContext* InContext, const FString& InLibraryName, UControlRigShapeLibrary* InShapeLibrary, bool bReplaceExisting, bool bLogResults)
+void UBaseControlRig::OnAddShapeLibrary(const FControlRigExecuteContext* InContext, const FString& InLibraryName, UControlRigShapeLibrary* InShapeLibrary, bool bReplaceExisting, bool bLogResults)
 {
 	// don't ever change the CDO
 	if (HasAnyFlags(RF_ClassDefaultObject))
@@ -195,7 +195,7 @@ void UControlRig::OnAddShapeLibrary(const FControlRigExecuteContext* InContext,
 	// if we've removed all shape libraries - let's add the ones from the CDO back
 	if (ShapeLibraries.IsEmpty())
 	{
-		UControlRig* CDO = GetClass()->GetDefaultObject<UControlRig>();
+		UBaseControlRig* CDO = GetClass()->GetDefaultObject<UBaseControlRig>();
 		ShapeLibraries = CDO->ShapeLibraries;
 	}
 
@@ -285,7 +285,7 @@ void UControlRig::OnAddShapeLibrary(const FControlRigExecuteContext* InContext,
 #endif
 }
 
-bool UControlRig::OnShapeExists(const FName& InShapeName) const
+bool UBaseControlRig::OnShapeExists(const FName& InShapeName) const
 {
 	const TArray<TSoftObjectPtr<UControlRigShapeLibrary>>& Libraries = GetShapeLibraries();
 	if (UControlRigShapeLibrary::GetShapeByName(InShapeName, GetShapeLibraries(), ShapeLibraryNameMap))
@@ -296,9 +296,9 @@ bool UControlRig::OnShapeExists(const FName& InShapeName) const
 	return false;
 }
 
-bool UControlRig::InitializeVM(const FName& InEventName)
+bool UBaseControlRig::InitializeVM(const FName& InEventName)
 {
-	if(!Super::InitializeVM(InEventName))
+	if(!InitializeVMs(InEventName))
 	{
 		return false;
 	}
@@ -327,13 +327,13 @@ bool UControlRig::InitializeVM(const FName& InEventName)
 	return true;
 }
 
-void UControlRig::Evaluate_AnyThread()
+void UBaseControlRig::Evaluate_AnyThread()
 {
 	if (bIsAdditive)
 	{
 		if (bRequiresInitExecution)
 		{
-			Super::Evaluate_AnyThread();
+			EvaluateVMs_AnyThread();
 		}
 		
 		// we can have other systems trying to poke into running instances of Control Rigs
@@ -427,11 +427,11 @@ void UControlRig::Evaluate_AnyThread()
 	}
 	else
 	{
-		Super::Evaluate_AnyThread();
+		EvaluateVMs_AnyThread();
 	}
 }
 
-bool UControlRig::EvaluateSkeletalMeshComponent(double InDeltaTime)
+bool UBaseControlRig::EvaluateSkeletalMeshComponent(double InDeltaTime)
 {
 	if (USkeletalMeshComponent* SkelMeshComp = Cast<USkeletalMeshComponent>(GetObjectBinding()->GetBoundObject()))
 	{
@@ -448,17 +448,17 @@ bool UControlRig::EvaluateSkeletalMeshComponent(double InDeltaTime)
 	return false;
 }
 
-void UControlRig::ResetControlValues()
+void UBaseControlRig::ResetControlValues()
 {
 	ControlValues.Reset();
 }
 
-void UControlRig::ClearPoseBeforeBackwardsSolve()
+void UBaseControlRig::ClearPoseBeforeBackwardsSolve()
 {
 	PoseBeforeBackwardsSolve.Reset();
 }
 
-TArray<FRigControlElement*> UControlRig::InvertInputPose(const TArray<FRigElementKey>& InElements, EControlRigSetKey InSetKey)
+TArray<FRigControlElement*> UBaseControlRig::InvertInputPose(const TArray<FRigElementKey>& InElements, EControlRigSetKey InSetKey)
 {
 	TArray<FRigControlElement*> ModifiedElements;
 	ModifiedElements.Reserve(ControlsAfterBackwardsSolve.Num());
@@ -484,9 +484,9 @@ TArray<FRigControlElement*> UControlRig::InvertInputPose(const TArray<FRigElemen
 	return ModifiedElements;
 }
 
-void UControlRig::InitializeFromCDO()
+void UBaseControlRig::InitializeFromCDO()
 {
-	Super::InitializeFromCDO();
+	InitializeVMsFromCDO();
 	
 	DECLARE_SCOPE_HIERARCHICAL_COUNTER_FUNC()
 
@@ -496,7 +496,7 @@ void UControlRig::InitializeFromCDO()
 		// similar to FControlRigBlueprintCompilerContext::CopyTermDefaultsToDefaultObject,
 		// where CDO is initialized from BP there,
 		// we initialize all other instances of Control Rig from the CDO here
-		UControlRig* CDO = GetClass()->GetDefaultObject<UControlRig>();
+		UBaseControlRig* CDO = GetClass()->GetDefaultObject<UBaseControlRig>();
 		URigHierarchy* Hierarchy = GetHierarchy();
 
 		// copy hierarchy
@@ -531,38 +531,38 @@ void UControlRig::InitializeFromCDO()
 	}
 }
 
-UControlRig::FAnimAttributeContainerPtrScope::FAnimAttributeContainerPtrScope(UControlRig* InControlRig,
+UBaseControlRig::FAnimAttributeContainerPtrScope::FAnimAttributeContainerPtrScope(UBaseControlRig* InControlRig,
 	UE::Anim::FStackAttributeContainer& InExternalContainer)
 {
 	ControlRig = InControlRig;
 	ControlRig->ExternalAnimAttributeContainer = &InExternalContainer;
 }
 
-UControlRig::FAnimAttributeContainerPtrScope::~FAnimAttributeContainerPtrScope()
+UBaseControlRig::FAnimAttributeContainerPtrScope::~FAnimAttributeContainerPtrScope()
 {
 	// control rig should not hold on to this container since it is stack allocated
 	// and should not be used outside of stack, see FPoseContext
 	ControlRig->ExternalAnimAttributeContainer = nullptr;
 }
 
-AActor* UControlRig::GetHostingActor() const
+AActor* UBaseControlRig::GetHostingActor() const
 {
 	return ObjectBinding ? ObjectBinding->GetHostingActor() : nullptr;
 }
 
 #if WITH_EDITOR
-FText UControlRig::GetCategory() const
+FText UBaseControlRig::GetCategory() const
 {
 	return LOCTEXT("DefaultControlRigCategory", "Animation|ControlRigs");
 }
 
-FText UControlRig::GetToolTipText() const
+FText UBaseControlRig::GetToolTipText() const
 {
 	return LOCTEXT("DefaultControlRigTooltip", "ControlRig");
 }
 #endif
 
-bool UControlRig::AllConnectorsAreResolved(FString* OutFailureReason, FRigElementKey* OutConnector) const
+bool UBaseControlRig::AllConnectorsAreResolved(FString* OutFailureReason, FRigElementKey* OutConnector) const
 {
 	if(const URigHierarchy* Hierarchy = GetHierarchy())
 	{
@@ -620,7 +620,7 @@ bool UControlRig::AllConnectorsAreResolved(FString* OutFailureReason, FRigElemen
 	return true;
 }
 
-bool UControlRig::Execute(const FName& InEventName)
+bool UBaseControlRig::Execute(const FName& InEventName)
 {
 	if(!CanExecute())
 	{
@@ -679,14 +679,14 @@ bool UControlRig::Execute(const FName& InEventName)
 		// only set a valid first entry event later when execution
 		// has passed the initialization stage and there are multiple events present in one evaluation
 		// first entry event is used to determined when to clear data during an evaluation
-		VM->SetFirstEntryEventInEventQueue(ExtendedExecuteContext, NAME_None);
+		SetFirstEntryEventInEventQueue(ExtendedExecuteContext, NAME_None);
 #endif
 	}
 
 #if WITH_EDITOR
 	if (IsInDebugMode())
 	{
-		if (UControlRig* CDO = GetClass()->GetDefaultObject<UControlRig>())
+		if (UBaseControlRig* CDO = GetClass()->GetDefaultObject<UBaseControlRig>())
 		{
 			// Copy the breakpoints. This will not override the state of the breakpoints
 			DebugInfo.SetBreakpoints(CDO->DebugInfo.GetBreakpoints());
@@ -740,7 +740,7 @@ bool UControlRig::Execute(const FName& InEventName)
 	// draw container contains persistent draw instructions, 
 	// so we cannot call Reset(), which will clear them,
 	// instead, we re-initialize them from the CDO
-	if (UControlRig* CDO = GetClass()->GetDefaultObject<UControlRig>())
+	if (UBaseControlRig* CDO = GetClass()->GetDefaultObject<UBaseControlRig>())
 	{
 		DrawContainer = CDO->DrawContainer;
 	}
@@ -759,13 +759,13 @@ bool UControlRig::Execute(const FName& InEventName)
 	check(PublicContext.Hierarchy);
 
 	// allow access to the shape libraries
-	PublicContext.OnAddShapeLibraryDelegate.BindUObject(this, &UControlRig::OnAddShapeLibrary);
-	PublicContext.OnShapeExistsDelegate.BindUObject(this, &UControlRig::OnShapeExists);
+	PublicContext.OnAddShapeLibraryDelegate.BindUObject(this, &UBaseControlRig::OnAddShapeLibrary);
+	PublicContext.OnShapeExistsDelegate.BindUObject(this, &UBaseControlRig::OnShapeExists);
 
 	// allow access to the default hierarchy to allow to reset
 	if(!HasAnyFlags(RF_ClassDefaultObject))
 	{
-		if (UControlRig* CDO = Cast<UControlRig>(GetClass()->GetDefaultObject()))
+		if (UBaseControlRig* CDO = Cast<UBaseControlRig>(GetClass()->GetDefaultObject()))
 		{
 			if(URigHierarchy* DefaultHierarchy = CDO->GetHierarchy())
 			{
@@ -915,7 +915,7 @@ bool UControlRig::Execute(const FName& InEventName)
 					// clone the shape libraries again from the CDO 
 					if (!HasAnyFlags(RF_ClassDefaultObject) && ShapeLibraries.IsEmpty())
 					{
-						UControlRig* CDO = GetClass()->GetDefaultObject<UControlRig>();
+						UBaseControlRig* CDO = GetClass()->GetDefaultObject<UBaseControlRig>();
 						ShapeLibraries = CDO->ShapeLibraries;
 					}
 
@@ -982,7 +982,7 @@ bool UControlRig::Execute(const FName& InEventName)
 		// has passed the initialization stage and there are multiple events present
 		if (EventQueueToRun.Num() >= 2 && VM)
 		{
-			VM->SetFirstEntryEventInEventQueue(ExtendedExecuteContext, EventQueueToRun[0]);
+			SetFirstEntryEventInEventQueue(ExtendedExecuteContext, EventQueueToRun[0]);
 		}
 
 		// Transform Overrride is generated using a Transient Control 
@@ -1322,149 +1322,28 @@ bool UControlRig::Execute(const FName& InEventName)
 	return bSuccess;
 }
 
-bool UControlRig::Execute_Internal(const FName& InEventName)
+void UBaseControlRig::RequestInit()
 {
-	if(IsRigModule())
-	{
-		FString ConnectorWarning;
-		if(!AllConnectorsAreResolved(&ConnectorWarning))
-		{
-#if WITH_EDITOR
-			LogOnce(EMessageSeverity::Warning, INDEX_NONE, ConnectorWarning);
-#endif
-			return false;
-		}
-	}
-	
-	if (VM)
-	{
-		FRigVMExtendedExecuteContext& Context = GetRigVMExtendedExecuteContext();
-
-		static constexpr TCHAR InvalidatedVMFormat[] = TEXT("%s: Invalidated VM - aborting execution.");
-		if(VM->IsNativized())
-		{
-			if(!IsValidLowLevel() ||
-				!VM->IsValidLowLevel())
-			{
-				UE_LOG(LogControlRig, Warning, InvalidatedVMFormat, *GetClass()->GetName());
-				return false;
-			}
-		}
-		else
-		{
-			// sanity check the validity of the VM to ensure stability.
-			if(!VM->IsContextValidForExecution(Context)
-				|| !IsValidLowLevel()
-				|| !VM->IsValidLowLevel()
-			)
-			{
-				UE_LOG(LogControlRig, Warning, InvalidatedVMFormat, *GetClass()->GetName());
-				return false;
-			}
-		}
-		
-#if UE_RIGVM_PROFILE_EXECUTE_UNITS_NUM
-		const uint64 StartCycles = FPlatformTime::Cycles64();
-		if(ProfilingRunsLeft <= 0)
-		{
-			ProfilingRunsLeft = UE_RIGVM_PROFILE_EXECUTE_UNITS_NUM;
-			AccumulatedCycles = 0;
-		}
-#endif
-		
-		const bool bUseDebuggingSnapshots = !VM->IsNativized();
-		
-#if WITH_EDITOR
-		if(bUseDebuggingSnapshots)
-		{
-			if(URigVM* SnapShotVM = GetSnapshotVM(false)) // don't create it for normal runs
-			{
-				const bool bIsEventFirstInQueue = !EventQueueToRun.IsEmpty() && EventQueueToRun[0] == InEventName; 
-				const bool bIsEventLastInQueue = !EventQueueToRun.IsEmpty() && EventQueueToRun.Last() == InEventName;
-
-				if (GetHaltedAtBreakpoint().IsValid())
-				{
-					if(bIsEventFirstInQueue)
-					{
-						CopyVMMemory(GetRigVMExtendedExecuteContext(), GetSnapshotContext());
-					}
-				}
-				else if(bIsEventLastInQueue)
-				{
-					CopyVMMemory(GetSnapshotContext(), GetRigVMExtendedExecuteContext());
-				}
-			}
-		}
-#endif
-
-		URigHierarchy* Hierarchy = GetHierarchy();
-#if WITH_EDITOR
-
-		bool bRecordTransformsAtRuntime = true;
-		if(const UObject* Outer = GetOuter())
-		{
-			if(Outer->IsA<UControlRigComponent>())
-			{
-				bRecordTransformsAtRuntime = false;
-			}
-		}
-		TGuardValue<bool> RecordTransformsPerInstructionGuard(Hierarchy->bRecordTransformsAtRuntime, bRecordTransformsAtRuntime);
-		
-		if(Hierarchy->bRecordTransformsAtRuntime)
-		{
-			Hierarchy->ReadTransformsAtRuntime.Reset();
-			Hierarchy->WrittenTransformsAtRuntime.Reset();
-		}
-		
-#endif
-		FRigHierarchyExecuteContextBracket HierarchyContextGuard(Hierarchy, &Context);
-
-		// setup the module information
-		FControlRigExecuteContext& PublicContext = Context.GetPublicDataSafe<FControlRigExecuteContext>();
-		FControlRigExecuteContextRigModuleGuard RigModuleGuard(PublicContext, this);
-		FRigHierarchyRedirectorGuard ElementRedirectorGuard(this);
-
-		const bool bSuccess = VM->ExecuteVM(Context, InEventName) != ERigVMExecuteResult::Failed;
-
-#if UE_RIGVM_PROFILE_EXECUTE_UNITS_NUM
-		const uint64 EndCycles = FPlatformTime::Cycles64();
-		const uint64 Cycles = EndCycles - StartCycles;
-		AccumulatedCycles += Cycles;
-		ProfilingRunsLeft--;
-		if(ProfilingRunsLeft == 0)
-		{
-			const double Milliseconds = FPlatformTime::ToMilliseconds64(AccumulatedCycles);
-			UE_LOG(LogControlRig, Display, TEXT("%s: %d runs took %.03lfms."), *GetClass()->GetName(), UE_RIGVM_PROFILE_EXECUTE_UNITS_NUM, Milliseconds);
-		}
-#endif
-
-		return bSuccess;
-	}
-	return false;
-}
-
-void UControlRig::RequestInit()
-{
-	Super::RequestInit();
+	RequestInitVMs();
 	RequestConstruction();
 }
 
-void UControlRig::RequestConstruction()
+void UBaseControlRig::RequestConstruction()
 {
 	RequestRunOnceEvent(FRigUnit_PrepareForExecution::EventName, 0);
 }
 
-bool UControlRig::IsConstructionRequired() const
+bool UBaseControlRig::IsConstructionRequired() const
 {
 	return IsRunOnceEvent(FRigUnit_PrepareForExecution::EventName);
 }
 
-bool UControlRig::SupportsBackwardsSolve() const
+bool UBaseControlRig::SupportsBackwardsSolve() const
 {
 	return SupportsEvent(FRigUnit_InverseExecution::EventName);
 }
 
-void UControlRig::AdaptEventQueueForEvaluate(TArray<FName>& InOutEventQueueToRun)
+void UBaseControlRig::AdaptEventQueueForEvaluate(TArray<FName>& InOutEventQueueToRun)
 {
 	Super::AdaptEventQueueForEvaluate(InOutEventQueueToRun);
 
@@ -1487,7 +1366,7 @@ void UControlRig::AdaptEventQueueForEvaluate(TArray<FName>& InOutEventQueueToRun
 	}
 }
 
-void UControlRig::GetMappableNodeData(TArray<FName>& OutNames, TArray<FNodeItem>& OutNodeItems) const
+void UBaseControlRig::GetMappableNodeData(TArray<FName>& OutNames, TArray<FNodeItem>& OutNodeItems) const
 {
 	DECLARE_SCOPE_HIERARCHICAL_COUNTER_FUNC()
 
@@ -1512,7 +1391,7 @@ void UControlRig::GetMappableNodeData(TArray<FName>& OutNames, TArray<FNodeItem>
 	});
 }
 
-UAnimationDataSourceRegistry* UControlRig::GetDataSourceRegistry()
+UAnimationDataSourceRegistry* UBaseControlRig::GetDataSourceRegistry()
 {
 	if (DataSourceRegistry)
 	{
@@ -1536,7 +1415,7 @@ UAnimationDataSourceRegistry* UControlRig::GetDataSourceRegistry()
 
 #if WITH_EDITORONLY_DATA
 
-void UControlRig::PostReinstanceCallback(const UControlRig* Old)
+void UBaseControlRig::PostReinstanceCallback(const UBaseControlRig* Old)
 {
 	DECLARE_SCOPE_HIERARCHICAL_COUNTER_FUNC()
 
@@ -1546,9 +1425,9 @@ void UControlRig::PostReinstanceCallback(const UControlRig* Old)
 
 #endif // WITH_EDITORONLY_DATA
 
-TArray<UControlRig*> UControlRig::FindControlRigs(UObject* Outer, TSubclassOf<UControlRig> OptionalClass)
+TArray<UBaseControlRig*> UBaseControlRig::FindControlRigs(UObject* Outer, TSubclassOf<UBaseControlRig> OptionalClass)
 {
-	TArray<UControlRig*> Result;
+	TArray<UBaseControlRig*> Result;
 	
 	if(Outer == nullptr)
 	{
@@ -1561,9 +1440,9 @@ TArray<UControlRig*> UControlRig::FindControlRigs(UObject* Outer, TSubclassOf<UC
 		OuterActor = Outer->GetTypedOuter<AActor>();
 	}
 	
-	for (TObjectIterator<UControlRig> Itr; Itr; ++Itr)
+	for (TObjectIterator<UBaseControlRig> Itr; Itr; ++Itr)
 	{
-		UControlRig* RigInstance = *Itr;
+		UBaseControlRig* RigInstance = *Itr;
 		const UClass* RigInstanceClass = RigInstance ? RigInstance->GetClass() : nullptr;
 		if (OptionalClass == nullptr || (RigInstanceClass && RigInstanceClass->IsChildOf(OptionalClass)))
 		{
@@ -1599,7 +1478,7 @@ TArray<UControlRig*> UControlRig::FindControlRigs(UObject* Outer, TSubclassOf<UC
 	return Result;
 }
 
-void UControlRig::Serialize(FArchive& Ar)
+void UBaseControlRig::Serialize(FArchive& Ar)
 {
 	DECLARE_SCOPE_HIERARCHICAL_COUNTER_FUNC()
 
@@ -1608,7 +1487,7 @@ void UControlRig::Serialize(FArchive& Ar)
 	Ar.UsingCustomVersion(FControlRigObjectVersion::GUID);
 }
 
-void UControlRig::PostLoad()
+void UBaseControlRig::PostLoad()
 {
 	Super::PostLoad();
 
@@ -1642,25 +1521,25 @@ void UControlRig::PostLoad()
 	Influences = NewInfluences;
 }
 
-const FRigModuleSettings& UControlRig::GetRigModuleSettings() const
+const FRigModuleSettings& UBaseControlRig::GetRigModuleSettings() const
 {
 	if(HasAnyFlags(RF_ClassDefaultObject))
 	{
 		return RigModuleSettings;
 	}
-	if (const UControlRig* CDO = Cast<UControlRig>(GetClass()->GetDefaultObject()))
+	if (const UBaseControlRig* CDO = Cast<UBaseControlRig>(GetClass()->GetDefaultObject()))
 	{
 		return CDO->GetRigModuleSettings();
 	}
 	return RigModuleSettings;
 }
 
-bool UControlRig::IsRigModule() const
+bool UBaseControlRig::IsRigModule() const
 {
 	return GetRigModuleSettings().IsValidModule(false);
 }
 
-bool UControlRig::IsRigModuleInstance() const
+bool UBaseControlRig::IsRigModuleInstance() const
 {
 	if(IsRigModule())
 	{
@@ -1669,16 +1548,16 @@ bool UControlRig::IsRigModuleInstance() const
 	return false;
 }
 
-UControlRig* UControlRig::GetParentRig() const
+UBaseControlRig* UBaseControlRig::GetParentRig() const
 {
-	return GetTypedOuter<UControlRig>();
+	return GetTypedOuter<UBaseControlRig>();
 }
 
-const FString& UControlRig::GetRigModuleNameSpace() const
+const FString& UBaseControlRig::GetRigModuleNameSpace() const
 {
 	if(IsRigModule())
 	{
-		if(const UControlRig* ParentRig = GetParentRig())
+		if(const UBaseControlRig* ParentRig = GetParentRig())
 		{
 			const FString& ParentNameSpace = ParentRig->GetRigModuleNameSpace();
 			static constexpr TCHAR JoinFormat[] = TEXT("%s%s:");
@@ -1697,7 +1576,7 @@ const FString& UControlRig::GetRigModuleNameSpace() const
 	return EmptyNameSpace;
 }
 
-FRigElementKeyRedirector& UControlRig::GetElementKeyRedirector()
+FRigElementKeyRedirector& UBaseControlRig::GetElementKeyRedirector()
 {
 	// if we are an instance on a modular rig, use our local info
 	if(IsRigModuleInstance())
@@ -1715,7 +1594,7 @@ FRigElementKeyRedirector& UControlRig::GetElementKeyRedirector()
 	return EmptyRedirector;
 }
 
-TArray<FRigControlElement*> UControlRig::AvailableControls() const
+TArray<FRigControlElement*> UBaseControlRig::AvailableControls() const
 {
 	if(DynamicHierarchy)
 	{
@@ -1724,7 +1603,7 @@ TArray<FRigControlElement*> UControlRig::AvailableControls() const
 	return TArray<FRigControlElement*>();
 }
 
-FRigControlElement* UControlRig::FindControl(const FName& InControlName) const
+FRigControlElement* UBaseControlRig::FindControl(const FName& InControlName) const
 {
 	if(DynamicHierarchy == nullptr)
 	{
@@ -1733,12 +1612,12 @@ FRigControlElement* UControlRig::FindControl(const FName& InControlName) const
 	return DynamicHierarchy->Find<FRigControlElement>(FRigElementKey(InControlName, ERigElementType::Control));
 }
 
-bool UControlRig::IsConstructionModeEnabled() const
+bool UBaseControlRig::IsConstructionModeEnabled() const
 {
 	return EventQueueToRun.Num() == 1 && EventQueueToRun.Contains(FRigUnit_PrepareForExecution::EventName);
 }
 
-FTransform UControlRig::SetupControlFromGlobalTransform(const FName& InControlName, const FTransform& InGlobalTransform)
+FTransform UBaseControlRig::SetupControlFromGlobalTransform(const FName& InControlName, const FTransform& InGlobalTransform)
 {
 	if (IsConstructionModeEnabled())
 	{
@@ -1763,7 +1642,7 @@ FTransform UControlRig::SetupControlFromGlobalTransform(const FName& InControlNa
 	return InGlobalTransform;
 }
 
-void UControlRig::CreateRigControlsForCurveContainer()
+void UBaseControlRig::CreateRigControlsForCurveContainer()
 {
 	const bool bCreateFloatControls = CVarControlRigCreateFloatControlsForCurves->GetInt() == 0 ? false : true;
 	if(bCreateFloatControls && DynamicHierarchy)
@@ -1796,11 +1675,11 @@ void UControlRig::CreateRigControlsForCurveContainer()
 			return true;
 		});
 
-		ControlModified().AddUObject(this, &UControlRig::HandleOnControlModified);
+		ControlModified().AddUObject(this, &UBaseControlRig::HandleOnControlModified);
 	}
 }
 
-void UControlRig::HandleOnControlModified(UControlRig* Subject, FRigControlElement* Control, const FRigControlModifiedContext& Context)
+void UBaseControlRig::HandleOnControlModified(UBaseControlRig* Subject, FRigControlElement* Control, const FRigControlModifiedContext& Context)
 {
 	if (Control->Settings.bIsCurve && DynamicHierarchy)
 	{
@@ -1809,12 +1688,12 @@ void UControlRig::HandleOnControlModified(UControlRig* Subject, FRigControlEleme
 	}	
 }
 
-bool UControlRig::IsCurveControl(const FRigControlElement* InControlElement) const
+bool UBaseControlRig::IsCurveControl(const FRigControlElement* InControlElement) const
 {
 	return InControlElement->Settings.bIsCurve;
 }
 
-FTransform UControlRig::GetControlGlobalTransform(const FName& InControlName) const
+FTransform UBaseControlRig::GetControlGlobalTransform(const FName& InControlName) const
 {
 	if(DynamicHierarchy == nullptr)
 	{
@@ -1823,7 +1702,7 @@ FTransform UControlRig::GetControlGlobalTransform(const FName& InControlName) co
 	return DynamicHierarchy->GetGlobalTransform(FRigElementKey(InControlName, ERigElementType::Control), false);
 }
 
-FRigControlValue UControlRig::GetControlValue(FRigControlElement* InControl, const ERigControlValueType& InValueType)
+FRigControlValue UBaseControlRig::GetControlValue(FRigControlElement* InControl, const ERigControlValueType& InValueType)
 {
 	if (bIsAdditive && InValueType == ERigControlValueType::Current)
 	{
@@ -1849,7 +1728,7 @@ FRigControlValue UControlRig::GetControlValue(FRigControlElement* InControl, con
 	return GetHierarchy()->GetControlValue(InControl, InValueType);
 }
 
-void UControlRig::SetControlValueImpl(const FName& InControlName, const FRigControlValue& InValue, bool bNotify,
+void UBaseControlRig::SetControlValueImpl(const FName& InControlName, const FRigControlValue& InValue, bool bNotify,
 	const FRigControlModifiedContext& Context, bool bSetupUndo, bool bPrintPythonCommnds, bool bFixEulerFlips)
 {
 	const FRigElementKey Key(InControlName, ERigElementType::Control);
@@ -1876,7 +1755,7 @@ void UControlRig::SetControlValueImpl(const FName& InControlName, const FRigCont
 	}
 }
 
-bool UControlRig::SetControlGlobalTransform(const FName& InControlName, const FTransform& InGlobalTransform, bool bNotify, const FRigControlModifiedContext& Context, bool bSetupUndo, bool bPrintPythonCommands, bool bFixEulerFlips)
+bool UBaseControlRig::SetControlGlobalTransform(const FName& InControlName, const FTransform& InGlobalTransform, bool bNotify, const FRigControlModifiedContext& Context, bool bSetupUndo, bool bPrintPythonCommands, bool bFixEulerFlips)
 {
 	FTransform GlobalTransform = InGlobalTransform;
 	ERigTransformType::Type TransformType = ERigTransformType::CurrentGlobal;
@@ -1921,7 +1800,7 @@ bool UControlRig::SetControlGlobalTransform(const FName& InControlName, const FT
 	return true;
 }
 
-FRigControlValue UControlRig::GetControlValueFromGlobalTransform(const FName& InControlName, const FTransform& InGlobalTransform, ERigTransformType::Type InTransformType)
+FRigControlValue UBaseControlRig::GetControlValueFromGlobalTransform(const FName& InControlName, const FTransform& InGlobalTransform, ERigTransformType::Type InTransformType)
 {
 	FRigControlValue Value;
 
@@ -1964,7 +1843,7 @@ FRigControlValue UControlRig::GetControlValueFromGlobalTransform(const FName& In
 	return Value;
 }
 
-void UControlRig::SetControlLocalTransform(const FName& InControlName, const FTransform& InLocalTransform, bool bNotify, const FRigControlModifiedContext& Context, bool bSetupUndo, bool bFixEulerFlips)
+void UBaseControlRig::SetControlLocalTransform(const FName& InControlName, const FTransform& InLocalTransform, bool bNotify, const FRigControlModifiedContext& Context, bool bSetupUndo, bool bFixEulerFlips)
 {
 	if (FRigControlElement* ControlElement = FindControl(InControlName))
 	{
@@ -1980,7 +1859,7 @@ void UControlRig::SetControlLocalTransform(const FName& InControlName, const FTr
 	}
 }
 
-FTransform UControlRig::GetControlLocalTransform(const FName& InControlName)
+FTransform UBaseControlRig::GetControlLocalTransform(const FName& InControlName)
 {
 	if(DynamicHierarchy == nullptr)
 	{
@@ -2000,13 +1879,13 @@ FTransform UControlRig::GetControlLocalTransform(const FName& InControlName)
 	return DynamicHierarchy->GetLocalTransform(FRigElementKey(InControlName, ERigElementType::Control));
 }
 
-const TArray<TSoftObjectPtr<UControlRigShapeLibrary>>& UControlRig::GetShapeLibraries() const
+const TArray<TSoftObjectPtr<UControlRigShapeLibrary>>& UBaseControlRig::GetShapeLibraries() const
 {
 	const TArray<TSoftObjectPtr<UControlRigShapeLibrary>>* LibrariesPtr = &ShapeLibraries;
 
 	if(!GetClass()->IsNative() && ShapeLibraries.IsEmpty())
 	{
-		if (UControlRig* CDO = Cast<UControlRig>(GetClass()->GetDefaultObject()))
+		if (UBaseControlRig* CDO = Cast<UBaseControlRig>(GetClass()->GetDefaultObject()))
 		{
 			LibrariesPtr = &CDO->ShapeLibraries;
 		}
@@ -2024,7 +1903,7 @@ const TArray<TSoftObjectPtr<UControlRigShapeLibrary>>& UControlRig::GetShapeLibr
 	return Libraries;
 }
 
-void UControlRig::SelectControl(const FName& InControlName, bool bSelect)
+void UBaseControlRig::SelectControl(const FName& InControlName, bool bSelect)
 {
 	if(DynamicHierarchy)
 	{
@@ -2035,7 +1914,7 @@ void UControlRig::SelectControl(const FName& InControlName, bool bSelect)
 	}
 }
 
-bool UControlRig::ClearControlSelection()
+bool UBaseControlRig::ClearControlSelection()
 {
 	if(DynamicHierarchy)
 	{
@@ -2047,7 +1926,7 @@ bool UControlRig::ClearControlSelection()
 	return false;
 }
 
-TArray<FName> UControlRig::CurrentControlSelection() const
+TArray<FName> UBaseControlRig::CurrentControlSelection() const
 {
 	TArray<FName> SelectedControlNames;
 
@@ -2068,7 +1947,7 @@ TArray<FName> UControlRig::CurrentControlSelection() const
 	return SelectedControlNames;
 }
 
-bool UControlRig::IsControlSelected(const FName& InControlName)const
+bool UBaseControlRig::IsControlSelected(const FName& InControlName)const
 {
 	if(DynamicHierarchy)
 	{
@@ -2080,7 +1959,7 @@ bool UControlRig::IsControlSelected(const FName& InControlName)const
 	return false;
 }
 
-void UControlRig::HandleHierarchyModified(ERigHierarchyNotification InNotification, URigHierarchy* InHierarchy,
+void UBaseControlRig::HandleHierarchyModified(ERigHierarchyNotification InNotification, URigHierarchy* InHierarchy,
     const FRigBaseElement* InElement)
 {
 	switch(InNotification)
@@ -2138,7 +2017,7 @@ void UControlRig::HandleHierarchyModified(ERigHierarchyNotification InNotificati
 
 #if WITH_EDITOR
 
-bool UControlRig::CanAddTransientControl(const URigVMUnitNode* InNode, const FRigDirectManipulationTarget& InTarget, FString* OutFailureReason)
+bool UBaseControlRig::CanAddTransientControl(const URigVMUnitNode* InNode, const FRigDirectManipulationTarget& InTarget, FString* OutFailureReason)
 {
 	if (InNode == nullptr)
 	{
@@ -2192,7 +2071,7 @@ bool UControlRig::CanAddTransientControl(const URigVMUnitNode* InNode, const FRi
 	return false;
 }
 
-FName UControlRig::AddTransientControl(const URigVMUnitNode* InNode, const FRigDirectManipulationTarget& InTarget)
+FName UBaseControlRig::AddTransientControl(const URigVMUnitNode* InNode, const FRigDirectManipulationTarget& InTarget)
 {
 	if(!CanAddTransientControl(InNode, InTarget, nullptr))
 	{
@@ -2248,7 +2127,7 @@ FName UControlRig::AddTransientControl(const URigVMUnitNode* InNode, const FRigD
 	return Info->ControlKey.Name;
 }
 
-bool UControlRig::SetTransientControlValue(const URigVMUnitNode* InNode, TSharedPtr<FRigDirectManipulationInfo> InInfo)
+bool UBaseControlRig::SetTransientControlValue(const URigVMUnitNode* InNode, TSharedPtr<FRigDirectManipulationInfo> InInfo)
 {
 	check(InNode);
 	check(DynamicHierarchy);
@@ -2276,7 +2155,7 @@ bool UControlRig::SetTransientControlValue(const URigVMUnitNode* InNode, TShared
 	return bResult;
 }
 
-FName UControlRig::RemoveTransientControl(const URigVMUnitNode* InNode, const FRigDirectManipulationTarget& InTarget)
+FName UBaseControlRig::RemoveTransientControl(const URigVMUnitNode* InNode, const FRigDirectManipulationTarget& InTarget)
 {
 	if ((InNode == nullptr) || (DynamicHierarchy == nullptr))
 	{
@@ -2310,7 +2189,7 @@ FName UControlRig::RemoveTransientControl(const URigVMUnitNode* InNode, const FR
 	return NAME_None;
 }
 
-FName UControlRig::AddTransientControl(const FRigElementKey& InElement)
+FName UBaseControlRig::AddTransientControl(const FRigElementKey& InElement)
 {
 	if (!InElement.IsValid())
 	{
@@ -2385,7 +2264,7 @@ FName UControlRig::AddTransientControl(const FRigElementKey& InElement)
 			if(FRigBoneElement* BoneElement = DynamicHierarchy->Find<FRigBoneElement>(InElement))
 			{
 				// add a modify bone AnimNode internally that the transient control controls for imported bones only
-				// for user created bones, refer to UControlRig::TransformOverrideForUserCreatedBones 
+				// for user created bones, refer to UBaseControlRig::TransformOverrideForUserCreatedBones 
 				if (BoneElement->BoneType == ERigBoneType::Imported)
 				{ 
 					if (PreviewInstance)
@@ -2408,7 +2287,7 @@ FName UControlRig::AddTransientControl(const FRigElementKey& InElement)
 	return ControlKey.Name;
 }
 
-bool UControlRig::SetTransientControlValue(const FRigElementKey& InElement)
+bool UBaseControlRig::SetTransientControlValue(const FRigElementKey& InElement)
 {
 	if (!InElement.IsValid())
 	{
@@ -2477,7 +2356,7 @@ bool UControlRig::SetTransientControlValue(const FRigElementKey& InElement)
 	return false;
 }
 
-FName UControlRig::RemoveTransientControl(const FRigElementKey& InElement)
+FName UBaseControlRig::RemoveTransientControl(const FRigElementKey& InElement)
 {
 	if (!InElement.IsValid())
 	{
@@ -2511,7 +2390,7 @@ FName UControlRig::RemoveTransientControl(const FRigElementKey& InElement)
 	return NAME_None;
 }
 
-FName UControlRig::GetNameForTransientControl(const URigVMUnitNode* InNode, const FRigDirectManipulationTarget& InTarget) const
+FName UBaseControlRig::GetNameForTransientControl(const URigVMUnitNode* InNode, const FRigDirectManipulationTarget& InTarget) const
 {
 	check(InNode);
 	check(DynamicHierarchy);
@@ -2520,7 +2399,7 @@ FName UControlRig::GetNameForTransientControl(const URigVMUnitNode* InNode, cons
 	return DynamicHierarchy->GetSanitizedName(FRigName(FString::Printf(TEXT("ControlForNode|%s|%s"), *NodeName, *InTarget.Name)));
 }
 
-FString UControlRig::GetNodeNameFromTransientControl(const FRigElementKey& InKey)
+FString UBaseControlRig::GetNodeNameFromTransientControl(const FRigElementKey& InKey)
 {
 	FString Name = InKey.Name.ToString();
 	if(Name.StartsWith(TEXT("ControlForNode|")))
@@ -2535,7 +2414,7 @@ FString UControlRig::GetNodeNameFromTransientControl(const FRigElementKey& InKey
 	return Name;
 }
 
-FString UControlRig::GetTargetFromTransientControl(const FRigElementKey& InKey)
+FString UBaseControlRig::GetTargetFromTransientControl(const FRigElementKey& InKey)
 {
 	FString Name = InKey.Name.ToString();
 	if(Name.StartsWith(TEXT("ControlForNode|")))
@@ -2550,7 +2429,7 @@ FString UControlRig::GetTargetFromTransientControl(const FRigElementKey& InKey)
 	return Name;
 }
 
-TSharedPtr<FRigDirectManipulationInfo> UControlRig::GetRigUnitManipulationInfoForTransientControl(
+TSharedPtr<FRigDirectManipulationInfo> UBaseControlRig::GetRigUnitManipulationInfoForTransientControl(
 	const FRigElementKey& InKey)
 {
 	const TSharedPtr<FRigDirectManipulationInfo>* InfoPtr = RigUnitManipulationInfos.FindByPredicate(
@@ -2567,7 +2446,7 @@ TSharedPtr<FRigDirectManipulationInfo> UControlRig::GetRigUnitManipulationInfoFo
 	return TSharedPtr<FRigDirectManipulationInfo>();
 }
 
-FName UControlRig::GetNameForTransientControl(const FRigElementKey& InElement)
+FName UBaseControlRig::GetNameForTransientControl(const FRigElementKey& InElement)
 {
 	if (InElement.Type == ERigElementType::Control)
 	{
@@ -2578,7 +2457,7 @@ FName UControlRig::GetNameForTransientControl(const FRigElementKey& InElement)
 	return *FString::Printf(TEXT("ControlForRigElement_%s_%s"), *EnumName.ToString(), *InElement.Name.ToString());
 }
 
-FRigElementKey UControlRig::GetElementKeyFromTransientControl(const FRigElementKey& InKey)
+FRigElementKey UBaseControlRig::GetElementKeyFromTransientControl(const FRigElementKey& InKey)
 {
 	if(InKey.Type != ERigElementType::Control)
 	{
@@ -2611,7 +2490,7 @@ FRigElementKey UControlRig::GetElementKeyFromTransientControl(const FRigElementK
 	return FRigElementKey();;
 }
 
-void UControlRig::ClearTransientControls()
+void UBaseControlRig::ClearTransientControls()
 {
 	if(DynamicHierarchy == nullptr)
 	{
@@ -2643,7 +2522,7 @@ void UControlRig::ClearTransientControls()
 	}
 }
 
-void UControlRig::ApplyTransformOverrideForUserCreatedBones()
+void UBaseControlRig::ApplyTransformOverrideForUserCreatedBones()
 {
 	if(DynamicHierarchy == nullptr)
 	{
@@ -2656,7 +2535,7 @@ void UControlRig::ApplyTransformOverrideForUserCreatedBones()
 	}
 }
 
-void UControlRig::ApplySelectionPoseForConstructionMode(const FName& InEventName)
+void UBaseControlRig::ApplySelectionPoseForConstructionMode(const FName& InEventName)
 {
 	FRigControlModifiedContext ControlValueContext;
 	ControlValueContext.EventName = InEventName;
@@ -2672,7 +2551,7 @@ void UControlRig::ApplySelectionPoseForConstructionMode(const FName& InEventName
 
 #endif
 
-void UControlRig::HandleHierarchyEvent(URigHierarchy* InHierarchy, const FRigEventContext& InEvent)
+void UBaseControlRig::HandleHierarchyEvent(URigHierarchy* InHierarchy, const FRigEventContext& InEvent)
 {
 	if (RigEventDelegate.IsBound())
 	{
@@ -2712,7 +2591,7 @@ void UControlRig::HandleHierarchyEvent(URigHierarchy* InHierarchy, const FRigEve
 	}
 }
 
-void UControlRig::GetControlsInOrder(TArray<FRigControlElement*>& SortedControls) const
+void UBaseControlRig::GetControlsInOrder(TArray<FRigControlElement*>& SortedControls) const
 {
 	SortedControls.Reset();
 
@@ -2724,9 +2603,9 @@ void UControlRig::GetControlsInOrder(TArray<FRigControlElement*>& SortedControls
 	SortedControls = DynamicHierarchy->GetControls(true);
 }
 
-const FRigInfluenceMap* UControlRig::FindInfluenceMap(const FName& InEventName)
+const FRigInfluenceMap* UBaseControlRig::FindInfluenceMap(const FName& InEventName)
 {
-	if (UControlRig* CDO = GetClass()->GetDefaultObject<UControlRig>())
+	if (UBaseControlRig* CDO = GetClass()->GetDefaultObject<UBaseControlRig>())
 	{
 		return CDO->Influences.Find(InEventName);
 	}
@@ -2735,7 +2614,7 @@ const FRigInfluenceMap* UControlRig::FindInfluenceMap(const FName& InEventName)
 
 #if WITH_EDITOR
 
-void UControlRig::PreEditChange(FProperty* PropertyAboutToChange)
+void UBaseControlRig::PreEditChange(FProperty* PropertyAboutToChange)
 {
 	// for BP user authored properties let's ignore changes since they
 	// will be distributed from the BP anyway to all archetype instances.
@@ -2760,13 +2639,13 @@ void UControlRig::PreEditChange(FProperty* PropertyAboutToChange)
 	Super::PreEditChange(PropertyAboutToChange);
 }
 
-void UControlRig::PostEditChangeProperty(FPropertyChangedEvent& PropertyChangedEvent)
+void UBaseControlRig::PostEditChangeProperty(FPropertyChangedEvent& PropertyChangedEvent)
 {
 	Super::PostEditChangeProperty(PropertyChangedEvent);
 }
 #endif
 
-FRigUnit* UControlRig::GetRigUnitInstanceFromScope(TSharedPtr<FStructOnScope> InScope)
+FRigUnit* UBaseControlRig::GetRigUnitInstanceFromScope(TSharedPtr<FStructOnScope> InScope)
 {
 	if(InScope.IsValid())
 	{
@@ -2782,7 +2661,7 @@ FRigUnit* UControlRig::GetRigUnitInstanceFromScope(TSharedPtr<FStructOnScope> In
 	return (FRigUnit*)DefaultRigUnitInstance.GetStructMemory();
 }
 
-const TArray<UAssetUserData*>* UControlRig::GetAssetUserDataArray() const
+const TArray<UAssetUserData*>* UBaseControlRig::GetAssetUserDataArray() const
 {
 	if(HasAnyFlags(RF_ClassDefaultObject))
 	{
@@ -2799,7 +2678,7 @@ const TArray<UAssetUserData*>* UControlRig::GetAssetUserDataArray() const
 
 	CombinedAssetUserData.Reset();
 
-	if (UControlRig* CDO = Cast<UControlRig>(GetClass()->GetDefaultObject(false)))
+	if (UBaseControlRig* CDO = Cast<UBaseControlRig>(GetClass()->GetDefaultObject(false)))
 	{
 		CombinedAssetUserData.Append(*CDO->GetAssetUserDataArray());
 	}
@@ -2892,7 +2771,7 @@ const TArray<UAssetUserData*>* UControlRig::GetAssetUserDataArray() const
 	return &ToRawPtrTArrayUnsafe(CombinedAssetUserData);
 }
 
-void UControlRig::CopyPoseFromOtherRig(UControlRig* Subject)
+void UBaseControlRig::CopyPoseFromOtherRig(UBaseControlRig* Subject)
 {
 	check(DynamicHierarchy);
 	check(Subject);
@@ -2927,7 +2806,7 @@ void UControlRig::CopyPoseFromOtherRig(UControlRig* Subject)
 	}
 }
 
-void UControlRig::SetBoneInitialTransformsFromAnimInstance(UAnimInstance* InAnimInstance)
+void UBaseControlRig::SetBoneInitialTransformsFromAnimInstance(UAnimInstance* InAnimInstance)
 {
 	FMemMark Mark(FMemStack::Get());
 	FCompactPose OutPose;
@@ -2935,7 +2814,7 @@ void UControlRig::SetBoneInitialTransformsFromAnimInstance(UAnimInstance* InAnim
 	SetBoneInitialTransformsFromCompactPose(&OutPose);
 }
 
-void UControlRig::SetBoneInitialTransformsFromAnimInstanceProxy(const FAnimInstanceProxy* InAnimInstanceProxy)
+void UBaseControlRig::SetBoneInitialTransformsFromAnimInstanceProxy(const FAnimInstanceProxy* InAnimInstanceProxy)
 {
 	FMemMark Mark(FMemStack::Get());
 	FCompactPose OutPose;
@@ -2943,7 +2822,7 @@ void UControlRig::SetBoneInitialTransformsFromAnimInstanceProxy(const FAnimInsta
 	SetBoneInitialTransformsFromCompactPose(&OutPose);
 }
 
-void UControlRig::SetBoneInitialTransformsFromSkeletalMeshComponent(USkeletalMeshComponent* InSkelMeshComp, bool bUseAnimInstance)
+void UBaseControlRig::SetBoneInitialTransformsFromSkeletalMeshComponent(USkeletalMeshComponent* InSkelMeshComp, bool bUseAnimInstance)
 {
 	check(InSkelMeshComp);
 	check(DynamicHierarchy);
@@ -2959,7 +2838,7 @@ void UControlRig::SetBoneInitialTransformsFromSkeletalMeshComponent(USkeletalMes
 }
 
 
-void UControlRig::SetBoneInitialTransformsFromSkeletalMesh(USkeletalMesh* InSkeletalMesh)
+void UBaseControlRig::SetBoneInitialTransformsFromSkeletalMesh(USkeletalMesh* InSkeletalMesh)
 {
 	if (ensure(InSkeletalMesh))
 	{ 
@@ -2967,7 +2846,7 @@ void UControlRig::SetBoneInitialTransformsFromSkeletalMesh(USkeletalMesh* InSkel
 	}
 }
 
-void UControlRig::SetBoneInitialTransformsFromRefSkeleton(const FReferenceSkeleton& InReferenceSkeleton)
+void UBaseControlRig::SetBoneInitialTransformsFromRefSkeleton(const FReferenceSkeleton& InReferenceSkeleton)
 {
 	check(DynamicHierarchy);
 
@@ -2989,7 +2868,7 @@ void UControlRig::SetBoneInitialTransformsFromRefSkeleton(const FReferenceSkelet
 	RequestConstruction();
 }
 
-void UControlRig::SetBoneInitialTransformsFromCompactPose(FCompactPose* InCompactPose)
+void UBaseControlRig::SetBoneInitialTransformsFromCompactPose(FCompactPose* InCompactPose)
 {
 	check(InCompactPose);
 
@@ -3026,7 +2905,7 @@ void UControlRig::SetBoneInitialTransformsFromCompactPose(FCompactPose* InCompac
 	RequestConstruction();
 }
 
-const FRigControlElementCustomization* UControlRig::GetControlCustomization(const FRigElementKey& InControl) const
+const FRigControlElementCustomization* UBaseControlRig::GetControlCustomization(const FRigElementKey& InControl) const
 {
 	check(InControl.Type == ERigElementType::Control);
 
@@ -3046,14 +2925,14 @@ const FRigControlElementCustomization* UControlRig::GetControlCustomization(cons
 	return nullptr;
 }
 
-void UControlRig::SetControlCustomization(const FRigElementKey& InControl, const FRigControlElementCustomization& InCustomization)
+void UBaseControlRig::SetControlCustomization(const FRigElementKey& InControl, const FRigControlElementCustomization& InCustomization)
 {
 	check(InControl.Type == ERigElementType::Control);
 	
 	ControlCustomizations.FindOrAdd(InControl) = InCustomization;
 }
 
-void UControlRig::PostInitInstanceIfRequired()
+void UBaseControlRig::PostInitInstanceIfRequired()
 {
 	if(GetHierarchy() == nullptr || VM == nullptr)
 	{
@@ -3063,7 +2942,7 @@ void UControlRig::PostInitInstanceIfRequired()
 		}
 		else
 		{
-			UControlRig* CDO = GetClass()->GetDefaultObject<UControlRig>();
+			UBaseControlRig* CDO = GetClass()->GetDefaultObject<UBaseControlRig>();
 			ensure(VM == nullptr || VM == CDO->VM);
 			PostInitInstance(CDO);
 		}
@@ -3071,7 +2950,7 @@ void UControlRig::PostInitInstanceIfRequired()
 }
 
 #if WITH_EDITORONLY_DATA
-void UControlRig::DeclareConstructClasses(TArray<FTopLevelAssetPath>& OutConstructClasses, const UClass* SpecificSubclass)
+void UBaseControlRig::DeclareConstructClasses(TArray<FTopLevelAssetPath>& OutConstructClasses, const UClass* SpecificSubclass)
 {
 	Super::DeclareConstructClasses(OutConstructClasses, SpecificSubclass);
 	OutConstructClasses.Add(FTopLevelAssetPath(URigHierarchy::StaticClass()));
@@ -3080,14 +2959,14 @@ void UControlRig::DeclareConstructClasses(TArray<FTopLevelAssetPath>& OutConstru
 
 #endif
 
-USceneComponent* UControlRig::GetOwningSceneComponent()
+USceneComponent* UBaseControlRig::GetOwningSceneComponent()
 {
 	if(OuterSceneComponent == nullptr)
 	{
 		const FControlRigExecuteContext& PublicContext = GetRigVMExtendedExecuteContext().GetPublicDataSafe<FControlRigExecuteContext>();
 		const FRigUnitContext& Context = PublicContext.UnitContext;
 
-		USceneComponent* SceneComponentFromRegistry = Context.DataSourceRegistry->RequestSource<USceneComponent>(UControlRig::OwnerComponent);
+		USceneComponent* SceneComponentFromRegistry = Context.DataSourceRegistry->RequestSource<USceneComponent>(UBaseControlRig::OwnerComponent);
 		if (SceneComponentFromRegistry)
 		{
 			OuterSceneComponent = SceneComponentFromRegistry;
@@ -3101,7 +2980,7 @@ USceneComponent* UControlRig::GetOwningSceneComponent()
 	return OuterSceneComponent.Get();
 }
 
-void UControlRig::PostInitInstance(URigVMHost* InCDO)
+void UBaseControlRig::PostInitInstance(URigVMHost* InCDO)
 {
 	const EObjectFlags SubObjectFlags =
 	HasAnyFlags(RF_ClassDefaultObject) ?
@@ -3113,7 +2992,7 @@ void UControlRig::PostInitInstance(URigVMHost* InCDO)
 	Context.SetContextPublicDataStruct(FControlRigExecuteContext::StaticStruct());
 
 	Context.ExecutionReachedExit().RemoveAll(this);
-	Context.ExecutionReachedExit().AddUObject(this, &UControlRig::HandleExecutionReachedExit);
+	Context.ExecutionReachedExit().AddUObject(this, &UBaseControlRig::HandleExecutionReachedExit);
 	UpdateVMSettings();
 
 	// set up the hierarchy
@@ -3135,8 +3014,8 @@ void UControlRig::PostInitInstance(URigVMHost* InCDO)
 	}
 
 #if WITH_EDITOR
-		const TWeakObjectPtr<UControlRig> WeakThis = this;
-		DynamicHierarchy->OnUndoRedo().AddStatic(&UControlRig::OnHierarchyTransformUndoRedoWeak, WeakThis);
+		const TWeakObjectPtr<UBaseControlRig> WeakThis = this;
+		DynamicHierarchy->OnUndoRedo().AddStatic(&UBaseControlRig::OnHierarchyTransformUndoRedoWeak, WeakThis);
 #endif
 
 	if(!HasAnyFlags(RF_ClassDefaultObject))
@@ -3149,7 +3028,7 @@ void UControlRig::PostInitInstance(URigVMHost* InCDO)
 				VM = InCDO->GetVM();
 			}
 
-			DynamicHierarchy->CopyHierarchy(CastChecked<UControlRig>(InCDO)->GetHierarchy());
+			DynamicHierarchy->CopyHierarchy(CastChecked<UBaseControlRig>(InCDO)->GetHierarchy());
 		}
 	}
 	else // we are the CDO
@@ -3176,7 +3055,7 @@ void UControlRig::PostInitInstance(URigVMHost* InCDO)
 		}
 	}
 
-	if(UControlRig* CDOControlRig = Cast<UControlRig>(InCDO))
+	if(UBaseControlRig* CDOControlRig = Cast<UBaseControlRig>(InCDO))
 	{
 		ElementKeyRedirector = FRigElementKeyRedirector(CDOControlRig->ElementKeyRedirector, DynamicHierarchy);
 	}
@@ -3184,7 +3063,7 @@ void UControlRig::PostInitInstance(URigVMHost* InCDO)
 	RequestInit();
 }
 
-UTransformableControlHandle* UControlRig::CreateTransformableControlHandle(
+UTransformableControlHandle* UBaseControlRig::CreateTransformableControlHandle(
 	const FName& InControlName) const
 {
 	auto IsConstrainable = [this](const FName& InControlName)
@@ -3220,7 +3099,7 @@ UTransformableControlHandle* UControlRig::CreateTransformableControlHandle(
 	return CtrlHandle;
 }
 
-void UControlRig::OnHierarchyTransformUndoRedo(URigHierarchy* InHierarchy, const FRigElementKey& InKey, ERigTransformType::Type InTransformType, const FTransform& InTransform, bool bIsUndo)
+void UBaseControlRig::OnHierarchyTransformUndoRedo(URigHierarchy* InHierarchy, const FRigElementKey& InKey, ERigTransformType::Type InTransformType, const FTransform& InTransform, bool bIsUndo)
 {
 	if(InKey.Type == ERigElementType::Control)
 	{
@@ -3231,7 +3110,7 @@ void UControlRig::OnHierarchyTransformUndoRedo(URigHierarchy* InHierarchy, const
 	}
 }
 
-UControlRig::FPoseScope::FPoseScope(UControlRig* InControlRig, ERigElementType InFilter, const TArray<FRigElementKey>& InElements, const ERigTransformType::Type InTransformType)
+UBaseControlRig::FPoseScope::FPoseScope(UBaseControlRig* InControlRig, ERigElementType InFilter, const TArray<FRigElementKey>& InElements, const ERigTransformType::Type InTransformType)
 : ControlRig(InControlRig)
 , Filter(InFilter)
 , TransformType(InTransformType)
@@ -3241,7 +3120,7 @@ UControlRig::FPoseScope::FPoseScope(UControlRig* InControlRig, ERigElementType I
 	CachedPose = InControlRig->GetHierarchy()->GetPose(IsInitial(InTransformType), InFilter, ElementView);
 }
 
-UControlRig::FPoseScope::~FPoseScope()
+UBaseControlRig::FPoseScope::~FPoseScope()
 {
 	check(ControlRig);
 
@@ -3250,7 +3129,7 @@ UControlRig::FPoseScope::~FPoseScope()
 
 #if WITH_EDITOR
 
-UControlRig::FTransientControlScope::FTransientControlScope(TObjectPtr<URigHierarchy> InHierarchy)
+UBaseControlRig::FTransientControlScope::FTransientControlScope(TObjectPtr<URigHierarchy> InHierarchy)
 	:Hierarchy(InHierarchy)
 {
 	for (FRigControlElement* Control : Hierarchy->GetTransientControls())
@@ -3268,7 +3147,7 @@ UControlRig::FTransientControlScope::FTransientControlScope(TObjectPtr<URigHiera
 	}
 }
 
-UControlRig::FTransientControlScope::~FTransientControlScope()
+UBaseControlRig::FTransientControlScope::~FTransientControlScope()
 {
 	if (URigHierarchyController* Controller = Hierarchy->GetController())
 	{
diff --git a/Engine/Plugins/Animation/ControlRig/Source/ControlRig/Private/Constraints/ControlRigTransformableHandle.cpp b/Engine/Plugins/Animation/ControlRig/Source/ControlRig/Private/Constraints/ControlRigTransformableHandle.cpp
--- a/Engine/Plugins/Animation/ControlRig/Source/ControlRig/Private/Constraints/ControlRigTransformableHandle.cpp
+++ b/Engine/Plugins/Animation/ControlRig/Source/ControlRig/Private/Constraints/ControlRigTransformableHandle.cpp
@@ -4,6 +4,7 @@
 #include "Constraints/ControlRigTransformableHandle.h"
 
 #include "Components/SkeletalMeshComponent.h"
+#include "BaseControlRig.h"
 #include "ControlRigComponent.h"
 #include "BaseControlRig.h"
 #include "IControlRigObjectBinding.h"
@@ -176,7 +177,7 @@ FTickFunction* UTransformableControlHandle::GetTickFunction() const
 	return BoundComponent ? &BoundComponent->PrimaryComponentTick : nullptr;
 }
 
-uint32 UTransformableControlHandle::ComputeHash(const UControlRig* InControlRig, const FName& InControlName)
+uint32 UTransformableControlHandle::ComputeHash(const UBaseControlRig* InControlRig, const FName& InControlName)
 {
 	return HashCombine(GetTypeHash(InControlRig), GetTypeHash(InControlName));
 }
@@ -407,7 +408,7 @@ void UTransformableControlHandle::OnHierarchyModified(
 }
 
 void UTransformableControlHandle::OnControlModified(
-	UControlRig* InControlRig,
+	UBaseControlRig* InControlRig,
 	FRigControlElement* InControl,
 	const FRigControlModifiedContext& InContext)
 {
@@ -446,7 +447,7 @@ void UTransformableControlHandle::OnControlModified(
 	}
 }
 
-void UTransformableControlHandle::OnControlRigBound(UControlRig* InControlRig)
+void UTransformableControlHandle::OnControlRigBound(UBaseControlRig* InControlRig)
 {
 	if (!InControlRig)
 	{
@@ -511,7 +512,7 @@ bool UTransformableControlHandle::AddTransformKeys(const TArray<FFrameNumber>& I
 	}
 	auto KeyframeFunc = [this, bLocal](const FTransform& InTransform, const FRigControlModifiedContext& InKeyframeContext)
 	{
-		UControlRig* InControlRig = ControlRig.Get();
+		UBaseControlRig* InControlRig = ControlRig.Get();
 		static constexpr bool bNotify = true;
 		static constexpr bool bUndo = false;
 		static constexpr bool bFixEuler = true;
@@ -551,7 +552,7 @@ bool UTransformableControlHandle::AddTransformKeys(const TArray<FFrameNumber>& I
 //for control rig need to check to see if the control rig is different then we may need to update it based upon what we are now bound to
 void UTransformableControlHandle::ResolveBoundObjects(FMovieSceneSequenceID LocalSequenceID, IMovieScenePlayer& Player, UObject* SubObject)
 {
-	if (UControlRig* InControlRig = Cast<UControlRig>(SubObject))
+	if (UBaseControlRig* InControlRig = Cast<UBaseControlRig>(SubObject))
 	{
 		if (ControlRig != InControlRig)
 		{
@@ -609,7 +610,7 @@ void UTransformableControlHandle::OnObjectsReplaced(const TMap<UObject*, UObject
 {
 	if (UObject* NewObject = InOldToNewInstances.FindRef(ControlRig.Get()))
 	{
-		if (UControlRig* NewControlRig = Cast<UControlRig>(NewObject))
+		if (UBaseControlRig* NewControlRig = Cast<UBaseControlRig>(NewObject))
 		{
 			if (URigHierarchy* Hierarchy = ControlRig->GetHierarchy())
 			{
diff --git a/Engine/Plugins/Animation/ControlRig/Source/ControlRig/Private/ControlRigBlueprintGeneratedClass.cpp b/Engine/Plugins/Animation/ControlRig/Source/ControlRig/Private/ControlRigBlueprintGeneratedClass.cpp
--- a/Engine/Plugins/Animation/ControlRig/Source/ControlRig/Private/ControlRigBlueprintGeneratedClass.cpp
+++ b/Engine/Plugins/Animation/ControlRig/Source/ControlRig/Private/ControlRigBlueprintGeneratedClass.cpp
@@ -34,7 +34,7 @@ void UControlRigBlueprintGeneratedClass::Serialize(FArchive& Ar)
 	const FString VMDebugName = FString::Printf(Format, *GetName(), uint32(++NumVMs));
 	URigVM* VM = NewObject<URigVM>(GetTransientPackage(), *VMDebugName);
 
-	if (UControlRig* CDO = Cast<UControlRig>(GetDefaultObject(true)))
+	if (UBaseControlRig* CDO = Cast<UBaseControlRig>(GetDefaultObject(true)))
 	{
 		if (Ar.IsSaving() && CDO->VM)
 		{
@@ -44,7 +44,7 @@ void UControlRigBlueprintGeneratedClass::Serialize(FArchive& Ar)
 	
 	VM->Serialize(Ar);
 
-	if (UControlRig* CDO = Cast<UControlRig>(GetDefaultObject(false)))
+	if (UBaseControlRig* CDO = Cast<UBaseControlRig>(GetDefaultObject(false)))
 	{
 		if (Ar.IsLoading() && CDO->VM)
 		{
diff --git a/Engine/Plugins/Animation/ControlRig/Source/ControlRig/Private/ControlRigComponent.cpp b/Engine/Plugins/Animation/ControlRig/Source/ControlRig/Private/ControlRigComponent.cpp
--- a/Engine/Plugins/Animation/ControlRig/Source/ControlRig/Private/ControlRigComponent.cpp
+++ b/Engine/Plugins/Animation/ControlRig/Source/ControlRig/Private/ControlRigComponent.cpp
@@ -206,7 +206,7 @@ void UControlRigComponent::TickComponent(float DeltaTime, enum ELevelTick TickTy
 		return;
 	}
 
-	if (UControlRig* CR = SetupControlRigIfRequired())
+	if (UBaseControlRig* CR = SetupControlRigIfRequired())
 	{
 		FScopeLock Lock(&gPendingSkeletalMeshesLock);
 		TArray<FSkeletalMeshToMap>* PendingSkeletalMeshes = gPendingSkeletalMeshes.Find(this);
@@ -283,7 +283,7 @@ FBoxSphereBounds UControlRigComponent::CalcBounds(const FTransform& LocalToWorld
 }
 
 
-UControlRig* UControlRigComponent::GetControlRig()
+UBaseControlRig* UControlRigComponent::GetControlRig()
 {
 	return SetupControlRigIfRequired();
 }
@@ -295,7 +295,7 @@ bool UControlRigComponent::CanExecute()
 		return false;
 	}
 	
-	if(UControlRig* CR = GetControlRig())
+	if(UBaseControlRig* CR = GetControlRig())
 	{
 		return CR->CanExecute();
 	}
@@ -370,7 +370,7 @@ void UControlRigComponent::Initialize()
 		OnPreInitialize(this);
 	}
 	
-	if(UControlRig* CR = SetupControlRigIfRequired())
+	if(UBaseControlRig* CR = SetupControlRigIfRequired())
 	{
 		if (CR->IsInitializing())
 		{
@@ -438,7 +438,7 @@ void UControlRigComponent::Initialize()
 
 void UControlRigComponent::Update(float DeltaTime)
 {
-	if(UControlRig* CR = SetupControlRigIfRequired())
+	if(UBaseControlRig* CR = SetupControlRigIfRequired())
 	{
 		if(!CanExecute())
 		{
@@ -531,7 +531,7 @@ TArray<FName> UControlRigComponent::GetElementNames(ERigElementType ElementType)
 {
 	TArray<FName> Names;
 
-	if (UControlRig* CR = SetupControlRigIfRequired())
+	if (UBaseControlRig* CR = SetupControlRigIfRequired())
 	{
 		for (FRigBaseElement* Element : *CR->GetHierarchy())
 		{
@@ -547,7 +547,7 @@ TArray<FName> UControlRigComponent::GetElementNames(ERigElementType ElementType)
 
 bool UControlRigComponent::DoesElementExist(FName Name, ERigElementType ElementType)
 {
-	if (UControlRig* CR = SetupControlRigIfRequired())
+	if (UBaseControlRig* CR = SetupControlRigIfRequired())
 	{
 		return CR->GetHierarchy()->GetIndex(FRigElementKey(Name, ElementType)) != INDEX_NONE;
 	}
@@ -637,7 +637,7 @@ void UControlRigComponent::AddMappedSkeletalMesh(USkeletalMeshComponent* Skeleta
 		return;
 	}
 
-	UControlRig* CR = SetupControlRigIfRequired();
+	UBaseControlRig* CR = SetupControlRigIfRequired();
 	if (CR == nullptr)
 	{
 		// if we don't have a valid rig yet - delay it until tick component
@@ -753,7 +753,7 @@ void UControlRigComponent::SetBoneInitialTransformsFromSkeletalMesh(USkeletalMes
 {
 	if (InSkeletalMesh)
 	{
-		if (UControlRig* CR = SetupControlRigIfRequired())
+		if (UBaseControlRig* CR = SetupControlRigIfRequired())
 		{
 			CR->SetBoneInitialTransformsFromSkeletalMesh(InSkeletalMesh);
 			bResetInitialsBeforeConstruction = false;
@@ -763,7 +763,7 @@ void UControlRigComponent::SetBoneInitialTransformsFromSkeletalMesh(USkeletalMes
 
 FTransform UControlRigComponent::GetBoneTransform(FName BoneName, EControlRigComponentSpace Space)
 {
-	if (UControlRig* CR = SetupControlRigIfRequired())
+	if (UBaseControlRig* CR = SetupControlRigIfRequired())
 	{
 		const int32 BoneIndex = CR->GetHierarchy()->GetIndex(FRigElementKey(BoneName, ERigElementType::Bone));
 		if (BoneIndex != INDEX_NONE)
@@ -786,7 +786,7 @@ FTransform UControlRigComponent::GetBoneTransform(FName BoneName, EControlRigCom
 
 FTransform UControlRigComponent::GetInitialBoneTransform(FName BoneName, EControlRigComponentSpace Space)
 {
-	if (UControlRig* CR = SetupControlRigIfRequired())
+	if (UBaseControlRig* CR = SetupControlRigIfRequired())
 	{
 		const int32 BoneIndex = CR->GetHierarchy()->GetIndex(FRigElementKey(BoneName, ERigElementType::Bone));
 		if (BoneIndex != INDEX_NONE)
@@ -814,7 +814,7 @@ void UControlRigComponent::SetBoneTransform(FName BoneName, FTransform Transform
 		return;
 	}
 
-	if (UControlRig* CR = SetupControlRigIfRequired())
+	if (UBaseControlRig* CR = SetupControlRigIfRequired())
 	{
 		int32 BoneIndex = CR->GetHierarchy()->GetIndex(FRigElementKey(BoneName, ERigElementType::Bone));
 		if (BoneIndex != INDEX_NONE)
@@ -853,7 +853,7 @@ void UControlRigComponent::SetBoneTransform(FName BoneName, FTransform Transform
 
 void UControlRigComponent::SetInitialBoneTransform(FName BoneName, FTransform InitialTransform, EControlRigComponentSpace Space, bool bPropagateToChildren)
 {
-	if (UControlRig* CR = SetupControlRigIfRequired())
+	if (UBaseControlRig* CR = SetupControlRigIfRequired())
 	{
 		const int32 BoneIndex = CR->GetHierarchy()->GetIndex(FRigElementKey(BoneName, ERigElementType::Bone));
 		if (BoneIndex != INDEX_NONE)
@@ -882,7 +882,7 @@ void UControlRigComponent::SetInitialBoneTransform(FName BoneName, FTransform In
 
 bool UControlRigComponent::GetControlBool(FName ControlName)
 {
-	if (UControlRig* CR = SetupControlRigIfRequired())
+	if (UBaseControlRig* CR = SetupControlRigIfRequired())
 	{
 		if(FRigControlElement* ControlElement = CR->GetHierarchy()->Find<FRigControlElement>(FRigElementKey(ControlName, ERigElementType::Control)))
 		{
@@ -898,7 +898,7 @@ bool UControlRigComponent::GetControlBool(FName ControlName)
 
 float UControlRigComponent::GetControlFloat(FName ControlName)
 {
-	if (UControlRig* CR = SetupControlRigIfRequired())
+	if (UBaseControlRig* CR = SetupControlRigIfRequired())
 	{
 		if(FRigControlElement* ControlElement = CR->GetHierarchy()->Find<FRigControlElement>(FRigElementKey(ControlName, ERigElementType::Control)))
 		{
@@ -914,7 +914,7 @@ float UControlRigComponent::GetControlFloat(FName ControlName)
 
 int32 UControlRigComponent::GetControlInt(FName ControlName)
 {
-	if (UControlRig* CR = SetupControlRigIfRequired())
+	if (UBaseControlRig* CR = SetupControlRigIfRequired())
 	{
 		if(FRigControlElement* ControlElement = CR->GetHierarchy()->Find<FRigControlElement>(FRigElementKey(ControlName, ERigElementType::Control)))
 		{
@@ -930,7 +930,7 @@ int32 UControlRigComponent::GetControlInt(FName ControlName)
 
 FVector2D UControlRigComponent::GetControlVector2D(FName ControlName)
 {
-	if (UControlRig* CR = SetupControlRigIfRequired())
+	if (UBaseControlRig* CR = SetupControlRigIfRequired())
 	{
 		if(FRigControlElement* ControlElement = CR->GetHierarchy()->Find<FRigControlElement>(FRigElementKey(ControlName, ERigElementType::Control)))
 		{
@@ -962,7 +962,7 @@ FVector UControlRigComponent::GetControlScale(FName ControlName, EControlRigComp
 
 FTransform UControlRigComponent::GetControlTransform(FName ControlName, EControlRigComponentSpace Space)
 {
-	if (UControlRig* CR = SetupControlRigIfRequired())
+	if (UBaseControlRig* CR = SetupControlRigIfRequired())
 	{
 		if(FRigControlElement* ControlElement = CR->GetHierarchy()->Find<FRigControlElement>(FRigElementKey(ControlName, ERigElementType::Control)))
 		{
@@ -985,7 +985,7 @@ FTransform UControlRigComponent::GetControlTransform(FName ControlName, EControl
 
 void UControlRigComponent::SetControlBool(FName ControlName, bool Value)
 {
-	if (UControlRig* CR = SetupControlRigIfRequired())
+	if (UBaseControlRig* CR = SetupControlRigIfRequired())
 	{
 		CR->SetControlValue<bool>(ControlName, Value);
 	}
@@ -993,7 +993,7 @@ void UControlRigComponent::SetControlBool(FName ControlName, bool Value)
 
 void UControlRigComponent::SetControlFloat(FName ControlName, float Value)
 {
-	if (UControlRig* CR = SetupControlRigIfRequired())
+	if (UBaseControlRig* CR = SetupControlRigIfRequired())
 	{
 		CR->SetControlValue<float>(ControlName, Value);
 	}
@@ -1001,7 +1001,7 @@ void UControlRigComponent::SetControlFloat(FName ControlName, float Value)
 
 void UControlRigComponent::SetControlInt(FName ControlName, int32 Value)
 {
-	if (UControlRig* CR = SetupControlRigIfRequired())
+	if (UBaseControlRig* CR = SetupControlRigIfRequired())
 	{
 		CR->SetControlValue<int32>(ControlName, Value);
 	}
@@ -1009,7 +1009,7 @@ void UControlRigComponent::SetControlInt(FName ControlName, int32 Value)
 
 void UControlRigComponent::SetControlVector2D(FName ControlName, FVector2D Value)
 {
-	if (UControlRig* CR = SetupControlRigIfRequired())
+	if (UBaseControlRig* CR = SetupControlRigIfRequired())
 	{
 		CR->SetControlValue<FVector2D>(ControlName, Value);
 	}
@@ -1017,7 +1017,7 @@ void UControlRigComponent::SetControlVector2D(FName ControlName, FVector2D Value
 
 void UControlRigComponent::SetControlPosition(FName ControlName, FVector Value, EControlRigComponentSpace Space)
 {
-	if (UControlRig* CR = SetupControlRigIfRequired())
+	if (UBaseControlRig* CR = SetupControlRigIfRequired())
 	{
 		if(FRigControlElement* ControlElement = CR->GetHierarchy()->Find<FRigControlElement>(FRigElementKey(ControlName, ERigElementType::Control)))
 		{
@@ -1044,7 +1044,7 @@ void UControlRigComponent::SetControlPosition(FName ControlName, FVector Value,
 
 void UControlRigComponent::SetControlRotator(FName ControlName, FRotator Value, EControlRigComponentSpace Space)
 {
-	if (UControlRig* CR = SetupControlRigIfRequired())
+	if (UBaseControlRig* CR = SetupControlRigIfRequired())
 	{
 		if(FRigControlElement* ControlElement = CR->GetHierarchy()->Find<FRigControlElement>(FRigElementKey(ControlName, ERigElementType::Control)))
 		{
@@ -1071,7 +1071,7 @@ void UControlRigComponent::SetControlRotator(FName ControlName, FRotator Value,
 
 void UControlRigComponent::SetControlScale(FName ControlName, FVector Value, EControlRigComponentSpace Space)
 {
-	if (UControlRig* CR = SetupControlRigIfRequired())
+	if (UBaseControlRig* CR = SetupControlRigIfRequired())
 	{
 		if(FRigControlElement* ControlElement = CR->GetHierarchy()->Find<FRigControlElement>(FRigElementKey(ControlName, ERigElementType::Control)))
 		{
@@ -1098,7 +1098,7 @@ void UControlRigComponent::SetControlScale(FName ControlName, FVector Value, ECo
 
 void UControlRigComponent::SetControlTransform(FName ControlName, FTransform Value, EControlRigComponentSpace Space)
 {
-	if (UControlRig* CR = SetupControlRigIfRequired())
+	if (UBaseControlRig* CR = SetupControlRigIfRequired())
 	{
 		if(FRigControlElement* ControlElement = CR->GetHierarchy()->Find<FRigControlElement>(FRigElementKey(ControlName, ERigElementType::Control)))
 		{
@@ -1117,7 +1117,7 @@ void UControlRigComponent::SetControlTransform(FName ControlName, FTransform Val
 
 FTransform UControlRigComponent::GetControlOffset(FName ControlName, EControlRigComponentSpace Space)
 {
-	if (UControlRig* CR = SetupControlRigIfRequired())
+	if (UBaseControlRig* CR = SetupControlRigIfRequired())
 	{
 		if(FRigControlElement* ControlElement = CR->GetHierarchy()->Find<FRigControlElement>(FRigElementKey(ControlName, ERigElementType::Control)))
 		{
@@ -1139,7 +1139,7 @@ FTransform UControlRigComponent::GetControlOffset(FName ControlName, EControlRig
 
 void UControlRigComponent::SetControlOffset(FName ControlName, FTransform OffsetTransform, EControlRigComponentSpace Space)
 {
-	if (UControlRig* CR = SetupControlRigIfRequired())
+	if (UBaseControlRig* CR = SetupControlRigIfRequired())
 	{
 		if(FRigControlElement* ControlElement = CR->GetHierarchy()->Find<FRigControlElement>(FRigElementKey(ControlName, ERigElementType::Control)))
 		{
@@ -1156,7 +1156,7 @@ void UControlRigComponent::SetControlOffset(FName ControlName, FTransform Offset
 
 FTransform UControlRigComponent::GetSpaceTransform(FName SpaceName, EControlRigComponentSpace Space)
 {
-	if (UControlRig* CR = SetupControlRigIfRequired())
+	if (UBaseControlRig* CR = SetupControlRigIfRequired())
 	{
 		if(FRigNullElement* NullElement = CR->GetHierarchy()->Find<FRigNullElement>(FRigElementKey(SpaceName, ERigElementType::Control)))
 		{
@@ -1178,7 +1178,7 @@ FTransform UControlRigComponent::GetSpaceTransform(FName SpaceName, EControlRigC
 
 FTransform UControlRigComponent::GetInitialSpaceTransform(FName SpaceName, EControlRigComponentSpace Space)
 {
-	if (UControlRig* CR = SetupControlRigIfRequired())
+	if (UBaseControlRig* CR = SetupControlRigIfRequired())
 	{
 		if(FRigNullElement* NullElement = CR->GetHierarchy()->Find<FRigNullElement>(FRigElementKey(SpaceName, ERigElementType::Control)))
 		{
@@ -1200,7 +1200,7 @@ FTransform UControlRigComponent::GetInitialSpaceTransform(FName SpaceName, ECont
 
 void UControlRigComponent::SetInitialSpaceTransform(FName SpaceName, FTransform InitialTransform, EControlRigComponentSpace Space)
 {
-	if (UControlRig* CR = SetupControlRigIfRequired())
+	if (UBaseControlRig* CR = SetupControlRigIfRequired())
 	{
 		if(FRigNullElement* NullElement = CR->GetHierarchy()->Find<FRigNullElement>(FRigElementKey(SpaceName, ERigElementType::Control)))
 		{
@@ -1217,7 +1217,7 @@ void UControlRigComponent::SetInitialSpaceTransform(FName SpaceName, FTransform
 	}
 }
 
-UControlRig* UControlRigComponent::SetupControlRigIfRequired()
+UBaseControlRig* UControlRigComponent::SetupControlRigIfRequired()
 {
 	if(ControlRig != nullptr)
 	{
@@ -1239,7 +1239,7 @@ UControlRig* UControlRigComponent::SetupControlRigIfRequired()
 
 	if(ControlRigClass)
 	{
-		ControlRig = NewObject<UControlRig>(this, ControlRigClass);
+		ControlRig = NewObject<UBaseControlRig>(this, ControlRigClass);
 
 		SetControlRig(ControlRig);
 
@@ -1253,7 +1253,7 @@ UControlRig* UControlRigComponent::SetupControlRigIfRequired()
 
 	return ControlRig;
 }
-void UControlRigComponent::SetControlRig(UControlRig* InControlRig)
+void UControlRigComponent::SetControlRig(UBaseControlRig* InControlRig)
 {
 	if (ControlRig)
 	{
@@ -1272,7 +1272,7 @@ void UControlRigComponent::SetControlRig(UControlRig* InControlRig)
 	ControlRig->OnPostForwardsSolve_AnyThread().AddUObject(this, &UControlRigComponent::HandleControlRigPostForwardsSolveEvent);
 	ControlRig->OnExecuted_AnyThread().AddUObject(this, &UControlRigComponent::HandleControlRigExecutedEvent);
 
-	ControlRig->GetDataSourceRegistry()->RegisterDataSource(UControlRig::OwnerComponent, this);
+	ControlRig->GetDataSourceRegistry()->RegisterDataSource(UBaseControlRig::OwnerComponent, this);
 	if(ObjectBinding.IsValid())
 	{
 		ControlRig->SetObjectBinding(ObjectBinding);
@@ -1281,7 +1281,7 @@ void UControlRigComponent::SetControlRig(UControlRig* InControlRig)
 	ControlRig->Initialize();
 }
 
-void UControlRigComponent::SetControlRigClass(TSubclassOf<UControlRig> InControlRigClass)
+void UControlRigComponent::SetControlRigClass(TSubclassOf<UBaseControlRig> InControlRigClass)
 {
 	ControlRig = nullptr;
 	ControlRigClass = InControlRigClass;
@@ -1296,7 +1296,7 @@ void UControlRigComponent::SetObjectBinding(UObject* InObjectToBind)
 	}
 	ObjectBinding->BindToObject(InObjectToBind);
 
-	if(UControlRig* CR = SetupControlRigIfRequired())
+	if(UBaseControlRig* CR = SetupControlRigIfRequired())
 	{
 		CR->SetObjectBinding(ObjectBinding);
 	}
@@ -1727,7 +1727,7 @@ void UControlRigComponent::HandleControlRigInitializedEvent(URigVMHost* InContro
 	}
 }
 
-void UControlRigComponent::HandleControlRigPreConstructionEvent(UControlRig* InControlRig, const FName& InEventName)
+void UControlRigComponent::HandleControlRigPreConstructionEvent(UBaseControlRig* InControlRig, const FName& InEventName)
 {
 	TArray<USkeletalMeshComponent*> ComponentsToTick;
 
@@ -1777,7 +1777,7 @@ void UControlRigComponent::HandleControlRigPreConstructionEvent(UControlRig* InC
 	}
 }
 
-void UControlRigComponent::HandleControlRigPostConstructionEvent(UControlRig* InControlRig, const FName& InEventName)
+void UControlRigComponent::HandleControlRigPostConstructionEvent(UBaseControlRig* InControlRig, const FName& InEventName)
 {
 #if WITH_EDITOR
 	if (bUpdateInEditor)
@@ -1792,7 +1792,7 @@ void UControlRigComponent::HandleControlRigPostConstructionEvent(UControlRig* In
 	}
 }
 
-void UControlRigComponent::HandleControlRigPreForwardsSolveEvent(UControlRig* InControlRig, const FName& InEventName)
+void UControlRigComponent::HandleControlRigPreForwardsSolveEvent(UBaseControlRig* InControlRig, const FName& InEventName)
 {
 #if WITH_EDITOR
 	if (bUpdateInEditor)
@@ -1807,7 +1807,7 @@ void UControlRigComponent::HandleControlRigPreForwardsSolveEvent(UControlRig* In
 	}
 }
 
-void UControlRigComponent::HandleControlRigPostForwardsSolveEvent(UControlRig* InControlRig, const FName& InEventName)
+void UControlRigComponent::HandleControlRigPostForwardsSolveEvent(UBaseControlRig* InControlRig, const FName& InEventName)
 {
 #if WITH_EDITOR
 	if (bUpdateInEditor)
@@ -1879,7 +1879,7 @@ void UControlRigComponent::ConvertTransformFromRigSpace(FTransform& InOutTransfo
 
 bool UControlRigComponent::EnsureCalledOutsideOfBracket(const TCHAR* InCallingFunctionName)
 {
-	if (UControlRig* CR = SetupControlRigIfRequired())
+	if (UBaseControlRig* CR = SetupControlRigIfRequired())
 	{
 		if (CR->IsRunningPreConstruction())
 		{
diff --git a/Engine/Plugins/Animation/ControlRig/Source/ControlRig/Private/ControlRigControlActor.cpp b/Engine/Plugins/Animation/ControlRig/Source/ControlRig/Private/ControlRigControlActor.cpp
--- a/Engine/Plugins/Animation/ControlRig/Source/ControlRig/Private/ControlRigControlActor.cpp
+++ b/Engine/Plugins/Animation/ControlRig/Source/ControlRig/Private/ControlRigControlActor.cpp
@@ -156,7 +156,7 @@ void AControlRigControlActor::Refresh()
 
 	if (ControlRig == nullptr)
 	{
-		TArray<UControlRig*> Rigs = UControlRig::FindControlRigs(ActorToTrack, ControlRigClass);
+		TArray<UBaseControlRig*> Rigs = UBaseControlRig::FindControlRigs(ActorToTrack, ControlRigClass);
 		if(Rigs.Num() > 0)
 		{
 			ControlRig = Rigs[0];
diff --git a/Engine/Plugins/Animation/ControlRig/Source/ControlRig/Private/ControlRigGizmoActor.cpp b/Engine/Plugins/Animation/ControlRig/Source/ControlRig/Private/ControlRigGizmoActor.cpp
--- a/Engine/Plugins/Animation/ControlRig/Source/ControlRig/Private/ControlRigGizmoActor.cpp
+++ b/Engine/Plugins/Animation/ControlRig/Source/ControlRig/Private/ControlRigGizmoActor.cpp
@@ -113,7 +113,7 @@ void AControlRigShapeActor::SetShapeColor(const FLinearColor& InColor)
 
 bool AControlRigShapeActor::UpdateControlSettings(
 	ERigHierarchyNotification InNotif,
-	UControlRig* InControlRig,
+	UBaseControlRig* InControlRig,
 	const FRigControlElement* InControlElement,
 	bool bHideManipulators,
 	bool bIsInLevelEditor)
diff --git a/Engine/Plugins/Animation/ControlRig/Source/ControlRig/Private/ControlRigModule.cpp b/Engine/Plugins/Animation/ControlRig/Source/ControlRig/Private/ControlRigModule.cpp
--- a/Engine/Plugins/Animation/ControlRig/Source/ControlRig/Private/ControlRigModule.cpp
+++ b/Engine/Plugins/Animation/ControlRig/Source/ControlRig/Private/ControlRigModule.cpp
@@ -36,7 +36,7 @@ void FControlRigModule::RegisterTransformableCustomization() const
 	// register UControlRig and AControlRigShapeActor
 	auto CreateControlHandle = [](UObject* InObject, const FName& InControlName)->UTransformableHandle*
 	{
-		if (const UControlRig* ControlRig = Cast<UControlRig>(InObject))
+		if (const UBaseControlRig* ControlRig = Cast<UBaseControlRig>(InObject))
 		{
 			return ControlRig->CreateTransformableControlHandle(InControlName);
 		}
@@ -45,7 +45,7 @@ void FControlRigModule::RegisterTransformableCustomization() const
 
 	auto GetControlHash = [](const UObject* InObject, const FName& InControlName)->uint32
 	{
-		if (const UControlRig* ControlRig = Cast<UControlRig>(InObject))
+		if (const UBaseControlRig* ControlRig = Cast<UBaseControlRig>(InObject))
 		{
 			return UTransformableControlHandle::ComputeHash(ControlRig, InControlName);
 		}
@@ -74,7 +74,7 @@ void FControlRigModule::RegisterTransformableCustomization() const
 	// as AControlRigShapeActor is only available if the ControlRig plugin is loaded.
 	FTransformableRegistry& Registry = FTransformableRegistry::Get();
 	Registry.Register(AControlRigShapeActor::StaticClass(), CreateControlHandleFromActor, GetControlHashFromActor);
-	Registry.Register(UControlRig::StaticClass(), CreateControlHandle, GetControlHash);
+	Registry.Register(UBaseControlRig::StaticClass(), CreateControlHandle, GetControlHash);
 }
 
 IMPLEMENT_MODULE(FControlRigModule, ControlRig)
