Control Rig: Rename ControlRig to BaseControlRig and created a specialized ControlRig

diff --git a/Engine/Plugins/Animation/ControlRig/Source/ControlRig/Private/ControlRigTestData.cpp b/Engine/Plugins/Animation/ControlRig/Source/ControlRig/Private/ControlRigTestData.cpp
--- a/Engine/Plugins/Animation/ControlRig/Source/ControlRig/Private/ControlRigTestData.cpp
+++ b/Engine/Plugins/Animation/ControlRig/Source/ControlRig/Private/ControlRigTestData.cpp
@@ -9,7 +9,7 @@
 
 #include UE_INLINE_GENERATED_CPP_BY_NAME(ControlRigTestData)
 
-bool FControlRigTestDataFrame::Store(UControlRig* InControlRig, bool bInitial)
+bool FControlRigTestDataFrame::Store(UBaseControlRig* InControlRig, bool bInitial)
 {
 	if(InControlRig == nullptr)
 	{
@@ -46,7 +46,7 @@ bool FControlRigTestDataFrame::Store(UControlRig* InControlRig, bool bInitial)
 	return true;
 }
 
-bool FControlRigTestDataFrame::Restore(UControlRig* InControlRig, bool bInitial) const
+bool FControlRigTestDataFrame::Restore(UBaseControlRig* InControlRig, bool bInitial) const
 {
 	if(InControlRig == nullptr)
 	{
@@ -71,7 +71,7 @@ bool FControlRigTestDataFrame::Restore(UControlRig* InControlRig, bool bInitial)
 	return RestoreVariables(InControlRig);
 }
 
-bool FControlRigTestDataFrame::RestoreVariables(UControlRig* InControlRig) const
+bool FControlRigTestDataFrame::RestoreVariables(UBaseControlRig* InControlRig) const
 {
 	class FControlRigTestDataFrame_ErrorPipe : public FOutputDevice
 	{
@@ -214,7 +214,7 @@ int32 UControlRigTestData::GetFrameIndexForTime(double InSeconds, bool bInput) c
 	return LastFrameIndex;
 }
 
-bool UControlRigTestData::Record(UControlRig* InControlRig, double InRecordingDuration)
+bool UControlRigTestData::Record(UBaseControlRig* InControlRig, double InRecordingDuration)
 {
 	if(InControlRig == nullptr)
 	{
@@ -232,7 +232,7 @@ bool UControlRigTestData::Record(UControlRig* InControlRig, double InRecordingDu
 	{
 		InControlRig->RequestInit();
 		PreConstructionHandle = InControlRig->OnPreConstruction_AnyThread().AddLambda(
-			[this](UControlRig* InControlRig, const FName& InEventName)
+			[this](UBaseControlRig* InControlRig, const FName& InEventName)
 			{
 				Initial.Store(InControlRig, true);
 			}
@@ -240,7 +240,7 @@ bool UControlRigTestData::Record(UControlRig* InControlRig, double InRecordingDu
 	}
 
 	PreForwardHandle = InControlRig->OnPreForwardsSolve_AnyThread().AddLambda(
-		[this](UControlRig* InControlRig, const FName& InEventName)
+		[this](UBaseControlRig* InControlRig, const FName& InEventName)
 		{
 			FControlRigTestDataFrame Frame;
 			Frame.Store(InControlRig);
@@ -254,7 +254,7 @@ bool UControlRigTestData::Record(UControlRig* InControlRig, double InRecordingDu
 	);
 
 	PostForwardHandle = InControlRig->OnPostForwardsSolve_AnyThread().AddLambda(
-		[this](UControlRig* InControlRig, const FName& InEventName)
+		[this](UBaseControlRig* InControlRig, const FName& InEventName)
 		{
 			FControlRigTestDataFrame Frame;
 			Frame.Store(InControlRig);
@@ -277,7 +277,7 @@ bool UControlRigTestData::Record(UControlRig* InControlRig, double InRecordingDu
 	return true;
 }
 
-bool UControlRigTestData::SetupReplay(UControlRig* InControlRig, bool bGroundTruth)
+bool UControlRigTestData::SetupReplay(UBaseControlRig* InControlRig, bool bGroundTruth)
 {
 	ReleaseReplay();
 	ClearDelegates(InControlRig);
@@ -299,14 +299,14 @@ bool UControlRigTestData::SetupReplay(UControlRig* InControlRig, bool bGroundTru
 	
 	InControlRig->RequestInit();
 	PreConstructionHandle = InControlRig->OnPreConstruction_AnyThread().AddLambda(
-		[this](UControlRig* InControlRig, const FName& InEventName)
+		[this](UBaseControlRig* InControlRig, const FName& InEventName)
 		{
 			Initial.Restore(InControlRig, true);
 		}
 	);
 
 	PreForwardHandle = InControlRig->OnPreForwardsSolve_AnyThread().AddLambda(
-		[this](UControlRig* InControlRig, const FName& InEventName)
+		[this](UBaseControlRig* InControlRig, const FName& InEventName)
 		{
 			// loop the animation data
 			if(InControlRig->GetAbsoluteTime() < GetTimeRange().X - SMALL_NUMBER ||
@@ -327,7 +327,7 @@ bool UControlRigTestData::SetupReplay(UControlRig* InControlRig, bool bGroundTru
 	);
 
 	PostForwardHandle = InControlRig->OnPostForwardsSolve_AnyThread().AddLambda(
-		[this](UControlRig* InControlRig, const FName& InEventName)
+		[this](UBaseControlRig* InControlRig, const FName& InEventName)
 		{
 			const FRigPose CurrentPose = InControlRig->GetHierarchy()->GetPose();
 
@@ -369,7 +369,7 @@ bool UControlRigTestData::SetupReplay(UControlRig* InControlRig, bool bGroundTru
 
 void UControlRigTestData::ReleaseReplay()
 {
-	if(UControlRig* ControlRig = ReplayControlRig.Get())
+	if(UBaseControlRig* ControlRig = ReplayControlRig.Get())
 	{
 		ClearDelegates(ControlRig);
 		ReplayControlRig.Reset();
@@ -394,7 +394,7 @@ bool UControlRigTestData::IsReplaying() const
 	return ReplayControlRig.IsValid();
 }
 
-void UControlRigTestData::ClearDelegates(UControlRig* InControlRig)
+void UControlRigTestData::ClearDelegates(UBaseControlRig* InControlRig)
 {
 	if(InControlRig)
 	{
diff --git a/Engine/Plugins/Animation/ControlRig/Source/ControlRig/Private/ControlRigValidationPass.cpp b/Engine/Plugins/Animation/ControlRig/Source/ControlRig/Private/ControlRigValidationPass.cpp
--- a/Engine/Plugins/Animation/ControlRig/Source/ControlRig/Private/ControlRigValidationPass.cpp
+++ b/Engine/Plugins/Animation/ControlRig/Source/ControlRig/Private/ControlRigValidationPass.cpp
@@ -111,9 +111,9 @@ void UControlRigValidator::RemovePass(UClass* InClass)
 	}
 }
 
-void UControlRigValidator::SetControlRig(UControlRig* InControlRig)
+void UControlRigValidator::SetControlRig(UBaseControlRig* InControlRig)
 {
-	if (UControlRig* ControlRig = WeakControlRig.Get())
+	if (UBaseControlRig* ControlRig = WeakControlRig.Get())
 	{
 		ControlRig->OnInitialized_AnyThread().RemoveAll(this);
 		ControlRig->OnExecuted_AnyThread().RemoveAll(this);
@@ -122,7 +122,7 @@ void UControlRigValidator::SetControlRig(UControlRig* InControlRig)
 	ValidationContext.DrawInterface = nullptr;
 	WeakControlRig = InControlRig;
 
-	if (UControlRig* ControlRig = WeakControlRig.Get())
+	if (UBaseControlRig* ControlRig = WeakControlRig.Get())
 	{
 		OnControlRigInitialized(ControlRig, FRigUnit_BeginExecution::EventName);
 		ControlRig->OnInitialized_AnyThread().AddUObject(this, &UControlRigValidator::OnControlRigInitialized);
@@ -133,7 +133,7 @@ void UControlRigValidator::SetControlRig(UControlRig* InControlRig)
 
 void UControlRigValidator::OnControlRigInitialized(URigVMHost* Subject, const FName& EventName)
 {
-	if (UControlRig* ControlRig = WeakControlRig.Get())
+	if (UBaseControlRig* ControlRig = WeakControlRig.Get())
 	{
 		if (ControlRig != Subject)
 		{
@@ -151,7 +151,7 @@ void UControlRigValidator::OnControlRigInitialized(URigVMHost* Subject, const FN
 
 void UControlRigValidator::OnControlRigExecuted(URigVMHost* Subject, const FName& EventName)
 {
-	if (UControlRig* ControlRig = WeakControlRig.Get())
+	if (UBaseControlRig* ControlRig = WeakControlRig.Get())
 	{
 		if (ControlRig != Subject)
 		{
diff --git a/Engine/Plugins/Animation/ControlRig/Source/ControlRig/Private/Rigs/AdditiveControlRig.cpp b/Engine/Plugins/Animation/ControlRig/Source/ControlRig/Private/Rigs/AdditiveControlRig.cpp
--- a/Engine/Plugins/Animation/ControlRig/Source/ControlRig/Private/Rigs/AdditiveControlRig.cpp
+++ b/Engine/Plugins/Animation/ControlRig/Source/ControlRig/Private/Rigs/AdditiveControlRig.cpp
@@ -20,6 +20,13 @@ UAdditiveControlRig::UAdditiveControlRig(const FObjectInitializer& ObjectInitial
 	bResetInitialTransformsBeforeConstruction = false;
 }
 
+#if WITH_EDITOR
+void UAdditiveControlRig::SetFirstEntryEventInEventQueue(FRigVMExtendedExecuteContext& Context, const FName& InFirstEventName)
+{
+	VM->SetFirstEntryEventInEventQueue(Context, NAME_None);
+}
+#endif
+
 FName UAdditiveControlRig::GetControlName(const FName& InBoneName)
 {
 	if (InBoneName != NAME_None)
diff --git a/Engine/Plugins/Animation/ControlRig/Source/ControlRig/Private/Rigs/FKControlRig.cpp b/Engine/Plugins/Animation/ControlRig/Source/ControlRig/Private/Rigs/FKControlRig.cpp
--- a/Engine/Plugins/Animation/ControlRig/Source/ControlRig/Private/Rigs/FKControlRig.cpp
+++ b/Engine/Plugins/Animation/ControlRig/Source/ControlRig/Private/Rigs/FKControlRig.cpp
@@ -260,8 +260,15 @@ bool UFKControlRig::Execute_Internal(const FName& InEventName)
 	return true;
 }
 
+#if WITH_EDITOR
+void UFKControlRig::SetFirstEntryEventInEventQueue(FRigVMExtendedExecuteContext& Context, const FName& InFirstEventName)
+{
+	VM->SetFirstEntryEventInEventQueue(Context, NAME_None);
+}
+#endif
+
 void UFKControlRig::SetBoneInitialTransformsFromSkeletalMeshComponent(USkeletalMeshComponent* InSkelMeshComp,
-	bool bUseAnimInstance)
+                                                                      bool bUseAnimInstance)
 {
 	Super::SetBoneInitialTransformsFromSkeletalMeshComponent(InSkelMeshComp, bUseAnimInstance);
 
diff --git a/Engine/Plugins/Animation/ControlRig/Source/ControlRig/Private/Rigs/RigHierarchy.cpp b/Engine/Plugins/Animation/ControlRig/Source/ControlRig/Private/Rigs/RigHierarchy.cpp
--- a/Engine/Plugins/Animation/ControlRig/Source/ControlRig/Private/Rigs/RigHierarchy.cpp
+++ b/Engine/Plugins/Animation/ControlRig/Source/ControlRig/Private/Rigs/RigHierarchy.cpp
@@ -74,7 +74,7 @@ FAutoConsoleCommandWithWorldAndArgs FCmdControlRigHierarchyTraceFrames
 			
 			// we'll just trace all of them for now
 			//if(Instance->GetWorld() == InWorld)
-			if(Instance->GetTypedOuter<UControlRig>() != nullptr)
+			if(Instance->GetTypedOuter<UBaseControlRig>() != nullptr)
 			{
 				CastChecked<URigHierarchy>(Instance)->TraceFrames(NumFrames);
 			}
@@ -2628,7 +2628,7 @@ void URigHierarchy::SendAutoKeyEvent(FRigElementKey InElement, float InOffsetInS
 	Context.Event = ERigEvent::RequestAutoKey;
 	Context.Key = InElement;
 	Context.LocalTime = InOffsetInSeconds;
-	if(UControlRig* Rig = Cast<UControlRig>(GetOuter()))
+	if(UBaseControlRig* Rig = Cast<UBaseControlRig>(GetOuter()))
 	{
 		Context.LocalTime += Rig->AbsoluteTime;
 	}
@@ -3336,7 +3336,7 @@ void URigHierarchy::SetTransform(FRigTransformElement* InTransformElement, const
 		{
 			BlueprintName = Blueprint->GetFName().ToString();
 		}
-		else if (UControlRig* Rig = Cast<UControlRig>(GetOuter()))
+		else if (UBaseControlRig* Rig = Cast<UBaseControlRig>(GetOuter()))
 		{
 			if (UBlueprint* BlueprintCR = Cast<UBlueprint>(Rig->GetClass()->ClassGeneratedBy))
 			{
@@ -3560,7 +3560,7 @@ void URigHierarchy::SetControlOffsetTransform(FRigControlElement* InControlEleme
 		{
 			BlueprintName = Blueprint->GetFName().ToString();
 		}
-		else if (UControlRig* Rig = Cast<UControlRig>(GetOuter()))
+		else if (UBaseControlRig* Rig = Cast<UBaseControlRig>(GetOuter()))
 		{
 			if (UBlueprint* BlueprintCR = Cast<UBlueprint>(Rig->GetClass()->ClassGeneratedBy))
 			{
@@ -3692,7 +3692,7 @@ void URigHierarchy::SetControlShapeTransform(FRigControlElement* InControlElemen
 		{
 			BlueprintName = Blueprint->GetFName().ToString();
 		}
-		else if (UControlRig* Rig = Cast<UControlRig>(GetOuter()))
+		else if (UBaseControlRig* Rig = Cast<UBaseControlRig>(GetOuter()))
 		{
 			if (UBlueprint* BlueprintCR = Cast<UBlueprint>(Rig->GetClass()->ClassGeneratedBy))
 			{
@@ -3760,7 +3760,7 @@ void URigHierarchy::SetControlSettings(FRigControlElement* InControlElement, FRi
 		{
 			BlueprintName = Blueprint->GetFName().ToString();
 		}
-		else if (UControlRig* Rig = Cast<UControlRig>(GetOuter()))
+		else if (UBaseControlRig* Rig = Cast<UBaseControlRig>(GetOuter()))
 		{
 			if (UBlueprint* BlueprintCR = Cast<UBlueprint>(Rig->GetClass()->ClassGeneratedBy))
 			{
@@ -4038,7 +4038,7 @@ void URigHierarchy::SetControlValue(FRigControlElement* InControlElement, const
 					{
 						BlueprintName = Blueprint->GetFName().ToString();
 					}
-					else if (UControlRig* Rig = Cast<UControlRig>(GetOuter()))
+					else if (UBaseControlRig* Rig = Cast<UBaseControlRig>(GetOuter()))
 					{
 						if (UBlueprint* BlueprintCR = Cast<UBlueprint>(Rig->GetClass()->ClassGeneratedBy))
 						{
@@ -4142,7 +4142,7 @@ void URigHierarchy::SetConnectorSettings(FRigConnectorElement* InConnectorElemen
 		{
 			BlueprintName = Blueprint->GetFName().ToString();
 		}
-		else if (UControlRig* Rig = Cast<UControlRig>(GetOuter()))
+		else if (UBaseControlRig* Rig = Cast<UBaseControlRig>(GetOuter()))
 		{
 			if (UBlueprint* BlueprintCR = Cast<UBlueprint>(Rig->GetClass()->ClassGeneratedBy))
 			{
@@ -5725,7 +5725,7 @@ bool URigHierarchy::ApplyTransformFromStack(const FRigTransformStackEntry& InEnt
 		// this might be a transient control which had been removed.
 		if(InEntry.Key.Type == ERigElementType::Control)
 		{
-			const FRigElementKey TargetKey = UControlRig::GetElementKeyFromTransientControl(InEntry.Key);
+			const FRigElementKey TargetKey = UBaseControlRig::GetElementKeyFromTransientControl(InEntry.Key);
 			Element = Find(TargetKey);
 			bApplyInitialForCurrent = Element != nullptr;
 		}
@@ -6606,7 +6606,7 @@ TArray<FString> URigHierarchy::ConnectorSettingsToPythonCommands(const FRigConne
 
 #endif
 
-FRigHierarchyRedirectorGuard::FRigHierarchyRedirectorGuard(UControlRig* InControlRig)
+FRigHierarchyRedirectorGuard::FRigHierarchyRedirectorGuard(UBaseControlRig* InControlRig)
 : Guard(InControlRig->GetHierarchy()->ElementKeyRedirector, &InControlRig->GetElementKeyRedirector())
 {
 }
\ No newline at end of file
diff --git a/Engine/Plugins/Animation/ControlRig/Source/ControlRig/Private/Sequencer/ControlRigLayerInstance.cpp b/Engine/Plugins/Animation/ControlRig/Source/ControlRig/Private/Sequencer/ControlRigLayerInstance.cpp
--- a/Engine/Plugins/Animation/ControlRig/Source/ControlRig/Private/Sequencer/ControlRigLayerInstance.cpp
+++ b/Engine/Plugins/Animation/ControlRig/Source/ControlRig/Private/Sequencer/ControlRigLayerInstance.cpp
@@ -34,7 +34,7 @@ void UControlRigLayerInstance::NativeUpdateAnimation(float DeltaSeconds)
 	
 	if (UAnimInstance* SourceAnimInstance = GetSourceAnimInstance())
 	{
-		if (UControlRig* ControlRig = GetFirstAvailableControlRig())
+		if (UBaseControlRig* ControlRig = GetFirstAvailableControlRig())
 		{
 			if (ControlRig->IsAdditive())
 			{
@@ -69,7 +69,7 @@ void UControlRigLayerInstance::ResetPose()
 	GetProxyOnGameThread<FControlRigLayerInstanceProxy>().ResetPose();
 }
 
-UControlRig* UControlRigLayerInstance::GetFirstAvailableControlRig() const
+UBaseControlRig* UControlRigLayerInstance::GetFirstAvailableControlRig() const
 {
 	return GetProxyOnGameThread<FControlRigLayerInstanceProxy>().GetFirstAvailableControlRig();
 }
@@ -100,7 +100,7 @@ void UControlRigLayerInstance::SetSourceAnimInstance(UAnimInstance* SourceAnimIn
 }
 
 /** ControlRig related support */
-void UControlRigLayerInstance::AddControlRigTrack(int32 ControlRigID, UControlRig* InControlRig)
+void UControlRigLayerInstance::AddControlRigTrack(int32 ControlRigID, UBaseControlRig* InControlRig)
 {
 	if (InControlRig->IsAdditive())
 	{
@@ -127,7 +127,7 @@ void UControlRigLayerInstance::RemoveControlRigTrack(int32 ControlRigID)
 	GetProxyOnGameThread<FControlRigLayerInstanceProxy>().RemoveControlRigTrack(ControlRigID);
 	if (UAnimInstance* SourceAnimInstance = GetSourceAnimInstance())
 	{
-		if (UControlRig* ControlRig = GetFirstAvailableControlRig())
+		if (UBaseControlRig* ControlRig = GetFirstAvailableControlRig())
 		{
 			SourceAnimInstance->bUseMultiThreadedAnimationUpdate = !ControlRig->IsAdditive();
 		}
diff --git a/Engine/Plugins/Animation/ControlRig/Source/ControlRig/Private/Sequencer/ControlRigLayerInstanceProxy.cpp b/Engine/Plugins/Animation/ControlRig/Source/ControlRig/Private/Sequencer/ControlRigLayerInstanceProxy.cpp
--- a/Engine/Plugins/Animation/ControlRig/Source/ControlRig/Private/Sequencer/ControlRigLayerInstanceProxy.cpp
+++ b/Engine/Plugins/Animation/ControlRig/Source/ControlRig/Private/Sequencer/ControlRigLayerInstanceProxy.cpp
@@ -173,7 +173,7 @@ void FControlRigLayerInstanceProxy::SetSourceAnimInstance(UAnimInstance* SourceA
 }
 
 /** ControlRig related support */
-void FControlRigLayerInstanceProxy::AddControlRigTrack(int32 ControlRigID, UControlRig* InControlRig)
+void FControlRigLayerInstanceProxy::AddControlRigTrack(int32 ControlRigID, UBaseControlRig* InControlRig)
 {
 	FAnimNode_ControlRig_ExternalSource* Node = FindControlRigNode(ControlRigID);
 
@@ -407,11 +407,11 @@ UAnimSequencerInstance* FControlRigLayerInstanceProxy::GetSequencerAnimInstance(
 	return Cast<UAnimSequencerInstance>(CurrentSourceAnimInstance);
 }
 
-UControlRig* FControlRigLayerInstanceProxy::GetFirstAvailableControlRig() const
+UBaseControlRig* FControlRigLayerInstanceProxy::GetFirstAvailableControlRig() const
 {
 	for (const TSharedPtr<FAnimNode_ControlRig_ExternalSource>& ControlRigNode : ControlRigNodes)
 	{
-		if (UControlRig* Rig = ControlRigNode->GetControlRig())
+		if (UBaseControlRig* Rig = ControlRigNode->GetControlRig())
 		{
 			return Rig;
 		}
diff --git a/Engine/Plugins/Animation/ControlRig/Source/ControlRig/Private/Sequencer/ControlRigLayerInstanceProxy.h b/Engine/Plugins/Animation/ControlRig/Source/ControlRig/Private/Sequencer/ControlRigLayerInstanceProxy.h
--- a/Engine/Plugins/Animation/ControlRig/Source/ControlRig/Private/Sequencer/ControlRigLayerInstanceProxy.h
+++ b/Engine/Plugins/Animation/ControlRig/Source/ControlRig/Private/Sequencer/ControlRigLayerInstanceProxy.h
@@ -6,7 +6,7 @@
 #include "AnimNode_ControlRig_ExternalSource.h"
 #include "ControlRigLayerInstanceProxy.generated.h"
 
-class UControlRig;
+class UBaseControlRig;
 class UAnimSequencerInstance;
 
 /** Custom internal Input Pose node that handles any AnimInstance */
@@ -79,7 +79,7 @@ public:
 	UAnimInstance* GetSourceAnimInstance() const { return CurrentSourceAnimInstance; }
 
 	/** ControlRig related support */
-	void AddControlRigTrack(int32 ControlRigID, UControlRig* InControlRig);
+	void AddControlRigTrack(int32 ControlRigID, UBaseControlRig* InControlRig);
 	void UpdateControlRigTrack(int32 ControlRigID, float Weight, const FControlRigIOSettings& InputSettings, bool bExecute);
 	void RemoveControlRigTrack(int32 ControlRigID);
 	bool HasControlRigTrack(int32 ControlRigID);
@@ -99,7 +99,7 @@ public:
 	virtual void ConstructNodes();
 
 	/** return first available control rig from the node it has */
-	UControlRig* GetFirstAvailableControlRig() const;
+	UBaseControlRig* GetFirstAvailableControlRig() const;
 
 	virtual void AddReferencedObjects(UAnimInstance* InAnimInstance, FReferenceCollector& Collector) override;
 
diff --git a/Engine/Plugins/Animation/ControlRig/Source/ControlRig/Private/Sequencer/ControlRigSequenceObjectReference.cpp b/Engine/Plugins/Animation/ControlRig/Source/ControlRig/Private/Sequencer/ControlRigSequenceObjectReference.cpp
--- a/Engine/Plugins/Animation/ControlRig/Source/ControlRig/Private/Sequencer/ControlRigSequenceObjectReference.cpp
+++ b/Engine/Plugins/Animation/ControlRig/Source/ControlRig/Private/Sequencer/ControlRigSequenceObjectReference.cpp
@@ -5,7 +5,7 @@
 
 #include UE_INLINE_GENERATED_CPP_BY_NAME(ControlRigSequenceObjectReference)
 
-FControlRigSequenceObjectReference FControlRigSequenceObjectReference::Create(UControlRig* InControlRig)
+FControlRigSequenceObjectReference FControlRigSequenceObjectReference::Create(UBaseControlRig* InControlRig)
 {
 	check(InControlRig);
 
diff --git a/Engine/Plugins/Animation/ControlRig/Source/ControlRig/Private/Sequencer/ControlRigSequencerHelpers.cpp b/Engine/Plugins/Animation/ControlRig/Source/ControlRig/Private/Sequencer/ControlRigSequencerHelpers.cpp
--- a/Engine/Plugins/Animation/ControlRig/Source/ControlRig/Private/Sequencer/ControlRigSequencerHelpers.cpp
+++ b/Engine/Plugins/Animation/ControlRig/Source/ControlRig/Private/Sequencer/ControlRigSequencerHelpers.cpp
@@ -6,7 +6,7 @@
 #include "Sequencer/MovieSceneControlRigParameterSection.h"
 
 TPair<const FChannelMapInfo*, int32> FControlRigSequencerHelpers::GetInfoAndNumFloatChannels(
-	const UControlRig* InControlRig,
+	const UBaseControlRig* InControlRig,
 	const FName& InControlName,
 	const UMovieSceneControlRigParameterSection* InSection)
 {
@@ -37,7 +37,7 @@ TPair<const FChannelMapInfo*, int32> FControlRigSequencerHelpers::GetInfoAndNumF
 }
 
 
-TArrayView<FMovieSceneFloatChannel*>  FControlRigSequencerHelpers::GetFloatChannels(const UControlRig* InControlRig,
+TArrayView<FMovieSceneFloatChannel*>  FControlRigSequencerHelpers::GetFloatChannels(const UBaseControlRig* InControlRig,
 	const FName& InControlName, const UMovieSceneSection* InSection)
 {
 	// no floats for transform sections
diff --git a/Engine/Plugins/Animation/ControlRig/Source/ControlRig/Private/Sequencer/MovieSceneControlRigParameterSection.cpp b/Engine/Plugins/Animation/ControlRig/Source/ControlRig/Private/Sequencer/MovieSceneControlRigParameterSection.cpp
--- a/Engine/Plugins/Animation/ControlRig/Source/ControlRig/Private/Sequencer/MovieSceneControlRigParameterSection.cpp
+++ b/Engine/Plugins/Animation/ControlRig/Source/ControlRig/Private/Sequencer/MovieSceneControlRigParameterSection.cpp
@@ -37,7 +37,7 @@
 
 struct FParameterFloatChannelEditorData
 {
-	FParameterFloatChannelEditorData(UControlRig *InControlRig, const FName& InName, bool bEnabledOverride, const FText& GroupName, int SortStartIndex)
+	FParameterFloatChannelEditorData(UBaseControlRig *InControlRig, const FName& InName, bool bEnabledOverride, const FText& GroupName, int SortStartIndex)
 	{
 		ControlRig = InControlRig;
 		ParameterName = InName;
@@ -56,7 +56,7 @@ struct FParameterFloatChannelEditorData
 		
 	}
 
-	static TOptional<float> GetValue(UControlRig* ControlRig, FName ParameterName, UObject& InObject, FTrackInstancePropertyBindings* Bindings)
+	static TOptional<float> GetValue(UBaseControlRig* ControlRig, FName ParameterName, UObject& InObject, FTrackInstancePropertyBindings* Bindings)
 	{
 		if (ControlRig)
 		{
@@ -104,13 +104,13 @@ struct FParameterFloatChannelEditorData
 	FMovieSceneChannelMetaData      MetaData;
 	TMovieSceneExternalValue<float> ExternalValues;
 	FName ParameterName;
-	UControlRig *ControlRig;
+	UBaseControlRig *ControlRig;
 };
 
 //Set up with all 4 Channels so it can be used by all vector types.
 struct FParameterVectorChannelEditorData
 {
-	FParameterVectorChannelEditorData(UControlRig *InControlRig, const FName& InName, bool bEnabledOverride, const FText& GroupName, int SortStartIndex, int32 NumChannels)
+	FParameterVectorChannelEditorData(UBaseControlRig *InControlRig, const FName& InName, bool bEnabledOverride, const FText& GroupName, int SortStartIndex, int32 NumChannels)
 	{
 		ControlRig = InControlRig;
 		ParameterName = InName;
@@ -174,7 +174,7 @@ struct FParameterVectorChannelEditorData
 
 	}
 
-	static FVector4 GetPropertyValue(UControlRig* ControlRig, FName ParameterName, UObject& InObject,int32 NumChannels)
+	static FVector4 GetPropertyValue(UBaseControlRig* ControlRig, FName ParameterName, UObject& InObject,int32 NumChannels)
 	{
 		if (ControlRig)
 		{
@@ -206,19 +206,19 @@ struct FParameterVectorChannelEditorData
 		return FVector4();
 	}
 
-	static TOptional<float> ExtractChannelX(UObject& InObject, UControlRig* ControlRig, FName ParameterName, int32 NumChannels)
+	static TOptional<float> ExtractChannelX(UObject& InObject, UBaseControlRig* ControlRig, FName ParameterName, int32 NumChannels)
 	{
 		return GetPropertyValue(ControlRig, ParameterName, InObject, NumChannels).X;
 	}
-	static TOptional<float> ExtractChannelY(UObject& InObject, UControlRig* ControlRig, FName ParameterName, int32 NumChannels)
+	static TOptional<float> ExtractChannelY(UObject& InObject, UBaseControlRig* ControlRig, FName ParameterName, int32 NumChannels)
 	{
 		return GetPropertyValue(ControlRig, ParameterName, InObject, NumChannels).Y;
 	}
-	static TOptional<float> ExtractChannelZ(UObject& InObject, UControlRig* ControlRig, FName ParameterName, int32 NumChannels)
+	static TOptional<float> ExtractChannelZ(UObject& InObject, UBaseControlRig* ControlRig, FName ParameterName, int32 NumChannels)
 	{
 		return GetPropertyValue(ControlRig, ParameterName, InObject, NumChannels).Z;
 	}
-	static TOptional<float> ExtractChannelW(UObject& InObject, UControlRig* ControlRig, FName ParameterName, int32 NumChannels)
+	static TOptional<float> ExtractChannelW(UObject& InObject, UBaseControlRig* ControlRig, FName ParameterName, int32 NumChannels)
 	{
 		return GetPropertyValue(ControlRig, ParameterName, InObject, NumChannels).W;
 	}
@@ -332,12 +332,12 @@ struct FParameterVectorChannelEditorData
 	FMovieSceneChannelMetaData      MetaData[4];
 	TMovieSceneExternalValue<float> ExternalValues[4];
 	FName ParameterName;
-	UControlRig *ControlRig;
+	UBaseControlRig *ControlRig;
 };
 
 struct FParameterTransformChannelEditorData
 {
-	FParameterTransformChannelEditorData(UControlRig *InControlRig, const FName& InName, bool bEnabledOverride, EMovieSceneTransformChannel Mask, 
+	FParameterTransformChannelEditorData(UBaseControlRig *InControlRig, const FName& InName, bool bEnabledOverride, EMovieSceneTransformChannel Mask, 
 		const FText& GroupName, int SortStartIndex)
 	{
 		ControlRig = InControlRig;
@@ -563,7 +563,7 @@ struct FParameterTransformChannelEditorData
 
 	}
 
-	static TOptional<FVector> GetTranslation(UControlRig* ControlRig, FName ParameterName, UObject& InObject, FTrackInstancePropertyBindings* Bindings)
+	static TOptional<FVector> GetTranslation(UBaseControlRig* ControlRig, FName ParameterName, UObject& InObject, FTrackInstancePropertyBindings* Bindings)
 	{
 		if (ControlRig)
 		{
@@ -616,7 +616,7 @@ struct FParameterTransformChannelEditorData
 		return TOptional<FVector>();
 	}
 
-	static TOptional<FRotator> GetRotator(UControlRig* ControlRig, FName ParameterName, UObject& InObject, FTrackInstancePropertyBindings* Bindings)
+	static TOptional<FRotator> GetRotator(UBaseControlRig* ControlRig, FName ParameterName, UObject& InObject, FTrackInstancePropertyBindings* Bindings)
 	{
 		if (ControlRig)
 
@@ -643,7 +643,7 @@ struct FParameterTransformChannelEditorData
 		return TOptional<FRotator>();
 	}
 
-	static TOptional<FVector> GetScale(UControlRig* ControlRig, FName ParameterName, UObject& InObject, FTrackInstancePropertyBindings* Bindings)
+	static TOptional<FVector> GetScale(UBaseControlRig* ControlRig, FName ParameterName, UObject& InObject, FTrackInstancePropertyBindings* Bindings)
 	{
 		if (ControlRig)
 		{
@@ -753,7 +753,7 @@ public:
 	FMovieSceneChannelMetaData      MetaData[9];
 	TMovieSceneExternalValue<float> ExternalValues[9];
 	FName ParameterName;
-	UControlRig *ControlRig;
+	UBaseControlRig *ControlRig;
 };
 
 #endif // WITH_EDITOR
@@ -2407,7 +2407,7 @@ void UMovieSceneControlRigParameterSection::KeyWeightValue(FFrameNumber InFrame,
 	Weight.AutoSetTangents();
 }
 
-void UMovieSceneControlRigParameterSection::RecreateWithThisControlRig(UControlRig* InControlRig, bool bSetDefault)
+void UMovieSceneControlRigParameterSection::RecreateWithThisControlRig(UBaseControlRig* InControlRig, bool bSetDefault)
 {
 	bool bSameControlRig = (ControlRig == InControlRig);
 	SetControlRig(InControlRig);
@@ -2591,7 +2591,7 @@ void UMovieSceneControlRigParameterSection::RecreateWithThisControlRig(UControlR
 	ReconstructChannelProxy();
 }
 
-void UMovieSceneControlRigParameterSection::SetControlRig(UControlRig* InControlRig)
+void UMovieSceneControlRigParameterSection::SetControlRig(UBaseControlRig* InControlRig)
 {
 	ControlRig = InControlRig;
 	ControlRigClass = ControlRig ? ControlRig->GetClass() : nullptr;
@@ -3184,7 +3184,7 @@ bool UMovieSceneControlRigParameterSection::LoadAnimSequenceIntoThisSection(UAni
 	// this ensures that the topology version matches in case of a dynamic hierarchy
 	if(!ControlRig->GetClass()->IsNative())
 	{
-		if (UControlRig* CDO = Cast<UControlRig>(ControlRig->GetClass()->GetDefaultObject()))
+		if (UBaseControlRig* CDO = Cast<UBaseControlRig>(ControlRig->GetClass()->GetDefaultObject()))
 		{
 			ControlRig->GetHierarchy()->CopyHierarchy(CDO->GetHierarchy());
 		}
diff --git a/Engine/Plugins/Animation/ControlRig/Source/ControlRig/Private/Sequencer/MovieSceneControlRigParameterTemplate.cpp b/Engine/Plugins/Animation/ControlRig/Source/ControlRig/Private/Sequencer/MovieSceneControlRigParameterTemplate.cpp
--- a/Engine/Plugins/Animation/ControlRig/Source/ControlRig/Private/Sequencer/MovieSceneControlRigParameterTemplate.cpp
+++ b/Engine/Plugins/Animation/ControlRig/Source/ControlRig/Private/Sequencer/MovieSceneControlRigParameterTemplate.cpp
@@ -198,7 +198,7 @@ using FControlRigAnimTypeIDsPtr = TSharedPtr<FControlRigAnimTypeIDs, ESPMode::Th
 struct FControlRigAnimTypeIDs
 {
 	/** Get the anim type IDs for the specified section */
-	static FControlRigAnimTypeIDsPtr Get(const UControlRig* ControlRig)
+	static FControlRigAnimTypeIDsPtr Get(const UBaseControlRig* ControlRig)
 	{
 		struct FControlRigAnimTypeIDsAnnotation
 		{
@@ -576,7 +576,7 @@ namespace MovieScene
 } // namespace UE
 
 //since initialization can blow up selection, may need to just reselect, used in a few places
-static void SelectControls(UControlRig* ControlRig, TArray<FName>& SelectedNames)
+static void SelectControls(UBaseControlRig* ControlRig, TArray<FName>& SelectedNames)
 {
 	if (ControlRig)
 	{
@@ -588,7 +588,7 @@ static void SelectControls(UControlRig* ControlRig, TArray<FName>& SelectedNames
 	}
 }
 
-void FControlRigBindingHelper::BindToSequencerInstance(UControlRig* ControlRig)
+void FControlRigBindingHelper::BindToSequencerInstance(UBaseControlRig* ControlRig)
 {
 	check(ControlRig);
 	if (USkeletalMeshComponent* SkeletalMeshComponent = Cast<USkeletalMeshComponent>(ControlRig->GetObjectBinding()->GetBoundObject()))
@@ -640,7 +640,7 @@ void FControlRigBindingHelper::BindToSequencerInstance(UControlRig* ControlRig)
 	}
 }
 
-void FControlRigBindingHelper::UnBindFromSequencerInstance(UControlRig* ControlRig)
+void FControlRigBindingHelper::UnBindFromSequencerInstance(UBaseControlRig* ControlRig)
 {
 	check(ControlRig);
 
@@ -763,7 +763,7 @@ struct FControlRigParameterPreAnimatedTokenProducer : IMovieScenePreAnimatedToke
 			virtual void RestoreState(UObject& InObject, const UE::MovieScene::FRestoreStateParams& Params) override
 			{
 
-				if (UControlRig* ControlRig = Cast<UControlRig>(&InObject))
+				if (UBaseControlRig* ControlRig = Cast<UBaseControlRig>(&InObject))
 				{
 					if (ControlRig->GetObjectBinding())
 					{
@@ -938,7 +938,7 @@ struct FControlRigParameterPreAnimatedTokenProducer : IMovieScenePreAnimatedToke
 
 		FToken Token(SequenceID);
 
-		if (UControlRig* ControlRig = Cast<UControlRig>(&Object))
+		if (UBaseControlRig* ControlRig = Cast<UBaseControlRig>(&Object))
 		{
 			URigHierarchy* RigHierarchy = ControlRig->GetHierarchy();
 			TArray<FRigControlElement*> Controls = ControlRig->AvailableControls();
@@ -1088,9 +1088,9 @@ struct FControlRigParameterPreAnimatedTokenProducer : IMovieScenePreAnimatedToke
 
 };
 
-static UControlRig* GetControlRig(const UMovieSceneControlRigParameterSection* Section,UObject* BoundObject)
+static UBaseControlRig* GetControlRig(const UMovieSceneControlRigParameterSection* Section,UObject* BoundObject)
 {
-	UControlRig* ControlRig = Section->GetControlRig();
+	UBaseControlRig* ControlRig = Section->GetControlRig();
 	if (ControlRig->GetObjectBinding())
 	{
 		if (UControlRigComponent* ControlRigComponent = Cast<UControlRigComponent>(ControlRig->GetObjectBinding()->GetBoundObject()))
@@ -1214,7 +1214,7 @@ struct FControlRigParameterExecutionToken : IMovieSceneExecutionToken
 		MOVIESCENE_DETAILED_SCOPE_CYCLE_COUNTER(MovieSceneEval_ControlRigParameterTrack_TokenExecute)
 		
 		FMovieSceneSequenceID SequenceID = Operand.SequenceID;
-		UControlRig* ControlRig = Section->GetControlRig();
+		UBaseControlRig* ControlRig = Section->GetControlRig();
 
 		// Update the animation's state
 		TArrayView<TWeakObjectPtr<>> BoundObjects = Player.FindBoundObjects(Operand);
@@ -1467,7 +1467,7 @@ struct TControlRigParameterActuatorFloat : TMovieSceneBlendingActuator<FControlR
 	{
 		const UMovieSceneControlRigParameterSection* Section = SectionData.Get();
 
-		UControlRig* ControlRig = Section ? GetControlRig(Section, InObject) : nullptr;
+		UBaseControlRig* ControlRig = Section ? GetControlRig(Section, InObject) : nullptr;
 
 		if (ControlRig)
 		{
@@ -1495,7 +1495,7 @@ struct TControlRigParameterActuatorFloat : TMovieSceneBlendingActuator<FControlR
 			bWasDoNotKey = Section->GetDoNotKey();
 			Section->SetDoNotKey(true);
 
-			UControlRig* ControlRig = GetControlRig(Section, InObject);
+			UBaseControlRig* ControlRig = GetControlRig(Section, InObject);
 
 			if (ControlRig && (Section->ControlsToSet.Num() == 0 || Section->ControlsToSet.Contains(ParameterName)))
 			{
@@ -1541,7 +1541,7 @@ struct TControlRigParameterActuatorVector2D : TMovieSceneBlendingActuator<FContr
 	{
 		const UMovieSceneControlRigParameterSection* Section = SectionData.Get();
 
-		UControlRig* ControlRig = Section ? GetControlRig(Section, InObject) : nullptr;
+		UBaseControlRig* ControlRig = Section ? GetControlRig(Section, InObject) : nullptr;
 
 		if (ControlRig)
 		{
@@ -1568,7 +1568,7 @@ struct TControlRigParameterActuatorVector2D : TMovieSceneBlendingActuator<FContr
 			bWasDoNotKey = Section->GetDoNotKey();
 			Section->SetDoNotKey(true);
 
-			UControlRig* ControlRig = GetControlRig(Section, InObject);
+			UBaseControlRig* ControlRig = GetControlRig(Section, InObject);
 
 
 			if (ControlRig && (Section->ControlsToSet.Num() == 0 || Section->ControlsToSet.Contains(ParameterName)))
@@ -1614,7 +1614,7 @@ struct TControlRigParameterActuatorVector : TMovieSceneBlendingActuator<FControl
 	{
 		const UMovieSceneControlRigParameterSection* Section = SectionData.Get();
 
-		UControlRig* ControlRig = Section ? GetControlRig(Section, InObject) : nullptr;
+		UBaseControlRig* ControlRig = Section ? GetControlRig(Section, InObject) : nullptr;
 
 		if (ControlRig)
 		{
@@ -1646,7 +1646,7 @@ struct TControlRigParameterActuatorVector : TMovieSceneBlendingActuator<FControl
 			const bool bSetupUndo = false;
 			bWasDoNotKey = Section->GetDoNotKey();
 			Section->SetDoNotKey(true);
-			UControlRig* ControlRig = GetControlRig(Section, InObject);
+			UBaseControlRig* ControlRig = GetControlRig(Section, InObject);
 
 			if (ControlRig && (Section->ControlsToSet.Num() == 0 || Section->ControlsToSet.Contains(ParameterName)))
 			{
@@ -1694,7 +1694,7 @@ struct TControlRigParameterActuatorTransform : TMovieSceneBlendingActuator<FCont
 	{
 		const UMovieSceneControlRigParameterSection* Section = SectionData.Get();
 
-		UControlRig* ControlRig = Section ? GetControlRig(Section, InObject) : nullptr;
+		UBaseControlRig* ControlRig = Section ? GetControlRig(Section, InObject) : nullptr;
 		URigHierarchy* Hierarchy = ControlRig->GetHierarchy();
 
 		if (ControlRig && Section && (Section->ControlsToSet.Num() == 0 || Section->ControlsToSet.Contains(ParameterName)))
@@ -1750,7 +1750,7 @@ struct TControlRigParameterActuatorTransform : TMovieSceneBlendingActuator<FCont
 			bWasDoNotKey = Section->GetDoNotKey();
 			Section->SetDoNotKey(true);
 			const bool bSetupUndo = false;
-			UControlRig* ControlRig = GetControlRig(Section, InObject);
+			UBaseControlRig* ControlRig = GetControlRig(Section, InObject);
 
 			if (ControlRig && (Section->ControlsToSet.Num() == 0 || Section->ControlsToSet.Contains(ParameterName)))
 			{
diff --git a/Engine/Plugins/Animation/ControlRig/Source/ControlRig/Private/Sequencer/MovieSceneControlRigParameterTrack.cpp b/Engine/Plugins/Animation/ControlRig/Source/ControlRig/Private/Sequencer/MovieSceneControlRigParameterTrack.cpp
--- a/Engine/Plugins/Animation/ControlRig/Source/ControlRig/Private/Sequencer/MovieSceneControlRigParameterTrack.cpp
+++ b/Engine/Plugins/Animation/ControlRig/Source/ControlRig/Private/Sequencer/MovieSceneControlRigParameterTrack.cpp
@@ -206,7 +206,7 @@ FText UMovieSceneControlRigParameterTrack::GetDefaultDisplayName() const
 #endif
 
 
-UMovieSceneSection* UMovieSceneControlRigParameterTrack::CreateControlRigSection(FFrameNumber StartTime, UControlRig* InControlRig, bool bInOwnsControlRig)
+UMovieSceneSection* UMovieSceneControlRigParameterTrack::CreateControlRigSection(FFrameNumber StartTime, UBaseControlRig* InControlRig, bool bInOwnsControlRig)
 {
 	if (InControlRig == nullptr)
 	{
@@ -558,7 +558,7 @@ void UMovieSceneControlRigParameterTrack::HandleOnPostConstructed_GameThread()
 	}
 }
 
-void UMovieSceneControlRigParameterTrack::HandleOnPostConstructed(UControlRig* Subject, const FName& InEventName)
+void UMovieSceneControlRigParameterTrack::HandleOnPostConstructed(UBaseControlRig* Subject, const FName& InEventName)
 {
 	MovieSceneControlRigTrack::AsyncHelpers::ExecuteOnGameThread<void>([this]()
 	{
@@ -648,7 +648,7 @@ void UMovieSceneControlRigParameterTrack::RenameParameterName(const FName& OldPa
 	}
 }
 
-void UMovieSceneControlRigParameterTrack::ReplaceControlRig(UControlRig* NewControlRig, bool RecreateChannels)
+void UMovieSceneControlRigParameterTrack::ReplaceControlRig(UBaseControlRig* NewControlRig, bool RecreateChannels)
 {
 	if (IsValid(ControlRig))
 	{
diff --git a/Engine/Plugins/Animation/ControlRig/Source/ControlRig/Private/Tools/ControlRigPose.cpp b/Engine/Plugins/Animation/ControlRig/Source/ControlRig/Private/Tools/ControlRigPose.cpp
--- a/Engine/Plugins/Animation/ControlRig/Source/ControlRig/Private/Tools/ControlRigPose.cpp
+++ b/Engine/Plugins/Animation/ControlRig/Source/ControlRig/Private/Tools/ControlRigPose.cpp
@@ -12,7 +12,7 @@
 
 #define LOCTEXT_NAMESPACE "ControlRigPose"
 
-void FControlRigControlPose::SavePose(UControlRig* ControlRig, bool bUseAll)
+void FControlRigControlPose::SavePose(UBaseControlRig* ControlRig, bool bUseAll)
 {
 	TArray<FRigControlElement*> CurrentControls;
 	ControlRig->GetControlsInOrder(CurrentControls);
@@ -50,7 +50,7 @@ void FControlRigControlPose::SavePose(UControlRig* ControlRig, bool bUseAll)
 	SetUpControlMap();
 }
 
-void FControlRigControlPose::PastePose(UControlRig* ControlRig, bool bDoKey, bool bDoMirror)
+void FControlRigControlPose::PastePose(UBaseControlRig* ControlRig, bool bDoKey, bool bDoMirror)
 {
 	PastePoseInternal(ControlRig, bDoKey, bDoMirror, CopyOfControls);
 	ControlRig->Evaluate_AnyThread();
@@ -58,7 +58,7 @@ void FControlRigControlPose::PastePose(UControlRig* ControlRig, bool bDoKey, boo
 
 }
 
-void FControlRigControlPose::SetControlMirrorTransform(bool bDoLocal, UControlRig* ControlRig, const FName& Name, bool bIsMatched, 
+void FControlRigControlPose::SetControlMirrorTransform(bool bDoLocal, UBaseControlRig* ControlRig, const FName& Name, bool bIsMatched, 
 	const FTransform& GlobalTransform, const FTransform& LocalTransform, bool bNotify, const  FRigControlModifiedContext&Context,bool bSetupUndo)
 {
 	if (bDoLocal || bIsMatched)
@@ -72,7 +72,7 @@ void FControlRigControlPose::SetControlMirrorTransform(bool bDoLocal, UControlRi
 	}	
 }
 
-void FControlRigControlPose::PastePoseInternal(UControlRig* ControlRig, bool bDoKey, bool bDoMirror, const TArray<FRigControlCopy>& ControlsToPaste)
+void FControlRigControlPose::PastePoseInternal(UBaseControlRig* ControlRig, bool bDoKey, bool bDoMirror, const TArray<FRigControlCopy>& ControlsToPaste)
 {
 	FRigControlModifiedContext Context;
 	Context.SetKey = bDoKey ? EControlRigSetKey::Always : EControlRigSetKey::DoNotCare;
@@ -178,7 +178,7 @@ void FControlRigControlPose::PastePoseInternal(UControlRig* ControlRig, bool bDo
 	}
 }
 
-void FControlRigControlPose::BlendWithInitialPoses(FControlRigControlPose& InitialPose, UControlRig* ControlRig, bool bDoKey, bool bDoMirror, float BlendValue)
+void FControlRigControlPose::BlendWithInitialPoses(FControlRigControlPose& InitialPose, UBaseControlRig* ControlRig, bool bDoKey, bool bDoMirror, float BlendValue)
 {
 	if (InitialPose.CopyOfControls.Num() == 0)
 	{
@@ -345,12 +345,12 @@ void UControlRigPoseAsset::PostLoad()
 	Pose.SetUpControlMap();
 }
 
-void UControlRigPoseAsset::SavePose(UControlRig* InControlRig, bool bUseAll)
+void UControlRigPoseAsset::SavePose(UBaseControlRig* InControlRig, bool bUseAll)
 {
 	Pose.SavePose(InControlRig,bUseAll);
 }
 
-void UControlRigPoseAsset::PastePose(UControlRig* InControlRig, bool bDoKey, bool bDoMirror)
+void UControlRigPoseAsset::PastePose(UBaseControlRig* InControlRig, bool bDoKey, bool bDoMirror)
 {
 #if WITH_EDITOR
 	FScopedTransaction ScopedTransaction(LOCTEXT("PastePoseTransaction", "Paste Pose"));
@@ -359,7 +359,7 @@ void UControlRigPoseAsset::PastePose(UControlRig* InControlRig, bool bDoKey, boo
 	Pose.PastePose(InControlRig,bDoKey, bDoMirror);
 }
 
-void UControlRigPoseAsset::SelectControls(UControlRig* InControlRig, bool bDoMirror)
+void UControlRigPoseAsset::SelectControls(UBaseControlRig* InControlRig, bool bDoMirror)
 {
 #if WITH_EDITOR
 	FScopedTransaction ScopedTransaction(LOCTEXT("SelectControlTransaction", "Select Control"));
@@ -395,20 +395,20 @@ void UControlRigPoseAsset::SelectControls(UControlRig* InControlRig, bool bDoMir
 	}
 }
 
-void UControlRigPoseAsset::GetCurrentPose(UControlRig* InControlRig, FControlRigControlPose& OutPose)
+void UControlRigPoseAsset::GetCurrentPose(UBaseControlRig* InControlRig, FControlRigControlPose& OutPose)
 {
 	OutPose.SavePose(InControlRig, true);
 }
 
 
-TArray<FRigControlCopy> UControlRigPoseAsset::GetCurrentPose(UControlRig* InControlRig) 
+TArray<FRigControlCopy> UControlRigPoseAsset::GetCurrentPose(UBaseControlRig* InControlRig) 
 {
 	FControlRigControlPose TempPose;
 	TempPose.SavePose(InControlRig,true);
 	return TempPose.GetPoses();
 }
 
-void UControlRigPoseAsset::BlendWithInitialPoses(FControlRigControlPose& InitialPose, UControlRig* InControlRig, bool bDoKey, bool bDoMirror, float BlendValue)
+void UControlRigPoseAsset::BlendWithInitialPoses(FControlRigControlPose& InitialPose, UBaseControlRig* InControlRig, bool bDoKey, bool bDoMirror, float BlendValue)
 {
 	if (BlendValue > 0.0f)
 	{
@@ -426,7 +426,7 @@ void UControlRigPoseAsset::ReplaceControlName(const FName& CurrentName, const FN
 	Pose.ReplaceControlName(CurrentName, NewName);
 }
 
-bool UControlRigPoseAsset::DoesMirrorMatch(UControlRig* ControlRig, const FName& ControlName) const
+bool UControlRigPoseAsset::DoesMirrorMatch(UBaseControlRig* ControlRig, const FName& ControlName) const
 {
 	FControlRigPoseMirrorTable MirrorTable;
 	MirrorTable.SetUpMirrorTable(ControlRig);
diff --git a/Engine/Plugins/Animation/ControlRig/Source/ControlRig/Private/Tools/ControlRigPoseMirrorTable.cpp b/Engine/Plugins/Animation/ControlRig/Source/ControlRig/Private/Tools/ControlRigPoseMirrorTable.cpp
--- a/Engine/Plugins/Animation/ControlRig/Source/ControlRig/Private/Tools/ControlRigPoseMirrorTable.cpp
+++ b/Engine/Plugins/Animation/ControlRig/Source/ControlRig/Private/Tools/ControlRigPoseMirrorTable.cpp
@@ -6,7 +6,7 @@
 #include "Tools/ControlRigPose.h"
 #include "RigVMFunctions/Math/RigVMMathLibrary.h"
 
-void FControlRigPoseMirrorTable::SetUpMirrorTable(const UControlRig* ControlRig)
+void FControlRigPoseMirrorTable::SetUpMirrorTable(const UBaseControlRig* ControlRig)
 {
 	const UControlRigPoseMirrorSettings* Settings = GetDefault<UControlRigPoseMirrorSettings>();
 	MatchedControls.Reset();
diff --git a/Engine/Plugins/Animation/ControlRig/Source/ControlRig/Private/Units/Execution/RigUnit_DynamicHierarchy.cpp b/Engine/Plugins/Animation/ControlRig/Source/ControlRig/Private/Units/Execution/RigUnit_DynamicHierarchy.cpp
--- a/Engine/Plugins/Animation/ControlRig/Source/ControlRig/Private/Units/Execution/RigUnit_DynamicHierarchy.cpp
+++ b/Engine/Plugins/Animation/ControlRig/Source/ControlRig/Private/Units/Execution/RigUnit_DynamicHierarchy.cpp
@@ -289,7 +289,7 @@ FRigUnit_HierarchyImportFromSkeleton_Execute()
 	{
 		FRigHierarchyControllerInstructionBracket InstructionBracket(Controller, ExecuteContext.GetInstructionIndex());
 		
-		if(const USkeletalMeshComponent* SkelMeshComponent = ExecuteContext.UnitContext.DataSourceRegistry->RequestSource<USkeletalMeshComponent>(UControlRig::OwnerComponent))
+		if(const USkeletalMeshComponent* SkelMeshComponent = ExecuteContext.UnitContext.DataSourceRegistry->RequestSource<USkeletalMeshComponent>(UBaseControlRig::OwnerComponent))
 		{
 			if(SkelMeshComponent->GetSkeletalMeshAsset())
 			{
diff --git a/Engine/Plugins/Animation/ControlRig/Source/ControlRig/Private/Units/RigUnitContext.cpp b/Engine/Plugins/Animation/ControlRig/Source/ControlRig/Private/Units/RigUnitContext.cpp
--- a/Engine/Plugins/Animation/ControlRig/Source/ControlRig/Private/Units/RigUnitContext.cpp
+++ b/Engine/Plugins/Animation/ControlRig/Source/ControlRig/Private/Units/RigUnitContext.cpp
@@ -72,7 +72,7 @@ FName FControlRigExecuteContext::AdaptMetadataName(bool bUseNameSpace, const FNa
 	return InMetadataName;
 }
 
-FControlRigExecuteContextRigModuleGuard::FControlRigExecuteContextRigModuleGuard(FControlRigExecuteContext& InContext, const UControlRig* InControlRig)
+FControlRigExecuteContextRigModuleGuard::FControlRigExecuteContextRigModuleGuard(FControlRigExecuteContext& InContext, const UBaseControlRig* InControlRig)
 	: Context(InContext)
 	, PreviousRigModuleNameSpace(InContext.RigModuleNameSpace)
 	, PreviousRigModuleNameSpaceHash(InContext.RigModuleNameSpaceHash)
diff --git a/Engine/Plugins/Animation/ControlRig/Source/ControlRig/Private/Validation/ControlRigNumericalValidationPass.cpp b/Engine/Plugins/Animation/ControlRig/Source/ControlRig/Private/Validation/ControlRigNumericalValidationPass.cpp
--- a/Engine/Plugins/Animation/ControlRig/Source/ControlRig/Private/Validation/ControlRigNumericalValidationPass.cpp
+++ b/Engine/Plugins/Animation/ControlRig/Source/ControlRig/Private/Validation/ControlRigNumericalValidationPass.cpp
@@ -20,7 +20,7 @@ UControlRigNumericalValidationPass::UControlRigNumericalValidationPass(const FOb
 {
 }
 
-void UControlRigNumericalValidationPass::OnSubjectChanged(UControlRig* InControlRig, FControlRigValidationContext* InContext)
+void UControlRigNumericalValidationPass::OnSubjectChanged(UBaseControlRig* InControlRig, FControlRigValidationContext* InContext)
 {
 	Pose.Reset();
 
@@ -36,12 +36,12 @@ void UControlRigNumericalValidationPass::OnSubjectChanged(UControlRig* InControl
 	}
 }
 
-void UControlRigNumericalValidationPass::OnInitialize(UControlRig* InControlRig, FControlRigValidationContext* InContext)
+void UControlRigNumericalValidationPass::OnInitialize(UBaseControlRig* InControlRig, FControlRigValidationContext* InContext)
 {
 	OnSubjectChanged(InControlRig, InContext);
 }
 
-void UControlRigNumericalValidationPass::OnEvent(UControlRig* InControlRig, const FName& InEventName, FControlRigValidationContext* InContext)
+void UControlRigNumericalValidationPass::OnEvent(UBaseControlRig* InControlRig, const FName& InEventName, FControlRigValidationContext* InContext)
 {
 	if (InEventName == EventNameA)
 	{
diff --git a/Engine/Plugins/Animation/ControlRig/Source/ControlRig/Private/Validation/ControlRigNumericalValidationPass.h b/Engine/Plugins/Animation/ControlRig/Source/ControlRig/Private/Validation/ControlRigNumericalValidationPass.h
--- a/Engine/Plugins/Animation/ControlRig/Source/ControlRig/Private/Validation/ControlRigNumericalValidationPass.h
+++ b/Engine/Plugins/Animation/ControlRig/Source/ControlRig/Private/Validation/ControlRigNumericalValidationPass.h
@@ -15,9 +15,9 @@ class CONTROLRIG_API UControlRigNumericalValidationPass : public UControlRigVali
 public:
 
 	// UControlRigValidationPass interface
-	virtual void OnSubjectChanged(UControlRig* InControlRig, FControlRigValidationContext* InContext) override;
-	virtual void OnInitialize(UControlRig* InControlRig, FControlRigValidationContext* InContext) override;
-	virtual void OnEvent(UControlRig* InControlRig, const FName& InEventName, FControlRigValidationContext* InContext) override;
+	virtual void OnSubjectChanged(UBaseControlRig* InControlRig, FControlRigValidationContext* InContext) override;
+	virtual void OnInitialize(UBaseControlRig* InControlRig, FControlRigValidationContext* InContext) override;
+	virtual void OnEvent(UBaseControlRig* InControlRig, const FName& InEventName, FControlRigValidationContext* InContext) override;
 
 	// If set to true the pass will validate the poses of all bones
 	UPROPERTY(EditAnywhere, Category = "Settings")
diff --git a/Engine/Plugins/Animation/ControlRig/Source/ControlRig/Public/AnimNode_ControlRig.h b/Engine/Plugins/Animation/ControlRig/Source/ControlRig/Public/AnimNode_ControlRig.h
--- a/Engine/Plugins/Animation/ControlRig/Source/ControlRig/Public/AnimNode_ControlRig.h
+++ b/Engine/Plugins/Animation/ControlRig/Source/ControlRig/Public/AnimNode_ControlRig.h
@@ -21,8 +21,8 @@ struct CONTROLRIG_API FAnimNode_ControlRig : public FAnimNode_ControlRigBase
 	FAnimNode_ControlRig();
 	~FAnimNode_ControlRig();
 
-	virtual UControlRig* GetControlRig() const override { return ControlRig; }
-	virtual TSubclassOf<UControlRig> GetControlRigClass() const override { return ControlRigClass; }
+	virtual UBaseControlRig* GetControlRig() const override { return ControlRig; }
+	virtual TSubclassOf<UBaseControlRig> GetControlRigClass() const override { return ControlRigClass; }
 
 	// FAnimNode_Base interface
 	virtual void OnInitializeAnimInstance(const FAnimInstanceProxy* InProxy, const UAnimInstance* InAnimInstance) override;
@@ -47,10 +47,10 @@ private:
 
 	/** Cached ControlRig */
 	UPROPERTY(EditAnywhere, Category = ControlRig)
-	TSubclassOf<UControlRig> ControlRigClass;
+	TSubclassOf<UBaseControlRig> ControlRigClass;
 
 	UPROPERTY(transient)
-	TObjectPtr<UControlRig> ControlRig;
+	TObjectPtr<UBaseControlRig> ControlRig;
 
 	// alpha value handler
 	UPROPERTY(EditAnywhere, Category = Settings, meta = (PinShownByDefault))
@@ -122,8 +122,8 @@ private:
 
 protected:
 	virtual UClass* GetTargetClass() const override { return *ControlRigClass; }
-	virtual void UpdateInput(UControlRig* InControlRig, const FPoseContext& InOutput) override;
-	virtual void UpdateOutput(UControlRig* InControlRig, FPoseContext& InOutput) override;
+	virtual void UpdateInput(UBaseControlRig* InControlRig, const FPoseContext& InOutput) override;
+	virtual void UpdateOutput(UBaseControlRig* InControlRig, FPoseContext& InOutput) override;
 
 	// Helper function to update the initial ref pose within the Control Rig if needed
 	void UpdateControlRigRefPoseIfNeeded(const FAnimInstanceProxy* InProxy, bool bIncludePoseInHash = false);
diff --git a/Engine/Plugins/Animation/ControlRig/Source/ControlRig/Public/AnimNode_ControlRigBase.h b/Engine/Plugins/Animation/ControlRig/Source/ControlRig/Public/AnimNode_ControlRigBase.h
--- a/Engine/Plugins/Animation/ControlRig/Source/ControlRig/Public/AnimNode_ControlRigBase.h
+++ b/Engine/Plugins/Animation/ControlRig/Source/ControlRig/Public/AnimNode_ControlRigBase.h
@@ -7,7 +7,7 @@
 #include "AnimNode_ControlRigBase.generated.h"
 
 class UNodeMappingContainer;
-class UControlRig;
+class UBaseControlRig;
 
 /** Struct defining the settings to override when driving a control rig */
 USTRUCT()
@@ -63,8 +63,8 @@ struct CONTROLRIG_API FAnimNode_ControlRigBase : public FAnimNode_CustomProperty
 	FAnimNode_ControlRigBase();
 
 	/* return Control Rig of current object */
-	virtual UControlRig* GetControlRig() const PURE_VIRTUAL(FAnimNode_ControlRigBase::GetControlRig, return nullptr; );
-	virtual TSubclassOf<UControlRig> GetControlRigClass() const PURE_VIRTUAL(FAnimNode_ControlRigBase::GetControlRigClass, return nullptr; );
+	virtual UBaseControlRig* GetControlRig() const PURE_VIRTUAL(FAnimNode_ControlRigBase::GetControlRig, return nullptr; );
+	virtual TSubclassOf<UBaseControlRig> GetControlRigClass() const PURE_VIRTUAL(FAnimNode_ControlRigBase::GetControlRigClass, return nullptr; );
 	
 	// FAnimNode_Base interface
 	virtual void OnInitializeAnimInstance(const FAnimInstanceProxy* InProxy, const UAnimInstance* InAnimInstance) override;
@@ -170,17 +170,17 @@ protected:
 
 	virtual bool CanExecute();
 	// update input/output to control rig
-	virtual void UpdateInput(UControlRig* ControlRig, const FPoseContext& InOutput);
-	virtual void UpdateOutput(UControlRig* ControlRig, FPoseContext& InOutput);
+	virtual void UpdateInput(UBaseControlRig* ControlRig, const FPoseContext& InOutput);
+	virtual void UpdateOutput(UBaseControlRig* ControlRig, FPoseContext& InOutput);
 	virtual UClass* GetTargetClass() const override;
 	
 	// execute control rig on the input pose and outputs the result
 	void ExecuteControlRig(FPoseContext& InOutput);
 
-	void QueueControlRigDrawInstructions(UControlRig* ControlRig, FAnimInstanceProxy* Proxy) const;
+	void QueueControlRigDrawInstructions(UBaseControlRig* ControlRig, FAnimInstanceProxy* Proxy) const;
 
 	TArray<TObjectPtr<UAssetUserData>> GetAssetUserData() const { return AssetUserData; }
-	void UpdateGetAssetUserDataDelegate(UControlRig* InControlRig) const;
+	void UpdateGetAssetUserDataDelegate(UBaseControlRig* InControlRig) const;
 
 	bool bControlRigRequiresInitialization;
 	uint16 LastBonesSerialNumberForCacheBones;
diff --git a/Engine/Plugins/Animation/ControlRig/Source/ControlRig/Public/AnimNode_ControlRig_ExternalSource.h b/Engine/Plugins/Animation/ControlRig/Source/ControlRig/Public/AnimNode_ControlRig_ExternalSource.h
--- a/Engine/Plugins/Animation/ControlRig/Source/ControlRig/Public/AnimNode_ControlRig_ExternalSource.h
+++ b/Engine/Plugins/Animation/ControlRig/Source/ControlRig/Public/AnimNode_ControlRig_ExternalSource.h
@@ -15,13 +15,13 @@ struct CONTROLRIG_API FAnimNode_ControlRig_ExternalSource : public FAnimNode_Con
 
 	FAnimNode_ControlRig_ExternalSource();
 
-	void SetControlRig(UControlRig* InControlRig);
-	virtual UControlRig* GetControlRig() const override;
-	virtual TSubclassOf<UControlRig> GetControlRigClass() const override;
+	void SetControlRig(UBaseControlRig* InControlRig);
+	virtual UBaseControlRig* GetControlRig() const override;
+	virtual TSubclassOf<UBaseControlRig> GetControlRigClass() const override;
 	virtual void Initialize_AnyThread(const FAnimationInitializeContext& Context) override;
 
 private:
 	UPROPERTY(transient)
-	TWeakObjectPtr<UControlRig> ControlRig;
+	TWeakObjectPtr<UBaseControlRig> ControlRig;
 };
 
diff --git a/Engine/Plugins/Animation/ControlRig/Source/ControlRig/Public/BaseControlRig.h b/Engine/Plugins/Animation/ControlRig/Source/ControlRig/Public/BaseControlRig.h
--- a/Engine/Plugins/Animation/ControlRig/Source/ControlRig/Public/BaseControlRig.h
+++ b/Engine/Plugins/Animation/ControlRig/Source/ControlRig/Public/BaseControlRig.h
@@ -47,7 +47,7 @@ CONTROLRIG_API DECLARE_LOG_CATEGORY_EXTERN(LogControlRig, Log, All);
 
 /** Runs logic for mapping input data to transforms (the "Rig") */
 UCLASS(Blueprintable, Abstract, editinlinenew)
-class CONTROLRIG_API UControlRig : public URigVMHost, public INodeMappingProviderInterface
+class CONTROLRIG_API UBaseControlRig : public URigVMHost, public INodeMappingProviderInterface
 {
 	GENERATED_UCLASS_BODY()
 
@@ -57,27 +57,27 @@ class CONTROLRIG_API UControlRig : public URigVMHost, public INodeMappingProvide
 public:
 
 	/** Bindable event for external objects to contribute to / filter a control value */
-	DECLARE_EVENT_ThreeParams(UControlRig, FFilterControlEvent, UControlRig*, FRigControlElement*, FRigControlValue&);
+	DECLARE_EVENT_ThreeParams(UBaseControlRig, FFilterControlEvent, UBaseControlRig*, FRigControlElement*, FRigControlValue&);
 
 	/** Bindable event for external objects to be notified of Control changes */
-	DECLARE_EVENT_ThreeParams(UControlRig, FControlModifiedEvent, UControlRig*, FRigControlElement*, const FRigControlModifiedContext&);
+	DECLARE_EVENT_ThreeParams(UBaseControlRig, FControlModifiedEvent, UBaseControlRig*, FRigControlElement*, const FRigControlModifiedContext&);
 
 	/** Bindable event for external objects to be notified that a Control is Selected */
-	DECLARE_EVENT_ThreeParams(UControlRig, FControlSelectedEvent, UControlRig*, FRigControlElement*, bool);
+	DECLARE_EVENT_ThreeParams(UBaseControlRig, FControlSelectedEvent, UBaseControlRig*, FRigControlElement*, bool);
 
 	/** Bindable event to manage undo / redo brackets in the client */
-	DECLARE_EVENT_TwoParams(UControlRig, FControlUndoBracketEvent, UControlRig*, bool /* bOpen */);
-	
+	DECLARE_EVENT_TwoParams(UBaseControlRig, FControlUndoBracketEvent, UBaseControlRig*, bool /* bOpen */);
+
 	// To support Blueprints/scripting, we need a different delegate type (a 'Dynamic' delegate) which supports looser style UFunction binding (using names).
-	DECLARE_DYNAMIC_MULTICAST_SPARSE_DELEGATE_ThreeParams(FOnControlSelectedBP, UControlRig, OnControlSelected_BP, UControlRig*, Rig, const FRigControlElement&, Control, bool, bSelected);
+	DECLARE_DYNAMIC_MULTICAST_SPARSE_DELEGATE_ThreeParams(FOnControlSelectedBP, UBaseControlRig, OnControlSelected_BP, UBaseControlRig*, Rig, const FRigControlElement&, Control, bool, bSelected);
 
 	/** Bindable event to notify object binding change. */
-	DECLARE_EVENT_OneParam(UControlRig, FControlRigBoundEvent, UControlRig*);
+	DECLARE_EVENT_OneParam(UBaseControlRig, FControlRigBoundEvent, UBaseControlRig*);
 
 	static const FName OwnerComponent;
 
 	UFUNCTION(BlueprintCallable, Category = ControlRig)
-	static TArray<UControlRig*> FindControlRigs(UObject* Outer, TSubclassOf<UControlRig> OptionalClass);
+	static TArray<UBaseControlRig*> FindControlRigs(UObject* Outer, TSubclassOf<UBaseControlRig> OptionalClass);
 
 public:
 	virtual UWorld* GetWorld() const override;
@@ -96,7 +96,7 @@ public:
 	bool IsRigModuleInstance() const;
 
 	// Returns the parent rig hosting this module instance
-	UControlRig* GetParentRig() const;
+	UBaseControlRig* GetParentRig() const;
 
 	// Returns the namespace of this module (for example ArmModule::)
 	const FString& GetRigModuleNameSpace() const;
@@ -124,9 +124,15 @@ public:
 	/** Initialize the VM */
 	virtual bool InitializeVM(const FName& InEventName) override;
 
+	virtual void InitializeVMs(bool bInitRigUnits = true) PURE_VIRTUAL(UBaseControlRig::InitializeVMs, return Super::Initialize(bInitRigUnits); );
+	virtual bool InitializeVMs(const FName& InEventName) PURE_VIRTUAL(UBaseControlRig::InitializeVMs, return Super::InitializeVM(InEventName); );
+
 	/** Evaluates the ControlRig */
 	virtual void Evaluate_AnyThread() override;
 
+	/** Calls Evaluate_AnyThread on all VMs in this Rig */
+	virtual void EvaluateVMs_AnyThread() PURE_VIRTUAL(UBaseControlRig::EvaluateVMs_AnyThread, return Super::Evaluate_AnyThread(); );
+
 	/** Ticks animation of the skeletal mesh component bound to this control rig */
 	bool EvaluateSkeletalMeshComponent(double InDeltaTime);
 
@@ -171,7 +177,7 @@ public:
 #if WITH_EDITOR
 
 	// called after post reinstance when compilng blueprint by Sequencer
-	void PostReinstanceCallback(const UControlRig* Old);
+	void PostReinstanceCallback(const UBaseControlRig* Old);
 
 #endif // WITH_EDITOR
 	
@@ -186,8 +192,9 @@ public:
 	FRigHierarchySettings HierarchySettings;
 
 	virtual bool Execute(const FName& InEventName) override;
-	virtual bool Execute_Internal(const FName& InEventName) override;
+	virtual bool Execute_Internal(const FName& InEventName) override PURE_VIRTUAL(UBaseControlRig::Execute_Internal, return Super::Execute_Internal(InEventName); );
 	virtual void RequestInit() override;
+	virtual void RequestInitVMs() PURE_VIRTUAL(UBaseControlRig::RequestInitVMs, return; );
 
 	bool AllConnectorsAreResolved(FString* OutFailureReason = nullptr, FRigElementKey* OutConnector = nullptr) const;
 
@@ -312,7 +319,7 @@ public:
 
 	bool IsCurveControl(const FRigControlElement* InControlElement) const;
 
-	DECLARE_EVENT_TwoParams(UControlRig, FControlRigExecuteEvent, class UControlRig*, const FName&);
+	DECLARE_EVENT_TwoParams(UBaseControlRig, FControlRigExecuteEvent, class UBaseControlRig*, const FName&);
 #if WITH_EDITOR
 	FControlRigExecuteEvent& OnPreConstructionForUI_AnyThread() { return PreConstructionForUIEvent; }
 #endif
@@ -394,17 +401,17 @@ protected:
 
 private:
 	
-	void HandleOnControlModified(UControlRig* Subject, FRigControlElement* Control, const FRigControlModifiedContext& Context);
+	void HandleOnControlModified(UBaseControlRig* Subject, FRigControlElement* Control, const FRigControlModifiedContext& Context);
 
 public:
 	
 	class CONTROLRIG_API FAnimAttributeContainerPtrScope
 	{
 	public:
-		FAnimAttributeContainerPtrScope(UControlRig* InControlRig, UE::Anim::FStackAttributeContainer& InExternalContainer);
+		FAnimAttributeContainerPtrScope(UBaseControlRig* InControlRig, UE::Anim::FStackAttributeContainer& InExternalContainer);
 		~FAnimAttributeContainerPtrScope();
 
-		UControlRig* ControlRig;
+		UBaseControlRig* ControlRig;
 	};
 	
 private:
@@ -450,6 +457,7 @@ private:
 	/** Handle changes within the hierarchy */
 	void HandleHierarchyModified(ERigHierarchyNotification InNotification, URigHierarchy* InHierarchy, const FRigBaseElement* InElement);
 
+private:
 #if WITH_EDITOR
 	/** Add a transient / temporary control used to interact with a node */
 	FName AddTransientControl(const URigVMUnitNode* InNode, const FRigDirectManipulationTarget& InTarget);
@@ -492,13 +500,14 @@ public:
 	
 #endif
 
-private: 
+protected:
 
 	void HandleHierarchyEvent(URigHierarchy* InHierarchy, const FRigEventContext& InEvent);
 	FRigEventDelegate RigEventDelegate;
 
 	void OnAddShapeLibrary(const FControlRigExecuteContext* InContext, const FString& InLibraryName, UControlRigShapeLibrary* InShapeLibrary, bool bReplaceExisting, bool bLogResults);
 	bool OnShapeExists(const FName& InShapeName) const;
+	virtual void InitializeVMsFromCDO() PURE_VIRTUAL(UBaseControlRig::InitializeVMsFromCDO, return; );
 	virtual void InitializeFromCDO() override;
 
 
@@ -507,46 +516,11 @@ private:
 
 	const FRigInfluenceMap* FindInfluenceMap(const FName& InEventName);
 
-#if WITH_EDITORONLY_DATA
-	UPROPERTY()
-	TObjectPtr<UControlRig> InteractionRig_DEPRECATED;
-
-	UPROPERTY()
-	TSubclassOf<UControlRig> InteractionRigClass_DEPRECATED;
-#endif
 
 	FRigElementKeyRedirector ElementKeyRedirector;
 
 public:
 
-	UE_DEPRECATED(5.4, "InteractionRig is no longer used")
-	UFUNCTION(BlueprintGetter, meta = (DeprecatedFunction, DeprecationMessage = "InteractionRig is no longer used"))
-	UControlRig* GetInteractionRig() const 
-	{
-#if WITH_EDITORONLY_DATA
-		return InteractionRig_DEPRECATED;
-#endif
-		return nullptr;
-	}
-
-	UE_DEPRECATED(5.4, "InteractionRig is no longer used")
-	UFUNCTION(BlueprintSetter, meta = (DeprecatedFunction, DeprecationMessage = "InteractionRig is no longer used"))
-	void SetInteractionRig(UControlRig* InInteractionRig) {}
-
-	UE_DEPRECATED(5.4, "InteractionRig is no longer used")
-	UFUNCTION(BlueprintGetter, meta = (DeprecatedFunction, DeprecationMessage = "InteractionRig is no longer used"))
-	TSubclassOf<UControlRig> GetInteractionRigClass() const 
-	{ 
-#if WITH_EDITORONLY_DATA
-		return InteractionRigClass_DEPRECATED;
-#endif
-		return nullptr;
-	}
-
-	UE_DEPRECATED(5.4, "InteractionRig is no longer used")
-	UFUNCTION(BlueprintSetter, meta = (DeprecatedFunction, DeprecationMessage = "InteractionRig is no longer used"))
-	void SetInteractionRigClass(TSubclassOf<UControlRig> InInteractionRigClass) {}
-
 	// UObject interface
 #if WITH_EDITOR
 	virtual void PreEditChange(FProperty* PropertyAboutToChange) override;
@@ -571,7 +545,7 @@ protected:
 
 private:
 
-	void CopyPoseFromOtherRig(UControlRig* Subject);
+	void CopyPoseFromOtherRig(UBaseControlRig* Subject);
 
 protected:
 	bool bCopyHierarchyBeforeConstruction;
@@ -620,8 +594,11 @@ protected:
 		return InterRigSyncBracket > 0;
 	}
 
+
 #if WITH_EDITOR
-	static void OnHierarchyTransformUndoRedoWeak(URigHierarchy* InHierarchy, const FRigElementKey& InKey, ERigTransformType::Type InTransformType, const FTransform& InTransform, bool bIsUndo, TWeakObjectPtr<UControlRig> WeakThis)
+	virtual void SetFirstEntryEventInEventQueue(FRigVMExtendedExecuteContext& Context, const FName& InFirstEventName) PURE_VIRTUAL(UBaseControlRig::SetFirstEntryEventInEventQueue, return; );
+
+	static void OnHierarchyTransformUndoRedoWeak(URigHierarchy* InHierarchy, const FRigElementKey& InKey, ERigTransformType::Type InTransformType, const FTransform& InTransform, bool bIsUndo, TWeakObjectPtr<UBaseControlRig> WeakThis)
 	{
 		if(WeakThis.IsValid() && InHierarchy != nullptr)
 		{
@@ -653,6 +630,8 @@ private:
 
 #endif
 
+
+protected:
 	/** An additive control rig runs a backwards solve before applying additive control values
 	 * and running the forward solve
 	 */
@@ -672,6 +651,8 @@ private:
 	FRigPose ControlsAfterBackwardsSolve;
 	TMap<FRigElementKey, FRigSetControlValueInfo> ControlValues; // Additive values in local space (to add after backwards solve)
 
+
+private:
 	float DebugBoneRadiusMultiplier;
 	
 public:
@@ -697,13 +678,13 @@ private:
 	class FPoseScope
 	{
 	public:
-		FPoseScope(UControlRig* InControlRig, ERigElementType InFilter = ERigElementType::All,
+		FPoseScope(UBaseControlRig* InControlRig, ERigElementType InFilter = ERigElementType::All,
 			const TArray<FRigElementKey>& InElements = TArray<FRigElementKey>(), const ERigTransformType::Type InTransformType = ERigTransformType::CurrentLocal);
 		~FPoseScope();
 
 	private:
 
-		UControlRig* ControlRig;
+		UBaseControlRig* ControlRig;
 		ERigElementType Filter;
 		FRigPose CachedPose;
 		ERigTransformType::Type TransformType;
@@ -715,6 +696,9 @@ private:
 	UPROPERTY(transient)
 	mutable FString RigModuleNameSpace;
 
+
+public:
+
 #if WITH_EDITOR
 
 	// Class used to temporarily cache current transient controls
@@ -742,14 +726,13 @@ private:
 		TObjectPtr<URigHierarchy> Hierarchy;
 	};
 
-public:
 	// Class used to temporarily cache current pose of transient controls
 	// restore them after a ResetPoseToInitial call,
 	// which allows user to move bones in construction mode
 	class FTransientControlPoseScope
 	{
 	public:
-		FTransientControlPoseScope(TObjectPtr<UControlRig> InControlRig)
+		FTransientControlPoseScope(TObjectPtr<UBaseControlRig> InControlRig)
 		{
 			ControlRig = InControlRig;
 
@@ -777,7 +760,7 @@ public:
 	
 	private:
 		
-		UControlRig* ControlRig;
+		UBaseControlRig* ControlRig;
 		FRigPose CachedPose;	
 	};	
 
@@ -839,7 +822,7 @@ class CONTROLRIG_API FControlRigInteractionScope
 {
 public:
 
-	FControlRigInteractionScope(UControlRig* InControlRig)
+	FControlRigInteractionScope(UBaseControlRig* InControlRig)
 		: ControlRig(InControlRig)
 		, InteractionBracketScope(InControlRig->InteractionBracket)
 		, SyncBracketScope(InControlRig->InterRigSyncBracket)
@@ -848,7 +831,7 @@ public:
 	}
 
 	FControlRigInteractionScope(
-		UControlRig* InControlRig,
+		UBaseControlRig* InControlRig,
 		const FRigElementKey& InKey,
 		EControlRigInteractionType InInteractionType = EControlRigInteractionType::All
 	)
@@ -863,7 +846,7 @@ public:
 	}
 
 	FControlRigInteractionScope(
-		UControlRig* InControlRig,
+		UBaseControlRig* InControlRig,
 		const TArray<FRigElementKey>& InKeys,
 		EControlRigInteractionType InInteractionType = EControlRigInteractionType::All
 	)
@@ -890,7 +873,7 @@ public:
 
 private:
 
-	TWeakObjectPtr<UControlRig> ControlRig;
+	TWeakObjectPtr<UBaseControlRig> ControlRig;
 	FControlRigBracketScope InteractionBracketScope;
 	FControlRigBracketScope SyncBracketScope;
 };
diff --git a/Engine/Plugins/Animation/ControlRig/Source/ControlRig/Public/Constraints/ControlRigTransformableHandle.h b/Engine/Plugins/Animation/ControlRig/Source/ControlRig/Public/Constraints/ControlRigTransformableHandle.h
--- a/Engine/Plugins/Animation/ControlRig/Source/ControlRig/Public/Constraints/ControlRigTransformableHandle.h
+++ b/Engine/Plugins/Animation/ControlRig/Source/ControlRig/Public/Constraints/ControlRigTransformableHandle.h
@@ -10,7 +10,7 @@
 struct FRigBaseElement;
 struct FRigControlElement;
 
-class UControlRig;
+class UBaseControlRig;
 class USkeletalMeshComponent;
 class UControlRigComponent;
 class URigHierarchy;
@@ -48,7 +48,7 @@ public:
 	virtual FTickFunction* GetTickFunction() const override;
 
 	/** Generates a hash value based on ControlRig and ControlName. */
-	static uint32 ComputeHash(const UControlRig* InControlRig, const FName& InControlName);
+	static uint32 ComputeHash(const UBaseControlRig* InControlRig, const FName& InControlName);
 	virtual uint32 GetHash() const override;
 	
 	/** Returns the underlying targeted object. */
@@ -92,7 +92,7 @@ public:
 
 	/** The ControlRig that this handle is pointing at. */
 	UPROPERTY(BlueprintReadOnly, Category = "Object")
-	TSoftObjectPtr<UControlRig> ControlRig;
+	TSoftObjectPtr<UBaseControlRig> ControlRig;
 
 	/** The ControlName of the control that this handle is pointing at. */
 	UPROPERTY(BlueprintReadOnly, Category = "Object")
@@ -100,7 +100,7 @@ public:
 
 	/** @todo document */
 	void OnControlModified(
-		UControlRig* InControlRig,
+		UBaseControlRig* InControlRig,
 		FRigControlElement* InControl,
 		const FRigControlModifiedContext& InContext);
 	
@@ -117,7 +117,7 @@ private:
 		URigHierarchy* InHierarchy,
 		const FRigBaseElement* InElement);
 
-	void OnControlRigBound(UControlRig* InControlRig);
+	void OnControlRigBound(UBaseControlRig* InControlRig);
 	void OnObjectBoundToControlRig(UObject* InObject);
 
 #if WITH_EDITOR
diff --git a/Engine/Plugins/Animation/ControlRig/Source/ControlRig/Public/ControlRig.h b/Engine/Plugins/Animation/ControlRig/Source/ControlRig/Public/ControlRig.h
new file mode 100644
--- /dev/null
+++ b/Engine/Plugins/Animation/ControlRig/Source/ControlRig/Public/ControlRig.h
@@ -0,0 +1,65 @@
+// Copyright Epic Games, Inc. All Rights Reserved.
+
+#pragma once
+#include "BaseControlRig.h"
+
+#include "ControlRig.generated.h"
+
+/** Runs logic for mapping input data to transforms (the "Rig") */
+UCLASS(Blueprintable, Abstract, editinlinenew)
+class CONTROLRIG_API UControlRig : public UBaseControlRig
+{
+	GENERATED_UCLASS_BODY()
+
+public:
+
+	// BEGIN BaseControlRig
+	virtual void InitializeVMs(bool bRequestInit = true) override { URigVMHost::Initialize(bRequestInit); }
+	virtual bool InitializeVMs(const FName& InEventName) override { return URigVMHost::InitializeVM(InEventName); }
+	virtual void InitializeVMsFromCDO() override { URigVMHost::InitializeFromCDO(); }
+	virtual void RequestInitVMs() override { URigVMHost::RequestInit(); }
+	virtual bool Execute_Internal(const FName& InEventName) override;
+	virtual void EvaluateVMs_AnyThread() override { URigVMHost::Evaluate_AnyThread(); }
+#if WITH_EDITOR
+	virtual void SetFirstEntryEventInEventQueue(FRigVMExtendedExecuteContext& Context, const FName& InFirstEventName) override;
+#endif
+	// END BaseControlRig
+
+	UE_DEPRECATED(5.4, "InteractionRig is no longer used") UFUNCTION(BlueprintGetter, meta = (DeprecatedFunction, DeprecationMessage = "InteractionRig is no longer used"))
+	UControlRig* GetInteractionRig() const
+	{
+#if WITH_EDITORONLY_DATA
+		return InteractionRig_DEPRECATED;
+#endif
+		return nullptr;
+	}
+
+	UE_DEPRECATED(5.4, "InteractionRig is no longer used")
+	UFUNCTION(BlueprintSetter, meta = (DeprecatedFunction, DeprecationMessage = "InteractionRig is no longer used"))
+	void SetInteractionRig(UControlRig* InInteractionRig) {}
+
+	UE_DEPRECATED(5.4, "InteractionRig is no longer used")
+	UFUNCTION(BlueprintGetter, meta = (DeprecatedFunction, DeprecationMessage = "InteractionRig is no longer used"))
+	TSubclassOf<UControlRig> GetInteractionRigClass() const
+	{
+#if WITH_EDITORONLY_DATA
+		return InteractionRigClass_DEPRECATED;
+#endif
+		return nullptr;
+	}
+
+	UE_DEPRECATED(5.4, "InteractionRig is no longer used")
+	UFUNCTION(BlueprintSetter, meta = (DeprecatedFunction, DeprecationMessage = "InteractionRig is no longer used"))
+	void SetInteractionRigClass(TSubclassOf<UControlRig> InInteractionRigClass) {}
+	
+private:
+#if WITH_EDITORONLY_DATA
+	UPROPERTY()
+	TObjectPtr<UControlRig> InteractionRig_DEPRECATED;
+
+	UPROPERTY()
+	TSubclassOf<UControlRig> InteractionRigClass_DEPRECATED;
+#endif
+
+	
+};
diff --git a/Engine/Plugins/Animation/ControlRig/Source/ControlRig/Public/ControlRigComponent.h b/Engine/Plugins/Animation/ControlRig/Source/ControlRig/Public/ControlRigComponent.h
--- a/Engine/Plugins/Animation/ControlRig/Source/ControlRig/Public/ControlRigComponent.h
+++ b/Engine/Plugins/Animation/ControlRig/Source/ControlRig/Public/ControlRigComponent.h
@@ -180,7 +180,7 @@ public:
 
 	/** The class of control rig to instantiate */
 	UPROPERTY(BlueprintReadWrite, EditAnywhere, Category = ControlRig, meta = (DisallowCreateNew), BlueprintSetter = SetControlRigClass)
-	TSubclassOf<UControlRig> ControlRigClass;
+	TSubclassOf<UBaseControlRig> ControlRigClass;
 
 	/** Event fired just before this component's ControlRig is initialized */
 	UPROPERTY(BlueprintAssignable, Category = "ControlRig", meta = (DisplayName = "On Pre Initialize"))
@@ -230,7 +230,7 @@ public:
 
 	/** Get the ControlRig hosted by this component */
 	UFUNCTION(BlueprintPure, Category = "ControlRig")
-	UControlRig* GetControlRig();
+	UBaseControlRig* GetControlRig();
 
 	/** Returns true if the Component can execute its Control Rig */
 	UFUNCTION(BlueprintPure, Category = "ControlRig")
@@ -595,10 +595,10 @@ public:
 	FControlRigComponentEvent& OnControlRigCreated() { return ControlRigCreatedEvent; }
 	public:
 		
-	void SetControlRig(UControlRig* ControlRig);
+	void SetControlRig(UBaseControlRig* ControlRig);
 
 	UFUNCTION(BlueprintCallable, Category = "ControlRig")
-	void SetControlRigClass(TSubclassOf<UControlRig> InControlRigClass);
+	void SetControlRigClass(TSubclassOf<UBaseControlRig> InControlRigClass);
 
 	UFUNCTION(BlueprintCallable, Category = "ControlRig")
 	void SetObjectBinding(UObject* InObjectToBind);
@@ -636,17 +636,17 @@ private:
 			}
 	};
 
-	UControlRig* SetupControlRigIfRequired();
+	UBaseControlRig* SetupControlRigIfRequired();
 	void ValidateMappingData();
 	void TransferInputs();
 	void TransferOutputs();
 	static FName GetComponentNameWithinActor(UActorComponent* InComponent);
 
 	void HandleControlRigInitializedEvent(URigVMHost* InControlRig, const FName& InEventName);
-	void HandleControlRigPreConstructionEvent(UControlRig* InControlRig, const FName& InEventName);
-	void HandleControlRigPostConstructionEvent(UControlRig* InControlRig, const FName& InEventName);
-	void HandleControlRigPreForwardsSolveEvent(UControlRig* InControlRig, const FName& InEventName);
-	void HandleControlRigPostForwardsSolveEvent(UControlRig* InControlRig, const FName& InEventName);
+	void HandleControlRigPreConstructionEvent(UBaseControlRig* InControlRig, const FName& InEventName);
+	void HandleControlRigPostConstructionEvent(UBaseControlRig* InControlRig, const FName& InEventName);
+	void HandleControlRigPreForwardsSolveEvent(UBaseControlRig* InControlRig, const FName& InEventName);
+	void HandleControlRigPostForwardsSolveEvent(UBaseControlRig* InControlRig, const FName& InEventName);
 	void HandleControlRigExecutedEvent(URigVMHost* InControlRig, const FName& InEventName);
 
 	void ConvertTransformToRigSpace(FTransform& InOutTransform, EControlRigComponentSpace FromSpace);
@@ -660,7 +660,7 @@ private:
 #endif
 
 	UPROPERTY(transient)
-	TObjectPtr<UControlRig> ControlRig;
+	TObjectPtr<UBaseControlRig> ControlRig;
 
 	TMap<USkeletalMeshComponent*, FCachedSkeletalMeshComponentSettings> CachedSkeletalMeshComponentSettings;
 
diff --git a/Engine/Plugins/Animation/ControlRig/Source/ControlRig/Public/ControlRigControlActor.h b/Engine/Plugins/Animation/ControlRig/Source/ControlRig/Public/ControlRigControlActor.h
--- a/Engine/Plugins/Animation/ControlRig/Source/ControlRig/Public/ControlRigControlActor.h
+++ b/Engine/Plugins/Animation/ControlRig/Source/ControlRig/Public/ControlRigControlActor.h
@@ -39,7 +39,7 @@ public:
 	TObjectPtr<class AActor> ActorToTrack;
 
 	UPROPERTY(BlueprintReadWrite, EditAnywhere, Interp, Category = "Control Actor")
-	TSubclassOf<UControlRig> ControlRigClass;
+	TSubclassOf<UBaseControlRig> ControlRigClass;
 
 	UPROPERTY(BlueprintReadWrite, EditAnywhere, Interp, Category = "Control Actor")
 	bool bRefreshOnTick;
@@ -71,7 +71,7 @@ private:
 	TObjectPtr<class USceneComponent> ActorRootComponent;
 
 	UPROPERTY(transient)
-	TSoftObjectPtr<UControlRig>  ControlRig;
+	TSoftObjectPtr<UBaseControlRig>  ControlRig;
 
 	UPROPERTY(transient)
 	TArray<FName> ControlNames;
diff --git a/Engine/Plugins/Animation/ControlRig/Source/ControlRig/Public/ControlRigGizmoActor.h b/Engine/Plugins/Animation/ControlRig/Source/ControlRig/Public/ControlRigGizmoActor.h
--- a/Engine/Plugins/Animation/ControlRig/Source/ControlRig/Public/ControlRigGizmoActor.h
+++ b/Engine/Plugins/Animation/ControlRig/Source/ControlRig/Public/ControlRigGizmoActor.h
@@ -9,7 +9,7 @@
 #include "Rigs/RigControlHierarchy.h"
 #include "ControlRigGizmoActor.generated.h"
 
-class UControlRig;
+class UBaseControlRig;
 struct FActorSpawnParameters;
 
 USTRUCT()
@@ -36,7 +36,7 @@ struct FControlShapeActorCreationParam
 
 	UObject*	ManipObj;
 	int32		ControlRigIndex;
-	UControlRig* ControlRig;
+	UBaseControlRig* ControlRig;
 	FName		ControlName;
 	FName		ShapeName;
 	FTransform	SpawnTransform;
@@ -72,7 +72,7 @@ public:
 
 	//  control rig this actor is referencing we can have multiple control rig's visible
 	UPROPERTY()
-	TWeakObjectPtr<UControlRig> ControlRig;
+	TWeakObjectPtr<UBaseControlRig> ControlRig;
 
 	// the name of the control this actor is referencing
 	UPROPERTY()
@@ -161,7 +161,7 @@ public:
 	// try to update the actor with the latest settings
 	bool UpdateControlSettings(
 		ERigHierarchyNotification InNotif,
-		UControlRig* InControlRig,
+		UBaseControlRig* InControlRig,
 		const FRigControlElement* InControlElement,
 		bool bHideManipulators,
 		bool bIsInLevelEditor);
diff --git a/Engine/Plugins/Animation/ControlRig/Source/ControlRig/Public/ControlRigTestData.h b/Engine/Plugins/Animation/ControlRig/Source/ControlRig/Public/ControlRigTestData.h
--- a/Engine/Plugins/Animation/ControlRig/Source/ControlRig/Public/ControlRigTestData.h
+++ b/Engine/Plugins/Animation/ControlRig/Source/ControlRig/Public/ControlRigTestData.h
@@ -8,7 +8,7 @@
 #include "ControlRigTestData.generated.h"
 
 class UControlRigTestData;
-class UControlRig;
+class UBaseControlRig;
 
 USTRUCT(BlueprintType)
 struct CONTROLRIG_API FControlRigTestDataVariable
@@ -45,9 +45,9 @@ struct CONTROLRIG_API FControlRigTestDataFrame
 		Pose.Reset();
 	}
 
-	bool Store(UControlRig* InControlRig, bool bInitial = false);
-	bool Restore(UControlRig* InControlRig, bool bInitial = false) const;
-	bool RestoreVariables(UControlRig* InControlRig) const;
+	bool Store(UBaseControlRig* InControlRig, bool bInitial = false);
+	bool Restore(UBaseControlRig* InControlRig, bool bInitial = false) const;
+	bool RestoreVariables(UBaseControlRig* InControlRig) const;
 
 	UPROPERTY(VisibleAnywhere, BlueprintReadOnly, Category = "ControlRigTestDataFrame")
 	double AbsoluteTime;
@@ -98,10 +98,10 @@ public:
 	int32 GetFrameIndexForTime(double InSeconds, bool bInput = false) const;
 
 	UFUNCTION(BlueprintCallable, Category = "ControlRigTestData")
-	bool Record(UControlRig* InControlRig, double InRecordingDuration = 0.0);
+	bool Record(UBaseControlRig* InControlRig, double InRecordingDuration = 0.0);
 
 	UFUNCTION(BlueprintCallable, Category = "ControlRigTestData")
-	bool SetupReplay(UControlRig* InControlRig, bool bGroundTruth = true);
+	bool SetupReplay(UBaseControlRig* InControlRig, bool bGroundTruth = true);
 
 	UFUNCTION(BlueprintCallable, Category = "ControlRigTestData")
 	void ReleaseReplay();
@@ -135,12 +135,12 @@ public:
 
 private:
 
-	void ClearDelegates(UControlRig* InControlRig);
+	void ClearDelegates(UBaseControlRig* InControlRig);
 
 	mutable int32 LastFrameIndex;
 	double DesiredRecordingDuration;
 	double TimeAtStartOfRecording;
-	TWeakObjectPtr<UControlRig> ReplayControlRig;
+	TWeakObjectPtr<UBaseControlRig> ReplayControlRig;
 	bool bIsApplyingOutputs;
 	FDelegateHandle PreConstructionHandle;
 	FDelegateHandle PreForwardHandle;
diff --git a/Engine/Plugins/Animation/ControlRig/Source/ControlRig/Public/ControlRigValidationPass.h b/Engine/Plugins/Animation/ControlRig/Source/ControlRig/Public/ControlRigValidationPass.h
--- a/Engine/Plugins/Animation/ControlRig/Source/ControlRig/Public/ControlRigValidationPass.h
+++ b/Engine/Plugins/Animation/ControlRig/Source/ControlRig/Public/ControlRigValidationPass.h
@@ -53,8 +53,8 @@ class CONTROLRIG_API UControlRigValidator : public UObject
 	UControlRigValidationPass* AddPass(UClass* InClass);
 	void RemovePass(UClass* InClass);
 
-	UControlRig* GetControlRig() const { return WeakControlRig.Get(); }
-	void SetControlRig(UControlRig* InControlRig);
+	UBaseControlRig* GetControlRig() const { return WeakControlRig.Get(); }
+	void SetControlRig(UBaseControlRig* InControlRig);
 
 	FControlRigValidationClearDelegate& OnClear() { return ValidationContext.OnClear(); }
 	FControlRigValidationReportDelegate& OnReport() { return ValidationContext.OnReport(); }
@@ -69,7 +69,7 @@ private:
 	void OnControlRigExecuted(URigVMHost* Subject, const FName& EventName);
 
 	FControlRigValidationContext ValidationContext;
-	TWeakObjectPtr<UControlRig> WeakControlRig;
+	TWeakObjectPtr<UBaseControlRig> WeakControlRig;
 	// todo FControlRigValidationControlRigChangedDelegate ControlRigChangedDelegate;
 };
 
@@ -83,11 +83,11 @@ class CONTROLRIG_API UControlRigValidationPass : public UObject
 public:
 
 	// Called whenever the rig being validated question is changed
-	virtual void OnSubjectChanged(UControlRig* InControlRig, FControlRigValidationContext* InContext) {}
+	virtual void OnSubjectChanged(UBaseControlRig* InControlRig, FControlRigValidationContext* InContext) {}
 
 	// Called whenever the rig in question is initialized
-	virtual void OnInitialize(UControlRig* InControlRig, FControlRigValidationContext* InContext) {}
+	virtual void OnInitialize(UBaseControlRig* InControlRig, FControlRigValidationContext* InContext) {}
 
 	// Called whenever the rig is running an event
-	virtual void OnEvent(UControlRig* InControlRig, const FName& InEventName, FControlRigValidationContext* InContext) {}
+	virtual void OnEvent(UBaseControlRig* InControlRig, const FName& InEventName, FControlRigValidationContext* InContext) {}
 };
diff --git a/Engine/Plugins/Animation/ControlRig/Source/ControlRig/Public/Rigs/AdditiveControlRig.h b/Engine/Plugins/Animation/ControlRig/Source/ControlRig/Public/Rigs/AdditiveControlRig.h
--- a/Engine/Plugins/Animation/ControlRig/Source/ControlRig/Public/Rigs/AdditiveControlRig.h
+++ b/Engine/Plugins/Animation/ControlRig/Source/ControlRig/Public/Rigs/AdditiveControlRig.h
@@ -11,15 +11,23 @@ struct FSmartNameMapping;
 
 /** Rig that allows additive layer editing per joint */
 UCLASS(NotBlueprintable)
-class CONTROLRIG_API UAdditiveControlRig : public UControlRig
+class CONTROLRIG_API UAdditiveControlRig : public UBaseControlRig
 {
 	GENERATED_UCLASS_BODY()
 
 public: 
-	// BEGIN ControlRig
+	// BEGIN BaseControlRig
 	virtual void Initialize(bool bInitRigUnits = true) override;
+	virtual void InitializeVMs(bool bRequestInit = true) override { URigVMHost::Initialize(bRequestInit); }
+	virtual bool InitializeVMs(const FName& InEventName) override { return URigVMHost::InitializeVM(InEventName); }
+	virtual void InitializeVMsFromCDO() override { URigVMHost::InitializeFromCDO(); }
+	virtual void RequestInitVMs() override { URigVMHost::RequestInit(); }
 	virtual bool Execute_Internal(const FName& InEventName) override;
-	// END ControlRig
+	virtual void EvaluateVMs_AnyThread() override { URigVMHost::Evaluate_AnyThread(); }
+#if WITH_EDITOR
+	virtual void SetFirstEntryEventInEventQueue(FRigVMExtendedExecuteContext& Context, const FName& InFirstEventName) override;
+#endif
+	// END BaseControlRig
 
 	// utility function to 
 	static FName GetControlName(const FName& InBoneName);
diff --git a/Engine/Plugins/Animation/ControlRig/Source/ControlRig/Public/Rigs/FKControlRig.h b/Engine/Plugins/Animation/ControlRig/Source/ControlRig/Public/Rigs/FKControlRig.h
--- a/Engine/Plugins/Animation/ControlRig/Source/ControlRig/Public/Rigs/FKControlRig.h
+++ b/Engine/Plugins/Animation/ControlRig/Source/ControlRig/Public/Rigs/FKControlRig.h
@@ -36,17 +36,25 @@ enum class EControlRigFKRigExecuteMode: uint8
 
 /** Rig that allows override editing per joint */
 UCLASS(NotBlueprintable, Meta = (DisplayName = "FK Control Rig"))
-class CONTROLRIG_API UFKControlRig : public UControlRig
+class CONTROLRIG_API UFKControlRig : public UBaseControlRig
 {
 	GENERATED_UCLASS_BODY()
 
 public: 
 
-	// BEGIN ControlRig
+	// BEGIN BaseControlRig
 	virtual void Initialize(bool bInitRigUnits = true) override;
+	virtual void InitializeVMs(bool bRequestInit = true) override { URigVMHost::Initialize(bRequestInit); }
+	virtual bool InitializeVMs(const FName& InEventName) override { return URigVMHost::InitializeVM(InEventName); }
+	virtual void InitializeVMsFromCDO() override { URigVMHost::InitializeFromCDO(); }
+	virtual void RequestInitVMs() override { URigVMHost::RequestInit(); }
 	virtual bool Execute_Internal(const FName& InEventName) override;
+	virtual void EvaluateVMs_AnyThread() override { URigVMHost::Evaluate_AnyThread(); }
+#if WITH_EDITOR
+	virtual void SetFirstEntryEventInEventQueue(FRigVMExtendedExecuteContext& Context, const FName& InFirstEventName) override;
+#endif
 	virtual void SetBoneInitialTransformsFromSkeletalMeshComponent(USkeletalMeshComponent* InSkelMeshComp, bool bUseAnimInstance = false) override;
-	// END ControlRig
+	// END BaseControlRig
 
 	// utility function to generate a valid control element name
 	static FName GetControlName(const FName& InName, const ERigElementType& InType);
diff --git a/Engine/Plugins/Animation/ControlRig/Source/ControlRig/Public/Rigs/RigHierarchy.h b/Engine/Plugins/Animation/ControlRig/Source/ControlRig/Public/Rigs/RigHierarchy.h
--- a/Engine/Plugins/Animation/ControlRig/Source/ControlRig/Public/Rigs/RigHierarchy.h
+++ b/Engine/Plugins/Animation/ControlRig/Source/ControlRig/Public/Rigs/RigHierarchy.h
@@ -18,7 +18,7 @@
 #include "Containers/Queue.h"
 #include "RigHierarchy.generated.h"
 
-class UControlRig;
+class UBaseControlRig;
 class URigHierarchy;
 class URigHierarchyController;
 
@@ -4639,6 +4639,7 @@ private:
 	void Reset_Impl(bool bResetElements);
 	
 	friend class URigHierarchyController;
+	friend class UBaseControlRig;
 	friend class UControlRig;
 	friend class FControlRigEditor;
 	friend struct FRigHierarchyValidityBracket;
@@ -4679,6 +4680,7 @@ private:
 	}
 
 	friend class URigHierarchy;
+	friend class UBaseControlRig;
 	friend class UControlRig;
 
 	// certain units are allowed to use this
@@ -4711,6 +4713,7 @@ private:
 	URigHierarchy* Hierarchy;
 	const FRigVMExtendedExecuteContext* PreviousContext;
 
+	friend class UBaseControlRig;
 	friend class UControlRig;
 };
 
@@ -4772,7 +4775,7 @@ public:
 	{
 	}
 
-	FRigHierarchyRedirectorGuard(UControlRig* InControlRig);
+	FRigHierarchyRedirectorGuard(UBaseControlRig* InControlRig);
 
 private:
 	TGuardValue<FRigElementKeyRedirector*> Guard;
diff --git a/Engine/Plugins/Animation/ControlRig/Source/ControlRig/Public/Rigs/RigHierarchyController.h b/Engine/Plugins/Animation/ControlRig/Source/ControlRig/Public/Rigs/RigHierarchyController.h
--- a/Engine/Plugins/Animation/ControlRig/Source/ControlRig/Public/Rigs/RigHierarchyController.h
+++ b/Engine/Plugins/Animation/ControlRig/Source/ControlRig/Public/Rigs/RigHierarchyController.h
@@ -750,6 +750,7 @@ private:
 		return Element;
 	}
 	
+	friend class UBaseControlRig;
 	friend class UControlRig;
 	friend class URigHierarchy;
 	friend class FRigHierarchyControllerInstructionBracket;
diff --git a/Engine/Plugins/Animation/ControlRig/Source/ControlRig/Public/Rigs/RigInfluenceMap.h b/Engine/Plugins/Animation/ControlRig/Source/ControlRig/Public/Rigs/RigInfluenceMap.h
--- a/Engine/Plugins/Animation/ControlRig/Source/ControlRig/Public/Rigs/RigInfluenceMap.h
+++ b/Engine/Plugins/Animation/ControlRig/Source/ControlRig/Public/Rigs/RigInfluenceMap.h
@@ -168,7 +168,7 @@ protected:
 	TMap<FRigElementKey, int32> KeyToIndex;
 
 	friend struct FRigInfluenceMapPerEvent;
-	friend class UControlRig;
+	friend class UBaseControlRig;
 };
 
 USTRUCT(BlueprintType)
diff --git a/Engine/Plugins/Animation/ControlRig/Source/ControlRig/Public/Sequencer/ControlRigLayerInstance.h b/Engine/Plugins/Animation/ControlRig/Source/ControlRig/Public/Sequencer/ControlRigLayerInstance.h
--- a/Engine/Plugins/Animation/ControlRig/Source/ControlRig/Public/Sequencer/ControlRigLayerInstance.h
+++ b/Engine/Plugins/Animation/ControlRig/Source/ControlRig/Public/Sequencer/ControlRigLayerInstance.h
@@ -12,7 +12,7 @@
 #include "AnimNode_ControlRigBase.h"
 #include "ControlRigLayerInstance.generated.h"
 
-class UControlRig;
+class UBaseControlRig;
 
 UCLASS(transient, NotBlueprintable)
 class CONTROLRIG_API UControlRigLayerInstance : public UAnimInstance, public ISequencerAnimationSupport
@@ -21,7 +21,7 @@ class CONTROLRIG_API UControlRigLayerInstance : public UAnimInstance, public ISe
 
 public:
 	/** ControlRig related support */
-	void AddControlRigTrack(int32 ControlRigID, UControlRig* InControlRig);
+	void AddControlRigTrack(int32 ControlRigID, UBaseControlRig* InControlRig);
 	void UpdateControlRigTrack(int32 ControlRigID, float Weight, const FControlRigIOSettings& InputSettings, bool bExecute);
 	void RemoveControlRigTrack(int32 ControlRigID);
 	bool HasControlRigTrack(int32 ControlRigID);
@@ -43,7 +43,7 @@ public:
 	virtual void SavePose() override;
 
 	/** Return the first available control rig */
-	UControlRig* GetFirstAvailableControlRig() const;
+	UBaseControlRig* GetFirstAvailableControlRig() const;
 
 	virtual UAnimInstance* GetSourceAnimInstance() override;
 	virtual void SetSourceAnimInstance(UAnimInstance* SourceAnimInstance) override;
diff --git a/Engine/Plugins/Animation/ControlRig/Source/ControlRig/Public/Sequencer/ControlRigSequenceObjectReference.h b/Engine/Plugins/Animation/ControlRig/Source/ControlRig/Public/Sequencer/ControlRigSequenceObjectReference.h
--- a/Engine/Plugins/Animation/ControlRig/Source/ControlRig/Public/Sequencer/ControlRigSequenceObjectReference.h
+++ b/Engine/Plugins/Animation/ControlRig/Source/ControlRig/Public/Sequencer/ControlRigSequenceObjectReference.h
@@ -25,7 +25,7 @@ struct FControlRigSequenceObjectReference
 	 *
 	 * @param InControlRig The ControlRig to create a reference for
 	 */
-	CONTROLRIG_API static FControlRigSequenceObjectReference Create(UControlRig* InControlRig);
+	CONTROLRIG_API static FControlRigSequenceObjectReference Create(UBaseControlRig* InControlRig);
 
 	/**
 	 * Check whether this object reference is valid or not
@@ -47,7 +47,7 @@ private:
 
 	/** The type of this animation ControlRig */
 	UPROPERTY()
-	TSubclassOf<UControlRig> ControlRigClass;
+	TSubclassOf<UBaseControlRig> ControlRigClass;
 };
 
 USTRUCT()
diff --git a/Engine/Plugins/Animation/ControlRig/Source/ControlRig/Public/Sequencer/ControlRigSequencerHelpers.h b/Engine/Plugins/Animation/ControlRig/Source/ControlRig/Public/Sequencer/ControlRigSequencerHelpers.h
--- a/Engine/Plugins/Animation/ControlRig/Source/ControlRig/Public/Sequencer/ControlRigSequencerHelpers.h
+++ b/Engine/Plugins/Animation/ControlRig/Source/ControlRig/Public/Sequencer/ControlRigSequencerHelpers.h
@@ -5,7 +5,7 @@
 #include "CoreMinimal.h"
 
 struct FChannelMapInfo;
-class UControlRig;
+class UBaseControlRig;
 class UMovieSceneControlRigParameterSection;
 struct FMovieSceneFloatChannel;
 class UMovieSceneSection;
@@ -14,14 +14,14 @@ class ISequencer;
 struct CONTROLRIG_API FControlRigSequencerHelpers
 {
 	static TPair<const FChannelMapInfo*, int32> GetInfoAndNumFloatChannels(
-		const UControlRig* InControlRig,
+		const UBaseControlRig* InControlRig,
 		const FName& InControlName,
 		const UMovieSceneControlRigParameterSection* InSection);
 
-	static TArrayView<FMovieSceneFloatChannel*> GetFloatChannels(const UControlRig* InControlRig,
+	static TArrayView<FMovieSceneFloatChannel*> GetFloatChannels(const UBaseControlRig* InControlRig,
 		const FName& InControlName, const UMovieSceneSection* InSection);
 
-	static UMovieSceneControlRigParameterSection* GetControlRigSection(ISequencer* InSequencer, const UControlRig* InControlRig);
+	static UMovieSceneControlRigParameterSection* GetControlRigSection(ISequencer* InSequencer, const UBaseControlRig* InControlRig);
 
 };
 
diff --git a/Engine/Plugins/Animation/ControlRig/Source/ControlRig/Public/Sequencer/MovieSceneControlRigParameterSection.h b/Engine/Plugins/Animation/ControlRig/Source/ControlRig/Public/Sequencer/MovieSceneControlRigParameterSection.h
--- a/Engine/Plugins/Animation/ControlRig/Source/ControlRig/Public/Sequencer/MovieSceneControlRigParameterSection.h
+++ b/Engine/Plugins/Animation/ControlRig/Source/ControlRig/Public/Sequencer/MovieSceneControlRigParameterSection.h
@@ -28,8 +28,8 @@ struct FKeyDataOptimizationParams;
 
 struct CONTROLRIG_API FControlRigBindingHelper
 {
-	static void BindToSequencerInstance(UControlRig* ControlRig);
-	static void UnBindFromSequencerInstance(UControlRig* ControlRig);
+	static void BindToSequencerInstance(UBaseControlRig* ControlRig);
+	static void UnBindFromSequencerInstance(UBaseControlRig* ControlRig);
 };
 
 struct FEnumParameterNameAndValue //uses uint8
@@ -220,13 +220,13 @@ private:
 	FSpaceChannelAddedEvent OnSpaceChannelAdded;
 	/** Control Rig that controls us*/
 	UPROPERTY()
-	TObjectPtr<UControlRig> ControlRig;
+	TObjectPtr<UBaseControlRig> ControlRig;
 
 public:
 
 	/** The class of control rig to instantiate */
 	UPROPERTY(EditAnywhere, Category = "Animation")
-	TSubclassOf<UControlRig> ControlRigClass;
+	TSubclassOf<UBaseControlRig> ControlRigClass;
 
 	/** Mask for controls themselves*/
 	UPROPERTY()
@@ -376,12 +376,12 @@ public:
 public:
 
 	/** Recreate with this Control Rig*/
-	void RecreateWithThisControlRig(UControlRig* InControlRig, bool bSetDefault);
+	void RecreateWithThisControlRig(UBaseControlRig* InControlRig, bool bSetDefault);
 
 	/* Set the control rig for this section */
-	void SetControlRig(UControlRig* InControlRig);
+	void SetControlRig(UBaseControlRig* InControlRig);
 	/* Get the control rig for this section */
-	UControlRig* GetControlRig() const { return ControlRig; }
+	UBaseControlRig* GetControlRig() const { return ControlRig; }
 
 	/** Whether or not to key currently, maybe evaluating so don't*/
 	void  SetDoNotKey(bool bIn) const { bDoNotKey = bIn; }
diff --git a/Engine/Plugins/Animation/ControlRig/Source/ControlRig/Public/Sequencer/MovieSceneControlRigParameterTrack.h b/Engine/Plugins/Animation/ControlRig/Source/ControlRig/Public/Sequencer/MovieSceneControlRigParameterTrack.h
--- a/Engine/Plugins/Animation/ControlRig/Source/ControlRig/Public/Sequencer/MovieSceneControlRigParameterTrack.h
+++ b/Engine/Plugins/Animation/ControlRig/Source/ControlRig/Public/Sequencer/MovieSceneControlRigParameterTrack.h
@@ -75,9 +75,9 @@ public:
 	virtual bool GetFbxCurveDataFromChannelMetadata(const FMovieSceneChannelMetaData& MetaData, FControlRigFbxCurveData& OutCurveData) override;
 #endif // WITH_EDITOR
 
-	//UControlRig Delegates
+	//UBaseControlRig Delegates
 	void HandleOnPostConstructed_GameThread();
-	void HandleOnPostConstructed(UControlRig* Subject, const FName& InEventName);
+	void HandleOnPostConstructed(UBaseControlRig* Subject, const FName& InEventName);
 
 #if WITH_EDITOR
 	void HandlePackageDone(const FEndLoadPackageContext& Context);
@@ -87,14 +87,14 @@ public:
 
 public:
 	/** Add a section at that start time*/
-	CONTROLRIG_API UMovieSceneSection* CreateControlRigSection(FFrameNumber StartTime, UControlRig* InControlRig, bool bInOwnsControlRig);
+	CONTROLRIG_API UMovieSceneSection* CreateControlRigSection(FFrameNumber StartTime, UBaseControlRig* InControlRig, bool bInOwnsControlRig);
 
-	CONTROLRIG_API void ReplaceControlRig(UControlRig* NewControlRig, bool RecreateChannels);
+	CONTROLRIG_API void ReplaceControlRig(UBaseControlRig* NewControlRig, bool RecreateChannels);
 
 	CONTROLRIG_API void RenameParameterName(const FName& OldParameterName, const FName& NewParameterName);
 
 public:
-	CONTROLRIG_API UControlRig* GetControlRig() const { return ControlRig; }
+	CONTROLRIG_API UBaseControlRig* GetControlRig() const { return ControlRig; }
 
 	/**
 	* Find all sections at the current time.
@@ -182,7 +182,7 @@ private:
 
 	/** Control Rig we control*/
 	UPROPERTY()
-	TObjectPtr<UControlRig> ControlRig;
+	TObjectPtr<UBaseControlRig> ControlRig;
 
 	/** Section we should Key */
 	UPROPERTY()
