Control Rig: Rename ControlRig to BaseControlRig and created a specialized ControlRig

diff --git a/Engine/Plugins/Animation/ControlRig/Source/ControlRig/Public/Tools/ControlRigPose.h b/Engine/Plugins/Animation/ControlRig/Source/ControlRig/Public/Tools/ControlRigPose.h
--- a/Engine/Plugins/Animation/ControlRig/Source/ControlRig/Public/Tools/ControlRigPose.h
+++ b/Engine/Plugins/Animation/ControlRig/Source/ControlRig/Public/Tools/ControlRigPose.h
@@ -13,7 +13,7 @@
 #include "Rigs/RigControlHierarchy.h"
 #include "ControlRigPose.generated.h"
 
-class UControlRig;
+class UBaseControlRig;
 /**
 * The Data Stored For Each Control in A Pose.
 */
@@ -84,18 +84,18 @@ struct CONTROLRIG_API FControlRigControlPose
 	GENERATED_USTRUCT_BODY()
 
 	FControlRigControlPose() {};
-	FControlRigControlPose(UControlRig* InControlRig, bool bUseAll)
+	FControlRigControlPose(UBaseControlRig* InControlRig, bool bUseAll)
 	{
 		SavePose(InControlRig, bUseAll);
 	}
 	~FControlRigControlPose() {};
 
-	void SavePose(UControlRig* ControlRig, bool bUseAll);
-	void PastePose(UControlRig* ControlRig, bool bDoKey, bool bDoMirror);
-	void SetControlMirrorTransform(bool bDoLocalSpace, UControlRig* ControlRig, const FName& Name, bool bIsMatched,
+	void SavePose(UBaseControlRig* ControlRig, bool bUseAll);
+	void PastePose(UBaseControlRig* ControlRig, bool bDoKey, bool bDoMirror);
+	void SetControlMirrorTransform(bool bDoLocalSpace, UBaseControlRig* ControlRig, const FName& Name, bool bIsMatched,
 		const FTransform& GlobalTrnnsform, const FTransform& LocalTransform, bool bNotify, const  FRigControlModifiedContext& Context, bool bSetupUndo);
-	void PastePoseInternal(UControlRig* ControlRig, bool bDoKey, bool bDoMirror, const TArray<FRigControlCopy>& ControlsToPaste);
-	void BlendWithInitialPoses(FControlRigControlPose& InitialPose, UControlRig* ControlRig, bool bDoKey, bool bDoMirror, float BlendValue);
+	void PastePoseInternal(UBaseControlRig* ControlRig, bool bDoKey, bool bDoMirror, const TArray<FRigControlCopy>& ControlsToPaste);
+	void BlendWithInitialPoses(FControlRigControlPose& InitialPose, UBaseControlRig* ControlRig, bool bDoKey, bool bDoMirror, float BlendValue);
 
 	bool ContainsName(const FName& Name) const;
 	void ReplaceControlName(const FName& Name, const FName& NewName);
@@ -127,18 +127,18 @@ public:
 	virtual void PostLoad() override;
 
 	UFUNCTION(BlueprintCallable, Category = "Pose")
-	void SavePose(UControlRig* InControlRig, bool bUseAll);
+	void SavePose(UBaseControlRig* InControlRig, bool bUseAll);
 
 	UFUNCTION(BlueprintCallable, Category = "Pose")
-	void PastePose(UControlRig* InControlRig, bool bDoKey = false, bool bDoMirror = false);
+	void PastePose(UBaseControlRig* InControlRig, bool bDoKey = false, bool bDoMirror = false);
 	
 	UFUNCTION(BlueprintCallable, Category = "Pose")
-	void SelectControls(UControlRig* InControlRig, bool bDoMirror = false);
+	void SelectControls(UBaseControlRig* InControlRig, bool bDoMirror = false);
 
-	TArray<FRigControlCopy> GetCurrentPose(UControlRig* InControlRig);
+	TArray<FRigControlCopy> GetCurrentPose(UBaseControlRig* InControlRig);
 
 	UFUNCTION(BlueprintCallable, Category = "Pose")
-	void GetCurrentPose(UControlRig* InControlRig, FControlRigControlPose& OutPose);
+	void GetCurrentPose(UBaseControlRig* InControlRig, FControlRigControlPose& OutPose);
 
 	UFUNCTION(BlueprintPure, Category = "Pose")
 	TArray<FName> GetControlNames() const;
@@ -147,9 +147,9 @@ public:
 	void ReplaceControlName(const FName& CurrentName, const FName& NewName);
 
 	UFUNCTION(BlueprintPure, Category = "Pose")
-	bool DoesMirrorMatch(UControlRig* ControlRig, const FName& ControlName) const;
+	bool DoesMirrorMatch(UBaseControlRig* ControlRig, const FName& ControlName) const;
 
-	void BlendWithInitialPoses(FControlRigControlPose& InitialPose, UControlRig* ControlRig, bool bDoKey, bool bdoMirror, float BlendValue);
+	void BlendWithInitialPoses(FControlRigControlPose& InitialPose, UBaseControlRig* ControlRig, bool bDoKey, bool bdoMirror, float BlendValue);
 
 	UPROPERTY(BlueprintReadWrite, EditAnywhere, Category = "Pose")
 	FControlRigControlPose Pose;
diff --git a/Engine/Plugins/Animation/ControlRig/Source/ControlRig/Public/Tools/ControlRigPoseMirrorTable.h b/Engine/Plugins/Animation/ControlRig/Source/ControlRig/Public/Tools/ControlRigPoseMirrorTable.h
--- a/Engine/Plugins/Animation/ControlRig/Source/ControlRig/Public/Tools/ControlRigPoseMirrorTable.h
+++ b/Engine/Plugins/Animation/ControlRig/Source/ControlRig/Public/Tools/ControlRigPoseMirrorTable.h
@@ -10,7 +10,7 @@
 #include "UObject/ObjectMacros.h"
 #include "UObject/Object.h"
 
-class UControlRig;
+class UBaseControlRig;
 struct FRigControl;
 struct FRigControlCopy;
 struct FControlRigControlPose;
@@ -22,7 +22,7 @@ public:
 	~FControlRigPoseMirrorTable() {};
 
 	/*Set up the Mirror Table*/
-	void SetUpMirrorTable(const UControlRig* ControlRig);
+	void SetUpMirrorTable(const UBaseControlRig* ControlRig);
 
 	/*Get the matched control with the given name*/
 	FRigControlCopy* GetControl(FControlRigControlPose& Pose, FName ControlrigName);
diff --git a/Engine/Plugins/Animation/ControlRig/Source/ControlRig/Public/Tools/ControlRigToolAsset.h b/Engine/Plugins/Animation/ControlRig/Source/ControlRig/Public/Tools/ControlRigToolAsset.h
--- a/Engine/Plugins/Animation/ControlRig/Source/ControlRig/Public/Tools/ControlRigToolAsset.h
+++ b/Engine/Plugins/Animation/ControlRig/Source/ControlRig/Public/Tools/ControlRigToolAsset.h
@@ -19,7 +19,7 @@ Create One of the Control Rig Assets, Currently Only support Pose Assets but may
 class FControlRigToolAsset
 {
 public:
-	template< typename Type> static  UObject* SaveAsset(UControlRig* InControlRig, const FString& CurrentPath, const FString& InString, bool bUseAllControls)
+	template< typename Type> static  UObject* SaveAsset(UBaseControlRig* InControlRig, const FString& CurrentPath, const FString& InString, bool bUseAllControls)
 	{
 		if (InControlRig)
 		{
diff --git a/Engine/Plugins/Animation/ControlRig/Source/ControlRig/Public/Units/RigUnitContext.h b/Engine/Plugins/Animation/ControlRig/Source/ControlRig/Public/Units/RigUnitContext.h
--- a/Engine/Plugins/Animation/ControlRig/Source/ControlRig/Public/Units/RigUnitContext.h
+++ b/Engine/Plugins/Animation/ControlRig/Source/ControlRig/Public/Units/RigUnitContext.h
@@ -11,7 +11,7 @@
 #include "RigVMCore/RigVMExecuteContext.h"
 #include "RigUnitContext.generated.h"
 
-class UControlRig;
+class UBaseControlRig;
 class UControlRigShapeLibrary;
 
 /**
@@ -196,7 +196,7 @@ private:
 class CONTROLRIG_API FControlRigExecuteContextRigModuleGuard
 {
 public:
-	FControlRigExecuteContextRigModuleGuard(FControlRigExecuteContext& InContext, const UControlRig* InControlRig);
+	FControlRigExecuteContextRigModuleGuard(FControlRigExecuteContext& InContext, const UBaseControlRig* InControlRig);
 	~FControlRigExecuteContextRigModuleGuard();
 
 private:
diff --git a/Engine/Plugins/Animation/ControlRig/Source/ControlRigDeveloper/Private/AnimGraphNode_ControlRig.cpp b/Engine/Plugins/Animation/ControlRig/Source/ControlRigDeveloper/Private/AnimGraphNode_ControlRig.cpp
--- a/Engine/Plugins/Animation/ControlRig/Source/ControlRigDeveloper/Private/AnimGraphNode_ControlRig.cpp
+++ b/Engine/Plugins/Animation/ControlRig/Source/ControlRigDeveloper/Private/AnimGraphNode_ControlRig.cpp
@@ -200,7 +200,7 @@ void UAnimGraphNode_ControlRig::CreateCustomPins(TArray<UEdGraphPin*>* OldPins)
 
 	if(const UClass* ControlRigClass = GetTargetClass())
 	{
-		if(UControlRig* CDO = ControlRigClass->GetDefaultObject<UControlRig>())
+		if(UBaseControlRig* CDO = ControlRigClass->GetDefaultObject<UBaseControlRig>())
 		{
 			if(const URigHierarchy* Hierarchy = CDO->GetHierarchy())
 			{
@@ -370,7 +370,7 @@ void UAnimGraphNode_ControlRig::RebuildExposedProperties()
 	// also add all of the controls
 	if(const UClass* TargetClass = GetTargetClass())
 	{
-		if(UControlRig* CDO = TargetClass->GetDefaultObject<UControlRig>())
+		if(UBaseControlRig* CDO = TargetClass->GetDefaultObject<UBaseControlRig>())
 		{
 			if(const URigHierarchy* Hierarchy = CDO->GetHierarchy())
 			{
@@ -400,7 +400,7 @@ FRigControlElement* UAnimGraphNode_ControlRig::FindControlElement(const FName& I
 {
 	if(const UClass* ControlRigClass = GetTargetClass())
 	{
-		if(UControlRig* CDO = ControlRigClass->GetDefaultObject<UControlRig>())
+		if(UBaseControlRig* CDO = ControlRigClass->GetDefaultObject<UBaseControlRig>())
 		{
 			if(const URigHierarchy* Hierarchy = CDO->GetHierarchy())
 			{
@@ -583,7 +583,7 @@ void UAnimGraphNode_ControlRig::GetVariables(bool bInput, TMap<FName, FRigVMExte
 		if (UControlRigBlueprint* RigBlueprint = Cast<UControlRigBlueprint>(TargetClass->ClassGeneratedBy))
 		{
 			//RigBlueprint->CleanupVariables();
-			UControlRig* ControlRig = TargetClass->GetDefaultObject<UControlRig>();
+			UBaseControlRig* ControlRig = TargetClass->GetDefaultObject<UBaseControlRig>();
 			if (ControlRig)
 			{
 				const TArray<FRigVMExternalVariable>& PublicVariables = ControlRig->GetPublicVariables();
@@ -636,7 +636,7 @@ void UAnimGraphNode_ControlRig::GetAvailableMapping(const FName& PathName, TArra
 		// also add all controls
 		if(const UClass* ControlRigClass = GetTargetClass())
 		{
-			if(UControlRig* CDO = ControlRigClass->GetDefaultObject<UControlRig>())
+			if(UBaseControlRig* CDO = ControlRigClass->GetDefaultObject<UBaseControlRig>())
 			{
 				if(const URigHierarchy* Hierarchy = CDO->GetHierarchy())
 				{
@@ -710,7 +710,7 @@ void UAnimGraphNode_ControlRig::CreateVariableMapping(const FString& FilteredTex
 		// add all controls
 		if(const UClass* ControlRigClass = GetTargetClass())
 		{
-			if(UControlRig* CDO = ControlRigClass->GetDefaultObject<UControlRig>())
+			if(UBaseControlRig* CDO = ControlRigClass->GetDefaultObject<UBaseControlRig>())
 			{
 				if(const URigHierarchy* Hierarchy = CDO->GetHierarchy())
 				{
diff --git a/Engine/Plugins/Animation/ControlRig/Source/ControlRigDeveloper/Private/ControlRigBlueprint.cpp b/Engine/Plugins/Animation/ControlRig/Source/ControlRigDeveloper/Private/ControlRigBlueprint.cpp
--- a/Engine/Plugins/Animation/ControlRig/Source/ControlRigDeveloper/Private/ControlRigBlueprint.cpp
+++ b/Engine/Plugins/Animation/ControlRig/Source/ControlRigDeveloper/Private/ControlRigBlueprint.cpp
@@ -104,13 +104,13 @@ void UControlRigBlueprint::PostEditChangeProperty(FPropertyChangedEvent& Propert
 		return Variable.VarName == PropertyChangedEvent.GetMemberPropertyName();
 	}))
 	{
-		if(UControlRig* DebuggedControlRig = Cast<UControlRig>(GetObjectBeingDebugged()))
+		if(UBaseControlRig* DebuggedControlRig = Cast<UBaseControlRig>(GetObjectBeingDebugged()))
 		{
 			if(const FProperty* PropertyOnRig = DebuggedControlRig->GetClass()->FindPropertyByName(PropertyChangedEvent.MemberProperty->GetFName()))
 			{
 				if(PropertyOnRig->SameType(PropertyChangedEvent.MemberProperty))
 				{
-					UControlRig* CDO = DebuggedControlRig->GetClass()->GetDefaultObject<UControlRig>();
+					UBaseControlRig* CDO = DebuggedControlRig->GetClass()->GetDefaultObject<UBaseControlRig>();
 					const uint8* SourceMemory = PropertyOnRig->ContainerPtrToValuePtr<uint8>(CDO);
 					uint8* TargetMemory = PropertyOnRig->ContainerPtrToValuePtr<uint8>(DebuggedControlRig);
 					PropertyOnRig->CopyCompleteValue(TargetMemory, SourceMemory);
@@ -129,7 +129,7 @@ void UControlRigBlueprint::PostEditChangeChainProperty(FPropertyChangedChainEven
 	if (PropertyChangedEvent.Property && PropertyChangedEvent.Property->GetFName() == GET_MEMBER_NAME_CHECKED(UControlRigBlueprint, ShapeLibraries))
 	{
 		URigVMBlueprintGeneratedClass* RigClass = GetRigVMBlueprintGeneratedClass();
-		UControlRig* CDO = Cast<UControlRig>(RigClass->GetDefaultObject(false /* create if needed */));
+		UBaseControlRig* CDO = Cast<UBaseControlRig>(RigClass->GetDefaultObject(false /* create if needed */));
 
 		TArray<UObject*> ArchetypeInstances;
 		CDO->GetArchetypeInstances(ArchetypeInstances);
@@ -138,7 +138,7 @@ void UControlRigBlueprint::PostEditChangeChainProperty(FPropertyChangedChainEven
 		// Propagate libraries to archetypes
 		for (UObject* Instance : ArchetypeInstances)
 		{
-			if (UControlRig* InstanceRig = Cast<UControlRig>(Instance))
+			if (UBaseControlRig* InstanceRig = Cast<UBaseControlRig>(Instance))
 			{
 				InstanceRig->ShapeLibraries = ShapeLibraries;
 			}
@@ -740,7 +740,7 @@ bool UControlRigBlueprint::ResolveConnector(const FRigElementKey& DraggedKey, co
 
 	PropagateHierarchyFromBPToInstances();
 
-	if(UControlRig* ControlRig = Cast<UControlRig>(GetObjectBeingDebugged()))
+	if(UBaseControlRig* ControlRig = Cast<UBaseControlRig>(GetObjectBeingDebugged()))
 	{
 		for (UEdGraph* Graph : UbergraphPages)
 		{
@@ -855,7 +855,7 @@ void UControlRigBlueprint::PostLoad()
 		}
 
 		URigVMBlueprintGeneratedClass* RigClass = GetRigVMBlueprintGeneratedClass();
-		UControlRig* CDO = Cast<UControlRig>(RigClass->GetDefaultObject(false /* create if needed */));
+		UBaseControlRig* CDO = Cast<UBaseControlRig>(RigClass->GetDefaultObject(false /* create if needed */));
 
 		TArray<UObject*> ArchetypeInstances;
 		CDO->GetArchetypeInstances(ArchetypeInstances);
@@ -863,7 +863,7 @@ void UControlRigBlueprint::PostLoad()
 
 		for (UObject* Instance : ArchetypeInstances)
 		{
-			if (UControlRig* InstanceRig = Cast<UControlRig>(Instance))
+			if (UBaseControlRig* InstanceRig = Cast<UBaseControlRig>(Instance))
 			{
 				InstanceRig->ShapeLibraries.Reset();
 				InstanceRig->GizmoLibrary_DEPRECATED.Reset();
@@ -884,7 +884,7 @@ void UControlRigBlueprint::HandlePackageDone()
 		}
 
 		URigVMBlueprintGeneratedClass* RigClass = GetRigVMBlueprintGeneratedClass();
-		UControlRig* CDO = Cast<UControlRig>(RigClass->GetDefaultObject(false /* create if needed */));
+		UBaseControlRig* CDO = Cast<UBaseControlRig>(RigClass->GetDefaultObject(false /* create if needed */));
 
 		TArray<UObject*> ArchetypeInstances;
 		CDO->GetArchetypeInstances(ArchetypeInstances);
@@ -892,7 +892,7 @@ void UControlRigBlueprint::HandlePackageDone()
 
 		for (UObject* Instance : ArchetypeInstances)
 		{
-			if (UControlRig* InstanceRig = Cast<UControlRig>(Instance))
+			if (UBaseControlRig* InstanceRig = Cast<UBaseControlRig>(Instance))
 			{
 				InstanceRig->ShapeLibraries = ShapeLibraries;
 			}
@@ -920,7 +920,7 @@ void UControlRigBlueprint::HandleConfigureRigVMController(const FRigVMClient* In
 			
 			if(const URigVMBlueprint* StrongThis = WeakThis.Get())
 			{
-				if(UControlRig* ControlRig = Cast<UControlRig>(StrongThis->GetObjectBeingDebugged()))
+				if(UBaseControlRig* ControlRig = Cast<UBaseControlRig>(StrongThis->GetObjectBeingDebugged()))
 				{
 					ControlRigNodeWorkflowOptions->Hierarchy = ControlRig->GetHierarchy();
 					ControlRigNodeWorkflowOptions->Selection = ControlRig->GetHierarchy()->GetSelectedKeys();
@@ -1070,7 +1070,7 @@ TArray<UControlRigBlueprint*> UControlRigBlueprint::GetCurrentlyOpenRigBlueprint
 const FControlRigShapeDefinition* UControlRigBlueprint::GetControlShapeByName(const FName& InName) const
 {
 	TMap<FString, FString> LibraryNameMap;
-	if(UControlRig* ControlRig = Cast<UControlRig>(GetObjectBeingDebugged()))
+	if(UBaseControlRig* ControlRig = Cast<UBaseControlRig>(GetObjectBeingDebugged()))
 	{
 		LibraryNameMap = ControlRig->ShapeLibraryNameMap;
 	}
@@ -1089,14 +1089,14 @@ FName UControlRigBlueprint::AddTransientControl(const URigVMUnitNode* InNode, co
 	ClearTransientControls();
 
 	URigVMBlueprintGeneratedClass* RigClass = GetRigVMBlueprintGeneratedClass();
-	UControlRig* CDO = Cast<UControlRig>(RigClass->GetDefaultObject(true /* create if needed */));
+	UBaseControlRig* CDO = Cast<UBaseControlRig>(RigClass->GetDefaultObject(true /* create if needed */));
 
 	FName ReturnName = NAME_None;
 	TArray<UObject*> ArchetypeInstances;
 	CDO->GetArchetypeInstances(ArchetypeInstances);
 	for (UObject* ArchetypeInstance : ArchetypeInstances)
 	{
-		UControlRig* InstancedControlRig = Cast<UControlRig>(ArchetypeInstance);
+		UBaseControlRig* InstancedControlRig = Cast<UBaseControlRig>(ArchetypeInstance);
 		if (InstancedControlRig)
 		{
 			FName ControlName = InstancedControlRig->AddTransientControl(InNode, InTarget);
@@ -1119,14 +1119,14 @@ FName UControlRigBlueprint::RemoveTransientControl(const URigVMUnitNode* InNode,
 	}
 
 	URigVMBlueprintGeneratedClass* RigClass = GetRigVMBlueprintGeneratedClass();
-	UControlRig* CDO = Cast<UControlRig>(RigClass->GetDefaultObject(true /* create if needed */));
+	UBaseControlRig* CDO = Cast<UBaseControlRig>(RigClass->GetDefaultObject(true /* create if needed */));
 
 	FName RemovedName = NAME_None;
 	TArray<UObject*> ArchetypeInstances;
 	CDO->GetArchetypeInstances(ArchetypeInstances);
 	for (UObject* ArchetypeInstance : ArchetypeInstances)
 	{
-		UControlRig* InstancedControlRig = Cast<UControlRig>(ArchetypeInstance);
+		UBaseControlRig* InstancedControlRig = Cast<UBaseControlRig>(ArchetypeInstance);
 		if (InstancedControlRig)
 		{
 			FName Name = InstancedControlRig->RemoveTransientControl(InNode, InTarget);
@@ -1148,7 +1148,7 @@ FName UControlRigBlueprint::AddTransientControl(const FRigElementKey& InElement)
 		ValueScope = MakeUnique<FControlValueScope>(this);
 	}
 	URigVMBlueprintGeneratedClass* RigClass = GetRigVMBlueprintGeneratedClass();
-	UControlRig* CDO = Cast<UControlRig>(RigClass->GetDefaultObject(true /* create if needed */));
+	UBaseControlRig* CDO = Cast<UBaseControlRig>(RigClass->GetDefaultObject(true /* create if needed */));
 
 	FName ReturnName = NAME_None;
 	TArray<UObject*> ArchetypeInstances;
@@ -1159,7 +1159,7 @@ FName UControlRigBlueprint::AddTransientControl(const FRigElementKey& InElement)
 	TMap<UObject*, FTransform> SavedElementLocalTransforms;
 	for (UObject* ArchetypeInstance : ArchetypeInstances)
 	{
-		UControlRig* InstancedControlRig = Cast<UControlRig>(ArchetypeInstance);
+		UBaseControlRig* InstancedControlRig = Cast<UBaseControlRig>(ArchetypeInstance);
 		if (InstancedControlRig)
 		{
 			if (InstancedControlRig->DynamicHierarchy)
@@ -1174,7 +1174,7 @@ FName UControlRigBlueprint::AddTransientControl(const FRigElementKey& InElement)
 	
 	for (UObject* ArchetypeInstance : ArchetypeInstances)
 	{
-		UControlRig* InstancedControlRig = Cast<UControlRig>(ArchetypeInstance);
+		UBaseControlRig* InstancedControlRig = Cast<UBaseControlRig>(ArchetypeInstance);
 		if (InstancedControlRig)
 		{
 			// restore the element transforms so that transient controls are created at the right place
@@ -1207,14 +1207,14 @@ FName UControlRigBlueprint::RemoveTransientControl(const FRigElementKey& InEleme
 	}
 
 	URigVMBlueprintGeneratedClass* RigClass = GetRigVMBlueprintGeneratedClass();
-	UControlRig* CDO = Cast<UControlRig>(RigClass->GetDefaultObject(true /* create if needed */));
+	UBaseControlRig* CDO = Cast<UBaseControlRig>(RigClass->GetDefaultObject(true /* create if needed */));
 
 	FName RemovedName = NAME_None;
 	TArray<UObject*> ArchetypeInstances;
 	CDO->GetArchetypeInstances(ArchetypeInstances);
 	for (UObject* ArchetypeInstance : ArchetypeInstances)
 	{
-		UControlRig* InstancedControlRig = Cast<UControlRig>(ArchetypeInstance);
+		UBaseControlRig* InstancedControlRig = Cast<UBaseControlRig>(ArchetypeInstance);
 		if (InstancedControlRig)
 		{
 			FName Name = InstancedControlRig->RemoveTransientControl(InElement);
@@ -1238,13 +1238,13 @@ void UControlRigBlueprint::ClearTransientControls()
 
 	if (URigVMBlueprintGeneratedClass* RigClass = GetRigVMBlueprintGeneratedClass())
 	{
-		UControlRig* CDO = Cast<UControlRig>(RigClass->GetDefaultObject(true /* create if needed */));
+		UBaseControlRig* CDO = Cast<UBaseControlRig>(RigClass->GetDefaultObject(true /* create if needed */));
 
 		TArray<UObject*> ArchetypeInstances;
 		CDO->GetArchetypeInstances(ArchetypeInstances);
 		for (UObject* ArchetypeInstance : ArchetypeInstances)
 		{
-			UControlRig* InstancedControlRig = Cast<UControlRig>(ArchetypeInstance);
+			UBaseControlRig* InstancedControlRig = Cast<UBaseControlRig>(ArchetypeInstance);
 			if (InstancedControlRig)
 			{
 				InstancedControlRig->ClearTransientControls();
@@ -1258,7 +1258,7 @@ void UControlRigBlueprint::ClearTransientControls()
 void UControlRigBlueprint::SetupDefaultObjectDuringCompilation(URigVMHost* InCDO)
 {
 	Super::SetupDefaultObjectDuringCompilation(InCDO);
-	CastChecked<UControlRig>(InCDO)->GetHierarchy()->CopyHierarchy(Hierarchy);
+	CastChecked<UBaseControlRig>(InCDO)->GetHierarchy()->CopyHierarchy(Hierarchy);
 }
 
 void UControlRigBlueprint::SetupPinRedirectorsForBackwardsCompatibility()
@@ -1696,14 +1696,14 @@ void UControlRigBlueprint::PatchVariableNodesOnLoad()
 	Super::PatchVariableNodesOnLoad();
 }
 
-void UControlRigBlueprint::UpdateElementKeyRedirector(UControlRig* InControlRig) const
+void UControlRigBlueprint::UpdateElementKeyRedirector(UBaseControlRig* InControlRig) const
 {
 	InControlRig->HierarchySettings = HierarchySettings;
 	InControlRig->RigModuleSettings = RigModuleSettings;
 	InControlRig->ElementKeyRedirector = FRigElementKeyRedirector(ConnectionMap, Hierarchy);
 }
 
-void UControlRigBlueprint::PropagatePoseFromInstanceToBP(UControlRig* InControlRig) const
+void UControlRigBlueprint::PropagatePoseFromInstanceToBP(UBaseControlRig* InControlRig) const
 {
 	check(InControlRig);
 	// current transforms in BP and CDO are meaningless, no need to copy them
@@ -1716,7 +1716,7 @@ void UControlRigBlueprint::PropagatePoseFromBPToInstances() const
 {
 	if (UClass* MyControlRigClass = GeneratedClass)
 	{
-		if (UControlRig* DefaultObject = Cast<UControlRig>(MyControlRigClass->GetDefaultObject(false)))
+		if (UBaseControlRig* DefaultObject = Cast<UBaseControlRig>(MyControlRigClass->GetDefaultObject(false)))
 		{
 			DefaultObject->PostInitInstanceIfRequired();
 			DefaultObject->GetHierarchy()->CopyPose(Hierarchy, true, true, true);
@@ -1725,7 +1725,7 @@ void UControlRigBlueprint::PropagatePoseFromBPToInstances() const
 			DefaultObject->GetArchetypeInstances(ArchetypeInstances);
 			for (UObject* ArchetypeInstance : ArchetypeInstances)
 			{
-				if (UControlRig* InstanceRig = Cast<UControlRig>(ArchetypeInstance))
+				if (UBaseControlRig* InstanceRig = Cast<UBaseControlRig>(ArchetypeInstance))
 				{
 					InstanceRig->PostInitInstanceIfRequired();
 					InstanceRig->GetHierarchy()->CopyPose(Hierarchy, true, true, true);
@@ -1739,7 +1739,7 @@ void UControlRigBlueprint::PropagateHierarchyFromBPToInstances() const
 {
 	if (UClass* MyControlRigClass = GeneratedClass)
 	{
-		if (UControlRig* DefaultObject = Cast<UControlRig>(MyControlRigClass->GetDefaultObject(false)))
+		if (UBaseControlRig* DefaultObject = Cast<UBaseControlRig>(MyControlRigClass->GetDefaultObject(false)))
 		{
 			DefaultObject->PostInitInstanceIfRequired();
 			DefaultObject->GetHierarchy()->CopyHierarchy(Hierarchy);
@@ -1755,7 +1755,7 @@ void UControlRigBlueprint::PropagateHierarchyFromBPToInstances() const
 			DefaultObject->GetArchetypeInstances(ArchetypeInstances);
 			for (UObject* ArchetypeInstance : ArchetypeInstances)
 			{
-				if (UControlRig* InstanceRig = Cast<UControlRig>(ArchetypeInstance))
+				if (UBaseControlRig* InstanceRig = Cast<UBaseControlRig>(ArchetypeInstance))
 				{
 					InstanceRig->PostInitInstanceIfRequired();
 					InstanceRig->GetHierarchy()->CopyHierarchy(Hierarchy);
@@ -1771,7 +1771,7 @@ void UControlRigBlueprint::PropagateDrawInstructionsFromBPToInstances() const
 {
 	if (UClass* MyControlRigClass = GeneratedClass)
 	{
-		if (UControlRig* DefaultObject = Cast<UControlRig>(MyControlRigClass->GetDefaultObject(false)))
+		if (UBaseControlRig* DefaultObject = Cast<UBaseControlRig>(MyControlRigClass->GetDefaultObject(false)))
 	{
 			DefaultObject->DrawContainer = DrawContainer;
 
@@ -1780,7 +1780,7 @@ void UControlRigBlueprint::PropagateDrawInstructionsFromBPToInstances() const
 
 			for (UObject* ArchetypeInstance : ArchetypeInstances)
 			{
-				if (UControlRig* InstanceRig = Cast<UControlRig>(ArchetypeInstance))
+				if (UBaseControlRig* InstanceRig = Cast<UBaseControlRig>(ArchetypeInstance))
 	{
 					InstanceRig->DrawContainer = DrawContainer;
 				}
@@ -1809,7 +1809,7 @@ void UControlRigBlueprint::PropagatePropertyFromBPToInstances(FRigElementKey InR
 
 	if (UClass* MyControlRigClass = GeneratedClass)
 	{
-		if (UControlRig* DefaultObject = Cast<UControlRig>(MyControlRigClass->GetDefaultObject(false)))
+		if (UBaseControlRig* DefaultObject = Cast<UBaseControlRig>(MyControlRigClass->GetDefaultObject(false)))
 		{
 			TArray<UObject*> ArchetypeInstances;
 			DefaultObject->GetArchetypeInstances(ArchetypeInstances);
@@ -1820,7 +1820,7 @@ void UControlRigBlueprint::PropagatePropertyFromBPToInstances(FRigElementKey InR
 			uint8* Source = ((uint8*)Hierarchy->Get(ElementIndex)) + PropertyOffset;
 			for (UObject* ArchetypeInstance : ArchetypeInstances)
 			{
-				if (UControlRig* InstanceRig = Cast<UControlRig>(ArchetypeInstance))
+				if (UBaseControlRig* InstanceRig = Cast<UBaseControlRig>(ArchetypeInstance))
 				{
 					InstanceRig->PostInitInstanceIfRequired();
 					uint8* Dest = ((uint8*)InstanceRig->GetHierarchy()->Get(ElementIndex)) + PropertyOffset;
@@ -1831,7 +1831,7 @@ void UControlRigBlueprint::PropagatePropertyFromBPToInstances(FRigElementKey InR
 	}
 }
 
-void UControlRigBlueprint::PropagatePropertyFromInstanceToBP(FRigElementKey InRigElement, const FProperty* InProperty, UControlRig* InInstance) const
+void UControlRigBlueprint::PropagatePropertyFromInstanceToBP(FRigElementKey InRigElement, const FProperty* InProperty, UBaseControlRig* InInstance) const
 {
 	const int32 ElementIndex = Hierarchy->GetIndex(InRigElement);
 	ensure(ElementIndex != INDEX_NONE);
@@ -1892,9 +1892,9 @@ void UControlRigBlueprint::HandleHierarchyModified(ERigHierarchyNotification InN
 
 			if(bClearTransientControls)
 			{
-				if(UControlRig* RigBeingDebugged = Cast<UControlRig>(GetObjectBeingDebugged()))
+				if(UBaseControlRig* RigBeingDebugged = Cast<UBaseControlRig>(GetObjectBeingDebugged()))
 				{
-					const FName TransientControlName = UControlRig::GetNameForTransientControl(InElement->GetKey());
+					const FName TransientControlName = UBaseControlRig::GetNameForTransientControl(InElement->GetKey());
 					const FRigElementKey TransientControlKey(TransientControlName, ERigElementType::Control);
 					if (const FRigControlElement* ControlElement = RigBeingDebugged->GetHierarchy()->Find<FRigControlElement>(TransientControlKey))
 					{
@@ -1940,7 +1940,7 @@ UControlRigBlueprint::FControlValueScope::FControlValueScope(UControlRigBlueprin
 #if WITH_EDITOR
 	check(Blueprint);
 
-	if (UControlRig* CR = Cast<UControlRig>(Blueprint->GetObjectBeingDebugged()))
+	if (UBaseControlRig* CR = Cast<UBaseControlRig>(Blueprint->GetObjectBeingDebugged()))
 	{
 		TArray<FRigControlElement*> Controls = CR->AvailableControls();
 		for (FRigControlElement* ControlElement : Controls)
@@ -1956,7 +1956,7 @@ UControlRigBlueprint::FControlValueScope::~FControlValueScope()
 #if WITH_EDITOR
 	check(Blueprint);
 
-	if (UControlRig* CR = Cast<UControlRig>(Blueprint->GetObjectBeingDebugged()))
+	if (UBaseControlRig* CR = Cast<UBaseControlRig>(Blueprint->GetObjectBeingDebugged()))
 	{
 		for (const TPair<FName, FRigControlValue>& Pair : ControlValues)
 		{
diff --git a/Engine/Plugins/Animation/ControlRig/Source/ControlRigDeveloper/Private/ControlRigBlueprintCompiler.cpp b/Engine/Plugins/Animation/ControlRig/Source/ControlRigDeveloper/Private/ControlRigBlueprintCompiler.cpp
--- a/Engine/Plugins/Animation/ControlRig/Source/ControlRigDeveloper/Private/ControlRigBlueprintCompiler.cpp
+++ b/Engine/Plugins/Animation/ControlRig/Source/ControlRigDeveloper/Private/ControlRigBlueprintCompiler.cpp
@@ -9,7 +9,7 @@ bool FControlRigBlueprintCompiler::CanCompile(const UBlueprint* Blueprint)
 {
 	DECLARE_SCOPE_HIERARCHICAL_COUNTER_FUNC()
 
-	if (Blueprint && Blueprint->ParentClass && Blueprint->ParentClass->IsChildOf(UControlRig::StaticClass()))
+	if (Blueprint && Blueprint->ParentClass && Blueprint->ParentClass->IsChildOf(UBaseControlRig::StaticClass()))
 	{
 		return true;
 	}
@@ -53,9 +53,9 @@ void FControlRigBlueprintCompilerContext::CopyTermDefaultsToDefaultObject(UObjec
 	if (ControlRigBlueprint)
 	{
 		// here, CDO is initialized from BP,
-		// and in UControlRig::InitializeFromCDO,
+		// and in UBaseControlRig::InitializeFromCDO,
 		// other Control Rig Instances are then initialized From the CDO 
-		UControlRig* ControlRig = CastChecked<UControlRig>(DefaultObject);
+		UBaseControlRig* ControlRig = CastChecked<UBaseControlRig>(DefaultObject);
 		
 		// copy hierarchy
 		{
diff --git a/Engine/Plugins/Animation/ControlRig/Source/ControlRigDeveloper/Private/ControlRigVisualGraphUtils.cpp b/Engine/Plugins/Animation/ControlRig/Source/ControlRigDeveloper/Private/ControlRigVisualGraphUtils.cpp
--- a/Engine/Plugins/Animation/ControlRig/Source/ControlRigDeveloper/Private/ControlRigVisualGraphUtils.cpp
+++ b/Engine/Plugins/Animation/ControlRig/Source/ControlRigDeveloper/Private/ControlRigVisualGraphUtils.cpp
@@ -66,7 +66,7 @@ FAutoConsoleCommandWithWorldAndArgs FCmdControlRigVisualGraphUtilsDumpHierarchy
 		{
 			if(UObject* Object = FindObject<UObject>(nullptr, *ObjectPathName, false))
 			{
-				if(UControlRig* CR = Cast<UControlRig>(Object))
+				if(UBaseControlRig* CR = Cast<UBaseControlRig>(Object))
 				{
 					Hierarchies.Add(CR->GetHierarchy());
 				}
@@ -227,7 +227,7 @@ FString FControlRigVisualGraphUtils::DumpRigHierarchyToDotGraph(URigHierarchy* I
 	
 	if(!InEventName.IsNone())
 	{
-		if(UControlRig* CR = InHierarchy->GetTypedOuter<UControlRig>())
+		if(UBaseControlRig* CR = InHierarchy->GetTypedOuter<UBaseControlRig>())
 		{
 			URigVM* VM = CR->GetVM();
 
diff --git a/Engine/Plugins/Animation/ControlRig/Source/ControlRigDeveloper/Private/Graph/ControlRigGraph.cpp b/Engine/Plugins/Animation/ControlRig/Source/ControlRigDeveloper/Private/Graph/ControlRigGraph.cpp
--- a/Engine/Plugins/Animation/ControlRig/Source/ControlRigDeveloper/Private/Graph/ControlRigGraph.cpp
+++ b/Engine/Plugins/Animation/ControlRig/Source/ControlRigDeveloper/Private/Graph/ControlRigGraph.cpp
@@ -27,7 +27,7 @@ void UControlRigGraph::InitializeFromBlueprint(URigVMBlueprint* InBlueprint)
 	const UControlRigBlueprint* ControlRigBlueprint = CastChecked<UControlRigBlueprint>(InBlueprint);
 	URigHierarchy* Hierarchy = ControlRigBlueprint->Hierarchy;
 
-	if(UControlRig* ControlRig = Cast<UControlRig>(ControlRigBlueprint->GetObjectBeingDebugged()))
+	if(UBaseControlRig* ControlRig = Cast<UBaseControlRig>(ControlRigBlueprint->GetObjectBeingDebugged()))
 	{
 		Hierarchy = ControlRig->GetHierarchy();
 	}
@@ -73,7 +73,7 @@ void UControlRigGraph::CacheNameLists(URigHierarchy* InHierarchy, const FRigVMDr
 	check(InHierarchy);
 	check(DrawContainer);
 
-	UControlRig* ControlRig = InHierarchy->GetTypedOuter<UControlRig>();
+	UBaseControlRig* ControlRig = InHierarchy->GetTypedOuter<UBaseControlRig>();
 
 	if(LastHierarchyTopologyVersion != InHierarchy->GetTopologyVersion())
 	{
@@ -170,7 +170,7 @@ const TArray<TSharedPtr<FRigVMStringWithTag>>* UControlRigGraph::GetElementNameL
 
 		UControlRigGraph* MutableThis = (UControlRigGraph*)this;
 		URigHierarchy* Hierarchy = Blueprint->Hierarchy;
-		if(UControlRig* ControlRig = Cast<UControlRig>(Blueprint->GetObjectBeingDebugged()))
+		if(UBaseControlRig* ControlRig = Cast<UBaseControlRig>(Blueprint->GetObjectBeingDebugged()))
 		{
 			Hierarchy = ControlRig->GetHierarchy();
 		}	
diff --git a/Engine/Plugins/Animation/ControlRig/Source/ControlRigDeveloper/Public/ControlRigBlueprint.h b/Engine/Plugins/Animation/ControlRig/Source/ControlRigDeveloper/Public/ControlRigBlueprint.h
--- a/Engine/Plugins/Animation/ControlRig/Source/ControlRigDeveloper/Public/ControlRigBlueprint.h
+++ b/Engine/Plugins/Animation/ControlRig/Source/ControlRigDeveloper/Public/ControlRigBlueprint.h
@@ -80,7 +80,7 @@ public:
 	UClass* GetControlRigClass();
 
 	UFUNCTION(BlueprintCallable, Category = "Control Rig Blueprint")
-	UControlRig* CreateControlRig() { return Cast<UControlRig>(CreateRigVMHost()); }
+	UBaseControlRig* CreateControlRig() { return Cast<UBaseControlRig>(CreateRigVMHost()); }
 
 	/** IInterface_PreviewMeshProvider interface */
 	UFUNCTION(BlueprintCallable, Category = "Control Rig Blueprint")
@@ -247,13 +247,13 @@ protected:
 	virtual void PatchVariableNodesOnLoad() override;
 
 public:
-	void UpdateElementKeyRedirector(UControlRig* InControlRig) const;
-	void PropagatePoseFromInstanceToBP(UControlRig* InControlRig) const;
+	void UpdateElementKeyRedirector(UBaseControlRig* InControlRig) const;
+	void PropagatePoseFromInstanceToBP(UBaseControlRig* InControlRig) const;
 	void PropagatePoseFromBPToInstances() const;
 	void PropagateHierarchyFromBPToInstances() const;
 	void PropagateDrawInstructionsFromBPToInstances() const;
 	void PropagatePropertyFromBPToInstances(FRigElementKey InRigElement, const FProperty* InProperty) const;
-	void PropagatePropertyFromInstanceToBP(FRigElementKey InRigElement, const FProperty* InProperty, UControlRig* InInstance) const;
+	void PropagatePropertyFromInstanceToBP(FRigElementKey InRigElement, const FProperty* InProperty, UBaseControlRig* InInstance) const;
 
 	/**
 	* Returns the modified event, which can be used to 
diff --git a/Engine/Plugins/Animation/ControlRig/Source/ControlRigDeveloper/Public/Graph/ControlRigGraph.h b/Engine/Plugins/Animation/ControlRig/Source/ControlRigDeveloper/Public/Graph/ControlRigGraph.h
--- a/Engine/Plugins/Animation/ControlRig/Source/ControlRigDeveloper/Public/Graph/ControlRigGraph.h
+++ b/Engine/Plugins/Animation/ControlRig/Source/ControlRigDeveloper/Public/Graph/ControlRigGraph.h
@@ -14,7 +14,7 @@
 
 class UControlRigBlueprint;
 class UControlRigGraphSchema;
-class UControlRig;
+class UBaseControlRig;
 class URigVMController;
 struct FRigCurveContainer;
 
@@ -79,7 +79,7 @@ private:
 	}
 
 	template<class T>
-	void CacheNameListForHierarchy(UControlRig* InControlRig, URigHierarchy* InHierarchy, TArray<TSharedPtr<FRigVMStringWithTag>>& OutNameList, bool bFilter = true)
+	void CacheNameListForHierarchy(UBaseControlRig* InControlRig, URigHierarchy* InHierarchy, TArray<TSharedPtr<FRigVMStringWithTag>>& OutNameList, bool bFilter = true)
 	{
         TArray<FRigVMStringWithTag> Names;
 		for (auto Element : *InHierarchy)
diff --git a/Engine/Plugins/Animation/ControlRig/Source/ControlRigEditor/Private/ControlRigAnimGraphDetails.cpp b/Engine/Plugins/Animation/ControlRig/Source/ControlRigEditor/Private/ControlRigAnimGraphDetails.cpp
--- a/Engine/Plugins/Animation/ControlRig/Source/ControlRigEditor/Private/ControlRigAnimGraphDetails.cpp
+++ b/Engine/Plugins/Animation/ControlRig/Source/ControlRigEditor/Private/ControlRigAnimGraphDetails.cpp
@@ -132,7 +132,7 @@ void FControlRigAnimNodeEventNameDetails::UpdateEntryNameList()
 	{
 		if(const UClass* Class = AnimNodeBeingCustomized->GetControlRigClass())
 		{
-			if(const UControlRig* CDO = Cast<UControlRig>(Class->GetDefaultObject(true)))
+			if(const UBaseControlRig* CDO = Cast<UBaseControlRig>(Class->GetDefaultObject(true)))
 			{
 				Algo::Transform(CDO->GetSupportedEvents(), EntryNameList,[](const FName& InEntryName)
 				{
diff --git a/Engine/Plugins/Animation/ControlRig/Source/ControlRigEditor/Private/ControlRigBlueprintActions.cpp b/Engine/Plugins/Animation/ControlRig/Source/ControlRigEditor/Private/ControlRigBlueprintActions.cpp
--- a/Engine/Plugins/Animation/ControlRig/Source/ControlRigEditor/Private/ControlRigBlueprintActions.cpp
+++ b/Engine/Plugins/Animation/ControlRig/Source/ControlRigEditor/Private/ControlRigBlueprintActions.cpp
@@ -3,6 +3,7 @@
 #include "ControlRigBlueprintActions.h"
 #include "ControlRigBlueprintFactory.h"
 #include "ControlRigBlueprint.h"
+#include "ControlRig.h"
 #include "Editor/RigVMEditorStyle.h"
 #include "IControlRigEditorModule.h"
 
@@ -392,7 +393,7 @@ void FControlRigBlueprintActions::OnSpawnedSkeletalMeshActorChanged(UObject* InO
 					Track = MovieScene->AddTrack<UMovieSceneControlRigParameterTrack>(ActorTrackGuid);
 				}
 
-				UControlRig* ControlRig = Track->GetControlRig();
+				UBaseControlRig* ControlRig = Track->GetControlRig();
 
 				FString ObjectName = (ControlRigClass->GetName());
 
@@ -403,10 +404,11 @@ void FControlRigBlueprintActions::OnSpawnedSkeletalMeshActorChanged(UObject* InO
 
 					ObjectName.RemoveFromEnd(TEXT("_C"));
 
-					ControlRig = NewObject<UControlRig>(Track, ControlRigClass, FName(*ObjectName), RF_Transactional);
+					// This is either a UControlRig or a UModularControlRig
+					ControlRig = NewObject<UBaseControlRig>(Track, ControlRigClass, FName(*ObjectName), RF_Transactional);
 					ControlRig->SetObjectBinding(MakeShared<FControlRigObjectBinding>());
 					ControlRig->GetObjectBinding()->BindToObject(MeshActor->GetSkeletalMeshComponent());
-					ControlRig->GetDataSourceRegistry()->RegisterDataSource(UControlRig::OwnerComponent, ControlRig->GetObjectBinding()->GetBoundObject());
+					ControlRig->GetDataSourceRegistry()->RegisterDataSource(UBaseControlRig::OwnerComponent, ControlRig->GetObjectBinding()->GetBoundObject());
 					ControlRig->Initialize();
 					ControlRig->Evaluate_AnyThread();
 					ControlRig->CreateRigControlsForCurveContainer();
diff --git a/Engine/Plugins/Animation/ControlRig/Source/ControlRigEditor/Private/ControlRigBlueprintFactory.cpp b/Engine/Plugins/Animation/ControlRig/Source/ControlRigEditor/Private/ControlRigBlueprintFactory.cpp
--- a/Engine/Plugins/Animation/ControlRig/Source/ControlRigEditor/Private/ControlRigBlueprintFactory.cpp
+++ b/Engine/Plugins/Animation/ControlRig/Source/ControlRigEditor/Private/ControlRigBlueprintFactory.cpp
@@ -42,7 +42,7 @@ public:
 	void Construct( const FArguments& InArgs )
 	{
 		bOkClicked = false;
-		ParentClass = UControlRig::StaticClass();
+		ParentClass = UBaseControlRig::StaticClass();
 
 		ChildSlot
 		[
@@ -134,6 +134,11 @@ private:
 			// If it appears on the allowed child-of classes list (or there is nothing on that list)
 			if (InClass)
 			{
+				if (InClass == UBaseControlRig::StaticClass())
+				{
+					return false;
+				}
+				
 				if (InFilterFuncs->IfInChildOfClassesSet(AllowedChildrenOfClasses, InClass) == EFilterReturn::Failed)
 				{
 					return false;
@@ -181,8 +186,8 @@ private:
 		TSharedPtr<FControlRigBlueprintParentFilter> Filter = MakeShareable(new FControlRigBlueprintParentFilter());
 		Options.ClassFilters.Add(Filter.ToSharedRef());
 
-		// All child child classes of UControlRig are valid.
-		Filter->AllowedChildrenOfClasses.Add(UControlRig::StaticClass());
+		// All child child classes of UBaseControlRig are valid.
+		Filter->AllowedChildrenOfClasses.Add(UBaseControlRig::StaticClass());
 
 		ParentClassContainer->ClearChildren();
 		ParentClassContainer->AddSlot()
@@ -266,14 +271,13 @@ UControlRigBlueprintFactory::UControlRigBlueprintFactory()
 	bCreateNew = true;
 	bEditAfterNew = true;
 	SupportedClass = UControlRigBlueprint::StaticClass();
-	ParentClass = UControlRig::StaticClass();
+	ParentClass = UBaseControlRig::StaticClass();
 }
 
 bool UControlRigBlueprintFactory::ConfigureProperties()
 {
-	// we don't need to do anything,
-	// let's return true to indicate that all properties have been configured to produce a control rig
-	return true;
+	TSharedRef<SControlRigBlueprintCreateDialog> Dialog = SNew(SControlRigBlueprintCreateDialog);
+	return Dialog->ConfigureProperties(this);
 };
 
 UObject* UControlRigBlueprintFactory::FactoryCreateNew(UClass* Class, UObject* InParent, FName Name, EObjectFlags Flags, UObject* Context, FFeedbackContext* Warn, FName CallingContext)
@@ -281,7 +285,7 @@ UObject* UControlRigBlueprintFactory::FactoryCreateNew(UClass* Class, UObject* I
 	// Make sure we are trying to factory a Control Rig Blueprint, then create and init one
 	check(Class->IsChildOf(UControlRigBlueprint::StaticClass()));
 
-	if ((ParentClass == nullptr) || !FKismetEditorUtilities::CanCreateBlueprintOfClass(ParentClass) || !ParentClass->IsChildOf(UControlRig::StaticClass()))
+	if ((ParentClass == nullptr) || !FKismetEditorUtilities::CanCreateBlueprintOfClass(ParentClass) || !ParentClass->IsChildOf(UBaseControlRig::StaticClass()))
 	{
 		FFormatNamedArguments Args;
 		Args.Add( TEXT("ClassName"), (ParentClass != nullptr) ? FText::FromString( ParentClass->GetName() ) : LOCTEXT("Null", "(null)") );
diff --git a/Engine/Plugins/Animation/ControlRig/Source/ControlRigEditor/Private/ControlRigBlueprintFactory.h b/Engine/Plugins/Animation/ControlRig/Source/ControlRigEditor/Private/ControlRigBlueprintFactory.h
--- a/Engine/Plugins/Animation/ControlRig/Source/ControlRigEditor/Private/ControlRigBlueprintFactory.h
+++ b/Engine/Plugins/Animation/ControlRig/Source/ControlRigEditor/Private/ControlRigBlueprintFactory.h
@@ -8,10 +8,9 @@
 #include "Engine/Blueprint.h"
 #include "Factories/Factory.h"
 #include "ControlRigBlueprint.h"
+#include "ControlRig.h"
 #include "ControlRigBlueprintFactory.generated.h"
 
-class UControlRig;
-
 UCLASS(HideCategories=Object)
 class CONTROLRIGEDITOR_API UControlRigBlueprintFactory : public UFactory
 {
diff --git a/Engine/Plugins/Animation/ControlRig/Source/ControlRigEditor/Private/ControlRigCompilerDetails.cpp b/Engine/Plugins/Animation/ControlRig/Source/ControlRigEditor/Private/ControlRigCompilerDetails.cpp
--- a/Engine/Plugins/Animation/ControlRig/Source/ControlRigEditor/Private/ControlRigCompilerDetails.cpp
+++ b/Engine/Plugins/Animation/ControlRig/Source/ControlRigEditor/Private/ControlRigCompilerDetails.cpp
@@ -144,7 +144,7 @@ void FRigVMCompileSettingsDetails::CustomizeChildren(TSharedRef<IPropertyHandle>
 					.ContentPadding(FMargin(2))
 					.Visibility_Lambda([]()
 					{
-						return UControlRig::AreNativizedVMsDisabled() ? EVisibility::Collapsed : EVisibility::Visible;
+						return UBaseControlRig::AreNativizedVMsDisabled() ? EVisibility::Collapsed : EVisibility::Visible;
 					})
 					.Content()
 					[
@@ -190,7 +190,7 @@ FReply FRigVMCompileSettingsDetails::OnInspectMemory(ERigVMMemoryType InMemoryTy
 {
 	if (BlueprintBeingCustomized)
 	{
-		if(UControlRig* DebuggedRig = Cast<UControlRig>(BlueprintBeingCustomized->GetObjectBeingDebugged()))
+		if(UBaseControlRig* DebuggedRig = Cast<UBaseControlRig>(BlueprintBeingCustomized->GetObjectBeingDebugged()))
 		{
 			if(FRigVMMemoryStorageStruct* MemoryStorage = DebuggedRig->GetMemoryByType(InMemoryType))
 			{
@@ -221,7 +221,7 @@ FReply FRigVMCompileSettingsDetails::OnCopyByteCodeClicked()
 	{
 		if (BlueprintBeingCustomized->GetDefaultModel())
 		{
-			if(UControlRig* ControlRig = Cast<UControlRig>(BlueprintBeingCustomized->GetObjectBeingDebugged()))
+			if(UBaseControlRig* ControlRig = Cast<UBaseControlRig>(BlueprintBeingCustomized->GetObjectBeingDebugged()))
 			{
 				FString ByteCodeContent = ControlRig->GetVM()->DumpByteCodeAsText(ControlRig->GetRigVMExtendedExecuteContext());
 				FPlatformApplicationMisc::ClipboardCopy(*ByteCodeContent);
@@ -235,7 +235,7 @@ FReply FRigVMCompileSettingsDetails::OnCopyHierarchyGraphClicked()
 {
 	if (BlueprintBeingCustomized)
 	{
-		if(UControlRig* ControlRig = Cast<UControlRig>(BlueprintBeingCustomized->GetObjectBeingDebugged()))
+		if(UBaseControlRig* ControlRig = Cast<UBaseControlRig>(BlueprintBeingCustomized->GetObjectBeingDebugged()))
 		{
 			FName EventName = FRigUnit_BeginExecution::EventName;
 			if(!ControlRig->GetEventQueue().IsEmpty())
@@ -257,7 +257,7 @@ FReply FRigVMCompileSettingsDetails::OnCopyGeneratedCodeClicked()
 		const FString ClassName = FString::Printf(TEXT("%sVM"), *BlueprintBeingCustomized->GetName());
 		if (BlueprintBeingCustomized->GeneratedClass)
 		{
-			if (UControlRig* CDO = Cast<UControlRig>(BlueprintBeingCustomized->GeneratedClass->GetDefaultObject()))
+			if (UBaseControlRig* CDO = Cast<UBaseControlRig>(BlueprintBeingCustomized->GeneratedClass->GetDefaultObject()))
 			{
 				if(CDO->GetVM())
 				{
@@ -287,7 +287,7 @@ FReply FRigVMCompileSettingsDetails::OnCopyVMExecutionClicked()
 {
 	if (BlueprintBeingCustomized)
 	{
-		if (UControlRig* ControlRig = Cast<UControlRig>(BlueprintBeingCustomized->GetObjectBeingDebugged()))
+		if (UBaseControlRig* ControlRig = Cast<UBaseControlRig>(BlueprintBeingCustomized->GetObjectBeingDebugged()))
 		{
 			FString DebugString = ControlRig->GetDebugExecutionString();
 			FPlatformApplicationMisc::ClipboardCopy(*DebugString);
diff --git a/Engine/Plugins/Animation/ControlRig/Source/ControlRigEditor/Private/ControlRigEditorModule.cpp b/Engine/Plugins/Animation/ControlRig/Source/ControlRigEditor/Private/ControlRigEditorModule.cpp
--- a/Engine/Plugins/Animation/ControlRig/Source/ControlRigEditor/Private/ControlRigEditorModule.cpp
+++ b/Engine/Plugins/Animation/ControlRig/Source/ControlRigEditor/Private/ControlRigEditorModule.cpp
@@ -121,7 +121,7 @@ void FControlRigEditorModule::StartupModule()
 	ClassesToUnregisterOnShutdown.Add(UControlRigBlueprint::StaticClass()->GetFName());
 	PropertyEditorModule.RegisterCustomClassLayout(ClassesToUnregisterOnShutdown.Last(), FOnGetDetailCustomizationInstance::CreateStatic(&FControlRigBlueprintDetails::MakeInstance));
 
-	ClassesToUnregisterOnShutdown.Add(UControlRig::StaticClass()->GetFName());
+	ClassesToUnregisterOnShutdown.Add(UBaseControlRig::StaticClass()->GetFName());
 
 	// same as ClassesToUnregisterOnShutdown but for properties, there is none right now
 	PropertiesToUnregisterOnShutdown.Reset();
@@ -632,7 +632,7 @@ void FControlRigEditorModule::GetDirectManipulationMenuActions(IRigVMClientHost*
     // Add direct manipulation context menu entries
 	if(UControlRigBlueprint* ControlRigBlueprint = Cast<UControlRigBlueprint>(RigVMClientHost))
 	{
-		UControlRig* DebuggedRig = Cast<UControlRig>(ControlRigBlueprint->GetObjectBeingDebugged());
+		UBaseControlRig* DebuggedRig = Cast<UBaseControlRig>(ControlRigBlueprint->GetObjectBeingDebugged());
 		if(DebuggedRig == nullptr)
 		{
 			return;
@@ -657,7 +657,7 @@ void FControlRigEditorModule::GetDirectManipulationMenuActions(IRigVMClientHost*
 				return;
 			}
 
-			const FRigUnit* UnitInstance = UControlRig::GetRigUnitInstanceFromScope(NodeInstance);
+			const FRigUnit* UnitInstance = UBaseControlRig::GetRigUnitInstanceFromScope(NodeInstance);
 			TArray<FRigDirectManipulationTarget> Targets;
 			if(UnitInstance->GetDirectManipulationTargets(UnitNode, NodeInstance, DebuggedRig->GetHierarchy(), Targets, nullptr))
 			{
@@ -687,7 +687,7 @@ void FControlRigEditorModule::GetDirectManipulationMenuActions(IRigVMClientHost*
 						
 						auto HasNoUnconstrainedAffectedPin = [NodeInstance, UnitNode, Target]() -> bool
 						{
-							const FRigUnit* UnitInstance = UControlRig::GetRigUnitInstanceFromScope(NodeInstance);
+							const FRigUnit* UnitInstance = UBaseControlRig::GetRigUnitInstanceFromScope(NodeInstance);
 							const TArray<const URigVMPin*> AffectedPins = UnitInstance->GetPinsForDirectManipulation(UnitNode, Target);
 
 							int32 NumAffectedPinsWithRootLinks = 0;
@@ -1056,10 +1056,10 @@ void FControlRigEditorModule::BakeToControlRig(UClass* ControlRigClass, UAnimSeq
 				FString ObjectName = (ControlRigClass->GetName());
 				ObjectName.RemoveFromEnd(TEXT("_C"));
 
-				UControlRig* ControlRig = NewObject<UControlRig>(Track, ControlRigClass, FName(*ObjectName), RF_Transactional);
+				UBaseControlRig* ControlRig = NewObject<UBaseControlRig>(Track, ControlRigClass, FName(*ObjectName), RF_Transactional);
 				ControlRig->SetObjectBinding(MakeShared<FControlRigObjectBinding>());
 				ControlRig->GetObjectBinding()->BindToObject(MeshActor);
-				ControlRig->GetDataSourceRegistry()->RegisterDataSource(UControlRig::OwnerComponent, ControlRig->GetObjectBinding()->GetBoundObject());
+				ControlRig->GetDataSourceRegistry()->RegisterDataSource(UBaseControlRig::OwnerComponent, ControlRig->GetObjectBinding()->GetBoundObject());
 				ControlRig->Initialize();
 				ControlRig->Evaluate_AnyThread();
 
@@ -1421,7 +1421,7 @@ bool FControlRigClassFilter::IsClassAllowed(const FClassViewerInitializationOpti
 {
 	if(InClass)
 	{
-		const bool bChildOfObjectClass = InClass->IsChildOf(UControlRig::StaticClass());
+		const bool bChildOfObjectClass = InClass->IsChildOf(UBaseControlRig::StaticClass());
 		const bool bMatchesFlags = !InClass->HasAnyClassFlags(CLASS_Hidden | CLASS_HideDropDown | CLASS_Deprecated | CLASS_Abstract);
 		const bool bNotNative = !InClass->IsNative();
 
@@ -1442,7 +1442,7 @@ bool FControlRigClassFilter::IsClassAllowed(const FClassViewerInitializationOpti
 
 bool FControlRigClassFilter::IsUnloadedClassAllowed(const FClassViewerInitializationOptions& InInitOptions, const TSharedRef< const IUnloadedBlueprintData > InUnloadedClassData, TSharedRef< FClassViewerFilterFuncs > InFilterFuncs)
 {
-	const bool bChildOfObjectClass = InUnloadedClassData->IsChildOf(UControlRig::StaticClass());
+	const bool bChildOfObjectClass = InUnloadedClassData->IsChildOf(UBaseControlRig::StaticClass());
 	const bool bMatchesFlags = !InUnloadedClassData->HasAnyClassFlags(CLASS_Hidden | CLASS_HideDropDown | CLASS_Deprecated | CLASS_Abstract);
 	if (bChildOfObjectClass && bMatchesFlags)
 	{
diff --git a/Engine/Plugins/Animation/ControlRig/Source/ControlRigEditor/Private/ControlRigElementDetails.cpp b/Engine/Plugins/Animation/ControlRig/Source/ControlRigEditor/Private/ControlRigElementDetails.cpp
--- a/Engine/Plugins/Animation/ControlRig/Source/ControlRigEditor/Private/ControlRigElementDetails.cpp
+++ b/Engine/Plugins/Animation/ControlRig/Source/ControlRigEditor/Private/ControlRigElementDetails.cpp
@@ -91,7 +91,7 @@ void RigElementKeyDetails_GetCustomizedInfo(TSharedRef<IPropertyHandle> InStruct
 			break;
 		}
 
-		if(const UControlRig* ControlRig = Object->GetTypedOuter<UControlRig>())
+		if(const UBaseControlRig* ControlRig = Object->GetTypedOuter<UBaseControlRig>())
 		{
 			OutBlueprint = Cast<UControlRigBlueprint>(ControlRig->GetClass()->ClassGeneratedBy);
 			if(OutBlueprint)
@@ -116,7 +116,7 @@ void RigElementKeyDetails_GetCustomizedInfo(TSharedRef<IPropertyHandle> InStruct
 			Package->GetDefaultSubobjects(SubObjects);
 			for (UObject* SubObject : SubObjects)
 			{
-				if (UControlRig* Rig = Cast<UControlRig>(SubObject))
+				if (UBaseControlRig* Rig = Cast<UBaseControlRig>(SubObject))
 				{
 					UControlRigBlueprint* Blueprint = Cast<UControlRigBlueprint>(Rig->GetClass()->ClassGeneratedBy);
 					if (Blueprint)
@@ -148,7 +148,7 @@ UControlRigBlueprint* RigElementDetails_GetBlueprintFromHierarchy(URigHierarchy*
 	UControlRigBlueprint* Blueprint = InHierarchy->GetTypedOuter<UControlRigBlueprint>();
 	if(Blueprint == nullptr)
 	{
-		UControlRig* Rig = InHierarchy->GetTypedOuter<UControlRig>();
+		UBaseControlRig* Rig = InHierarchy->GetTypedOuter<UBaseControlRig>();
 		if(Rig)
 		{
 			Blueprint = Cast<UControlRigBlueprint>(Rig->GetClass()->ClassGeneratedBy);
@@ -965,7 +965,7 @@ bool FRigBaseElementDetails::IsConstructionModeEnabled() const
 	
 	if(const UControlRigBlueprint* Blueprint = PerElementInfos[0].GetBlueprint())
 	{
-		if (const UControlRig* DebuggedRig = Cast<UControlRig>(Blueprint ->GetObjectBeingDebugged()))
+		if (const UBaseControlRig* DebuggedRig = Cast<UBaseControlRig>(Blueprint ->GetObjectBeingDebugged()))
 		{
 			return DebuggedRig->IsConstructionModeEnabled();
 		}
@@ -2091,7 +2091,7 @@ FDetailWidgetRow& FRigTransformElementDetails::CreateEulerTransformValueWidgetRo
 		const bool bInitial = CurrentTransformType == ERigTransformElementDetailsTransform::Initial; 
 
 		bool bConstructionModeEnabled = false;
-		if (UControlRig* DebuggedRig = Cast<UControlRig>(PerElementInfos[0].GetBlueprint()->GetObjectBeingDebugged()))
+		if (UBaseControlRig* DebuggedRig = Cast<UBaseControlRig>(PerElementInfos[0].GetBlueprint()->GetObjectBeingDebugged()))
 		{
 			bConstructionModeEnabled = DebuggedRig->IsConstructionModeEnabled();
 		}
@@ -2330,7 +2330,7 @@ FDetailWidgetRow& FRigTransformElementDetails::CreateEulerTransformValueWidgetRo
 			ESlateRotationRepresentation::Type Representation,
 			ESlateTransformSubComponent::Type SubComponent)
 		{
-			if (UControlRig* DebuggedRig = Cast<UControlRig>(PerElementInfos[0].GetBlueprint()->GetObjectBeingDebugged()))
+			if (UBaseControlRig* DebuggedRig = Cast<UBaseControlRig>(PerElementInfos[0].GetBlueprint()->GetObjectBeingDebugged()))
 			{
 				EControlRigInteractionType Type = EControlRigInteractionType::None;
 				switch (Component)
@@ -2360,7 +2360,7 @@ FDetailWidgetRow& FRigTransformElementDetails::CreateEulerTransformValueWidgetRo
 			ESlateTransformSubComponent::Type SubComponent,
 			FVector::FReal InNumericValue)
 		{
-			if (UControlRig* DebuggedRig = Cast<UControlRig>(PerElementInfos[0].GetBlueprint()->GetObjectBeingDebugged()))
+			if (UBaseControlRig* DebuggedRig = Cast<UBaseControlRig>(PerElementInfos[0].GetBlueprint()->GetObjectBeingDebugged()))
 			{
 				DebuggedRig->InteractionType = (uint8)EControlRigInteractionType::None;
 				DebuggedRig->ElementsBeingInteracted.Reset();
@@ -3130,7 +3130,7 @@ void FRigControlElementDetails::CustomizeControl(IDetailLayoutBuilder& DetailBui
 						Info.GetDefaultHierarchy()->SetControlValue(ControlBeingCustomized, InitialValue, ERigControlValueType::Initial, false, false, true);
 						Info.GetDefaultHierarchy()->SetControlValue(ControlBeingCustomized, CurrentValue, ERigControlValueType::Current, false, false, true);
 
-						if (UControlRig* DebuggedRig = Cast<UControlRig>(Info.GetBlueprint()->GetObjectBeingDebugged()))
+						if (UBaseControlRig* DebuggedRig = Cast<UBaseControlRig>(Info.GetBlueprint()->GetObjectBeingDebugged()))
 						{
 							URigHierarchy* DebuggedHierarchy = DebuggedRig->GetHierarchy();
 							if(FRigControlElement* DebuggedControlElement = DebuggedHierarchy->Find<FRigControlElement>(ControlBeingCustomized->GetKey()))
diff --git a/Engine/Plugins/Animation/ControlRig/Source/ControlRigEditor/Private/ControlRigElementDetails.h b/Engine/Plugins/Animation/ControlRig/Source/ControlRigEditor/Private/ControlRigElementDetails.h
--- a/Engine/Plugins/Animation/ControlRig/Source/ControlRigEditor/Private/ControlRigElementDetails.h
+++ b/Engine/Plugins/Animation/ControlRig/Source/ControlRigEditor/Private/ControlRigElementDetails.h
@@ -196,7 +196,7 @@ public:
 
 		UControlRigBlueprint* GetBlueprint() const
 		{
-			if(const UControlRig* ControlRig = GetHierarchy()->GetTypedOuter<UControlRig>())
+			if(const UBaseControlRig* ControlRig = GetHierarchy()->GetTypedOuter<UBaseControlRig>())
 			{
 				return Cast<UControlRigBlueprint>(ControlRig->GetClass()->ClassGeneratedBy);
 			}
diff --git a/Engine/Plugins/Animation/ControlRig/Source/ControlRigEditor/Private/ControlRigInfluenceMapDetails.cpp b/Engine/Plugins/Animation/ControlRig/Source/ControlRigEditor/Private/ControlRigInfluenceMapDetails.cpp
--- a/Engine/Plugins/Animation/ControlRig/Source/ControlRigEditor/Private/ControlRigInfluenceMapDetails.cpp
+++ b/Engine/Plugins/Animation/ControlRig/Source/ControlRigEditor/Private/ControlRigInfluenceMapDetails.cpp
@@ -40,7 +40,7 @@ UControlRigBlueprint* FRigInfluenceMapPerEventDetails::GetBlueprintFromDetailBui
 
 			for (UObject* SubObject : SubObjects)
 			{
-				if (UControlRig* Rig = Cast<UControlRig>(SubObject))
+				if (UBaseControlRig* Rig = Cast<UBaseControlRig>(SubObject))
 				{
 					if(UControlRigBlueprint* Blueprint = Cast<UControlRigBlueprint>(Rig->GetClass()->ClassGeneratedBy))
 					{
diff --git a/Engine/Plugins/Animation/ControlRig/Source/ControlRigEditor/Private/ControlRigSequencerEditorLibrary.cpp b/Engine/Plugins/Animation/ControlRig/Source/ControlRigEditor/Private/ControlRigSequencerEditorLibrary.cpp
--- a/Engine/Plugins/Animation/ControlRig/Source/ControlRigEditor/Private/ControlRigSequencerEditorLibrary.cpp
+++ b/Engine/Plugins/Animation/ControlRig/Source/ControlRigEditor/Private/ControlRigSequencerEditorLibrary.cpp
@@ -54,9 +54,9 @@
 
 #define LOCTEXT_NAMESPACE "ControlrigSequencerEditorLibrary"
 
-TArray<UControlRig*> UControlRigSequencerEditorLibrary::GetVisibleControlRigs()
+TArray<UBaseControlRig*> UControlRigSequencerEditorLibrary::GetVisibleControlRigs()
 {
-	TArray<UControlRig*> ControlRigs;
+	TArray<UBaseControlRig*> ControlRigs;
 	FControlRigEditMode* ControlRigEditMode = static_cast<FControlRigEditMode*>(GLevelEditorModeTools().GetActiveMode(FControlRigEditMode::ModeName));
 	if (ControlRigEditMode)
 	{
@@ -199,10 +199,10 @@ static TSharedPtr<ISequencer> GetSequencerFromAsset()
 	return Sequencer;
 }
 
-static UMovieSceneControlRigParameterTrack* AddControlRig(ULevelSequence* LevelSequence,const UClass* InClass, UObject* BoundActor, FGuid ObjectBinding, UControlRig* InExistingControlRig, bool bIsAdditiveControlRig)
+static UMovieSceneControlRigParameterTrack* AddControlRig(ULevelSequence* LevelSequence,const UClass* InClass, UObject* BoundActor, FGuid ObjectBinding, UBaseControlRig* InExistingControlRig, bool bIsAdditiveControlRig)
 {
 	FSlateApplication::Get().DismissAllMenus();
-	if (!InClass || !InClass->IsChildOf(UControlRig::StaticClass()) ||
+	if (!InClass || !InClass->IsChildOf(UBaseControlRig::StaticClass()) ||
 		!LevelSequence || !LevelSequence->GetMovieScene())
 	{
 		return nullptr;
@@ -218,7 +218,7 @@ static UMovieSceneControlRigParameterTrack* AddControlRig(ULevelSequence* LevelS
 	LevelSequence->Modify();
 	OwnerMovieScene->Modify();
 	
-	if (bIsAdditiveControlRig && InClass != UFKControlRig::StaticClass() && !InClass->GetDefaultObject<UControlRig>()->SupportsEvent(FRigUnit_InverseExecution::EventName))
+	if (bIsAdditiveControlRig && InClass != UFKControlRig::StaticClass() && !InClass->GetDefaultObject<UBaseControlRig>()->SupportsEvent(FRigUnit_InverseExecution::EventName))
 	{
 		UE_LOG(LogControlRigEditor, Error, TEXT("Cannot add an additive control rig which does not contain a backwards solve event."));
 		return nullptr;
@@ -233,10 +233,10 @@ static UMovieSceneControlRigParameterTrack* AddControlRig(ULevelSequence* LevelS
 		ObjectName.RemoveFromEnd(TEXT("_C"));
 
 		bool bSequencerOwnsControlRig = false;
-		UControlRig* ControlRig = InExistingControlRig;
+		UBaseControlRig* ControlRig = InExistingControlRig;
 		if (ControlRig == nullptr)
 		{
-			ControlRig = NewObject<UControlRig>(Track, InClass, FName(*ObjectName), RF_Transactional);
+			ControlRig = NewObject<UBaseControlRig>(Track, InClass, FName(*ObjectName), RF_Transactional);
 			bSequencerOwnsControlRig = true;
 		}
 
@@ -254,7 +254,7 @@ static UMovieSceneControlRigParameterTrack* AddControlRig(ULevelSequence* LevelS
 		}
 		ControlRig->SetObjectBinding(MakeShared<FControlRigObjectBinding>());
 		ControlRig->GetObjectBinding()->BindToObject(BoundActor);
-		ControlRig->GetDataSourceRegistry()->RegisterDataSource(UControlRig::OwnerComponent, ControlRig->GetObjectBinding()->GetBoundObject());
+		ControlRig->GetDataSourceRegistry()->RegisterDataSource(UBaseControlRig::OwnerComponent, ControlRig->GetObjectBinding()->GetBoundObject());
 		// Do not re-initialize existing control rig
 		if (!InExistingControlRig)
 		{
@@ -338,7 +338,7 @@ UMovieSceneTrack* UControlRigSequencerEditorLibrary::FindOrCreateControlRigTrack
 				USkeletalMeshComponent* SkeletalMeshComponent = nullptr;
 				AcquireSkeletonAndSkelMeshCompFromObject(BoundObject, &Skeleton, &SkeletalMeshComponent);
 
-				UControlRig* ControlRig = nullptr;
+				UBaseControlRig* ControlRig = nullptr;
 				if (Skeleton && SkeletalMeshComponent)
 				{
 					UMovieSceneControlRigParameterTrack* Track = AddControlRig(LevelSequence, ControlRigClass, SkeletalMeshComponent, InBinding.BindingID, nullptr, bIsAdditiveControlRig);
@@ -387,7 +387,7 @@ TArray<UMovieSceneTrack*> UControlRigSequencerEditorLibrary::FindOrCreateControl
 			BoundActor->GetComponents<UControlRigComponent>(ControlRigComponents);
 			for (UControlRigComponent* ControlRigComponent : ControlRigComponents)
 			{
-				if (UControlRig* CR = ControlRigComponent->GetControlRig())
+				if (UBaseControlRig* CR = ControlRigComponent->GetControlRig())
 				{
 					UMovieSceneControlRigParameterTrack* GoodTrack = nullptr;
 					if (FMovieSceneBinding* Binding = MovieScene->FindBinding(InBinding.BindingID))
@@ -428,7 +428,7 @@ TArray<UMovieSceneTrack*> UControlRigSequencerEditorLibrary::FindOrCreateControl
 	return Tracks;
 }
 
-bool UControlRigSequencerEditorLibrary::TweenControlRig(ULevelSequence* LevelSequence, UControlRig* ControlRig, float TweenValue)
+bool UControlRigSequencerEditorLibrary::TweenControlRig(ULevelSequence* LevelSequence, UBaseControlRig* ControlRig, float TweenValue)
 {
 	TWeakPtr<ISequencer> WeakSequencer = GetSequencerFromAsset();
 	if (WeakSequencer.IsValid() && WeakSequencer.Pin()->GetFocusedMovieSceneSequence() == LevelSequence
@@ -436,7 +436,7 @@ bool UControlRigSequencerEditorLibrary::TweenControlRig(ULevelSequence* LevelSeq
 	{
 		FControlsToTween ControlsToTween;
 		LevelSequence->GetMovieScene()->Modify();
-		TArray<UControlRig*> SelectedControlRigs;
+		TArray<UBaseControlRig*> SelectedControlRigs;
 		SelectedControlRigs.Add(ControlRig);
 		ControlsToTween.Setup(SelectedControlRigs, WeakSequencer);
 		ControlsToTween.Blend(WeakSequencer, TweenValue);
@@ -1012,7 +1012,7 @@ TArray<FTransform> UControlRigSequencerEditorLibrary::GetSkeletalMeshComponentWo
 	return OutWorldTransforms;
 }
 
-FTransform UControlRigSequencerEditorLibrary::GetControlRigWorldTransform(ULevelSequence* LevelSequence,UControlRig* ControlRig, FName ControlName, FFrameNumber Frame,
+FTransform UControlRigSequencerEditorLibrary::GetControlRigWorldTransform(ULevelSequence* LevelSequence,UBaseControlRig* ControlRig, FName ControlName, FFrameNumber Frame,
 	ESequenceTimeUnit TimeUnit)
 {
 	TArray<FFrameNumber> Frames;
@@ -1025,7 +1025,7 @@ FTransform UControlRigSequencerEditorLibrary::GetControlRigWorldTransform(ULevel
 	return FTransform::Identity;
 }
 
-TArray<FTransform> UControlRigSequencerEditorLibrary::GetControlRigWorldTransforms(ULevelSequence* LevelSequence,UControlRig* ControlRig, FName ControlName, const TArray<FFrameNumber>& InFrames,
+TArray<FTransform> UControlRigSequencerEditorLibrary::GetControlRigWorldTransforms(ULevelSequence* LevelSequence,UBaseControlRig* ControlRig, FName ControlName, const TArray<FFrameNumber>& InFrames,
 	ESequenceTimeUnit TimeUnit)
 {
 	TWeakPtr<ISequencer> WeakSequencer = GetSequencerFromAsset();
@@ -1057,7 +1057,7 @@ TArray<FTransform> UControlRigSequencerEditorLibrary::GetControlRigWorldTransfor
 }
 
 
-static void LocalSetControlRigWorldTransforms(ULevelSequence* LevelSequence,UControlRig* ControlRig, FName ControlName, EControlRigSetKey SetKey, const TArray<FFrameNumber>& InFrames, 
+static void LocalSetControlRigWorldTransforms(ULevelSequence* LevelSequence,UBaseControlRig* ControlRig, FName ControlName, EControlRigSetKey SetKey, const TArray<FFrameNumber>& InFrames, 
 	const TArray<FTransform>& WorldTransforms, ESequenceTimeUnit TimeUnit)
 {
 	TWeakPtr<ISequencer> WeakSequencer = GetSequencerFromAsset();
@@ -1111,7 +1111,7 @@ static void LocalSetControlRigWorldTransforms(ULevelSequence* LevelSequence,UCon
 	}
 }
 
-void UControlRigSequencerEditorLibrary::SetControlRigWorldTransform(ULevelSequence* LevelSequence, UControlRig* ControlRig, FName ControlName, FFrameNumber Frame, const FTransform& WorldTransform,
+void UControlRigSequencerEditorLibrary::SetControlRigWorldTransform(ULevelSequence* LevelSequence, UBaseControlRig* ControlRig, FName ControlName, FFrameNumber Frame, const FTransform& WorldTransform,
 	ESequenceTimeUnit TimeUnit, bool bSetKey)
 {
 	EControlRigSetKey SetKey = bSetKey ? EControlRigSetKey::Always : EControlRigSetKey::DoNotCare;
@@ -1124,7 +1124,7 @@ void UControlRigSequencerEditorLibrary::SetControlRigWorldTransform(ULevelSequen
 
 }
 
-void UControlRigSequencerEditorLibrary::SetControlRigWorldTransforms(ULevelSequence* LevelSequence, UControlRig* ControlRig, FName ControlName, const TArray<FFrameNumber>& Frames, const TArray<FTransform>& WorldTransforms,
+void UControlRigSequencerEditorLibrary::SetControlRigWorldTransforms(ULevelSequence* LevelSequence, UBaseControlRig* ControlRig, FName ControlName, const TArray<FFrameNumber>& Frames, const TArray<FTransform>& WorldTransforms,
 	ESequenceTimeUnit TimeUnit)
 {
 	LocalSetControlRigWorldTransforms(LevelSequence, ControlRig, ControlName, EControlRigSetKey::Always, Frames, WorldTransforms,TimeUnit);
@@ -1227,7 +1227,7 @@ bool UControlRigSequencerEditorLibrary::BakeToControlRig(UWorld* World, ULevelSe
 				{
 					FString ObjectName = InClass->GetName();
 					ObjectName.RemoveFromEnd(TEXT("_C"));
-					UControlRig* ControlRig = NewObject<UControlRig>(Track, InClass, FName(*ObjectName), RF_Transactional);
+					UBaseControlRig* ControlRig = NewObject<UBaseControlRig>(Track, InClass, FName(*ObjectName), RF_Transactional);
 					FName OldEventString = FName(FString(TEXT("Inverse")));
 
 					if (InClass != UFKControlRig::StaticClass() && !(ControlRig->SupportsEvent(FRigUnit_InverseExecution::EventName) || ControlRig->SupportsEvent(OldEventString)))
@@ -1253,7 +1253,7 @@ bool UControlRigSequencerEditorLibrary::BakeToControlRig(UWorld* World, ULevelSe
 						else
 						{
 							/* mz todo we don't unbind  will test more
-							UControlRig* OldControlRig = ControlRigEditMode->GetControlRig(false);
+							UBaseControlRig* OldControlRig = ControlRigEditMode->GetControlRig(false);
 							if (OldControlRig)
 							{
 								WeakSequencer.Pin()->ObjectImplicitlyRemoved(OldControlRig);
@@ -1265,7 +1265,7 @@ bool UControlRigSequencerEditorLibrary::BakeToControlRig(UWorld* World, ULevelSe
 					ControlRig->Modify();
 					ControlRig->SetObjectBinding(MakeShared<FControlRigObjectBinding>());
 					ControlRig->GetObjectBinding()->BindToObject(SkeletalMeshComp);
-					ControlRig->GetDataSourceRegistry()->RegisterDataSource(UControlRig::OwnerComponent, ControlRig->GetObjectBinding()->GetBoundObject());
+					ControlRig->GetDataSourceRegistry()->RegisterDataSource(UBaseControlRig::OwnerComponent, ControlRig->GetObjectBinding()->GetBoundObject());
 					ControlRig->Initialize();
 					ControlRig->RequestInit();
 					ControlRig->SetBoneInitialTransformsFromSkeletalMeshComponent(SkeletalMeshComp, true);
@@ -1358,7 +1358,7 @@ bool UControlRigSequencerEditorLibrary::LoadAnimSequenceIntoControlRigSection(UM
 }
 
 static bool LocalGetControlRigControlValues(IMovieScenePlayer* Player, UMovieSceneSequence* MovieSceneSequence, FMovieSceneSequenceIDRef Template, FMovieSceneSequenceTransform& RootToLocalTransform,
-	UControlRig* ControlRig, const FName& ControlName, ESequenceTimeUnit TimeUnit,
+	UBaseControlRig* ControlRig, const FName& ControlName, ESequenceTimeUnit TimeUnit,
 	const TArray<FFrameNumber>& InFrames, TArray<FRigControlValue>& OutValues)
 {
 	if (Player == nullptr || MovieSceneSequence == nullptr || ControlRig == nullptr)
@@ -1394,7 +1394,7 @@ static bool LocalGetControlRigControlValues(IMovieScenePlayer* Player, UMovieSce
 	return true;
 }
 
-static bool GetControlRigValues(ISequencer* Sequencer, UControlRig* ControlRig, const FName& ControlName, ESequenceTimeUnit TimeUnit,
+static bool GetControlRigValues(ISequencer* Sequencer, UBaseControlRig* ControlRig, const FName& ControlName, ESequenceTimeUnit TimeUnit,
 	const TArray<FFrameNumber>& Frames,  TArray<FRigControlValue>& OutValues)
 {
 	if (Sequencer->GetFocusedMovieSceneSequence())
@@ -1409,7 +1409,7 @@ static bool GetControlRigValues(ISequencer* Sequencer, UControlRig* ControlRig,
 	return false;
 }
 
-static bool GetControlRigValue(ISequencer* Sequencer, UControlRig* ControlRig, const FName& ControlName, ESequenceTimeUnit TimeUnit,
+static bool GetControlRigValue(ISequencer* Sequencer, UBaseControlRig* ControlRig, const FName& ControlName, ESequenceTimeUnit TimeUnit,
 	const FFrameNumber Frame, FRigControlValue& OutValue)
 {
 	if (Sequencer->GetFocusedMovieSceneSequence())
@@ -1430,7 +1430,7 @@ static bool GetControlRigValue(ISequencer* Sequencer, UControlRig* ControlRig, c
 	return false;
 }
 
-static bool GetControlRigValues(UWorld* World, ULevelSequence* LevelSequence, UControlRig* ControlRig, const FName& ControlName, ESequenceTimeUnit TimeUnit,
+static bool GetControlRigValues(UWorld* World, ULevelSequence* LevelSequence, UBaseControlRig* ControlRig, const FName& ControlName, ESequenceTimeUnit TimeUnit,
 	const TArray<FFrameNumber>& Frames, TArray<FRigControlValue>& OutValues)
 {
 	if (LevelSequence)
@@ -1449,7 +1449,7 @@ static bool GetControlRigValues(UWorld* World, ULevelSequence* LevelSequence, UC
 }
 
 
-float UControlRigSequencerEditorLibrary::GetLocalControlRigFloat(ULevelSequence* LevelSequence,UControlRig* ControlRig, FName ControlName, FFrameNumber Frame, ESequenceTimeUnit TimeUnit)
+float UControlRigSequencerEditorLibrary::GetLocalControlRigFloat(ULevelSequence* LevelSequence,UBaseControlRig* ControlRig, FName ControlName, FFrameNumber Frame, ESequenceTimeUnit TimeUnit)
 {
 	float Value = 0.0f;
 	TWeakPtr<ISequencer> WeakSequencer = GetSequencerFromAsset();
@@ -1473,7 +1473,7 @@ float UControlRigSequencerEditorLibrary::GetLocalControlRigFloat(ULevelSequence*
 	return Value;
 }
 
-TArray<float> UControlRigSequencerEditorLibrary::GetLocalControlRigFloats(ULevelSequence* LevelSequence,UControlRig* ControlRig, FName ControlName, const TArray<FFrameNumber>& Frames, ESequenceTimeUnit TimeUnit)
+TArray<float> UControlRigSequencerEditorLibrary::GetLocalControlRigFloats(ULevelSequence* LevelSequence,UBaseControlRig* ControlRig, FName ControlName, const TArray<FFrameNumber>& Frames, ESequenceTimeUnit TimeUnit)
 {
 	TArray<float> Values;
 	TWeakPtr<ISequencer> WeakSequencer = GetSequencerFromAsset();
@@ -1503,7 +1503,7 @@ TArray<float> UControlRigSequencerEditorLibrary::GetLocalControlRigFloats(ULevel
 	return Values;
 }
 
-void UControlRigSequencerEditorLibrary::SetLocalControlRigFloat(ULevelSequence* LevelSequence, UControlRig* ControlRig, FName ControlName, FFrameNumber Frame, float Value, ESequenceTimeUnit TimeUnit, bool bSetKey)
+void UControlRigSequencerEditorLibrary::SetLocalControlRigFloat(ULevelSequence* LevelSequence, UBaseControlRig* ControlRig, FName ControlName, FFrameNumber Frame, float Value, ESequenceTimeUnit TimeUnit, bool bSetKey)
 {
 	if (LevelSequence == nullptr || ControlRig == nullptr)
 	{
@@ -1531,7 +1531,7 @@ void UControlRigSequencerEditorLibrary::SetLocalControlRigFloat(ULevelSequence*
 	}
 }
 
-void UControlRigSequencerEditorLibrary::SetLocalControlRigFloats(ULevelSequence* LevelSequence, UControlRig* ControlRig, FName ControlName, const TArray<FFrameNumber>& Frames, const TArray<float> Values,
+void UControlRigSequencerEditorLibrary::SetLocalControlRigFloats(ULevelSequence* LevelSequence, UBaseControlRig* ControlRig, FName ControlName, const TArray<FFrameNumber>& Frames, const TArray<float> Values,
 	ESequenceTimeUnit TimeUnit)
 {
 	if (LevelSequence == nullptr || ControlRig == nullptr || (Frames.Num() != Values.Num()))
@@ -1566,7 +1566,7 @@ void UControlRigSequencerEditorLibrary::SetLocalControlRigFloats(ULevelSequence*
 }
 
 
-bool UControlRigSequencerEditorLibrary::GetLocalControlRigBool(ULevelSequence* LevelSequence, UControlRig* ControlRig, FName ControlName, FFrameNumber Frame, ESequenceTimeUnit TimeUnit)
+bool UControlRigSequencerEditorLibrary::GetLocalControlRigBool(ULevelSequence* LevelSequence, UBaseControlRig* ControlRig, FName ControlName, FFrameNumber Frame, ESequenceTimeUnit TimeUnit)
 {
 	bool Value = true;
 	TWeakPtr<ISequencer> WeakSequencer = GetSequencerFromAsset();
@@ -1590,7 +1590,7 @@ bool UControlRigSequencerEditorLibrary::GetLocalControlRigBool(ULevelSequence* L
 	return Value;
 }
 
-TArray<bool> UControlRigSequencerEditorLibrary::GetLocalControlRigBools(ULevelSequence* LevelSequence, UControlRig* ControlRig, FName ControlName, const TArray<FFrameNumber>& Frames, ESequenceTimeUnit TimeUnit)
+TArray<bool> UControlRigSequencerEditorLibrary::GetLocalControlRigBools(ULevelSequence* LevelSequence, UBaseControlRig* ControlRig, FName ControlName, const TArray<FFrameNumber>& Frames, ESequenceTimeUnit TimeUnit)
 {
 	TArray<bool> Values;
 	TWeakPtr<ISequencer> WeakSequencer = GetSequencerFromAsset();
@@ -1619,7 +1619,7 @@ TArray<bool> UControlRigSequencerEditorLibrary::GetLocalControlRigBools(ULevelSe
 	return Values;
 }
 
-void UControlRigSequencerEditorLibrary::SetLocalControlRigBool(ULevelSequence* LevelSequence, UControlRig* ControlRig, FName ControlName, FFrameNumber Frame, bool Value, ESequenceTimeUnit TimeUnit, bool bSetKey)
+void UControlRigSequencerEditorLibrary::SetLocalControlRigBool(ULevelSequence* LevelSequence, UBaseControlRig* ControlRig, FName ControlName, FFrameNumber Frame, bool Value, ESequenceTimeUnit TimeUnit, bool bSetKey)
 {
 	if (LevelSequence == nullptr || ControlRig == nullptr)
 	{
@@ -1647,7 +1647,7 @@ void UControlRigSequencerEditorLibrary::SetLocalControlRigBool(ULevelSequence* L
 	}
 }
 
-void UControlRigSequencerEditorLibrary::SetLocalControlRigBools(ULevelSequence* LevelSequence, UControlRig* ControlRig, FName ControlName, const TArray<FFrameNumber>& Frames,
+void UControlRigSequencerEditorLibrary::SetLocalControlRigBools(ULevelSequence* LevelSequence, UBaseControlRig* ControlRig, FName ControlName, const TArray<FFrameNumber>& Frames,
 	const TArray<bool> Values, ESequenceTimeUnit TimeUnit)
 {
 	if (LevelSequence == nullptr || ControlRig == nullptr || (Frames.Num() != Values.Num()))
@@ -1681,7 +1681,7 @@ void UControlRigSequencerEditorLibrary::SetLocalControlRigBools(ULevelSequence*
 	}
 }
 
-int32 UControlRigSequencerEditorLibrary::GetLocalControlRigInt(ULevelSequence* LevelSequence, UControlRig* ControlRig, FName ControlName, FFrameNumber Frame, ESequenceTimeUnit TimeUnit)
+int32 UControlRigSequencerEditorLibrary::GetLocalControlRigInt(ULevelSequence* LevelSequence, UBaseControlRig* ControlRig, FName ControlName, FFrameNumber Frame, ESequenceTimeUnit TimeUnit)
 {
 	int32 Value = 0;
 	TWeakPtr<ISequencer> WeakSequencer = GetSequencerFromAsset();
@@ -1705,7 +1705,7 @@ int32 UControlRigSequencerEditorLibrary::GetLocalControlRigInt(ULevelSequence* L
 	return Value;
 }
 
-TArray<int32> UControlRigSequencerEditorLibrary::GetLocalControlRigInts(ULevelSequence* LevelSequence, UControlRig* ControlRig, FName ControlName, const TArray<FFrameNumber>& Frames, ESequenceTimeUnit TimeUnit)
+TArray<int32> UControlRigSequencerEditorLibrary::GetLocalControlRigInts(ULevelSequence* LevelSequence, UBaseControlRig* ControlRig, FName ControlName, const TArray<FFrameNumber>& Frames, ESequenceTimeUnit TimeUnit)
 {
 	TArray<int32> Values;
 	TWeakPtr<ISequencer> WeakSequencer = GetSequencerFromAsset();
@@ -1734,7 +1734,7 @@ TArray<int32> UControlRigSequencerEditorLibrary::GetLocalControlRigInts(ULevelSe
 	return Values;
 }
 
-void UControlRigSequencerEditorLibrary::SetLocalControlRigInt(ULevelSequence* LevelSequence, UControlRig* ControlRig, FName ControlName, FFrameNumber Frame, int32 Value, ESequenceTimeUnit TimeUnit,bool bSetKey)
+void UControlRigSequencerEditorLibrary::SetLocalControlRigInt(ULevelSequence* LevelSequence, UBaseControlRig* ControlRig, FName ControlName, FFrameNumber Frame, int32 Value, ESequenceTimeUnit TimeUnit,bool bSetKey)
 {
 	if (LevelSequence == nullptr || ControlRig == nullptr)
 	{
@@ -1762,7 +1762,7 @@ void UControlRigSequencerEditorLibrary::SetLocalControlRigInt(ULevelSequence* Le
 	}
 }
 
-void UControlRigSequencerEditorLibrary::SetLocalControlRigInts(ULevelSequence* LevelSequence, UControlRig* ControlRig, FName ControlName, const TArray<FFrameNumber>& Frames,
+void UControlRigSequencerEditorLibrary::SetLocalControlRigInts(ULevelSequence* LevelSequence, UBaseControlRig* ControlRig, FName ControlName, const TArray<FFrameNumber>& Frames,
 	const TArray<int32> Values , ESequenceTimeUnit TimeUnit)
 {
 	if (LevelSequence == nullptr || ControlRig == nullptr || (Frames.Num() != Values.Num()))
@@ -1797,7 +1797,7 @@ void UControlRigSequencerEditorLibrary::SetLocalControlRigInts(ULevelSequence* L
 }
 
 
-FVector2D UControlRigSequencerEditorLibrary::GetLocalControlRigVector2D(ULevelSequence* LevelSequence, UControlRig* ControlRig, FName ControlName, FFrameNumber Frame, ESequenceTimeUnit TimeUnit)
+FVector2D UControlRigSequencerEditorLibrary::GetLocalControlRigVector2D(ULevelSequence* LevelSequence, UBaseControlRig* ControlRig, FName ControlName, FFrameNumber Frame, ESequenceTimeUnit TimeUnit)
 {
 	FVector2D Value;
 	TWeakPtr<ISequencer> WeakSequencer = GetSequencerFromAsset();
@@ -1822,7 +1822,7 @@ FVector2D UControlRigSequencerEditorLibrary::GetLocalControlRigVector2D(ULevelSe
 	return Value;
 }
 
-TArray<FVector2D> UControlRigSequencerEditorLibrary::GetLocalControlRigVector2Ds(ULevelSequence* LevelSequence, UControlRig* ControlRig, FName ControlName, const TArray<FFrameNumber>& Frames, ESequenceTimeUnit TimeUnit)
+TArray<FVector2D> UControlRigSequencerEditorLibrary::GetLocalControlRigVector2Ds(ULevelSequence* LevelSequence, UBaseControlRig* ControlRig, FName ControlName, const TArray<FFrameNumber>& Frames, ESequenceTimeUnit TimeUnit)
 {
 	TArray<FVector2D> Values;
 	TWeakPtr<ISequencer> WeakSequencer = GetSequencerFromAsset();
@@ -1852,7 +1852,7 @@ TArray<FVector2D> UControlRigSequencerEditorLibrary::GetLocalControlRigVector2Ds
 
 }
 
-void UControlRigSequencerEditorLibrary::SetLocalControlRigVector2D(ULevelSequence* LevelSequence, UControlRig* ControlRig, FName ControlName, FFrameNumber Frame, FVector2D Value, ESequenceTimeUnit TimeUnit,bool bSetKey)
+void UControlRigSequencerEditorLibrary::SetLocalControlRigVector2D(ULevelSequence* LevelSequence, UBaseControlRig* ControlRig, FName ControlName, FFrameNumber Frame, FVector2D Value, ESequenceTimeUnit TimeUnit,bool bSetKey)
 {
 	if (LevelSequence == nullptr || ControlRig == nullptr)
 	{
@@ -1880,7 +1880,7 @@ void UControlRigSequencerEditorLibrary::SetLocalControlRigVector2D(ULevelSequenc
 	}
 }
 
-void UControlRigSequencerEditorLibrary::SetLocalControlRigVector2Ds(ULevelSequence* LevelSequence, UControlRig* ControlRig, FName ControlName, const TArray<FFrameNumber>& Frames, 
+void UControlRigSequencerEditorLibrary::SetLocalControlRigVector2Ds(ULevelSequence* LevelSequence, UBaseControlRig* ControlRig, FName ControlName, const TArray<FFrameNumber>& Frames, 
 	const TArray<FVector2D> Values, ESequenceTimeUnit TimeUnit)
 {
 	if (LevelSequence == nullptr || ControlRig == nullptr || (Frames.Num() != Values.Num()))
@@ -1915,7 +1915,7 @@ void UControlRigSequencerEditorLibrary::SetLocalControlRigVector2Ds(ULevelSequen
 }
 
 
-FVector UControlRigSequencerEditorLibrary::GetLocalControlRigPosition(ULevelSequence* LevelSequence, UControlRig* ControlRig, FName ControlName, FFrameNumber Frame, ESequenceTimeUnit TimeUnit)
+FVector UControlRigSequencerEditorLibrary::GetLocalControlRigPosition(ULevelSequence* LevelSequence, UBaseControlRig* ControlRig, FName ControlName, FFrameNumber Frame, ESequenceTimeUnit TimeUnit)
 {
 	FVector Value;
 	TWeakPtr<ISequencer> WeakSequencer = GetSequencerFromAsset();
@@ -1939,7 +1939,7 @@ FVector UControlRigSequencerEditorLibrary::GetLocalControlRigPosition(ULevelSequ
 	return Value;
 }
 
-TArray<FVector> UControlRigSequencerEditorLibrary::GetLocalControlRigPositions(ULevelSequence* LevelSequence, UControlRig* ControlRig, FName ControlName, const TArray<FFrameNumber>& Frames, ESequenceTimeUnit TimeUnit)
+TArray<FVector> UControlRigSequencerEditorLibrary::GetLocalControlRigPositions(ULevelSequence* LevelSequence, UBaseControlRig* ControlRig, FName ControlName, const TArray<FFrameNumber>& Frames, ESequenceTimeUnit TimeUnit)
 {
 	TArray<FVector> Values;
 	TWeakPtr<ISequencer> WeakSequencer = GetSequencerFromAsset();
@@ -1968,7 +1968,7 @@ TArray<FVector> UControlRigSequencerEditorLibrary::GetLocalControlRigPositions(U
 	return Values;
 }
 
-void UControlRigSequencerEditorLibrary::SetLocalControlRigPosition(ULevelSequence* LevelSequence, UControlRig* ControlRig, FName ControlName, FFrameNumber Frame, FVector Value, ESequenceTimeUnit TimeUnit,bool bSetKey)
+void UControlRigSequencerEditorLibrary::SetLocalControlRigPosition(ULevelSequence* LevelSequence, UBaseControlRig* ControlRig, FName ControlName, FFrameNumber Frame, FVector Value, ESequenceTimeUnit TimeUnit,bool bSetKey)
 {
 	if (LevelSequence == nullptr || ControlRig == nullptr)
 	{
@@ -1996,7 +1996,7 @@ void UControlRigSequencerEditorLibrary::SetLocalControlRigPosition(ULevelSequenc
 	}
 }
 
-void UControlRigSequencerEditorLibrary::SetLocalControlRigPositions(ULevelSequence* LevelSequence, UControlRig* ControlRig, FName ControlName, const TArray<FFrameNumber>& Frames, 
+void UControlRigSequencerEditorLibrary::SetLocalControlRigPositions(ULevelSequence* LevelSequence, UBaseControlRig* ControlRig, FName ControlName, const TArray<FFrameNumber>& Frames, 
 	const TArray<FVector> Values, ESequenceTimeUnit TimeUnit)
 {
 	if (LevelSequence == nullptr || ControlRig == nullptr || (Frames.Num() != Values.Num()))
@@ -2031,7 +2031,7 @@ void UControlRigSequencerEditorLibrary::SetLocalControlRigPositions(ULevelSequen
 }
 
 
-FRotator UControlRigSequencerEditorLibrary::GetLocalControlRigRotator(ULevelSequence* LevelSequence, UControlRig* ControlRig, FName ControlName, FFrameNumber Frame, ESequenceTimeUnit TimeUnit)
+FRotator UControlRigSequencerEditorLibrary::GetLocalControlRigRotator(ULevelSequence* LevelSequence, UBaseControlRig* ControlRig, FName ControlName, FFrameNumber Frame, ESequenceTimeUnit TimeUnit)
 {
 	FRotator Value;
 	TWeakPtr<ISequencer> WeakSequencer = GetSequencerFromAsset();
@@ -2055,7 +2055,7 @@ FRotator UControlRigSequencerEditorLibrary::GetLocalControlRigRotator(ULevelSequ
 	return Value;
 }
 
-TArray<FRotator> UControlRigSequencerEditorLibrary::GetLocalControlRigRotators(ULevelSequence* LevelSequence, UControlRig* ControlRig, FName ControlName, const TArray<FFrameNumber>& Frames, ESequenceTimeUnit TimeUnit)
+TArray<FRotator> UControlRigSequencerEditorLibrary::GetLocalControlRigRotators(ULevelSequence* LevelSequence, UBaseControlRig* ControlRig, FName ControlName, const TArray<FFrameNumber>& Frames, ESequenceTimeUnit TimeUnit)
 {
 	TArray<FRotator> Values;
 	TWeakPtr<ISequencer> WeakSequencer = GetSequencerFromAsset();
@@ -2084,7 +2084,7 @@ TArray<FRotator> UControlRigSequencerEditorLibrary::GetLocalControlRigRotators(U
 	return Values;
 }
 
-void UControlRigSequencerEditorLibrary::SetLocalControlRigRotator(ULevelSequence* LevelSequence, UControlRig* ControlRig, FName ControlName, FFrameNumber Frame, FRotator Value, ESequenceTimeUnit TimeUnit, bool bSetKey)
+void UControlRigSequencerEditorLibrary::SetLocalControlRigRotator(ULevelSequence* LevelSequence, UBaseControlRig* ControlRig, FName ControlName, FFrameNumber Frame, FRotator Value, ESequenceTimeUnit TimeUnit, bool bSetKey)
 {
 	if (LevelSequence == nullptr || ControlRig == nullptr)
 	{
@@ -2112,7 +2112,7 @@ void UControlRigSequencerEditorLibrary::SetLocalControlRigRotator(ULevelSequence
 	}
 }
 
-void UControlRigSequencerEditorLibrary::SetLocalControlRigRotators(ULevelSequence* LevelSequence, UControlRig* ControlRig, FName ControlName, const TArray<FFrameNumber>& Frames,
+void UControlRigSequencerEditorLibrary::SetLocalControlRigRotators(ULevelSequence* LevelSequence, UBaseControlRig* ControlRig, FName ControlName, const TArray<FFrameNumber>& Frames,
 	const TArray<FRotator> Values, ESequenceTimeUnit TimeUnit)
 {
 	if (LevelSequence == nullptr || ControlRig == nullptr || (Frames.Num() != Values.Num()))
@@ -2147,7 +2147,7 @@ void UControlRigSequencerEditorLibrary::SetLocalControlRigRotators(ULevelSequenc
 }
 
 
-FVector UControlRigSequencerEditorLibrary::GetLocalControlRigScale(ULevelSequence* LevelSequence, UControlRig* ControlRig, FName ControlName, FFrameNumber Frame,  ESequenceTimeUnit TimeUnit)
+FVector UControlRigSequencerEditorLibrary::GetLocalControlRigScale(ULevelSequence* LevelSequence, UBaseControlRig* ControlRig, FName ControlName, FFrameNumber Frame,  ESequenceTimeUnit TimeUnit)
 {
 	FVector Value;
 	TWeakPtr<ISequencer> WeakSequencer = GetSequencerFromAsset();
@@ -2171,7 +2171,7 @@ FVector UControlRigSequencerEditorLibrary::GetLocalControlRigScale(ULevelSequenc
 	return Value;
 }
 
-TArray<FVector>UControlRigSequencerEditorLibrary::GetLocalControlRigScales(ULevelSequence* LevelSequence, UControlRig* ControlRig, FName ControlName, const TArray<FFrameNumber>& Frames, ESequenceTimeUnit TimeUnit)
+TArray<FVector>UControlRigSequencerEditorLibrary::GetLocalControlRigScales(ULevelSequence* LevelSequence, UBaseControlRig* ControlRig, FName ControlName, const TArray<FFrameNumber>& Frames, ESequenceTimeUnit TimeUnit)
 {
 	TArray<FVector> Values;
 	TWeakPtr<ISequencer> WeakSequencer = GetSequencerFromAsset();
@@ -2200,7 +2200,7 @@ TArray<FVector>UControlRigSequencerEditorLibrary::GetLocalControlRigScales(ULeve
 	return Values;
 }
 
-void UControlRigSequencerEditorLibrary::SetLocalControlRigScale(ULevelSequence* LevelSequence, UControlRig* ControlRig, FName ControlName, FFrameNumber Frame, FVector Value, ESequenceTimeUnit TimeUnit, bool bSetKey)
+void UControlRigSequencerEditorLibrary::SetLocalControlRigScale(ULevelSequence* LevelSequence, UBaseControlRig* ControlRig, FName ControlName, FFrameNumber Frame, FVector Value, ESequenceTimeUnit TimeUnit, bool bSetKey)
 {
 	if (LevelSequence == nullptr || ControlRig == nullptr)
 	{
@@ -2228,7 +2228,7 @@ void UControlRigSequencerEditorLibrary::SetLocalControlRigScale(ULevelSequence*
 	}
 }
 
-void UControlRigSequencerEditorLibrary::SetLocalControlRigScales(ULevelSequence* LevelSequence, UControlRig* ControlRig, FName ControlName, const TArray<FFrameNumber>& Frames, 
+void UControlRigSequencerEditorLibrary::SetLocalControlRigScales(ULevelSequence* LevelSequence, UBaseControlRig* ControlRig, FName ControlName, const TArray<FFrameNumber>& Frames, 
 	const TArray<FVector> Values, ESequenceTimeUnit TimeUnit)
 {
 	if (LevelSequence == nullptr || ControlRig == nullptr || (Frames.Num() != Values.Num()))
@@ -2263,7 +2263,7 @@ void UControlRigSequencerEditorLibrary::SetLocalControlRigScales(ULevelSequence*
 }
 
 
-FEulerTransform UControlRigSequencerEditorLibrary::GetLocalControlRigEulerTransform(ULevelSequence* LevelSequence, UControlRig* ControlRig, FName ControlName, FFrameNumber Frame, ESequenceTimeUnit TimeUnit)
+FEulerTransform UControlRigSequencerEditorLibrary::GetLocalControlRigEulerTransform(ULevelSequence* LevelSequence, UBaseControlRig* ControlRig, FName ControlName, FFrameNumber Frame, ESequenceTimeUnit TimeUnit)
 {
 	FEulerTransform Value = FEulerTransform::Identity;
 	TWeakPtr<ISequencer> WeakSequencer = GetSequencerFromAsset();
@@ -2287,7 +2287,7 @@ FEulerTransform UControlRigSequencerEditorLibrary::GetLocalControlRigEulerTransf
 	return Value;
 }
 
-TArray<FEulerTransform> UControlRigSequencerEditorLibrary::GetLocalControlRigEulerTransforms(ULevelSequence* LevelSequence, UControlRig* ControlRig, FName ControlName, const TArray<FFrameNumber>& Frames, ESequenceTimeUnit TimeUnit)
+TArray<FEulerTransform> UControlRigSequencerEditorLibrary::GetLocalControlRigEulerTransforms(ULevelSequence* LevelSequence, UBaseControlRig* ControlRig, FName ControlName, const TArray<FFrameNumber>& Frames, ESequenceTimeUnit TimeUnit)
 {
 	TArray<FEulerTransform> Values;
 	TWeakPtr<ISequencer> WeakSequencer = GetSequencerFromAsset();
@@ -2316,7 +2316,7 @@ TArray<FEulerTransform> UControlRigSequencerEditorLibrary::GetLocalControlRigEul
 	return Values;
 }
 
-void UControlRigSequencerEditorLibrary::SetLocalControlRigEulerTransform(ULevelSequence* LevelSequence, UControlRig* ControlRig, FName ControlName, FFrameNumber Frame, FEulerTransform Value,
+void UControlRigSequencerEditorLibrary::SetLocalControlRigEulerTransform(ULevelSequence* LevelSequence, UBaseControlRig* ControlRig, FName ControlName, FFrameNumber Frame, FEulerTransform Value,
 	ESequenceTimeUnit TimeUnit, bool bSetKey)
 {
 	if (LevelSequence == nullptr || ControlRig == nullptr)
@@ -2350,7 +2350,7 @@ void UControlRigSequencerEditorLibrary::SetLocalControlRigEulerTransform(ULevelS
 	}
 }
 
-void UControlRigSequencerEditorLibrary::SetLocalControlRigEulerTransforms(ULevelSequence* LevelSequence, UControlRig* ControlRig, FName ControlName, const TArray<FFrameNumber>& Frames,
+void UControlRigSequencerEditorLibrary::SetLocalControlRigEulerTransforms(ULevelSequence* LevelSequence, UBaseControlRig* ControlRig, FName ControlName, const TArray<FFrameNumber>& Frames,
 	const TArray<FEulerTransform> Values, ESequenceTimeUnit TimeUnit)
 {
 	if (LevelSequence == nullptr || ControlRig == nullptr || (Frames.Num() != Values.Num()))
@@ -2391,7 +2391,7 @@ void UControlRigSequencerEditorLibrary::SetLocalControlRigEulerTransforms(ULevel
 }
 
 
-FTransformNoScale UControlRigSequencerEditorLibrary::GetLocalControlRigTransformNoScale(ULevelSequence* LevelSequence, UControlRig* ControlRig, FName ControlName, FFrameNumber Frame, ESequenceTimeUnit TimeUnit)
+FTransformNoScale UControlRigSequencerEditorLibrary::GetLocalControlRigTransformNoScale(ULevelSequence* LevelSequence, UBaseControlRig* ControlRig, FName ControlName, FFrameNumber Frame, ESequenceTimeUnit TimeUnit)
 {
 	FTransformNoScale Value = FTransformNoScale::Identity;
 	TWeakPtr<ISequencer> WeakSequencer = GetSequencerFromAsset();
@@ -2415,7 +2415,7 @@ FTransformNoScale UControlRigSequencerEditorLibrary::GetLocalControlRigTransform
 	return Value;
 }
 
-TArray<FTransformNoScale> UControlRigSequencerEditorLibrary::GetLocalControlRigTransformNoScales(ULevelSequence* LevelSequence, UControlRig* ControlRig, FName ControlName, const TArray<FFrameNumber>& Frames, ESequenceTimeUnit TimeUnit)
+TArray<FTransformNoScale> UControlRigSequencerEditorLibrary::GetLocalControlRigTransformNoScales(ULevelSequence* LevelSequence, UBaseControlRig* ControlRig, FName ControlName, const TArray<FFrameNumber>& Frames, ESequenceTimeUnit TimeUnit)
 {
 	TArray<FTransformNoScale> Values;
 	TWeakPtr<ISequencer> WeakSequencer = GetSequencerFromAsset();
@@ -2444,7 +2444,7 @@ TArray<FTransformNoScale> UControlRigSequencerEditorLibrary::GetLocalControlRigT
 	return Values;
 }
 
-void UControlRigSequencerEditorLibrary::SetLocalControlRigTransformNoScale(ULevelSequence* LevelSequence, UControlRig* ControlRig, FName ControlName, FFrameNumber Frame, FTransformNoScale Value,
+void UControlRigSequencerEditorLibrary::SetLocalControlRigTransformNoScale(ULevelSequence* LevelSequence, UBaseControlRig* ControlRig, FName ControlName, FFrameNumber Frame, FTransformNoScale Value,
 	ESequenceTimeUnit TimeUnit,bool bSetKey)
 {
 	if (LevelSequence == nullptr || ControlRig == nullptr)
@@ -2473,7 +2473,7 @@ void UControlRigSequencerEditorLibrary::SetLocalControlRigTransformNoScale(ULeve
 	}
 }
 
-void UControlRigSequencerEditorLibrary::SetLocalControlRigTransformNoScales(ULevelSequence* LevelSequence, UControlRig* ControlRig, FName ControlName, const TArray<FFrameNumber>& Frames, 
+void UControlRigSequencerEditorLibrary::SetLocalControlRigTransformNoScales(ULevelSequence* LevelSequence, UBaseControlRig* ControlRig, FName ControlName, const TArray<FFrameNumber>& Frames, 
 	const TArray<FTransformNoScale> Values, ESequenceTimeUnit TimeUnit)
 {
 	if (LevelSequence == nullptr || ControlRig == nullptr || (Frames.Num() != Values.Num()))
@@ -2508,7 +2508,7 @@ void UControlRigSequencerEditorLibrary::SetLocalControlRigTransformNoScales(ULev
 }
 
 
-FTransform UControlRigSequencerEditorLibrary::GetLocalControlRigTransform(ULevelSequence* LevelSequence, UControlRig* ControlRig, FName ControlName, FFrameNumber Frame, ESequenceTimeUnit TimeUnit)
+FTransform UControlRigSequencerEditorLibrary::GetLocalControlRigTransform(ULevelSequence* LevelSequence, UBaseControlRig* ControlRig, FName ControlName, FFrameNumber Frame, ESequenceTimeUnit TimeUnit)
 {
 	FTransform Value = FTransform::Identity;
 	TWeakPtr<ISequencer> WeakSequencer = GetSequencerFromAsset();
@@ -2532,7 +2532,7 @@ FTransform UControlRigSequencerEditorLibrary::GetLocalControlRigTransform(ULevel
 	return Value;
 }
 
-TArray<FTransform> UControlRigSequencerEditorLibrary::GetLocalControlRigTransforms(ULevelSequence* LevelSequence, UControlRig* ControlRig, FName ControlName, const TArray<FFrameNumber>& Frames, ESequenceTimeUnit TimeUnit)
+TArray<FTransform> UControlRigSequencerEditorLibrary::GetLocalControlRigTransforms(ULevelSequence* LevelSequence, UBaseControlRig* ControlRig, FName ControlName, const TArray<FFrameNumber>& Frames, ESequenceTimeUnit TimeUnit)
 {
 	TArray<FTransform> Values;
 	TWeakPtr<ISequencer> WeakSequencer = GetSequencerFromAsset();
@@ -2561,7 +2561,7 @@ TArray<FTransform> UControlRigSequencerEditorLibrary::GetLocalControlRigTransfor
 	return Values;
 }
 
-void UControlRigSequencerEditorLibrary::SetLocalControlRigTransform(ULevelSequence* LevelSequence, UControlRig* ControlRig, FName ControlName, FFrameNumber Frame, FTransform Value,
+void UControlRigSequencerEditorLibrary::SetLocalControlRigTransform(ULevelSequence* LevelSequence, UBaseControlRig* ControlRig, FName ControlName, FFrameNumber Frame, FTransform Value,
 	ESequenceTimeUnit TimeUnit,bool bSetKey)
 {
 	if (LevelSequence == nullptr || ControlRig == nullptr)
@@ -2590,7 +2590,7 @@ void UControlRigSequencerEditorLibrary::SetLocalControlRigTransform(ULevelSequen
 	}
 }
 
-void UControlRigSequencerEditorLibrary::SetLocalControlRigTransforms(ULevelSequence* LevelSequence, UControlRig* ControlRig, FName ControlName, const TArray<FFrameNumber>& Frames, 
+void UControlRigSequencerEditorLibrary::SetLocalControlRigTransforms(ULevelSequence* LevelSequence, UBaseControlRig* ControlRig, FName ControlName, const TArray<FFrameNumber>& Frames, 
 	const TArray<FTransform> Values, ESequenceTimeUnit TimeUnit)
 {
 	if (LevelSequence == nullptr || ControlRig == nullptr || (Frames.Num() != Values.Num()))
@@ -2728,7 +2728,7 @@ bool UControlRigSequencerEditorLibrary::CollapseControlRigAnimLayers(ULevelSeque
 	return bValid;
 }
 
-bool UControlRigSequencerEditorLibrary::SetControlRigSpace(ULevelSequence* LevelSequence, UControlRig* ControlRig, FName ControlName, const FRigElementKey& InSpaceKey, FFrameNumber InTime, ESequenceTimeUnit TimeUnit)
+bool UControlRigSequencerEditorLibrary::SetControlRigSpace(ULevelSequence* LevelSequence, UBaseControlRig* ControlRig, FName ControlName, const FRigElementKey& InSpaceKey, FFrameNumber InTime, ESequenceTimeUnit TimeUnit)
 {
 	TWeakPtr<ISequencer> WeakSequencer = GetSequencerFromAsset();
 	bool bValid = false;
@@ -2764,7 +2764,7 @@ bool UControlRigSequencerEditorLibrary::SetControlRigSpace(ULevelSequence* Level
 	return bValid;
 }
 
-bool UControlRigSequencerEditorLibrary::BakeControlRigSpace(ULevelSequence* InSequence, UControlRig* InControlRig, const TArray<FName>& InControlNames, FRigSpacePickerBakeSettings InSettings, ESequenceTimeUnit TimeUnit)
+bool UControlRigSequencerEditorLibrary::BakeControlRigSpace(ULevelSequence* InSequence, UBaseControlRig* InControlRig, const TArray<FName>& InControlNames, FRigSpacePickerBakeSettings InSettings, ESequenceTimeUnit TimeUnit)
 {
 	TWeakPtr<ISequencer> WeakSequencer = GetSequencerFromAsset();
 	bool bValid = false;
@@ -2799,7 +2799,7 @@ bool UControlRigSequencerEditorLibrary::BakeControlRigSpace(ULevelSequence* InSe
 	return bValid;
 }
 
-bool UControlRigSequencerEditorLibrary::DeleteControlRigSpace(ULevelSequence* LevelSequence, UControlRig* ControlRig, FName ControlName, FFrameNumber InTime, ESequenceTimeUnit TimeUnit)
+bool UControlRigSequencerEditorLibrary::DeleteControlRigSpace(ULevelSequence* LevelSequence, UBaseControlRig* ControlRig, FName ControlName, FFrameNumber InTime, ESequenceTimeUnit TimeUnit)
 {
 	TWeakPtr<ISequencer> WeakSequencer = GetSequencerFromAsset();
 	bool bValid = false;
@@ -2856,7 +2856,7 @@ bool UControlRigSequencerEditorLibrary::DeleteControlRigSpace(ULevelSequence* Le
 	return bValid;
 }
 
-bool UControlRigSequencerEditorLibrary::MoveControlRigSpace(ULevelSequence* LevelSequence, UControlRig* ControlRig, FName ControlName, FFrameNumber InTime, FFrameNumber InNewTime, ESequenceTimeUnit TimeUnit)
+bool UControlRigSequencerEditorLibrary::MoveControlRigSpace(ULevelSequence* LevelSequence, UBaseControlRig* ControlRig, FName ControlName, FFrameNumber InTime, FFrameNumber InNewTime, ESequenceTimeUnit TimeUnit)
 {
 	TWeakPtr<ISequencer> WeakSequencer = GetSequencerFromAsset();
 	bool bValid = false;
@@ -2915,7 +2915,7 @@ bool UControlRigSequencerEditorLibrary::MoveControlRigSpace(ULevelSequence* Leve
 	return bValid;
 }
 
-bool UControlRigSequencerEditorLibrary::RenameControlRigControlChannels(ULevelSequence* LevelSequence, UControlRig* ControlRig, const TArray<FName>& InOldControlNames, const TArray<FName>& InNewControlNames)
+bool UControlRigSequencerEditorLibrary::RenameControlRigControlChannels(ULevelSequence* LevelSequence, UBaseControlRig* ControlRig, const TArray<FName>& InOldControlNames, const TArray<FName>& InNewControlNames)
 {
 	if (LevelSequence == nullptr || ControlRig == nullptr)
 	{
@@ -2984,7 +2984,7 @@ bool UControlRigSequencerEditorLibrary::GetControlsMask(UMovieSceneSection* InSe
 		return false;
 	}
 	
-	UControlRig* ControlRig = ParameterSection->GetControlRig();
+	UBaseControlRig* ControlRig = ParameterSection->GetControlRig();
 	if (!ControlRig)
 	{
 		FFrame::KismetExecutionMessage(TEXT("Section does not have a control rig"), ELogVerbosity::Error);
@@ -3016,7 +3016,7 @@ void UControlRigSequencerEditorLibrary::SetControlsMask(UMovieSceneSection* InSe
 		return;
 	}
 
-	UControlRig* ControlRig = ParameterSection->GetControlRig();
+	UBaseControlRig* ControlRig = ParameterSection->GetControlRig();
 	if (!ControlRig)
 	{
 		FFrame::KismetExecutionMessage(TEXT("Section does not have a control rig"), ELogVerbosity::Error);
@@ -3064,17 +3064,17 @@ void UControlRigSequencerEditorLibrary::HideAllControls(UMovieSceneSection* InSe
 	ParameterSection->FillControlsMask(false);
 }
 
-bool UControlRigSequencerEditorLibrary::IsFKControlRig(UControlRig* InControlRig)
+bool UControlRigSequencerEditorLibrary::IsFKControlRig(UBaseControlRig* InControlRig)
 {
 	return (InControlRig && InControlRig->IsA<UFKControlRig>());
 }
 
-bool UControlRigSequencerEditorLibrary::IsAdditiveControlRig(UControlRig* InControlRig)
+bool UControlRigSequencerEditorLibrary::IsAdditiveControlRig(UBaseControlRig* InControlRig)
 {
 	return (InControlRig && InControlRig->IsAdditive());
 }
 
-EControlRigFKRigExecuteMode UControlRigSequencerEditorLibrary::GetFKControlRigApplyMode(UControlRig* InControlRig)
+EControlRigFKRigExecuteMode UControlRigSequencerEditorLibrary::GetFKControlRigApplyMode(UBaseControlRig* InControlRig)
 {
 	EControlRigFKRigExecuteMode ApplyMode = EControlRigFKRigExecuteMode::Direct;
 	if (UFKControlRig* FKRig = Cast<UFKControlRig>(InControlRig))
@@ -3084,7 +3084,7 @@ EControlRigFKRigExecuteMode UControlRigSequencerEditorLibrary::GetFKControlRigAp
 	return ApplyMode;
 }
 
-bool UControlRigSequencerEditorLibrary::SetControlRigApplyMode(UControlRig* InControlRig, EControlRigFKRigExecuteMode InApplyMode)
+bool UControlRigSequencerEditorLibrary::SetControlRigApplyMode(UBaseControlRig* InControlRig, EControlRigFKRigExecuteMode InApplyMode)
 {
 	if (UFKControlRig* FKRig = Cast<UFKControlRig>(InControlRig))
 	{
diff --git a/Engine/Plugins/Animation/ControlRig/Source/ControlRigEditor/Private/ControlRigThumbnailRenderer.cpp b/Engine/Plugins/Animation/ControlRig/Source/ControlRigEditor/Private/ControlRigThumbnailRenderer.cpp
--- a/Engine/Plugins/Animation/ControlRig/Source/ControlRigEditor/Private/ControlRigThumbnailRenderer.cpp
+++ b/Engine/Plugins/Animation/ControlRig/Source/ControlRigEditor/Private/ControlRigThumbnailRenderer.cpp
@@ -100,16 +100,16 @@ void UControlRigThumbnailRenderer::AddAdditionalPreviewSceneContent(UObject* Obj
 			return;
 		}
 
-		UControlRig* ControlRig = nullptr;
+		UBaseControlRig* ControlRig = nullptr;
 
 		// reuse the current control rig if possible
-		UControlRig* CDO = Cast<UControlRig>(RigBlueprint->GeneratedClass->GetDefaultObject(true /* create if needed */));
+		UBaseControlRig* CDO = Cast<UBaseControlRig>(RigBlueprint->GeneratedClass->GetDefaultObject(true /* create if needed */));
 
 		TArray<UObject*> ArchetypeInstances;
 		CDO->GetArchetypeInstances(ArchetypeInstances);
 		for (UObject* ArchetypeInstance : ArchetypeInstances)
 		{
-			ControlRig = Cast<UControlRig>(ArchetypeInstance);
+			ControlRig = Cast<UBaseControlRig>(ArchetypeInstance);
 			break;
 		}
 
