Control Rig: Rename ControlRig to BaseControlRig and created a specialized ControlRig

diff --git a/Engine/Plugins/Animation/ControlRig/Source/ControlRigEditor/Private/EditMode/ControlRigBaseDockableView.cpp b/Engine/Plugins/Animation/ControlRig/Source/ControlRigEditor/Private/EditMode/ControlRigBaseDockableView.cpp
--- a/Engine/Plugins/Animation/ControlRig/Source/ControlRigEditor/Private/EditMode/ControlRigBaseDockableView.cpp
+++ b/Engine/Plugins/Animation/ControlRig/Source/ControlRigEditor/Private/EditMode/ControlRigBaseDockableView.cpp
@@ -25,9 +25,9 @@ FControlRigBaseDockableView::~FControlRigBaseDockableView()
 	}
 }
 
-TArray<UControlRig*> FControlRigBaseDockableView::GetControlRigs()  const
+TArray<UBaseControlRig*> FControlRigBaseDockableView::GetControlRigs()  const
 {
-	TArray<UControlRig*> ControlRigs;
+	TArray<UBaseControlRig*> ControlRigs;
 	if (FControlRigEditMode* EditMode = static_cast<FControlRigEditMode*>(ModeTools->GetActiveMode(FControlRigEditMode::ModeName)))
 	{
 		ControlRigs = EditMode->GetControlRigsArray(false);
@@ -48,15 +48,15 @@ void FControlRigBaseDockableView::SetEditMode(FControlRigEditMode& InEditMode)
 	}
 }
 
-void FControlRigBaseDockableView::HandleControlAdded(UControlRig* ControlRig, bool bIsAdded)
+void FControlRigBaseDockableView::HandleControlAdded(UBaseControlRig* ControlRig, bool bIsAdded)
 {
 }
 
-void FControlRigBaseDockableView::HandleControlSelected(UControlRig* ControlRig, FRigControlElement* InControl, bool bSelected)
+void FControlRigBaseDockableView::HandleControlSelected(UBaseControlRig* ControlRig, FRigControlElement* InControl, bool bSelected)
 {
 }
 
-void FControlRigBaseDockableView::HandlElementSelected(UControlRig* ControlRig, const FRigElementKey& Key, bool bSelected)
+void FControlRigBaseDockableView::HandlElementSelected(UBaseControlRig* ControlRig, const FRigElementKey& Key, bool bSelected)
 {
 	if (ControlRig)
 	{
diff --git a/Engine/Plugins/Animation/ControlRig/Source/ControlRigEditor/Private/EditMode/ControlRigBaseDockableView.h b/Engine/Plugins/Animation/ControlRig/Source/ControlRigEditor/Private/EditMode/ControlRigBaseDockableView.h
--- a/Engine/Plugins/Animation/ControlRig/Source/ControlRigEditor/Private/EditMode/ControlRigBaseDockableView.h
+++ b/Engine/Plugins/Animation/ControlRig/Source/ControlRigEditor/Private/EditMode/ControlRigBaseDockableView.h
@@ -8,7 +8,7 @@
 #include "EditorModeManager.h"
 
 
-class UControlRig;
+class UBaseControlRig;
 class ISequencer;
 class FControlRigEditMode;
 struct FRigControlElement;
@@ -19,15 +19,15 @@ class FControlRigBaseDockableView
 public:
 	FControlRigBaseDockableView();
 	virtual ~FControlRigBaseDockableView();
-	TArray<UControlRig*> GetControlRigs() const;
+	TArray<UBaseControlRig*> GetControlRigs() const;
 
 	virtual void SetEditMode(FControlRigEditMode& InEditMode);
 
 protected:
-	virtual void HandleControlSelected(UControlRig* Subject, FRigControlElement* InControl, bool bSelected);
-	virtual void HandleControlAdded(UControlRig* ControlRig, bool bIsAdded);
+	virtual void HandleControlSelected(UBaseControlRig* Subject, FRigControlElement* InControl, bool bSelected);
+	virtual void HandleControlAdded(UBaseControlRig* ControlRig, bool bIsAdded);
 
-	void HandlElementSelected(UControlRig* Subject, const FRigElementKey& Key, bool bSelected);
+	void HandlElementSelected(UBaseControlRig* Subject, const FRigElementKey& Key, bool bSelected);
 
 	ISequencer* GetSequencer() const;
 
diff --git a/Engine/Plugins/Animation/ControlRig/Source/ControlRigEditor/Private/EditMode/ControlRigControlsProxy.cpp b/Engine/Plugins/Animation/ControlRig/Source/ControlRigEditor/Private/EditMode/ControlRigControlsProxy.cpp
--- a/Engine/Plugins/Animation/ControlRig/Source/ControlRigEditor/Private/EditMode/ControlRigControlsProxy.cpp
+++ b/Engine/Plugins/Animation/ControlRig/Source/ControlRigEditor/Private/EditMode/ControlRigControlsProxy.cpp
@@ -829,7 +829,7 @@ void UControlRigBoolControlProxy::SetKey(const IPropertyHandle& KeyedPropertyHan
 
 //////UControlDetailPanelControlProxies////////
 
-UControlRigControlsProxy* UControlRigDetailPanelControlProxies::FindProxy(UControlRig* ControlRig, const FName& Name) const
+UControlRigControlsProxy* UControlRigDetailPanelControlProxies::FindProxy(UBaseControlRig* ControlRig, const FName& Name) const
 {
 	const FControlToProxyMap* ControlRigProxies = AllProxies.Find(ControlRig);
 	if (ControlRigProxies)
@@ -843,7 +843,7 @@ UControlRigControlsProxy* UControlRigDetailPanelControlProxies::FindProxy(UContr
 	return nullptr;
 }
 
-void UControlRigDetailPanelControlProxies::AddProxy(UControlRig* ControlRig, const FName& Name,  FRigControlElement* ControlElement)
+void UControlRigDetailPanelControlProxies::AddProxy(UBaseControlRig* ControlRig, const FName& Name,  FRigControlElement* ControlElement)
 {
 	UControlRigControlsProxy* Proxy = FindProxy(ControlRig,Name);
 	if (!Proxy && ControlElement != nullptr)
@@ -937,7 +937,7 @@ void UControlRigDetailPanelControlProxies::AddProxy(UControlRig* ControlRig, con
 	}
 }
 
-void UControlRigDetailPanelControlProxies::RemoveProxy(UControlRig* ControlRig, const FName& Name)
+void UControlRigDetailPanelControlProxies::RemoveProxy(UBaseControlRig* ControlRig, const FName& Name)
 {
 	UControlRigControlsProxy* ExistingProxy = FindProxy(ControlRig,Name);
 	if (ExistingProxy)
@@ -952,12 +952,12 @@ void UControlRigDetailPanelControlProxies::RemoveProxy(UControlRig* ControlRig,
 	}
 }
 
-void UControlRigDetailPanelControlProxies::RemoveAllProxies(UControlRig* ControlRig)
+void UControlRigDetailPanelControlProxies::RemoveAllProxies(UBaseControlRig* ControlRig)
 {
 	//no control rig remove all
 	if (ControlRig == nullptr)
 	{
-		for (TPair<TObjectPtr<UControlRig>, FControlToProxyMap>& ControlRigProxies : AllProxies)
+		for (TPair<TObjectPtr<UBaseControlRig>, FControlToProxyMap>& ControlRigProxies : AllProxies)
 		{
 			for (TPair<FName, TObjectPtr<UControlRigControlsProxy> >& Pair : ControlRigProxies.Value.ControlToProxy)
 			{
@@ -993,7 +993,7 @@ void UControlRigDetailPanelControlProxies::RemoveAllProxies(UControlRig* Control
 	}
 }
 
-void UControlRigDetailPanelControlProxies::RecreateAllProxies(UControlRig* ControlRig)
+void UControlRigDetailPanelControlProxies::RecreateAllProxies(UBaseControlRig* ControlRig)
 {
 	RemoveAllProxies(ControlRig);
 	TArray<FRigControlElement*> Controls = ControlRig->AvailableControls();
@@ -1006,7 +1006,7 @@ void UControlRigDetailPanelControlProxies::RecreateAllProxies(UControlRig* Contr
 	}
 }
 
-void UControlRigDetailPanelControlProxies::ProxyChanged(UControlRig* ControlRig, const FName& Name)
+void UControlRigDetailPanelControlProxies::ProxyChanged(UBaseControlRig* ControlRig, const FName& Name)
 {
 	if (IsInGameThread())
 	{
@@ -1019,7 +1019,7 @@ void UControlRigDetailPanelControlProxies::ProxyChanged(UControlRig* ControlRig,
 	}
 }
 
-void UControlRigDetailPanelControlProxies::SelectProxy(UControlRig* ControlRig,const FName& Name, bool bSelected)
+void UControlRigDetailPanelControlProxies::SelectProxy(UBaseControlRig* ControlRig,const FName& Name, bool bSelected)
 {
 	UControlRigControlsProxy* Proxy = FindProxy(ControlRig,Name);
 	if (Proxy)
@@ -1040,7 +1040,7 @@ void UControlRigDetailPanelControlProxies::SelectProxy(UControlRig* ControlRig,c
 	}
 }
 
-bool UControlRigDetailPanelControlProxies::IsSelected(UControlRig* InControlRig, const FName& Name) const
+bool UControlRigDetailPanelControlProxies::IsSelected(UBaseControlRig* InControlRig, const FName& Name) const
 {
 	if (UControlRigControlsProxy* Proxy = FindProxy(InControlRig, Name))
 	{
diff --git a/Engine/Plugins/Animation/ControlRig/Source/ControlRigEditor/Private/EditMode/ControlRigControlsProxy.h b/Engine/Plugins/Animation/ControlRig/Source/ControlRigEditor/Private/EditMode/ControlRigControlsProxy.h
--- a/Engine/Plugins/Animation/ControlRig/Source/ControlRigEditor/Private/EditMode/ControlRigControlsProxy.h
+++ b/Engine/Plugins/Animation/ControlRig/Source/ControlRigEditor/Private/EditMode/ControlRigControlsProxy.h
@@ -14,7 +14,7 @@
 #include "ControlRigControlsProxy.generated.h"
 
 struct FRigControlElement;
-class UControlRig;
+class UBaseControlRig;
 class IPropertyHandle;
 class FControlRigInteractionScope;
 
@@ -41,7 +41,7 @@ public:
 public:
 
 	FRigControlElement* GetControlElement() const;
-	TWeakObjectPtr<UControlRig> ControlRig;
+	TWeakObjectPtr<UBaseControlRig> ControlRig;
 
 	UPROPERTY()
 	bool bSelected;
@@ -319,21 +319,21 @@ class UControlRigDetailPanelControlProxies :public UObject
 protected:
 
 	UPROPERTY()
-	TMap<TObjectPtr<UControlRig>, FControlToProxyMap> AllProxies; //proxies themselves contain weakobjectptr to the controlrig
+	TMap<TObjectPtr<UBaseControlRig>, FControlToProxyMap> AllProxies; //proxies themselves contain weakobjectptr to the controlrig
 
 	UPROPERTY()
 	TArray< TObjectPtr<UControlRigControlsProxy>> SelectedProxies;
 
 
 public:
-	UControlRigControlsProxy* FindProxy(UControlRig* InControlRig, const FName& Name) const;
-	void AddProxy(UControlRig* InControlRig, const FName& Name,  FRigControlElement* ControlElement);
-	void RemoveProxy(UControlRig* InControlRig, const FName& Name );
-	void ProxyChanged(UControlRig* InControlRig, const FName& Name);
-	void RemoveAllProxies(UControlRig* InControlRig);
-	void RecreateAllProxies(UControlRig* InControlRig);
-	void SelectProxy(UControlRig* InControlRig, const FName& Name, bool bSelected);
+	UControlRigControlsProxy* FindProxy(UBaseControlRig* InControlRig, const FName& Name) const;
+	void AddProxy(UBaseControlRig* InControlRig, const FName& Name,  FRigControlElement* ControlElement);
+	void RemoveProxy(UBaseControlRig* InControlRig, const FName& Name );
+	void ProxyChanged(UBaseControlRig* InControlRig, const FName& Name);
+	void RemoveAllProxies(UBaseControlRig* InControlRig);
+	void RecreateAllProxies(UBaseControlRig* InControlRig);
+	void SelectProxy(UBaseControlRig* InControlRig, const FName& Name, bool bSelected);
 	const TArray<UControlRigControlsProxy*>& GetSelectedProxies() const { return SelectedProxies;}
-	bool IsSelected(UControlRig* InControlRig, const FName& Name) const;
+	bool IsSelected(UBaseControlRig* InControlRig, const FName& Name) const;
 
 };
diff --git a/Engine/Plugins/Animation/ControlRig/Source/ControlRigEditor/Private/EditMode/ControlRigEditMode.cpp b/Engine/Plugins/Animation/ControlRig/Source/ControlRigEditor/Private/EditMode/ControlRigEditMode.cpp
--- a/Engine/Plugins/Animation/ControlRig/Source/ControlRigEditor/Private/EditMode/ControlRigEditMode.cpp
+++ b/Engine/Plugins/Animation/ControlRig/Source/ControlRigEditor/Private/EditMode/ControlRigEditMode.cpp
@@ -188,7 +188,7 @@ FControlRigEditMode::~FControlRigEditMode()
 	OnControlRigAddedOrRemovedDelegate.Clear();
 	OnControlRigSelectedDelegate.Clear();
 
-	TArray<TWeakObjectPtr<UControlRig>> PreviousRuntimeRigs = RuntimeControlRigs;
+	TArray<TWeakObjectPtr<UBaseControlRig>> PreviousRuntimeRigs = RuntimeControlRigs;
 	for (int32 PreviousRuntimeRigIndex = 0; PreviousRuntimeRigIndex < PreviousRuntimeRigs.Num(); PreviousRuntimeRigIndex++)
 	{
 		if (PreviousRuntimeRigs[PreviousRuntimeRigIndex].IsValid())
@@ -210,7 +210,7 @@ bool FControlRigEditMode:: SetSequencer(TWeakPtr<ISequencer> InSequencer)
 	{
 		WeakSequencer = InSequencer;
 		DestroyShapesActors(nullptr);
-		TArray<TWeakObjectPtr<UControlRig>> PreviousRuntimeRigs = RuntimeControlRigs;
+		TArray<TWeakObjectPtr<UBaseControlRig>> PreviousRuntimeRigs = RuntimeControlRigs;
 		for (int32 PreviousRuntimeRigIndex = 0; PreviousRuntimeRigIndex < PreviousRuntimeRigs.Num(); PreviousRuntimeRigIndex++)
 		{
 			if (PreviousRuntimeRigs[PreviousRuntimeRigIndex].IsValid())
@@ -227,7 +227,7 @@ bool FControlRigEditMode:: SetSequencer(TWeakPtr<ISequencer> InSequencer)
 				TArray<FControlRigSequencerBindingProxy> Proxies = UControlRigSequencerEditorLibrary::GetControlRigs(LevelSequence);
 				for (FControlRigSequencerBindingProxy& Proxy : Proxies)
 				{
-					if (UControlRig* ControlRig = Proxy.ControlRig.Get())
+					if (UBaseControlRig* ControlRig = Proxy.ControlRig.Get())
 					{
 						AddControlRigInternal(ControlRig);
 					}
@@ -240,7 +240,7 @@ bool FControlRigEditMode:: SetSequencer(TWeakPtr<ISequencer> InSequencer)
 	return false;
 }
 
-bool FControlRigEditMode::AddControlRigObject(UControlRig* ControlRig, TWeakPtr<ISequencer> InSequencer)
+bool FControlRigEditMode::AddControlRigObject(UBaseControlRig* ControlRig, TWeakPtr<ISequencer> InSequencer)
 {
 	if (ControlRig)
 	{
@@ -260,9 +260,9 @@ bool FControlRigEditMode::AddControlRigObject(UControlRig* ControlRig, TWeakPtr<
 	return false;
 }
 
-void FControlRigEditMode::SetObjects(UControlRig* ControlRig,  UObject* BindingObject, TWeakPtr<ISequencer> InSequencer)
+void FControlRigEditMode::SetObjects(UBaseControlRig* ControlRig,  UObject* BindingObject, TWeakPtr<ISequencer> InSequencer)
 {
-	TArray<TWeakObjectPtr<UControlRig>> PreviousRuntimeRigs = RuntimeControlRigs;
+	TArray<TWeakObjectPtr<UBaseControlRig>> PreviousRuntimeRigs = RuntimeControlRigs;
 	for (int32 PreviousRuntimeRigIndex = 0; PreviousRuntimeRigIndex < PreviousRuntimeRigs.Num(); PreviousRuntimeRigIndex++)
 	{
 		if (PreviousRuntimeRigs[PreviousRuntimeRigIndex].IsValid())
@@ -315,9 +315,9 @@ void FControlRigEditMode::SetUpDetailPanel()
 void FControlRigEditMode::SetObjects_Internal()
 {
 	bool bHasValidRuntimeControlRig = false;
-	for (TWeakObjectPtr<UControlRig>& RuntimeRigPtr : RuntimeControlRigs)
+	for (TWeakObjectPtr<UBaseControlRig>& RuntimeRigPtr : RuntimeControlRigs)
 	{
-		if (UControlRig* RuntimeControlRig = RuntimeRigPtr.Get())
+		if (UBaseControlRig* RuntimeControlRig = RuntimeRigPtr.Get())
 		{
 			RuntimeControlRig->ControlModified().RemoveAll(this);
 			RuntimeControlRig->GetHierarchy()->OnModified().RemoveAll(this);
@@ -437,9 +437,9 @@ void FControlRigEditMode::Exit()
 	ClearOutAnyActiveTools();
 	OnControlRigAddedOrRemovedDelegate.Clear();
 	OnControlRigSelectedDelegate.Clear();
-	for (TWeakObjectPtr<UControlRig>& RuntimeRigPtr : RuntimeControlRigs)
+	for (TWeakObjectPtr<UBaseControlRig>& RuntimeRigPtr : RuntimeControlRigs)
 	{
-		if (UControlRig* ControlRig = RuntimeRigPtr.Get())
+		if (UBaseControlRig* ControlRig = RuntimeRigPtr.Get())
 		{
 			ControlRig->ClearControlSelection();
 		}
@@ -453,7 +453,7 @@ void FControlRigEditMode::Exit()
 			GEditor->EndTransaction();
 		}
 
-		for (TPair<UControlRig*,FControlRigInteractionScope*>& InteractionScope : InteractionScopes)
+		for (TPair<UBaseControlRig*,FControlRigInteractionScope*>& InteractionScope : InteractionScopes)
 		{
 			if (InteractionScope.Value)
 			{
@@ -473,7 +473,7 @@ void FControlRigEditMode::Exit()
 	DestroyShapesActors(nullptr);
 
 
-	TArray<TWeakObjectPtr<UControlRig>> PreviousRuntimeRigs = RuntimeControlRigs;
+	TArray<TWeakObjectPtr<UBaseControlRig>> PreviousRuntimeRigs = RuntimeControlRigs;
 	for (int32 PreviousRuntimeRigIndex = 0; PreviousRuntimeRigIndex < PreviousRuntimeRigs.Num(); PreviousRuntimeRigIndex++)
 	{
 		if (PreviousRuntimeRigs[PreviousRuntimeRigIndex].IsValid())
@@ -537,9 +537,9 @@ void FControlRigEditMode::Tick(FEditorViewportClient* ViewportClient, float Delt
 	if(bShowControlsAsOverlay != Settings->bShowControlsAsOverlay)
 	{
 		bShowControlsAsOverlay = Settings->bShowControlsAsOverlay;
-		for (TWeakObjectPtr<UControlRig>& RuntimeRigPtr : RuntimeControlRigs)
+		for (TWeakObjectPtr<UBaseControlRig>& RuntimeRigPtr : RuntimeControlRigs)
 		{
-			if (UControlRig* RuntimeControlRig = RuntimeRigPtr.Get())
+			if (UBaseControlRig* RuntimeControlRig = RuntimeRigPtr.Get())
 			{
 				UpdateSelectabilityOnSkeletalMeshes(RuntimeControlRig, !bShowControlsAsOverlay);
 			}
@@ -551,9 +551,9 @@ void FControlRigEditMode::Tick(FEditorViewportClient* ViewportClient, float Delt
 	if (RecreateControlShapesRequired != ERecreateControlRigShape::RecreateNone && !(FSlateApplication::Get().HasAnyMouseCaptor() || GUnrealEd->IsUserInteracting()))
 	{
 		RecreateControlShapeActors();
-		for (TWeakObjectPtr<UControlRig>& RuntimeRigPtr : RuntimeControlRigs)
+		for (TWeakObjectPtr<UBaseControlRig>& RuntimeRigPtr : RuntimeControlRigs)
 		{
-			if (UControlRig* ControlRig = RuntimeRigPtr.Get())
+			if (UBaseControlRig* ControlRig = RuntimeRigPtr.Get())
 			{
 				TArray<FRigElementKey> SelectedRigElements = GetSelectedRigElements(ControlRig);
 				for (const FRigElementKey& SelectedKey : SelectedRigElements)
@@ -595,9 +595,9 @@ void FControlRigEditMode::Tick(FEditorViewportClient* ViewportClient, float Delt
 
 		if (!AreEditingControlRigDirectly() == false) //only do this check if not in level editor
 		{
-			for (TWeakObjectPtr<UControlRig>& RuntimeRigPtr : RuntimeControlRigs)
+			for (TWeakObjectPtr<UBaseControlRig>& RuntimeRigPtr : RuntimeControlRigs)
 			{
-				if (UControlRig* ControlRig = RuntimeRigPtr.Get())
+				if (UBaseControlRig* ControlRig = RuntimeRigPtr.Get())
 				{
 					TArray<FRigElementKey> SelectedRigElements = GetSelectedRigElements(ControlRig);
 					UE::Widget::EWidgetMode CurrentWidgetMode = ViewportClient->GetWidgetMode();
@@ -661,7 +661,7 @@ struct  HFKRigBoneProxy : public HHitProxy
 	DECLARE_HIT_PROXY()
 
 	FName BoneName;
-	UControlRig* ControlRig;
+	UBaseControlRig* ControlRig;
 
 	HFKRigBoneProxy()
 		: HHitProxy(HPP_Foreground)
@@ -669,7 +669,7 @@ struct  HFKRigBoneProxy : public HHitProxy
 		, ControlRig(nullptr)
 	{}
 
-	HFKRigBoneProxy(FName InBoneName, UControlRig *InControlRig)
+	HFKRigBoneProxy(FName InBoneName, UBaseControlRig *InControlRig)
 		: HHitProxy(HPP_Foreground)
 		, BoneName(InBoneName)
 		, ControlRig(InControlRig)
@@ -684,7 +684,7 @@ struct  HFKRigBoneProxy : public HHitProxy
 IMPLEMENT_HIT_PROXY(HFKRigBoneProxy, HHitProxy)
 
 
-TSet<FName> FControlRigEditMode::GetActiveControlsFromSequencer(UControlRig* ControlRig)
+TSet<FName> FControlRigEditMode::GetActiveControlsFromSequencer(UBaseControlRig* ControlRig)
 {
 	TSet<FName> ActiveControls;
 	if (WeakSequencer.IsValid() == false)
@@ -756,9 +756,9 @@ void FControlRigEditMode::Render(const FSceneView* View, FViewport* Viewport, FP
 	const UControlRigEditModeSettings* Settings = GetDefault<UControlRigEditModeSettings>();
 	const bool bIsInGameView = !AreEditingControlRigDirectly() ? (ViewportToGameView.Find(Viewport) && ViewportToGameView[Viewport]) : false;
 	bool bRender = !Settings->bHideControlShapes;
-	for (TWeakObjectPtr<UControlRig>& ControlRigPtr : RuntimeControlRigs)
+	for (TWeakObjectPtr<UBaseControlRig>& ControlRigPtr : RuntimeControlRigs)
 	{
-		UControlRig* ControlRig = ControlRigPtr.Get();
+		UBaseControlRig* ControlRig = ControlRigPtr.Get();
 		//actor game view drawing is handled by not drawing in game via SetActorHiddenInGame().
 		if (bRender && ControlRig && ControlRig->GetControlsVisible())
 		{
@@ -1011,7 +1011,7 @@ bool FControlRigEditMode::EndTracking(FEditorViewportClient* InViewportClient, F
 			GEditor->EndTransaction();
 		}
 
-		for (TPair<UControlRig*, FControlRigInteractionScope*>& InteractionScope : InteractionScopes)
+		for (TPair<UBaseControlRig*, FControlRigInteractionScope*>& InteractionScope : InteractionScopes)
 		{
 			if (InteractionScope.Value)
 			{
@@ -1063,9 +1063,9 @@ bool FControlRigEditMode::StartTracking(FEditorViewportClient* InViewportClient,
 		{
 			if (AreEditingControlRigDirectly())
 			{
-				for (TWeakObjectPtr<UControlRig>& RuntimeRigPtr : RuntimeControlRigs)
+				for (TWeakObjectPtr<UBaseControlRig>& RuntimeRigPtr : RuntimeControlRigs)
 				{
-					if (UControlRig* ControlRig = RuntimeRigPtr.Get())
+					if (UBaseControlRig* ControlRig = RuntimeRigPtr.Get())
 					{
 						TArray<FRigElementKey> SelectedRigElements = GetSelectedRigElements(ControlRig);
 						for (const FRigElementKey& Key : SelectedRigElements)
@@ -1084,9 +1084,9 @@ bool FControlRigEditMode::StartTracking(FEditorViewportClient* InViewportClient,
 
 		if (AreEditingControlRigDirectly())
 		{
-			for (TWeakObjectPtr<UControlRig>& RuntimeRigPtr : RuntimeControlRigs)
+			for (TWeakObjectPtr<UBaseControlRig>& RuntimeRigPtr : RuntimeControlRigs)
 			{
-				if (UControlRig* ControlRig = RuntimeRigPtr.Get())
+				if (UBaseControlRig* ControlRig = RuntimeRigPtr.Get())
 				{
 					UObject* Blueprint = ControlRig->GetClass()->ClassGeneratedBy;
 					if (Blueprint)
@@ -1112,9 +1112,9 @@ bool FControlRigEditMode::StartTracking(FEditorViewportClient* InViewportClient,
 
 	if (!AreEditingControlRigDirectly())
 	{
-		for (TWeakObjectPtr<UControlRig>& RuntimeRigPtr : RuntimeControlRigs)
+		for (TWeakObjectPtr<UBaseControlRig>& RuntimeRigPtr : RuntimeControlRigs)
 		{
-			if (UControlRig* ControlRig = RuntimeRigPtr.Get())
+			if (UBaseControlRig* ControlRig = RuntimeRigPtr.Get())
 			{
 				if (AreRigElementSelectedAndMovable(ControlRig))
 				{
@@ -1132,7 +1132,7 @@ bool FControlRigEditMode::StartTracking(FEditorViewportClient* InViewportClient,
 	}
 	else if(UControlRigEditorSettings::Get()->bEnableUndoForPoseInteraction)
 	{
-		UControlRig* ControlRig = RuntimeControlRigs[0].Get();
+		UBaseControlRig* ControlRig = RuntimeControlRigs[0].Get();
 		FControlRigInteractionScope* InteractionScope = new FControlRigInteractionScope(ControlRig);
 		InteractionScopes.Add(ControlRig,InteractionScope);
 	}
@@ -1286,7 +1286,7 @@ bool FControlRigEditMode::HandleClick(FEditorViewportClient* InViewportClient, H
 					}
 					else if(Click.IsControlDown()) //if ctrl we toggle selection
 					{
-						if (UControlRig* ControlRig = ShapeActor->ControlRig.Get())
+						if (UBaseControlRig* ControlRig = ShapeActor->ControlRig.Get())
 						{
 							bool bIsSelected = ControlRig->IsControlSelected(ControlName);
 							SetRigElementSelection(ControlRig, ERigElementType::Control, ControlName, !bIsSelected);
@@ -1328,9 +1328,9 @@ bool FControlRigEditMode::HandleClick(FEditorViewportClient* InViewportClient, H
 			}
 			else 
 			{ 
-				for (TWeakObjectPtr<UControlRig>& RuntimeRigPtr : RuntimeControlRigs)
+				for (TWeakObjectPtr<UBaseControlRig>& RuntimeRigPtr : RuntimeControlRigs)
 				{
-					if (UControlRig* ControlRig = RuntimeRigPtr.Get())
+					if (UBaseControlRig* ControlRig = RuntimeRigPtr.Get())
 					{
 
 						//if we have an additive or fk control rig active select the control based upon the selected bone.
@@ -1393,9 +1393,9 @@ bool FControlRigEditMode::HandleClick(FEditorViewportClient* InViewportClient, H
 			}
 			else if (Click.IsControlDown()) //if ctrl we toggle selection
 			{
-				for (TWeakObjectPtr<UControlRig>& RuntimeRigPtr : RuntimeControlRigs)
+				for (TWeakObjectPtr<UBaseControlRig>& RuntimeRigPtr : RuntimeControlRigs)
 				{
-					if (UControlRig* ControlRig = RuntimeRigPtr.Get())
+					if (UBaseControlRig* ControlRig = RuntimeRigPtr.Get())
 					{
 						bool bIsSelected = ControlRig->IsControlSelected(ControlName);
 						SetRigElementSelection(FKBoneProxy->ControlRig, ERigElementType::Control, ControlName, !bIsSelected);
@@ -1414,7 +1414,7 @@ bool FControlRigEditMode::HandleClick(FEditorViewportClient* InViewportClient, H
 	{
 		if (RuntimeControlRigs.Num() > 0)
 		{
-			if (UControlRig* DebuggedControlRig = RuntimeControlRigs[0].Get())
+			if (UBaseControlRig* DebuggedControlRig = RuntimeControlRigs[0].Get())
 			{
 				URigHierarchy* Hierarchy = DebuggedControlRig->GetHierarchy();
 
@@ -1967,9 +1967,9 @@ bool FControlRigEditMode::InputDelta(FEditorViewportClient* InViewportClient, FV
 
 bool FControlRigEditMode::ShouldDrawWidget() const
 {
-	for (const TWeakObjectPtr<UControlRig>& RuntimeRigPtr : RuntimeControlRigs)
+	for (const TWeakObjectPtr<UBaseControlRig>& RuntimeRigPtr : RuntimeControlRigs)
 	{
-		if (UControlRig* ControlRig = RuntimeRigPtr.Get())
+		if (UBaseControlRig* ControlRig = RuntimeRigPtr.Get())
 		{
 			if (AreRigElementSelectedAndMovable(ControlRig))
 			{
@@ -2002,9 +2002,9 @@ void FControlRigEditMode::AddReferencedObjects( FReferenceCollector& Collector )
 
 void FControlRigEditMode::ClearRigElementSelection(uint32 InTypes)
 {
-	for (TWeakObjectPtr<UControlRig>& RuntimeRigPtr : RuntimeControlRigs)
+	for (TWeakObjectPtr<UBaseControlRig>& RuntimeRigPtr : RuntimeControlRigs)
 	{
-		if (UControlRig* ControlRig = RuntimeRigPtr.Get())
+		if (UBaseControlRig* ControlRig = RuntimeRigPtr.Get())
 		{
 			if (!AreEditingControlRigDirectly())
 			{
@@ -2022,7 +2022,7 @@ void FControlRigEditMode::ClearRigElementSelection(uint32 InTypes)
 }
 
 // internal private function that doesn't use guarding.
-void FControlRigEditMode::SetRigElementSelectionInternal(UControlRig* ControlRig, ERigElementType Type, const FName& InRigElementName, bool bSelected)
+void FControlRigEditMode::SetRigElementSelectionInternal(UBaseControlRig* ControlRig, ERigElementType Type, const FName& InRigElementName, bool bSelected)
 {
 	if(URigHierarchyController* Controller = ControlRig->GetHierarchy()->GetController())
 	{
@@ -2030,7 +2030,7 @@ void FControlRigEditMode::SetRigElementSelectionInternal(UControlRig* ControlRig
 	}
 }
 
-void FControlRigEditMode::SetRigElementSelection(UControlRig* ControlRig, ERigElementType Type, const FName& InRigElementName, bool bSelected)
+void FControlRigEditMode::SetRigElementSelection(UBaseControlRig* ControlRig, ERigElementType Type, const FName& InRigElementName, bool bSelected)
 {
 	if (!bSelecting)
 	{
@@ -2042,7 +2042,7 @@ void FControlRigEditMode::SetRigElementSelection(UControlRig* ControlRig, ERigEl
 	}
 }
 
-void FControlRigEditMode::SetRigElementSelection(UControlRig* ControlRig, ERigElementType Type, const TArray<FName>& InRigElementNames, bool bSelected)
+void FControlRigEditMode::SetRigElementSelection(UBaseControlRig* ControlRig, ERigElementType Type, const TArray<FName>& InRigElementNames, bool bSelected)
 {
 	if (!bSelecting)
 	{
@@ -2057,7 +2057,7 @@ void FControlRigEditMode::SetRigElementSelection(UControlRig* ControlRig, ERigEl
 	}
 }
 
-TArray<FRigElementKey> FControlRigEditMode::GetSelectedRigElements(UControlRig* ControlRig) const
+TArray<FRigElementKey> FControlRigEditMode::GetSelectedRigElements(UBaseControlRig* ControlRig) const
 {
 	if (ControlRig == nullptr && GetControlRigs().Num() > 0)
 	{
@@ -2079,7 +2079,7 @@ TArray<FRigElementKey> FControlRigEditMode::GetSelectedRigElements(UControlRig*
 	{
 		if (SelectedKeys[0].Type == ERigElementType::Bone || SelectedKeys[0].Type == ERigElementType::Null)
 		{
-			const FName ControlName = UControlRig::GetNameForTransientControl(SelectedKeys[0]);
+			const FName ControlName = UBaseControlRig::GetNameForTransientControl(SelectedKeys[0]);
 			const FRigElementKey TransientControlKey = FRigElementKey(ControlName, ERigElementType::Control);
 			if(ControlRig->GetHierarchy()->Contains(TransientControlKey))
 			{
@@ -2105,7 +2105,7 @@ TArray<FRigElementKey> FControlRigEditMode::GetSelectedRigElements(UControlRig*
 	return SelectedKeys;
 }
 
-bool FControlRigEditMode::AreRigElementsSelected(uint32 InTypes, UControlRig* InControlRig) const
+bool FControlRigEditMode::AreRigElementsSelected(uint32 InTypes, UBaseControlRig* InControlRig) const
 {
 	TArray<FRigElementKey> SelectedRigElements = GetSelectedRigElements(InControlRig);
 
@@ -2120,7 +2120,7 @@ bool FControlRigEditMode::AreRigElementsSelected(uint32 InTypes, UControlRig* In
 	return false;
 }
 
-int32 FControlRigEditMode::GetNumSelectedRigElements(uint32 InTypes, UControlRig* InControlRig) const
+int32 FControlRigEditMode::GetNumSelectedRigElements(uint32 InTypes, UBaseControlRig* InControlRig) const
 {
 	TArray<FRigElementKey> SelectedRigElements = GetSelectedRigElements(InControlRig);
 	if (FRigElementTypeHelper::DoesHave(InTypes, ERigElementType::All))
@@ -2171,9 +2171,9 @@ bool FControlRigEditMode::CanRemoveFromPreviewScene(const USceneComponent* InCom
 void FControlRigEditMode::RecalcPivotTransform()
 {
 	PivotTransforms.Reset();
-	for (TWeakObjectPtr<UControlRig>& RuntimeRigPtr : RuntimeControlRigs)
+	for (TWeakObjectPtr<UBaseControlRig>& RuntimeRigPtr : RuntimeControlRigs)
 	{
-		if (UControlRig* ControlRig = RuntimeRigPtr.Get())
+		if (UBaseControlRig* ControlRig = RuntimeRigPtr.Get())
 		{
 			FTransform PivotTransform = FTransform::Identity;
 			// Use average location as pivot location
@@ -2324,7 +2324,7 @@ bool FControlRigEditMode::HasPivotTransformsChanged() const
 	{
 		return true;
 	}
-	for (const TPair<UControlRig*, FTransform>& Transform : PivotTransforms)
+	for (const TPair<UBaseControlRig*, FTransform>& Transform : PivotTransforms)
 	{
 		if (const FTransform* LastTransform = LastPivotTransforms.Find(Transform.Key))
 		{
@@ -2437,9 +2437,9 @@ void FControlRigEditMode::BindCommands()
 bool FControlRigEditMode::IsControlSelected() const
 {
 	static uint32 TypeFlag = (uint32)ERigElementType::Control;
-	for (const TWeakObjectPtr<UControlRig>& RuntimeRigPtr : RuntimeControlRigs)
+	for (const TWeakObjectPtr<UBaseControlRig>& RuntimeRigPtr : RuntimeControlRigs)
 	{
-		if (UControlRig* ControlRig = RuntimeRigPtr.Get())
+		if (UBaseControlRig* ControlRig = RuntimeRigPtr.Get())
 		{
 			if (AreRigElementsSelected(TypeFlag,ControlRig))
 			{
@@ -2452,9 +2452,9 @@ bool FControlRigEditMode::IsControlSelected() const
 
 bool FControlRigEditMode::CanFrameSelection()
 {
-	for (TWeakObjectPtr<UControlRig>& RuntimeRigPtr : RuntimeControlRigs)
+	for (TWeakObjectPtr<UBaseControlRig>& RuntimeRigPtr : RuntimeControlRigs)
 	{
-		if (UControlRig* ControlRig = RuntimeRigPtr.Get())
+		if (UBaseControlRig* ControlRig = RuntimeRigPtr.Get())
 		{
 			if (GetSelectedRigElements(ControlRig).Num() > 0)
 			{
@@ -2491,9 +2491,9 @@ void FControlRigEditMode::FrameSelection()
     }
 
 	TArray<AActor*> Actors;
-	for (TWeakObjectPtr<UControlRig>& RuntimeRigPtr : RuntimeControlRigs)
+	for (TWeakObjectPtr<UBaseControlRig>& RuntimeRigPtr : RuntimeControlRigs)
 	{
-		if (UControlRig* ControlRig = RuntimeRigPtr.Get())
+		if (UBaseControlRig* ControlRig = RuntimeRigPtr.Get())
 		{
 			TArray<FRigElementKey> SelectedRigElements = GetSelectedRigElements(ControlRig);
 			for (const FRigElementKey& SelectedKey : SelectedRigElements)
@@ -2611,11 +2611,11 @@ void FControlRigEditMode::ToggleControlShapeTransformEdit()
 	}
 }
 
-void FControlRigEditMode::GetAllSelectedControls(TMap<UControlRig*, TArray<FRigElementKey>>& OutSelectedControls) const
+void FControlRigEditMode::GetAllSelectedControls(TMap<UBaseControlRig*, TArray<FRigElementKey>>& OutSelectedControls) const
 {
-	for (const TWeakObjectPtr<UControlRig>& RuntimeRigPtr : RuntimeControlRigs)
+	for (const TWeakObjectPtr<UBaseControlRig>& RuntimeRigPtr : RuntimeControlRigs)
 	{
-		if (UControlRig* ControlRig = RuntimeRigPtr.Get())
+		if (UBaseControlRig* ControlRig = RuntimeRigPtr.Get())
 		{
 			if (const URigHierarchy* Hierarchy = ControlRig->GetHierarchy())
 			{
@@ -2701,7 +2701,7 @@ bool FControlRigEditMode::IsDragAnimSliderToolPressed(FViewport* InViewport)
 
 void FControlRigEditMode::OpenSpacePickerWidget()
 {
-	TMap<UControlRig*, TArray<FRigElementKey>> SelectedControlRigsAndControls;
+	TMap<UBaseControlRig*, TArray<FRigElementKey>> SelectedControlRigsAndControls;
 	GetAllSelectedControls(SelectedControlRigsAndControls);
 
 	if (SelectedControlRigsAndControls.Num() < 1)
@@ -2709,14 +2709,14 @@ void FControlRigEditMode::OpenSpacePickerWidget()
 		return;
 	}
 
-	TArray<UControlRig*> ControlRigs;
+	TArray<UBaseControlRig*> ControlRigs;
 	TArray<TArray<FRigElementKey>> AllSelectedControls;
 	SelectedControlRigsAndControls.GenerateKeyArray(ControlRigs);
 	SelectedControlRigsAndControls.GenerateValueArray(AllSelectedControls);
 
 
 	//mz todo handle multiple control rigs with space picker
-	UControlRig* RuntimeRig = ControlRigs[0];
+	UBaseControlRig* RuntimeRig = ControlRigs[0];
 	TArray<FRigElementKey>& SelectedControls = AllSelectedControls[0];
 
 	URigHierarchy* Hierarchy = RuntimeRig->GetHierarchy();
@@ -2843,11 +2843,11 @@ void FControlRigEditMode::ToggleManipulators()
 {
 	if (!AreEditingControlRigDirectly())
 	{
-		TMap<UControlRig*, TArray<FRigElementKey>> SelectedControls;
+		TMap<UBaseControlRig*, TArray<FRigElementKey>> SelectedControls;
 		GetAllSelectedControls(SelectedControls);
-		TArray<UControlRig*> ControlRigs;
+		TArray<UBaseControlRig*> ControlRigs;
 		SelectedControls.GenerateKeyArray(ControlRigs);
-		for (UControlRig* ControlRig : ControlRigs)
+		for (UBaseControlRig* ControlRig : ControlRigs)
 		{
 			if (ControlRig)
 			{
@@ -2872,9 +2872,9 @@ void FControlRigEditMode::ToggleAllManipulators()
 	//turn on all if in level editor in case any where off
 	if (!AreEditingControlRigDirectly() && Settings->bHideControlShapes)
 	{
-		for (TWeakObjectPtr<UControlRig>& RuntimeRigPtr : RuntimeControlRigs)
+		for (TWeakObjectPtr<UBaseControlRig>& RuntimeRigPtr : RuntimeControlRigs)
 		{
-			if (UControlRig* ControlRig = RuntimeRigPtr.Get())
+			if (UBaseControlRig* ControlRig = RuntimeRigPtr.Get())
 			{
 				ControlRig->SetControlsVisible(true);
 			}
@@ -2885,10 +2885,10 @@ void FControlRigEditMode::ToggleAllManipulators()
 void FControlRigEditMode::ZeroTransforms(bool bSelectionOnly)
 {
 	// Gather up the control rigs for the selected controls
-	TArray<UControlRig*> ControlRigs;
-	for (TWeakObjectPtr<UControlRig>& RuntimeRigPtr : RuntimeControlRigs)
+	TArray<UBaseControlRig*> ControlRigs;
+	for (TWeakObjectPtr<UBaseControlRig>& RuntimeRigPtr : RuntimeControlRigs)
 	{
-		if (UControlRig* ControlRig = RuntimeRigPtr.Get())
+		if (UBaseControlRig* ControlRig = RuntimeRigPtr.Get())
 		{
 			if (!bSelectionOnly || ControlRig->CurrentControlSelection().Num() > 0)
 			{
@@ -2903,7 +2903,7 @@ void FControlRigEditMode::ZeroTransforms(bool bSelectionOnly)
 
 	FScopedTransaction Transaction(LOCTEXT("HierarchyZeroTransforms", "Zero Transforms"));
 
-	for (UControlRig* ControlRig : ControlRigs)
+	for (UBaseControlRig* ControlRig : ControlRigs)
 	{
 		TArray<FRigElementKey> SelectedRigElements = GetSelectedRigElements(ControlRig);
 		if (ControlRig->IsAdditive())
@@ -2979,7 +2979,7 @@ void FControlRigEditMode::ZeroTransforms(bool bSelectionOnly)
 				ControlElement = ControlRig->FindControl(ElementToReset.Name);
 				if (ControlElement->Settings.bIsTransientControl)
 				{
-					if(UControlRig::GetNodeNameFromTransientControl(ControlElement->GetKey()).IsEmpty())
+					if(UBaseControlRig::GetNodeNameFromTransientControl(ControlElement->GetKey()).IsEmpty())
 					{
 						ControlElement = nullptr;
 					}
@@ -3110,10 +3110,10 @@ void FControlRigEditMode::ZeroTransforms(bool bSelectionOnly)
 void FControlRigEditMode::InvertInputPose(bool bSelectionOnly)
 {
 	// Gather up the control rigs for the selected controls
-	TArray<UControlRig*> ControlRigs;
-	for (TWeakObjectPtr<UControlRig>& RuntimeRigPtr : RuntimeControlRigs)
+	TArray<UBaseControlRig*> ControlRigs;
+	for (TWeakObjectPtr<UBaseControlRig>& RuntimeRigPtr : RuntimeControlRigs)
 	{
-		if (UControlRig* ControlRig = RuntimeRigPtr.Get())
+		if (UBaseControlRig* ControlRig = RuntimeRigPtr.Get())
 		{
 			if (!bSelectionOnly || ControlRig->CurrentControlSelection().Num() > 0)
 			{
@@ -3128,7 +3128,7 @@ void FControlRigEditMode::InvertInputPose(bool bSelectionOnly)
 
 	FScopedTransaction Transaction(LOCTEXT("HierarchyInvertTransformsToRestPose", "Invert Transforms to Rest Pose"));
 
-	for (UControlRig* ControlRig : ControlRigs)
+	for (UBaseControlRig* ControlRig : ControlRigs)
 	{
 		if (!ControlRig->IsAdditive())
 		{
@@ -3207,11 +3207,11 @@ bool FControlRigEditMode::CheckMovieSceneSig()
 			{
 				if (ULevelSequence* LevelSequence = Cast<ULevelSequence>(Sequencer->GetFocusedMovieSceneSequence()))
 				{
-					TArray<TWeakObjectPtr<UControlRig>> CurrentControlRigs;
+					TArray<TWeakObjectPtr<UBaseControlRig>> CurrentControlRigs;
 					TArray<FControlRigSequencerBindingProxy> Proxies = UControlRigSequencerEditorLibrary::GetControlRigs(LevelSequence);
 					for (FControlRigSequencerBindingProxy& Proxy : Proxies)
 					{
-						if (UControlRig* ControlRig = Proxy.ControlRig.Get())
+						if (UBaseControlRig* ControlRig = Proxy.ControlRig.Get())
 						{
 							CurrentControlRigs.Add(ControlRig);
 							if (RuntimeControlRigs.Contains(ControlRig) == false)
@@ -3221,15 +3221,15 @@ bool FControlRigEditMode::CheckMovieSceneSig()
 							}
 						}
 					}
-					TArray<TWeakObjectPtr<UControlRig>> ControlRigsToRemove;
-					for (TWeakObjectPtr<UControlRig>& RuntimeRigPtr : RuntimeControlRigs)
+					TArray<TWeakObjectPtr<UBaseControlRig>> ControlRigsToRemove;
+					for (TWeakObjectPtr<UBaseControlRig>& RuntimeRigPtr : RuntimeControlRigs)
 					{
 						if (CurrentControlRigs.Contains(RuntimeRigPtr) == false)
 						{
 							ControlRigsToRemove.Add(RuntimeRigPtr);
 						}
 					}
-					for (TWeakObjectPtr<UControlRig>& OldRuntimeRigPtr : ControlRigsToRemove)
+					for (TWeakObjectPtr<UBaseControlRig>& OldRuntimeRigPtr : ControlRigsToRemove)
 					{
 						RemoveControlRig(OldRuntimeRigPtr.Get());
 					}
@@ -3251,7 +3251,7 @@ void FControlRigEditMode::PostUndo()
 	if (WeakSequencer.IsValid())
 	{
 		bool bHaveInvalidControlRig = false;
-		for (TWeakObjectPtr<UControlRig>& RuntimeRigPtr : RuntimeControlRigs)
+		for (TWeakObjectPtr<UBaseControlRig>& RuntimeRigPtr : RuntimeControlRigs)
 		{
 			if (RuntimeRigPtr.IsValid() == false)
 			{				
@@ -3262,7 +3262,7 @@ void FControlRigEditMode::PostUndo()
 		//if one is invalid we need to clear everything,since no longer have ptr to selectively delete
 		if (bHaveInvalidControlRig == true)
 		{
-			TArray<TWeakObjectPtr<UControlRig>> PreviousRuntimeRigs = RuntimeControlRigs;
+			TArray<TWeakObjectPtr<UBaseControlRig>> PreviousRuntimeRigs = RuntimeControlRigs;
 			for (int32 PreviousRuntimeRigIndex = 0; PreviousRuntimeRigIndex < PreviousRuntimeRigs.Num(); PreviousRuntimeRigIndex++)
 			{
 				if (PreviousRuntimeRigs[PreviousRuntimeRigIndex].IsValid())
@@ -3282,7 +3282,7 @@ void FControlRigEditMode::PostUndo()
 			TArray<FControlRigSequencerBindingProxy> Proxies = UControlRigSequencerEditorLibrary::GetControlRigs(LevelSequence);
 			for (FControlRigSequencerBindingProxy& Proxy : Proxies)
 			{
-				if (UControlRig* ControlRig = Proxy.ControlRig.Get())
+				if (UBaseControlRig* ControlRig = Proxy.ControlRig.Get())
 				{
 					if (RuntimeControlRigs.Contains(ControlRig) == false)
 					{
@@ -3302,7 +3302,7 @@ void FControlRigEditMode::PostUndo()
 	}
 	else
 	{
-		for (TWeakObjectPtr<UControlRig>& RuntimeRigPtr : RuntimeControlRigs)
+		for (TWeakObjectPtr<UBaseControlRig>& RuntimeRigPtr : RuntimeControlRigs)
 		{
 			if (RuntimeRigPtr.IsValid() == false)
 			{
@@ -3326,7 +3326,7 @@ void FControlRigEditMode::PostUndo()
 
 }
 
-void FControlRigEditMode::RequestToRecreateControlShapeActors(UControlRig* ControlRig)
+void FControlRigEditMode::RequestToRecreateControlShapeActors(UBaseControlRig* ControlRig)
 { 
 	if (ControlRig)
 	{
@@ -3376,9 +3376,9 @@ void FControlRigEditMode::RecreateControlShapeActors(const TArray<FRigElementKey
 {
 	if (RecreateControlShapesRequired == ERecreateControlRigShape::RecreateAll)
 	{
-		for (TWeakObjectPtr<UControlRig>& RuntimeRigPtr : RuntimeControlRigs)
+		for (TWeakObjectPtr<UBaseControlRig>& RuntimeRigPtr : RuntimeControlRigs)
 		{
-			if (UControlRig* RuntimeControlRig = RuntimeRigPtr.Get())
+			if (UBaseControlRig* RuntimeControlRig = RuntimeRigPtr.Get())
 			{
 				DestroyShapesActors(RuntimeControlRig);
 				CreateShapeActors(RuntimeControlRig);
@@ -3387,8 +3387,8 @@ void FControlRigEditMode::RecreateControlShapeActors(const TArray<FRigElementKey
 	}
 	else if (ControlRigsToRecreate.Num() > 0)
 	{
-		TArray < UControlRig*> ControlRigsCopy = ControlRigsToRecreate;
-		for (UControlRig* ControlRig : ControlRigsCopy)
+		TArray < UBaseControlRig*> ControlRigsCopy = ControlRigsToRecreate;
+		for (UBaseControlRig* ControlRig : ControlRigsCopy)
 		{
 			//check to see if actors have really changed, if not don't do it
 			bool bRecreateThem = true;
@@ -3467,7 +3467,7 @@ void FControlRigEditMode::RecreateControlShapeActors(const TArray<FRigElementKey
 	}
 }
 
-void FControlRigEditMode::CreateShapeActors(UControlRig* ControlRig)
+void FControlRigEditMode::CreateShapeActors(UBaseControlRig* ControlRig)
 {
 	// create gizmo actors
 	FActorSpawnParameters ActorSpawnParameters;
@@ -3615,7 +3615,7 @@ void FControlRigEditMode::OnObjectsReplaced(const TMap<UObject*, UObject*>& OldT
 		UObject* NewObject = OldToNewInstanceMap.FindRef(OldObject);
 		if (NewObject)
 		{
-			TArray<TWeakObjectPtr<UControlRig>> PreviousRuntimeRigs = RuntimeControlRigs;
+			TArray<TWeakObjectPtr<UBaseControlRig>> PreviousRuntimeRigs = RuntimeControlRigs;
 			for (int32 PreviousRuntimeRigIndex = 0; PreviousRuntimeRigIndex < PreviousRuntimeRigs.Num(); PreviousRuntimeRigIndex++)
 			{
 				if (PreviousRuntimeRigs[PreviousRuntimeRigIndex].IsValid())
@@ -3625,7 +3625,7 @@ void FControlRigEditMode::OnObjectsReplaced(const TMap<UObject*, UObject*>& OldT
 			}
 			RuntimeControlRigs.Reset();
 
-			UControlRig* NewRig = Cast<UControlRig>(NewObject);
+			UBaseControlRig* NewRig = Cast<UBaseControlRig>(NewObject);
 			AddControlRigInternal(NewRig);
 
 			NewRig->Initialize();
@@ -3640,7 +3640,7 @@ bool FControlRigEditMode::IsTransformDelegateAvailable() const
 	return (OnGetRigElementTransformDelegate.IsBound() && OnSetRigElementTransformDelegate.IsBound());
 }
 
-bool FControlRigEditMode::AreRigElementSelectedAndMovable(UControlRig* ControlRig) const
+bool FControlRigEditMode::AreRigElementSelectedAndMovable(UBaseControlRig* ControlRig) const
 {
 	const UControlRigEditModeSettings* Settings = GetDefault<UControlRigEditModeSettings>();
 
@@ -3698,7 +3698,7 @@ bool FControlRigEditMode::AreRigElementSelectedAndMovable(UControlRig* ControlRi
 	return true;
 }
 
-void FControlRigEditMode::ReplaceControlRig(UControlRig* OldControlRig, UControlRig* NewControlRig)
+void FControlRigEditMode::ReplaceControlRig(UBaseControlRig* OldControlRig, UBaseControlRig* NewControlRig)
 {
 	if (OldControlRig != nullptr)
 	{
@@ -3725,7 +3725,7 @@ void FControlRigEditMode::OnHierarchyModified(ERigHierarchyNotification InNotif,
 		case ERigHierarchyNotification::ElementReordered:
 		case ERigHierarchyNotification::HierarchyReset:
 		{
-			UControlRig* ControlRig = InHierarchy->GetTypedOuter<UControlRig>();
+			UBaseControlRig* ControlRig = InHierarchy->GetTypedOuter<UBaseControlRig>();
 			RequestToRecreateControlShapeActors(ControlRig);
 			break;
 		}
@@ -3734,7 +3734,7 @@ void FControlRigEditMode::OnHierarchyModified(ERigHierarchyNotification InNotif,
 		case ERigHierarchyNotification::ControlShapeTransformChanged:
 		{
 			const FRigElementKey Key = InElement->GetKey();
-			UControlRig* ControlRig = InHierarchy->GetTypedOuter<UControlRig>();
+			UBaseControlRig* ControlRig = InHierarchy->GetTypedOuter<UBaseControlRig>();
 			if (Key.Type == ERigElementType::Control)
 			{
 				if (const FRigControlElement* ControlElement = Cast<FRigControlElement>(InElement))
@@ -3768,7 +3768,7 @@ void FControlRigEditMode::OnHierarchyModified(ERigHierarchyNotification InNotif,
 				// are correctly selected in sequencer to match what the proxy control is driving.
 				if (FRigControlElement* ControlElement = InHierarchy->Find<FRigControlElement>(InElement->GetKey()))
 				{
-					UControlRig* ControlRig = InHierarchy->GetTypedOuter<UControlRig>();
+					UBaseControlRig* ControlRig = InHierarchy->GetTypedOuter<UBaseControlRig>();
 					if(ControlProxy->IsSelected(ControlRig, ControlElement->GetFName()))
 					{
 						// reselect the control - to affect the details panel / sequencer
@@ -3811,7 +3811,7 @@ void FControlRigEditMode::OnHierarchyModified(ERigHierarchyNotification InNotif,
 					const bool bSelected = InNotif == ERigHierarchyNotification::ElementSelected;
 					// users may select gizmo and control rig units, so we have to let them go through both of them if they do
 						// first go through gizmo actor
-					UControlRig* ControlRig = InHierarchy->GetTypedOuter<UControlRig>();
+					UBaseControlRig* ControlRig = InHierarchy->GetTypedOuter<UBaseControlRig>();
 					if (ControlRig == nullptr)
 					{
 						if (RuntimeControlRigs.Num() > 0)
@@ -3968,7 +3968,7 @@ void FControlRigEditMode::OnHierarchyModified_AnyThread(ERigHierarchyNotificatio
 	}, TStatId(), NULL, ENamedThreads::GameThread);
 }
 
-void FControlRigEditMode::OnControlModified(UControlRig* Subject, FRigControlElement* InControlElement, const FRigControlModifiedContext& Context)
+void FControlRigEditMode::OnControlModified(UBaseControlRig* Subject, FRigControlElement* InControlElement, const FRigControlModifiedContext& Context)
 {
 	//this makes sure the details panel ui get's updated, don't remove
 	ControlProxy->ProxyChanged(Subject,InControlElement->GetFName());
@@ -3978,7 +3978,7 @@ void FControlRigEditMode::OnControlModified(UControlRig* Subject, FRigControlEle
 	ControlProxy->Modify();
 	RecalcPivotTransform();
 
-	if (UControlRig* ControlRig = static_cast<UControlRig*>(Subject))
+	if (UBaseControlRig* ControlRig = static_cast<UBaseControlRig*>(Subject))
 	{
 		FTransform ComponentTransform = GetHostingSceneComponentTransform();
 		if (AControlRigShapeActor* const* Actor = GizmoToControlMap.FindKey(InControl.Index))
@@ -3989,12 +3989,12 @@ void FControlRigEditMode::OnControlModified(UControlRig* Subject, FRigControlEle
 	*/
 }
 
-void FControlRigEditMode::OnPreConstruction_AnyThread(UControlRig* InRig, const FName& InEventName)
+void FControlRigEditMode::OnPreConstruction_AnyThread(UBaseControlRig* InRig, const FName& InEventName)
 {
 	bIsConstructionEventRunning = true;
 }
 
-void FControlRigEditMode::OnPostConstruction_AnyThread(UControlRig* InRig, const FName& InEventName)
+void FControlRigEditMode::OnPostConstruction_AnyThread(UBaseControlRig* InRig, const FName& InEventName)
 {
 	bIsConstructionEventRunning = false;
 
@@ -4063,9 +4063,9 @@ bool FControlRigEditMode::CanChangeControlShapeTransform()
 {
 	if (AreEditingControlRigDirectly())
 	{
-		for (TWeakObjectPtr<UControlRig> RuntimeRigPtr : RuntimeControlRigs)
+		for (TWeakObjectPtr<UBaseControlRig> RuntimeRigPtr : RuntimeControlRigs)
 		{
-			if (UControlRig* ControlRig = RuntimeRigPtr.Get())
+			if (UBaseControlRig* ControlRig = RuntimeRigPtr.Get())
 			{
 				TArray<FRigElementKey> SelectedRigElements = GetSelectedRigElements(ControlRig);
 				// do not allow multi-select
@@ -4104,7 +4104,7 @@ void FControlRigEditMode::SetControlShapeTransform(
 	const FRigControlModifiedContext& InContext,
 	const bool bPrintPython) const
 {
-	UControlRig* ControlRig = InShapeActor->ControlRig.Get();
+	UBaseControlRig* ControlRig = InShapeActor->ControlRig.Get();
 	if (!ControlRig)
 	{
 		return;
@@ -4161,7 +4161,7 @@ void FControlRigEditMode::SetControlShapeTransform(
 
 FTransform FControlRigEditMode::GetControlShapeTransform(const AControlRigShapeActor* ShapeActor)
 {
-	if (const UControlRig* ControlRig = ShapeActor->ControlRig.Get())
+	if (const UBaseControlRig* ControlRig = ShapeActor->ControlRig.Get())
 	{
 		return ControlRig->GetControlGlobalTransform(ShapeActor->ControlName);
 	}
@@ -4176,7 +4176,7 @@ void FControlRigEditMode::MoveControlShape(AControlRigShapeActor* ShapeActor, co
 	//first case is where we do all controls by the local diff.
 	if (bUseLocal)
 	{
-		if (UControlRig* ControlRig = ShapeActor->ControlRig.Get())
+		if (UBaseControlRig* ControlRig = ShapeActor->ControlRig.Get())
 		{
 			FRigControlModifiedContext Context;
 			Context.EventName = FRigUnit_BeginExecution::EventName;
@@ -4251,7 +4251,7 @@ void FControlRigEditMode::MoveControlShape(AControlRigShapeActor* ShapeActor, co
 
 		if (bTransformChanged)
 		{
-			if (UControlRig* ControlRig = ShapeActor->ControlRig.Get())
+			if (UBaseControlRig* ControlRig = ShapeActor->ControlRig.Get())
 			{
 				ControlRig->InteractionType = InteractionType;
 				ControlRig->ElementsBeingInteracted.AddUnique(ShapeActor->GetElementKey());
@@ -4314,7 +4314,7 @@ void FControlRigEditMode::ChangeControlShapeTransform(AControlRigShapeActor* Sha
 
 	FTransform CurrentTransform = FTransform::Identity;
 
-	if (UControlRig* ControlRig = ShapeActor->ControlRig.Get())
+	if (UBaseControlRig* ControlRig = ShapeActor->ControlRig.Get())
 	{
 		if (FRigControlElement* ControlElement = ControlRig->GetHierarchy()->Find<FRigControlElement>(ShapeActor->GetElementKey()))
 		{
@@ -4349,7 +4349,7 @@ void FControlRigEditMode::ChangeControlShapeTransform(AControlRigShapeActor* Sha
 
 	if (bTransformChanged)
 	{
-		if (UControlRig* ControlRig = ShapeActor->ControlRig.Get())
+		if (UBaseControlRig* ControlRig = ShapeActor->ControlRig.Get())
 		{
 
 			FTransform NewTransform = CurrentTransform.GetRelativeTransform(ToWorldTransform);
@@ -4378,7 +4378,7 @@ void FControlRigEditMode::ChangeControlShapeTransform(AControlRigShapeActor* Sha
 
 bool FControlRigEditMode::ModeSupportedByShapeActor(const AControlRigShapeActor* ShapeActor, UE::Widget::EWidgetMode InMode) const
 {
-	if (UControlRig* ControlRig = ShapeActor->ControlRig.Get())
+	if (UBaseControlRig* ControlRig = ShapeActor->ControlRig.Get())
 	{
 		const FRigControlElement* ControlElement = ControlRig->FindControl(ShapeActor->ControlName);
 		if (ControlElement)
@@ -4475,7 +4475,7 @@ bool FControlRigEditMode::ModeSupportedByShapeActor(const AControlRigShapeActor*
 	return false;
 }
 
-bool FControlRigEditMode::IsControlRigSkelMeshVisible(UControlRig* ControlRig) const
+bool FControlRigEditMode::IsControlRigSkelMeshVisible(UBaseControlRig* ControlRig) const
 {
 	if (IsInLevelEditor())
 	{
@@ -4497,7 +4497,7 @@ void FControlRigEditMode::TickControlShape(AControlRigShapeActor* ShapeActor, co
 	const UControlRigEditModeSettings* Settings = GetDefault<UControlRigEditModeSettings>();
 	if (ShapeActor)
 	{
-		if (UControlRig* ControlRig = ShapeActor->ControlRig.Get())
+		if (UBaseControlRig* ControlRig = ShapeActor->ControlRig.Get())
 		{
 			const FTransform Transform = ControlRig->GetControlGlobalTransform(ShapeActor->ControlName);
 			ShapeActor->SetActorTransform(Transform * ComponentTransform);
@@ -4535,7 +4535,7 @@ void FControlRigEditMode::TickControlShape(AControlRigShapeActor* ShapeActor, co
 	}
 }
 
-AControlRigShapeActor* FControlRigEditMode::GetControlShapeFromControlName(UControlRig* InControlRig,const FName& ControlName) const
+AControlRigShapeActor* FControlRigEditMode::GetControlShapeFromControlName(UBaseControlRig* InControlRig,const FName& ControlName) const
 {
 	const auto* ShapeActors = ControlRigShapeActors.Find(InControlRig);
 	if (ShapeActors)
@@ -4552,7 +4552,7 @@ AControlRigShapeActor* FControlRigEditMode::GetControlShapeFromControlName(UCont
 	return nullptr;
 }
 
-void FControlRigEditMode::AddControlRigInternal(UControlRig* InControlRig)
+void FControlRigEditMode::AddControlRigInternal(UBaseControlRig* InControlRig)
 {
 	RuntimeControlRigs.AddUnique(InControlRig);
 	LastHierarchyHash.Add(INDEX_NONE);
@@ -4574,20 +4574,20 @@ void FControlRigEditMode::AddControlRigInternal(UControlRig* InControlRig)
 	UpdateSelectabilityOnSkeletalMeshes(InControlRig, !bShowControlsAsOverlay);
 }
 
-TArrayView<const TWeakObjectPtr<UControlRig>> FControlRigEditMode::GetControlRigs() const
+TArrayView<const TWeakObjectPtr<UBaseControlRig>> FControlRigEditMode::GetControlRigs() const
 {
 	return MakeArrayView(RuntimeControlRigs);
 }
 
-TArrayView<TWeakObjectPtr<UControlRig>> FControlRigEditMode::GetControlRigs() 
+TArrayView<TWeakObjectPtr<UBaseControlRig>> FControlRigEditMode::GetControlRigs() 
 {
 	return MakeArrayView(RuntimeControlRigs);
 }
 
-TArray<UControlRig*> FControlRigEditMode::GetControlRigsArray(bool bIsVisible)
+TArray<UBaseControlRig*> FControlRigEditMode::GetControlRigsArray(bool bIsVisible)
 {
-	TArray < UControlRig*> ControlRigs;
-	for (TWeakObjectPtr<UControlRig> ControlRigPtr : RuntimeControlRigs)
+	TArray < UBaseControlRig*> ControlRigs;
+	for (TWeakObjectPtr<UBaseControlRig> ControlRigPtr : RuntimeControlRigs)
 	{
 		if (ControlRigPtr.IsValid() && ControlRigPtr.Get() != nullptr && (bIsVisible == false ||ControlRigPtr.Get()->GetControlsVisible()))
 		{
@@ -4597,10 +4597,10 @@ TArray<UControlRig*> FControlRigEditMode::GetControlRigsArray(bool bIsVisible)
 	return ControlRigs;
 }
 
-TArray<const UControlRig*> FControlRigEditMode::GetControlRigsArray(bool bIsVisible) const
+TArray<const UBaseControlRig*> FControlRigEditMode::GetControlRigsArray(bool bIsVisible) const
 {
-	TArray<const UControlRig*> ControlRigs;
-	for (const TWeakObjectPtr<UControlRig> ControlRigPtr : RuntimeControlRigs)
+	TArray<const UBaseControlRig*> ControlRigs;
+	for (const TWeakObjectPtr<UBaseControlRig> ControlRigPtr : RuntimeControlRigs)
 	{
 		if (ControlRigPtr.IsValid() && ControlRigPtr.Get() != nullptr && (bIsVisible == false || ControlRigPtr.Get()->GetControlsVisible()))
 		{
@@ -4610,7 +4610,7 @@ TArray<const UControlRig*> FControlRigEditMode::GetControlRigsArray(bool bIsVisi
 	return ControlRigs;
 }
 
-void FControlRigEditMode::RemoveControlRig(UControlRig* InControlRig)
+void FControlRigEditMode::RemoveControlRig(UBaseControlRig* InControlRig)
 {
 	if (InControlRig == nullptr)
 	{
@@ -4652,9 +4652,9 @@ void FControlRigEditMode::RemoveControlRig(UControlRig* InControlRig)
 
 void FControlRigEditMode::TickManipulatableObjects(float DeltaTime)
 {
-	for (TWeakObjectPtr<UControlRig> RuntimeRigPtr : RuntimeControlRigs)
+	for (TWeakObjectPtr<UBaseControlRig> RuntimeRigPtr : RuntimeControlRigs)
 	{
-		if (UControlRig* ControlRig = RuntimeRigPtr.Get())
+		if (UBaseControlRig* ControlRig = RuntimeRigPtr.Get())
 		{
 			// tick skeletalmeshcomponent, that's how they update their transform from rig change
 			USceneComponent* SceneComponent = GetHostingSceneComponent(ControlRig);
@@ -4722,7 +4722,7 @@ void FControlRigEditMode::OnEditorClosed()
 	ControlRigsToRecreate.Reset();
 }
 
-void FControlRigEditMode::DestroyShapesActors(UControlRig* ControlRig)
+void FControlRigEditMode::DestroyShapesActors(UBaseControlRig* ControlRig)
 {
 	if (ControlRig == nullptr)
 	{
@@ -4771,7 +4771,7 @@ void FControlRigEditMode::DestroyShapesActors(UControlRig* ControlRig)
 	}
 }
 
-USceneComponent* FControlRigEditMode::GetHostingSceneComponent(const UControlRig* ControlRig) const
+USceneComponent* FControlRigEditMode::GetHostingSceneComponent(const UBaseControlRig* ControlRig) const
 {
 	if (ControlRig == nullptr && GetControlRigs().Num() > 0)
 	{
@@ -4812,7 +4812,7 @@ USceneComponent* FControlRigEditMode::GetHostingSceneComponent(const UControlRig
 	return nullptr;
 }
 
-FTransform FControlRigEditMode::GetHostingSceneComponentTransform(const UControlRig* ControlRig) const
+FTransform FControlRigEditMode::GetHostingSceneComponentTransform(const UBaseControlRig* ControlRig) const
 {
 	if (ControlRig == nullptr && GetControlRigs().Num() > 0)
 	{
@@ -4844,7 +4844,7 @@ void FControlRigEditMode::PostPoseUpdate()
 
 }
 
-void FControlRigEditMode::NotifyDrivenControls(UControlRig* InControlRig, const FRigElementKey& InKey)
+void FControlRigEditMode::NotifyDrivenControls(UBaseControlRig* InControlRig, const FRigElementKey& InKey)
 {
 	// if we are changing a proxy control - we also need to notify the change for the driven controls
 	if (FRigControlElement* ControlElement = InControlRig->GetHierarchy()->Find<FRigControlElement>(InKey))
@@ -4866,7 +4866,7 @@ void FControlRigEditMode::NotifyDrivenControls(UControlRig* InControlRig, const
 	}
 }
 
-void FControlRigEditMode::UpdateSelectabilityOnSkeletalMeshes(UControlRig* InControlRig, bool bEnabled)
+void FControlRigEditMode::UpdateSelectabilityOnSkeletalMeshes(UBaseControlRig* InControlRig, bool bEnabled)
 {
 	if(const USceneComponent* HostingComponent = GetHostingSceneComponent(InControlRig))
 	{
diff --git a/Engine/Plugins/Animation/ControlRig/Source/ControlRigEditor/Private/EditMode/SControlPicker.cpp b/Engine/Plugins/Animation/ControlRig/Source/ControlRigEditor/Private/EditMode/SControlPicker.cpp
--- a/Engine/Plugins/Animation/ControlRig/Source/ControlRigEditor/Private/EditMode/SControlPicker.cpp
+++ b/Engine/Plugins/Animation/ControlRig/Source/ControlRigEditor/Private/EditMode/SControlPicker.cpp
@@ -30,7 +30,7 @@ void SControlPicker::Construct(const FArguments& InArgs, UWorld* InWorld)
 	];
 }
 
-void SControlPicker::SetControlRig(UControlRig* InRig)
+void SControlPicker::SetControlRig(UBaseControlRig* InRig)
 {
 	if (InRig != RigPtr.Get())
 	{
@@ -38,7 +38,7 @@ void SControlPicker::SetControlRig(UControlRig* InRig)
 	}
 }
 
-UControlRig* SControlPicker::GetRig() const
+UBaseControlRig* SControlPicker::GetRig() const
 {
 	return RigPtr.Get();
 }
diff --git a/Engine/Plugins/Animation/ControlRig/Source/ControlRigEditor/Private/EditMode/SControlPicker.h b/Engine/Plugins/Animation/ControlRig/Source/ControlRigEditor/Private/EditMode/SControlPicker.h
--- a/Engine/Plugins/Animation/ControlRig/Source/ControlRigEditor/Private/EditMode/SControlPicker.h
+++ b/Engine/Plugins/Animation/ControlRig/Source/ControlRigEditor/Private/EditMode/SControlPicker.h
@@ -8,7 +8,7 @@
 #include "Delegates/Delegate.h"
 #include "UObject/WeakObjectPtr.h"
 
-class UControlRig;
+class UBaseControlRig;
 class SControlPicker;
 class SCanvas;
 class SScaleBox;
@@ -27,7 +27,7 @@ public:
 	void Construct(const FArguments& InArgs, UWorld* InWorld);
 
 	/** Set the rig to display manipulators for */
-	void SetControlRig(UControlRig* InRig);
+	void SetControlRig(UBaseControlRig* InRig);
 
 	/** Set the manipulators that are currently selected */
 	void SetSelectedManipulators(const TArray<FName>& Manipulators);
@@ -43,7 +43,7 @@ public:
 	bool IsManipulatorSelected(FName ManipulatorName);
 
 	/** Returns the rig we are displaying controls for */
-	UControlRig* GetRig() const;
+	UBaseControlRig* GetRig() const;
 
 	/** See if the limb/spine is in IK mode, or FK */
 	bool IsControlIK(FName ControlName) const;
@@ -59,5 +59,5 @@ protected:
 	TSharedPtr<SEditorUserWidgetHost> EditorUserWidgetHost;
 
 	/** Rig we are showing controls for */
-	TWeakObjectPtr<UControlRig> RigPtr;
+	TWeakObjectPtr<UBaseControlRig> RigPtr;
 };
diff --git a/Engine/Plugins/Animation/ControlRig/Source/ControlRigEditor/Private/EditMode/SControlRigBaseListWidget.cpp b/Engine/Plugins/Animation/ControlRig/Source/ControlRigEditor/Private/EditMode/SControlRigBaseListWidget.cpp
--- a/Engine/Plugins/Animation/ControlRig/Source/ControlRigEditor/Private/EditMode/SControlRigBaseListWidget.cpp
+++ b/Engine/Plugins/Animation/ControlRig/Source/ControlRigEditor/Private/EditMode/SControlRigBaseListWidget.cpp
@@ -423,7 +423,7 @@ void SControlRigPoseAnimSelectionToolbar::MakeControlRigAssetDialog(FControlRigA
 		return;
 	}
 
-	TMap<UControlRig*, TArray<FRigElementKey>> AllSelectedControls;
+	TMap<UBaseControlRig*, TArray<FRigElementKey>> AllSelectedControls;
 	ControlRigEditMode->GetAllSelectedControls(AllSelectedControls);
 	
 	if (AllSelectedControls.Num() > 1)
@@ -455,13 +455,13 @@ void SControlRigPoseAnimSelectionToolbar::MakeControlRigAssetDialog(FControlRigA
 			FControlRigEditMode* ControlRigEditMode =  OwningControlRigWidget->GetEditMode();
 			if (ControlRigEditMode )
 			{
-				TMap<UControlRig*, TArray<FRigElementKey>> AllSelectedControls;
+				TMap<UBaseControlRig*, TArray<FRigElementKey>> AllSelectedControls;
 				ControlRigEditMode->GetAllSelectedControls(AllSelectedControls);
 				if (AllSelectedControls.Num() == 1)
 				{
-					TArray<UControlRig*> ControlRigs;
+					TArray<UBaseControlRig*> ControlRigs;
 					AllSelectedControls.GenerateKeyArray(ControlRigs);
-					UControlRig* ControlRig =ControlRigs[0];
+					UBaseControlRig* ControlRig =ControlRigs[0];
 					UObject* NewAsset = nullptr;
 					switch (Type)
 					{
@@ -497,7 +497,7 @@ bool SControlRigPoseAnimSelectionToolbar::CanExecuteMakeControlRigAsset()
 		return false;
 	}
 	
-	TMap<UControlRig*, TArray<FRigElementKey>> AllSelectedControls;
+	TMap<UBaseControlRig*, TArray<FRigElementKey>> AllSelectedControls;
 	ControlRigEditMode->GetAllSelectedControls(AllSelectedControls);
 
 	if (AllSelectedControls.Num() != 1)
@@ -990,11 +990,11 @@ void SControlRigBaseListWidget::OnAssetsActivated(const TArray<FAssetData>& Sele
 					FControlRigEditMode* ControlRigEditMode = GetEditMode();
 					if (ControlRigEditMode)
 					{
-						TArray<UControlRig*> ControlRigs = ControlRigEditMode->GetControlRigsArray(true /*bIsVisible*/);
+						TArray<UBaseControlRig*> ControlRigs = ControlRigEditMode->GetControlRigsArray(true /*bIsVisible*/);
 						if (ControlRigs.Num() > 0)
 						{
 							const FScopedTransaction Transaction(LOCTEXT("SelectControls", "Select Controls"));
-							for (UControlRig* ControlRig : ControlRigs)
+							for (UBaseControlRig* ControlRig : ControlRigs)
 							{
 								ControlRig->Modify();
 								PoseAsset->SelectControls(ControlRig, SControlRigPoseView::IsMirror());
@@ -1368,11 +1368,11 @@ void SControlRigBaseListWidget::ExecutePastePose(UControlRigPoseAsset* PoseAsset
 		FControlRigEditMode* ControlRigEditMode = GetEditMode();
 		if (ControlRigEditMode)
 		{
-			TArray<UControlRig*> ControlRigs = ControlRigEditMode->GetControlRigsArray(true /*bIsVisible*/);
+			TArray<UBaseControlRig*> ControlRigs = ControlRigEditMode->GetControlRigsArray(true /*bIsVisible*/);
 			if (ControlRigs.Num() > 0)
 			{
 				const FScopedTransaction Transaction(LOCTEXT("PastePose", "Paste Pose"));
-				for (UControlRig* ControlRig : ControlRigs)
+				for (UBaseControlRig* ControlRig : ControlRigs)
 				{
 					ControlRig->Modify();
 					PoseAsset->PastePose(ControlRig);
@@ -1397,11 +1397,11 @@ void SControlRigBaseListWidget::ExecuteSelectControls(UControlRigPoseAsset* Pose
 	FControlRigEditMode* ControlRigEditMode = GetEditMode();
 	if (ControlRigEditMode)
 	{
-		TArray<UControlRig*> ControlRigs = ControlRigEditMode->GetControlRigsArray(true /*bIsVisible*/);
+		TArray<UBaseControlRig*> ControlRigs = ControlRigEditMode->GetControlRigsArray(true /*bIsVisible*/);
 		if (ControlRigs.Num() > 0)
 		{
 			const FScopedTransaction Transaction(LOCTEXT("SelectControls", "Select Controls"));
-			for (UControlRig* ControlRig : ControlRigs)
+			for (UBaseControlRig* ControlRig : ControlRigs)
 			{
 				ControlRig->Modify();
 				PoseAsset->SelectControls(ControlRig);
@@ -1434,11 +1434,11 @@ void SControlRigBaseListWidget::ExecutePasteMirrorPose(UControlRigPoseAsset* Pos
 		FControlRigEditMode* ControlRigEditMode = GetEditMode();
 		if (ControlRigEditMode)
 		{
-			TArray<UControlRig*> ControlRigs = ControlRigEditMode->GetControlRigsArray(true /*bIsVisible*/);
+			TArray<UBaseControlRig*> ControlRigs = ControlRigEditMode->GetControlRigsArray(true /*bIsVisible*/);
 			if (ControlRigs.Num() > 0)
 			{
 				const FScopedTransaction Transaction(LOCTEXT("PasteMirrorPose", "Paste Mirror Pose"));
-				for (UControlRig* ControlRig : ControlRigs)
+				for (UBaseControlRig* ControlRig : ControlRigs)
 				{
 					ControlRig->Modify();
 					PoseAsset->PastePose(ControlRig, false, true);
diff --git a/Engine/Plugins/Animation/ControlRig/Source/ControlRigEditor/Private/EditMode/SControlRigBaseListWidget.h b/Engine/Plugins/Animation/ControlRig/Source/ControlRigEditor/Private/EditMode/SControlRigBaseListWidget.h
--- a/Engine/Plugins/Animation/ControlRig/Source/ControlRigEditor/Private/EditMode/SControlRigBaseListWidget.h
+++ b/Engine/Plugins/Animation/ControlRig/Source/ControlRigEditor/Private/EditMode/SControlRigBaseListWidget.h
@@ -20,7 +20,7 @@
 #include "ContentBrowserDelegates.h"
 #include "EditMode/SControlRigControlViews.h"
 
-class UControlRig;
+class UBaseControlRig;
 class FControlRigEditMode;
 class UControlRigPoseAsset;
 class FControlRigEditModeToolkit;
diff --git a/Engine/Plugins/Animation/ControlRig/Source/ControlRigEditor/Private/EditMode/SControlRigControlViews.cpp b/Engine/Plugins/Animation/ControlRig/Source/ControlRigEditor/Private/EditMode/SControlRigControlViews.cpp
--- a/Engine/Plugins/Animation/ControlRig/Source/ControlRigEditor/Private/EditMode/SControlRigControlViews.cpp
+++ b/Engine/Plugins/Animation/ControlRig/Source/ControlRigEditor/Private/EditMode/SControlRigControlViews.cpp
@@ -446,8 +446,8 @@ void SControlRigPoseView::Construct(const FArguments& InArgs)
 	if (FControlRigEditMode* EditMode = OwningWidget->GetEditMode())
 	{
 		EditMode->OnControlRigAddedOrRemoved().AddRaw(this, &SControlRigPoseView::HandleControlAdded);
-		TArray<UControlRig*> ControlRigs = GetControlRigs();
-		for (UControlRig* ControlRig : ControlRigs)
+		TArray<UBaseControlRig*> ControlRigs = GetControlRigs();
+		for (UBaseControlRig* ControlRig : ControlRigs)
 		{
 			HandleControlAdded(ControlRig, true);
 		}
@@ -459,8 +459,8 @@ SControlRigPoseView::~SControlRigPoseView()
 	if (FControlRigEditMode* EditMode = OwningWidget->GetEditMode())
 	{
 		EditMode->OnControlRigAddedOrRemoved().RemoveAll(this);
-		TArray<UControlRig*> EditModeRigs = EditMode->GetControlRigsArray(false /*bIsVisible*/);
-		for (UControlRig* ControlRig : EditModeRigs)
+		TArray<UBaseControlRig*> EditModeRigs = EditMode->GetControlRigsArray(false /*bIsVisible*/);
+		for (UBaseControlRig* ControlRig : EditModeRigs)
 		{
 			if (ControlRig)
 			{
@@ -470,7 +470,7 @@ SControlRigPoseView::~SControlRigPoseView()
 	}
 	else
 	{
-		for (TWeakObjectPtr<UControlRig>& CurrentControlRig: CurrentControlRigs)
+		for (TWeakObjectPtr<UBaseControlRig>& CurrentControlRig: CurrentControlRigs)
 		{
 			if (CurrentControlRig.IsValid())
 			{
@@ -538,8 +538,8 @@ FReply SControlRigPoseView::OnPastePose()
 {
 	if (PoseAsset.IsValid())
 	{
-		TArray<UControlRig*> ControlRigs = GetControlRigs();
-		for (UControlRig* ControlRig : ControlRigs)
+		TArray<UBaseControlRig*> ControlRigs = GetControlRigs();
+		for (UBaseControlRig* ControlRig : ControlRigs)
 		{
 			if (ControlRig)
 			{
@@ -554,8 +554,8 @@ FReply SControlRigPoseView::OnSelectControls()
 {	
 	if (PoseAsset.IsValid())
 	{
-		TArray<UControlRig*> ControlRigs = GetControlRigs();
-		for (UControlRig* ControlRig : ControlRigs)
+		TArray<UBaseControlRig*> ControlRigs = GetControlRigs();
+		for (UBaseControlRig* ControlRig : ControlRigs)
 		{
 			if (ControlRig)
 			{
@@ -570,8 +570,8 @@ void SControlRigPoseView::OnPoseBlendChanged(float ChangedVal)
 {
 	if (PoseAsset.IsValid())
 	{
-		TArray<UControlRig*> ControlRigs = GetControlRigs();
-		for (UControlRig* ControlRig : ControlRigs)
+		TArray<UBaseControlRig*> ControlRigs = GetControlRigs();
+		for (UBaseControlRig* ControlRig : ControlRigs)
 		{
 			if (ControlRig)
 			{
@@ -610,11 +610,11 @@ void SControlRigPoseView::OnPoseBlendCommited(float ChangedVal, ETextCommit::Typ
 {
 	if (PoseAsset.IsValid())
 	{
-		TArray<UControlRig*> ControlRigs = GetControlRigs();
+		TArray<UBaseControlRig*> ControlRigs = GetControlRigs();
 		if (ControlRigs.Num() > 0)
 		{
 			FScopedTransaction ScopedTransaction(LOCTEXT("PastePoseTransaction", "Paste Pose"));
-			for (UControlRig* ControlRig : ControlRigs)
+			for (UBaseControlRig* ControlRig : ControlRigs)
 			{
 				if (ControlRig)
 				{
@@ -662,15 +662,15 @@ TSharedRef<SWidget> SControlRigPoseView::GetThumbnailWidget()
 		];
 }
 
-TArray<UControlRig*> SControlRigPoseView::GetControlRigs()
+TArray<UBaseControlRig*> SControlRigPoseView::GetControlRigs()
 {
 	FControlRigEditMode* EditMode = OwningWidget->GetEditMode();
-	TArray<UControlRig*> NewControlRigs;
+	TArray<UBaseControlRig*> NewControlRigs;
 	if (EditMode)
 	{
 		NewControlRigs =  EditMode->GetControlRigsArray(false /*bIsVisible*/);
 	}
-	for (TWeakObjectPtr<UControlRig> ControlRigPtr : CurrentControlRigs)
+	for (TWeakObjectPtr<UBaseControlRig> ControlRigPtr : CurrentControlRigs)
 	{
 		if (ControlRigPtr.IsValid())
 		{
@@ -701,7 +701,7 @@ void SControlRigPoseView::CreateControlList()
 	}
 }
 */
-void SControlRigPoseView::HandleControlAdded(UControlRig* ControlRig, bool bIsAdded)
+void SControlRigPoseView::HandleControlAdded(UBaseControlRig* ControlRig, bool bIsAdded)
 {
 	if (ControlRig)
 	{
@@ -718,7 +718,7 @@ void SControlRigPoseView::HandleControlAdded(UControlRig* ControlRig, bool bIsAd
 	UpdateStatusBlocks();
 }
 
-void SControlRigPoseView::HandleControlSelected(UControlRig* Subject, FRigControlElement* InControl, bool bSelected)
+void SControlRigPoseView::HandleControlSelected(UBaseControlRig* Subject, FRigControlElement* InControl, bool bSelected)
 {
 	UpdateStatusBlocks();
 }
@@ -727,7 +727,7 @@ void SControlRigPoseView::UpdateStatusBlocks()
 {
 	FText StatusText1;
 	FText StatusText2;
-	TArray<UControlRig*> ControlRigs = GetControlRigs();
+	TArray<UBaseControlRig*> ControlRigs = GetControlRigs();
 	if (PoseAsset.IsValid() && ControlRigs.Num() > 0)
 	{
 
@@ -737,7 +737,7 @@ void SControlRigPoseView::UpdateStatusBlocks()
 		int32 TotalSelected = 0;
 		uint32 Matching = 0;
 		uint32 MirrorMatching = 0;
-		for (UControlRig* ControlRig : ControlRigs)
+		for (UBaseControlRig* ControlRig : ControlRigs)
 		{
 
 			TArray<FName> SelectedNames = ControlRig->CurrentControlSelection();
diff --git a/Engine/Plugins/Animation/ControlRig/Source/ControlRigEditor/Private/EditMode/SControlRigControlViews.h b/Engine/Plugins/Animation/ControlRig/Source/ControlRigEditor/Private/EditMode/SControlRigControlViews.h
--- a/Engine/Plugins/Animation/ControlRig/Source/ControlRigEditor/Private/EditMode/SControlRigControlViews.h
+++ b/Engine/Plugins/Animation/ControlRig/Source/ControlRigEditor/Private/EditMode/SControlRigControlViews.h
@@ -19,7 +19,7 @@
 
 class FAssetThumbnail;
 class IDetailsView;
-class UControlRig;
+class UBaseControlRig;
 class SControlRigBaseListWidget;
 
 //Class to Hold Statics that are shared and externally callable
@@ -62,9 +62,9 @@ private:
 	FReply OnPastePose();
 	FReply OnSelectControls();
 	FReply OnCaptureThumbnail();
-	void HandleControlSelected(UControlRig* Subject, FRigControlElement* InControl, bool bSelected);
-	void HandleControlAdded(UControlRig* ControlRig, bool bIsAdded);
-	TArray<UControlRig*> GetControlRigs();
+	void HandleControlSelected(UBaseControlRig* Subject, FRigControlElement* InControl, bool bSelected);
+	void HandleControlAdded(UBaseControlRig* ControlRig, bool bIsAdded);
+	TArray<UBaseControlRig*> GetControlRigs();
 
 	static bool bIsKey;
 	static bool bIsMirror;
@@ -76,7 +76,7 @@ private:
 	TSharedRef<SWidget> GetThumbnailWidget();
 
 	TWeakObjectPtr<UControlRigPoseAsset> PoseAsset;
-	TArray<TWeakObjectPtr<UControlRig>> CurrentControlRigs;
+	TArray<TWeakObjectPtr<UBaseControlRig>> CurrentControlRigs;
 
 	TSharedPtr<STextBlock> TextStatusBlock1;
 	TSharedPtr<STextBlock> TextStatusBlock2;
diff --git a/Engine/Plugins/Animation/ControlRig/Source/ControlRigEditor/Private/EditMode/SControlRigDetails.cpp b/Engine/Plugins/Animation/ControlRig/Source/ControlRigEditor/Private/EditMode/SControlRigDetails.cpp
--- a/Engine/Plugins/Animation/ControlRig/Source/ControlRigEditor/Private/EditMode/SControlRigDetails.cpp
+++ b/Engine/Plugins/Animation/ControlRig/Source/ControlRigEditor/Private/EditMode/SControlRigDetails.cpp
@@ -157,7 +157,7 @@ void FControlRigEditModeGenericDetails::CustomizeDetails(class IDetailLayoutBuil
 				{
 					if (FControlRigEditMode* EditMode = static_cast<FControlRigEditMode*>(ModeTools->GetActiveMode(FControlRigEditMode::ModeName)))
 					{
-						UControlRig* ControlRig = Proxy->ControlRig.Get();
+						UBaseControlRig* ControlRig = Proxy->ControlRig.Get();
 						if (UObject* NestedProxy = EditMode->ControlProxy->FindProxy(ControlRig,ControlElement->GetFName()))
 						{
 							FName PropertyName(NAME_None);
@@ -379,7 +379,7 @@ SControlRigDetails::~SControlRigDetails()
 	//base class handles control rig related cleanup
 }
 
-void SControlRigDetails::HandleControlSelected(UControlRig* Subject, FRigControlElement* InControl, bool bSelected)
+void SControlRigDetails::HandleControlSelected(UBaseControlRig* Subject, FRigControlElement* InControl, bool bSelected)
 {
 	FControlRigBaseDockableView::HandleControlSelected(Subject, InControl, bSelected);
 	UpdateProxies();
@@ -420,7 +420,7 @@ void SControlRigDetails::UpdateProxies()
 		TArray<TWeakObjectPtr<>> IndividualBools;
 		TArray<TWeakObjectPtr<>> IndividualIntegers;
 		TArray<TWeakObjectPtr<>> IndividualEnums;
-		TArray<UControlRig*> ControlRigs = StrongThis->GetControlRigs();
+		TArray<UBaseControlRig*> ControlRigs = StrongThis->GetControlRigs();
 	
 		if (FControlRigEditMode* EditMode = static_cast<FControlRigEditMode*>(StrongThis->ModeTools->GetActiveMode(FControlRigEditMode::ModeName)))
 		{
@@ -662,7 +662,7 @@ bool SControlRigDetails::IsPropertyKeyable(const UClass* InObjectClass, const IP
 	{
 		if (UControlRigControlsProxy* Proxy = Cast< UControlRigControlsProxy>(OuterObjects[0]))
 		{
-			if(UControlRig* ControlRig = Proxy->ControlRig.Get())
+			if(UBaseControlRig* ControlRig = Proxy->ControlRig.Get())
 			{
 				const FRigElementKey Key = FRigElementKey(Proxy->ControlName, ERigElementType::Control);
 				if(const FRigControlElement* ControlElement = ControlRig->GetHierarchy()->Find<FRigControlElement>(Key))
diff --git a/Engine/Plugins/Animation/ControlRig/Source/ControlRigEditor/Private/EditMode/SControlRigDetails.h b/Engine/Plugins/Animation/ControlRig/Source/ControlRigEditor/Private/EditMode/SControlRigDetails.h
--- a/Engine/Plugins/Animation/ControlRig/Source/ControlRigEditor/Private/EditMode/SControlRigDetails.h
+++ b/Engine/Plugins/Animation/ControlRig/Source/ControlRigEditor/Private/EditMode/SControlRigDetails.h
@@ -19,7 +19,7 @@
 #include "DetailWidgetRow.h"
 
 class ISequencer;
-class UControlRig;
+class UBaseControlRig;
 
 class FControlRigEditModeGenericDetails : public IDetailCustomization
 {
@@ -75,7 +75,7 @@ private:
 	void SetEnumDetailsObjects(const TArray<TWeakObjectPtr<>>& InObjects, bool bIsIndividual);
 
 	void UpdateProxies();
-	virtual void HandleControlSelected(UControlRig* Subject, FRigControlElement* InControl, bool bSelected) override;
+	virtual void HandleControlSelected(UBaseControlRig* Subject, FRigControlElement* InControl, bool bSelected) override;
 
 	//these views will hold more than one of the same tyhpe
 	TSharedPtr<IDetailsView> ControlEulerTransformDetailsView;
diff --git a/Engine/Plugins/Animation/ControlRig/Source/ControlRigEditor/Private/EditMode/SControlRigEditModeTools.cpp b/Engine/Plugins/Animation/ControlRig/Source/ControlRigEditor/Private/EditMode/SControlRigEditModeTools.cpp
--- a/Engine/Plugins/Animation/ControlRig/Source/ControlRigEditor/Private/EditMode/SControlRigEditModeTools.cpp
+++ b/Engine/Plugins/Animation/ControlRig/Source/ControlRigEditor/Private/EditMode/SControlRigEditModeTools.cpp
@@ -46,9 +46,9 @@
 //statics to reuse in the UI
 FRigSpacePickerBakeSettings SControlRigEditModeTools::BakeSpaceSettings;
 
-void SControlRigEditModeTools::SetControlRigs(const TArrayView<TWeakObjectPtr<UControlRig>>& InControlRigs)
+void SControlRigEditModeTools::SetControlRigs(const TArrayView<TWeakObjectPtr<UBaseControlRig>>& InControlRigs)
 {
-	for (TWeakObjectPtr<UControlRig>& ControlRig : ControlRigs)
+	for (TWeakObjectPtr<UBaseControlRig>& ControlRig : ControlRigs)
 	{
 		if (ControlRig.IsValid())
 		{
@@ -56,7 +56,7 @@ void SControlRigEditModeTools::SetControlRigs(const TArrayView<TWeakObjectPtr<UC
 		}
 	}
 	ControlRigs = InControlRigs;
-	for (TWeakObjectPtr<UControlRig>& InControlRig : InControlRigs)
+	for (TWeakObjectPtr<UBaseControlRig>& InControlRig : InControlRigs)
 	{
 		if (InControlRig.IsValid())
 		{
@@ -65,7 +65,7 @@ void SControlRigEditModeTools::SetControlRigs(const TArrayView<TWeakObjectPtr<UC
 	}
 
 	//mz todo handle multiple rigs
-	UControlRig* Rig = ControlRigs.Num() > 0 ? ControlRigs[0].Get() : nullptr;
+	UBaseControlRig* Rig = ControlRigs.Num() > 0 ? ControlRigs[0].Get() : nullptr;
 	TArray<TWeakObjectPtr<>> Objects;
 	Objects.Add(Rig);
 	RigOptionsDetailsView->SetObjects(Objects);
@@ -79,7 +79,7 @@ void SControlRigEditModeTools::SetControlRigs(const TArrayView<TWeakObjectPtr<UC
 const URigHierarchy* SControlRigEditModeTools::GetHierarchy() const
 {
 	//mz todo handle multiple rigs
-	UControlRig* Rig = ControlRigs.Num() > 0 ? ControlRigs[0].Get() : nullptr;
+	UBaseControlRig* Rig = ControlRigs.Num() > 0 ? ControlRigs[0].Get() : nullptr;
 	if (Rig)
 	{
 		Rig->GetHierarchy();
@@ -809,7 +809,7 @@ void SControlRigEditModeTools::HandleSelectionChanged(TSharedPtr<FRigTreeElement
 }
 #endif
 
-void SControlRigEditModeTools::OnRigElementSelected(UControlRig* Subject, FRigControlElement* ControlElement, bool bSelected)
+void SControlRigEditModeTools::OnRigElementSelected(UBaseControlRig* Subject, FRigControlElement* ControlElement, bool bSelected)
 {
 #if USE_LOCAL_DETAILS
 	const FRigElementKey Key = ControlElement->GetKey();
@@ -849,8 +849,8 @@ void SControlRigEditModeTools::OnRigElementSelected(UControlRig* Subject, FRigCo
 
 const FRigControlElementCustomization* SControlRigEditModeTools::HandleGetControlElementCustomization(URigHierarchy* InHierarchy, const FRigElementKey& InControlKey)
 {
-	UControlRig* Rig = ControlRigs.Num() > 0 ? ControlRigs[0].Get() : nullptr;
-	for (TWeakObjectPtr<UControlRig>& ControlRig : ControlRigs)
+	UBaseControlRig* Rig = ControlRigs.Num() > 0 ? ControlRigs[0].Get() : nullptr;
+	for (TWeakObjectPtr<UBaseControlRig>& ControlRig : ControlRigs)
 	{
 		if (ControlRig.IsValid() && ControlRig->GetHierarchy() == InHierarchy)
 		{
@@ -866,7 +866,7 @@ void SControlRigEditModeTools::HandleActiveSpaceChanged(URigHierarchy* InHierarc
 
 	if (WeakSequencer.IsValid())
 	{
-		for (TWeakObjectPtr<UControlRig>& ControlRig : ControlRigs)
+		for (TWeakObjectPtr<UBaseControlRig>& ControlRig : ControlRigs)
 		{
 			if (ControlRig.IsValid() && ControlRig->GetHierarchy() == InHierarchy)
 			{
@@ -913,7 +913,7 @@ void SControlRigEditModeTools::HandleSpaceListChanged(URigHierarchy* InHierarchy
 {
 	FScopedTransaction Transaction(LOCTEXT("ChangeControlRigSpace", "Change Control Rig Space"));
 
-	for (TWeakObjectPtr<UControlRig>& ControlRig : ControlRigs)
+	for (TWeakObjectPtr<UBaseControlRig>& ControlRig : ControlRigs)
 	{
 		if (ControlRig.IsValid() && ControlRig->GetHierarchy() == InHierarchy)
 		{
@@ -979,7 +979,7 @@ FReply SControlRigEditModeTools::OnBakeControlsToNewSpaceButtonClicked()
 	}
 
 	bool bNoValidControlRig = true;
-	for (TWeakObjectPtr<UControlRig>& ControlRig : ControlRigs)
+	for (TWeakObjectPtr<UBaseControlRig>& ControlRig : ControlRigs)
 	{
 		if (ControlRig.IsValid() && SpacePickerWidget->GetHierarchy() == ControlRig->GetHierarchy())
 		{
@@ -997,7 +997,7 @@ FReply SControlRigEditModeTools::OnBakeControlsToNewSpaceButtonClicked()
 	{
 		return FReply::Unhandled();
 	}
-	for (TWeakObjectPtr<UControlRig>& ControlRig : ControlRigs)
+	for (TWeakObjectPtr<UBaseControlRig>& ControlRig : ControlRigs)
 	{
 		if (ControlRig.IsValid() && SpacePickerWidget->GetHierarchy() == ControlRig->GetHierarchy())
 		{
@@ -1098,7 +1098,7 @@ FReply SControlRigEditModeTools::HandleAddConstraintClicked()
 
 EVisibility SControlRigEditModeTools::GetRigOptionExpanderVisibility() const
 {
-	for (const TWeakObjectPtr<UControlRig>& ControlRig : ControlRigs)
+	for (const TWeakObjectPtr<UBaseControlRig>& ControlRig : ControlRigs)
 	{
 		if (ControlRig.IsValid())
 		{
@@ -1113,7 +1113,7 @@ EVisibility SControlRigEditModeTools::GetRigOptionExpanderVisibility() const
 
 void SControlRigEditModeTools::OnRigOptionFinishedChange(const FPropertyChangedEvent& PropertyChangedEvent)
 {
-	TArray<TWeakObjectPtr<UControlRig>> ControlRigsCopy = ControlRigs;
+	TArray<TWeakObjectPtr<UBaseControlRig>> ControlRigsCopy = ControlRigs;
 	SetControlRigs(ControlRigsCopy);
 
 	if (FControlRigEditMode* EditMode = static_cast<FControlRigEditMode*>(ModeTools->GetActiveMode(FControlRigEditMode::ModeName)))
diff --git a/Engine/Plugins/Animation/ControlRig/Source/ControlRigEditor/Private/EditMode/SControlRigEditModeTools.h b/Engine/Plugins/Animation/ControlRig/Source/ControlRigEditor/Private/EditMode/SControlRigEditModeTools.h
--- a/Engine/Plugins/Animation/ControlRig/Source/ControlRigEditor/Private/EditMode/SControlRigEditModeTools.h
+++ b/Engine/Plugins/Animation/ControlRig/Source/ControlRigEditor/Private/EditMode/SControlRigEditModeTools.h
@@ -17,7 +17,7 @@ class ISequencer;
 class SControlPicker;
 class SExpandableArea;
 class SRigHierarchyTreeView;
-class UControlRig;
+class UBaseControlRig;
 class URigHierarchy;
 class FToolBarBuilder;
 class FEditorModeTools;
@@ -50,7 +50,7 @@ public:
 	void SetSequencer(TWeakPtr<ISequencer> InSequencer);
 
 	/** Set The Control Rig we are using*/
-	void SetControlRigs(const TArrayView<TWeakObjectPtr<UControlRig>>& InControlRigs);
+	void SetControlRigs(const TArrayView<TWeakObjectPtr<UBaseControlRig>>& InControlRigs);
 
 	/** Returns the hierarchy currently being used */
 	const URigHierarchy* GetHierarchy() const;
@@ -95,7 +95,7 @@ private:
 	TSharedPtr<SRigSpacePickerWidget> SpacePickerWidget;
 
 	/** Storage for control rigs */
-	TArray<TWeakObjectPtr<UControlRig>> ControlRigs;
+	TArray<TWeakObjectPtr<UBaseControlRig>> ControlRigs;
 
 	/** Constraint edition widget. */
 	TSharedPtr<SExpandableArea> ConstraintPickerExpander = nullptr;
@@ -110,7 +110,7 @@ private:
 
 	void HandleModifiedEvent(ERigVMGraphNotifType InNotifType, URigVMGraph* InGraph, UObject* InSubject);
 	void HandleSelectionChanged(TSharedPtr<FRigTreeElement> Selection, ESelectInfo::Type SelectInfo);
-	void OnRigElementSelected(UControlRig* Subject, FRigControlElement* ControlElement, bool bSelected);
+	void OnRigElementSelected(UBaseControlRig* Subject, FRigControlElement* ControlElement, bool bSelected);
 
 	const FRigControlElementCustomization* HandleGetControlElementCustomization(URigHierarchy* InHierarchy, const FRigElementKey& InControlKey);
 	void HandleActiveSpaceChanged(URigHierarchy* InHierarchy, const FRigElementKey& InControlKey, const FRigElementKey& InSpaceKey);
diff --git a/Engine/Plugins/Animation/ControlRig/Source/ControlRigEditor/Private/EditMode/SControlRigOutliner.cpp b/Engine/Plugins/Animation/ControlRig/Source/ControlRigEditor/Private/EditMode/SControlRigOutliner.cpp
--- a/Engine/Plugins/Animation/ControlRig/Source/ControlRigEditor/Private/EditMode/SControlRigOutliner.cpp
+++ b/Engine/Plugins/Animation/ControlRig/Source/ControlRigEditor/Private/EditMode/SControlRigOutliner.cpp
@@ -51,7 +51,7 @@ FRigTreeDisplaySettings FMultiRigTreeDelegates::DefaultDisplaySettings;
 
 uint32 GetTypeHash(const FMultiRigData& Data)
 {
-	return GetTypeHash(TTuple<const UControlRig*, FRigElementKey>(Data.ControlRig.Get(), (Data.Key.IsSet() ? Data.Key.GetValue() : FRigElementKey())));
+	return GetTypeHash(TTuple<const UBaseControlRig*, FRigElementKey>(Data.ControlRig.Get(), (Data.Key.IsSet() ? Data.Key.GetValue() : FRigElementKey())));
 }
 
 FText FMultiRigData::GetName() const
@@ -233,7 +233,7 @@ void SMultiRigHierarchyItem::Construct(const FArguments& InArgs, const TSharedRe
 						//if no key is set then it's the control rig so we get that based upon it's state
 						if (WeakRigTreeElement.Pin()->Data.Key.IsSet() == false)
 						{
-							if (UControlRig* ControlRig = WeakRigTreeElement.Pin()->Data.ControlRig.Get())
+							if (UBaseControlRig* ControlRig = WeakRigTreeElement.Pin()->Data.ControlRig.Get())
 							{
 								if (ControlRig->GetControlsVisible())
 								{
@@ -280,7 +280,7 @@ FReply SMultiRigHierarchyItem::OnGetSelectedClicked()
 	{
 		if (Element->Data.Key.IsSet() == false)
 		{
-			if (UControlRig* ControlRig = Element->Data.ControlRig.Get())
+			if (UBaseControlRig* ControlRig = Element->Data.ControlRig.Get())
 			{
 				FScopedTransaction ScopedTransaction(LOCTEXT("ToggleControlsVisibility", "Toggle Controls Visibility"), !GIsTransacting);
 				ControlRig->Modify();
@@ -313,7 +313,7 @@ TPair<const FSlateBrush*, FSlateColor> SMultiRigHierarchyItem::GetBrushForElemen
 	}
 	else
 	{
-		if (UControlRig* ControlRig = InData.ControlRig.Get())
+		if (UBaseControlRig* ControlRig = InData.ControlRig.Get())
 		{
 			if (ControlRig->GetControlsVisible())
 			{
@@ -466,7 +466,7 @@ bool SMultiRigHierarchyTreeView::AddElement(const FMultiRigData& InData, const F
 	return true;
 }
 
-bool SMultiRigHierarchyTreeView::AddElement(UControlRig* InControlRig, const FRigBaseElement* InElement)
+bool SMultiRigHierarchyTreeView::AddElement(UBaseControlRig* InControlRig, const FRigBaseElement* InElement)
 {
 	check(InControlRig);
 	check(InElement);
@@ -717,9 +717,9 @@ void SMultiRigHierarchyTreeView::RefreshTreeView(bool bRebuildContent)
 
 	if (bRebuildContent)
 	{
-		for (const TWeakObjectPtr<UControlRig>& ControlRigPtr : ControlRigs)
+		for (const TWeakObjectPtr<UBaseControlRig>& ControlRigPtr : ControlRigs)
 		{
-			if (UControlRig* ControlRig  = ControlRigPtr.Get())
+			if (UBaseControlRig* ControlRig  = ControlRigPtr.Get())
 			{
 				FMultiRigData Empty(nullptr, FRigElementKey());
 				FMultiRigData CRData;
@@ -779,9 +779,9 @@ void SMultiRigHierarchyTreeView::RefreshTreeView(bool bRebuildContent)
 	{
 		ClearSelection();
 
-		for (const TWeakObjectPtr<UControlRig>& ControlRigPtr : ControlRigs)
+		for (const TWeakObjectPtr<UBaseControlRig>& ControlRigPtr : ControlRigs)
 		{
-			if (UControlRig* ControlRig = ControlRigPtr.Get())
+			if (UBaseControlRig* ControlRig = ControlRigPtr.Get())
 			{
 				if (const URigHierarchy* Hierarchy = ControlRig->GetHierarchy())
 				{
@@ -856,7 +856,7 @@ TArray<FMultiRigData> SMultiRigHierarchyTreeView::GetSelectedData() const
 TArray<URigHierarchy*> SMultiRigHierarchyTreeView::GetHierarchy() const
 {
 	TArray<URigHierarchy*> RigHierarchy;
-	for (const TWeakObjectPtr<UControlRig>& ControlRig : ControlRigs)
+	for (const TWeakObjectPtr<UBaseControlRig>& ControlRig : ControlRigs)
 	{
 		if (ControlRig.IsValid())
 		{
@@ -866,10 +866,10 @@ TArray<URigHierarchy*> SMultiRigHierarchyTreeView::GetHierarchy() const
 	return RigHierarchy;
 }
 
-void SMultiRigHierarchyTreeView::SetControlRigs(TArrayView < TWeakObjectPtr<UControlRig>>& InControlRigs)
+void SMultiRigHierarchyTreeView::SetControlRigs(TArrayView < TWeakObjectPtr<UBaseControlRig>>& InControlRigs)
 {
 	ControlRigs.SetNum(0);
-	for (TWeakObjectPtr<UControlRig>& ControlRig : InControlRigs)
+	for (TWeakObjectPtr<UBaseControlRig>& ControlRig : InControlRigs)
 	{
 		if (ControlRig.IsValid())
 		{
@@ -997,7 +997,7 @@ void SControlRigOutliner::OnObjectsReplaced(const TMap<UObject*, UObject*>& OldT
 	{
 		if(Pair.Key && Pair.Value)
 		{
-			if (Pair.Key->IsA<UControlRig>() && Pair.Value->IsA<UControlRig>())
+			if (Pair.Key->IsA<UBaseControlRig>() && Pair.Value->IsA<UBaseControlRig>())
 			{
 				bNewControlRig = false;
 				break;
@@ -1020,7 +1020,7 @@ SControlRigOutliner::~SControlRigOutliner()
 	FCoreUObjectDelegates::OnObjectsReplaced.RemoveAll(this);
 }
 
-void SControlRigOutliner::HandleControlSelected(UControlRig* Subject, FRigControlElement* ControlElement, bool bSelected)
+void SControlRigOutliner::HandleControlSelected(UBaseControlRig* Subject, FRigControlElement* ControlElement, bool bSelected)
 {
 	FControlRigBaseDockableView::HandleControlSelected(Subject, ControlElement, bSelected);
 	const FRigElementKey Key = ControlElement->GetKey();
@@ -1056,7 +1056,7 @@ void SControlRigOutliner::HandleSelectionChanged(TSharedPtr<FMultiRigTreeElement
 		return;
 	}
 	const TArray<FMultiRigData> NewSelection = HierarchyTreeView->GetTreeView()->GetSelectedData();
-	TMap<UControlRig*, TArray<FRigElementKey>> SelectedRigAndKeys;
+	TMap<UBaseControlRig*, TArray<FRigElementKey>> SelectedRigAndKeys;
 	for (const FMultiRigData& Data : NewSelection)
 	{
 		if (Data.ControlRig.IsValid() && (Data.Key.IsSet() && Data.Key.GetValue() != FRigElementKey()))
@@ -1081,9 +1081,9 @@ void SControlRigOutliner::HandleSelectionChanged(TSharedPtr<FMultiRigTreeElement
 	{
 		if (EditMode)
 		{		
-			TMap<UControlRig*, TArray<FRigElementKey>> SelectedControls;
+			TMap<UBaseControlRig*, TArray<FRigElementKey>> SelectedControls;
 			EditMode->GetAllSelectedControls(SelectedControls);
-			for (TPair<UControlRig*, TArray<FRigElementKey>>& CurrentSelection : SelectedControls)
+			for (TPair<UBaseControlRig*, TArray<FRigElementKey>>& CurrentSelection : SelectedControls)
 			{
 				if (CurrentSelection.Key)
 				{
@@ -1093,7 +1093,7 @@ void SControlRigOutliner::HandleSelectionChanged(TSharedPtr<FMultiRigTreeElement
 		}
 	}
 
-	for(TPair<UControlRig*, TArray<FRigElementKey>>& RigAndKeys: SelectedRigAndKeys)
+	for(TPair<UBaseControlRig*, TArray<FRigElementKey>>& RigAndKeys: SelectedRigAndKeys)
 	{ 
 		const URigHierarchy* Hierarchy = RigAndKeys.Key->GetHierarchy();
 		if (Hierarchy)
@@ -1122,17 +1122,17 @@ void SControlRigOutliner::SetEditMode(FControlRigEditMode& InEditMode)
 	ModeTools = InEditMode.GetModeManager();
 	if (FControlRigEditMode* EditMode = static_cast<FControlRigEditMode*>(ModeTools->GetActiveMode(FControlRigEditMode::ModeName)))
 	{
-		TArrayView<TWeakObjectPtr<UControlRig>> ControlRigs = EditMode->GetControlRigs();
+		TArrayView<TWeakObjectPtr<UBaseControlRig>> ControlRigs = EditMode->GetControlRigs();
 		HierarchyTreeView->GetTreeView()->SetControlRigs(ControlRigs); //will refresh tree
 	}
 }
 
-void SControlRigOutliner::HandleControlAdded(UControlRig* ControlRig, bool bIsAdded)
+void SControlRigOutliner::HandleControlAdded(UBaseControlRig* ControlRig, bool bIsAdded)
 {
 	FControlRigBaseDockableView::HandleControlAdded(ControlRig, bIsAdded);
 	if (FControlRigEditMode* EditMode = static_cast<FControlRigEditMode*>(ModeTools->GetActiveMode(FControlRigEditMode::ModeName)))
 	{
-		TArrayView<TWeakObjectPtr<UControlRig>> ControlRigs = EditMode->GetControlRigs();
+		TArrayView<TWeakObjectPtr<UBaseControlRig>> ControlRigs = EditMode->GetControlRigs();
 		HierarchyTreeView->GetTreeView()->SetControlRigs(ControlRigs); //will refresh tree
 	}
 }
diff --git a/Engine/Plugins/Animation/ControlRig/Source/ControlRigEditor/Private/EditMode/SControlRigOutliner.h b/Engine/Plugins/Animation/ControlRig/Source/ControlRigEditor/Private/EditMode/SControlRigOutliner.h
--- a/Engine/Plugins/Animation/ControlRig/Source/ControlRigEditor/Private/EditMode/SControlRigOutliner.h
+++ b/Engine/Plugins/Animation/ControlRig/Source/ControlRigEditor/Private/EditMode/SControlRigOutliner.h
@@ -77,14 +77,14 @@ class FMultiRigData
 
 public:
 	FMultiRigData() {};
-	FMultiRigData(UControlRig* InControlRig, FRigElementKey InKey) : ControlRig(InControlRig), Key(InKey) {};
+	FMultiRigData(UBaseControlRig* InControlRig, FRigElementKey InKey) : ControlRig(InControlRig), Key(InKey) {};
 	FText GetName() const;
 	FText GetDisplayName() const;
 	bool operator == (const FMultiRigData & Other) const;
 	bool IsValid() const;
 	URigHierarchy* GetHierarchy() const;
 public:
-	TWeakObjectPtr<UControlRig> ControlRig;
+	TWeakObjectPtr<UBaseControlRig> ControlRig;
 	TOptional<FRigElementKey> Key;
 };
 
@@ -177,7 +177,7 @@ public:
 	static TSharedPtr<FMultiRigTreeElement> FindElement(const FMultiRigData& InData, TSharedPtr<FMultiRigTreeElement> CurrentItem);
 
 	bool AddElement(const FMultiRigData& InData, const FMultiRigData& InParentData);
-	bool AddElement(UControlRig* InControlRig, const FRigBaseElement* InElement);
+	bool AddElement(UBaseControlRig* InControlRig, const FRigBaseElement* InElement);
 	bool ReparentElement(const FMultiRigData& InData, const FMultiRigData& InParentData);
 	bool RemoveElement(const FMultiRigData& InData);
 	void RefreshTreeView(bool bRebuildContent = true);
@@ -190,7 +190,7 @@ public:
 	FMultiRigTreeDelegates& GetTreeDelegates() { return Delegates; }
 
 	TArray<URigHierarchy*> GetHierarchy() const;
-	void SetControlRigs(TArrayView < TWeakObjectPtr<UControlRig>>& InControlRigs);
+	void SetControlRigs(TArrayView < TWeakObjectPtr<UBaseControlRig>>& InControlRigs);
 
 private:
 
@@ -210,7 +210,7 @@ private:
 
 	friend class SRigHierarchy;
 
-	TArray <TWeakObjectPtr<UControlRig>> ControlRigs;
+	TArray <TWeakObjectPtr<UBaseControlRig>> ControlRigs;
 };
 
 class SSearchableMultiRigHierarchyTreeView : public SCompoundWidget
@@ -249,8 +249,8 @@ class SControlRigOutliner : public FControlRigBaseDockableView, public SCompound
 	//FControlRigBaseDockableView overrides
 	virtual void SetEditMode(FControlRigEditMode& InEditMode) override;
 private:
-	virtual void HandleControlAdded(UControlRig* ControlRig, bool bIsAdded) override;
-	virtual void HandleControlSelected(UControlRig* Subject, FRigControlElement* InControl, bool bSelected) override;
+	virtual void HandleControlAdded(UBaseControlRig* ControlRig, bool bIsAdded) override;
+	virtual void HandleControlSelected(UBaseControlRig* Subject, FRigControlElement* InControl, bool bSelected) override;
 
 	void OnObjectsReplaced(const TMap<UObject*, UObject*>& OldToNewInstanceMap);
 
diff --git a/Engine/Plugins/Animation/ControlRig/Source/ControlRigEditor/Private/EditMode/SControlRigSnapper.cpp b/Engine/Plugins/Animation/ControlRig/Source/ControlRigEditor/Private/EditMode/SControlRigSnapper.cpp
--- a/Engine/Plugins/Animation/ControlRig/Source/ControlRigEditor/Private/EditMode/SControlRigSnapper.cpp
+++ b/Engine/Plugins/Animation/ControlRig/Source/ControlRigEditor/Private/EditMode/SControlRigSnapper.cpp
@@ -501,9 +501,9 @@ void SControlRigSnapper::SetStartEndFrames()
 FControlRigSnapperSelection SControlRigSnapper::GetSelection(bool bGetAll) 
 {
 	FControlRigSnapperSelection Selection;
-	TArray<UControlRig*> ControlRigs;
+	TArray<UBaseControlRig*> ControlRigs;
 	GetControlRigs(ControlRigs);
-	for (UControlRig* ControlRig : ControlRigs)
+	for (UBaseControlRig* ControlRig : ControlRigs)
 	{
 		if (ControlRig)
 		{
@@ -550,15 +550,15 @@ FControlRigSnapperSelection SControlRigSnapper::GetSelection(bool bGetAll)
 	return Selection;
 }
 
-void SControlRigSnapper::GetControlRigs(TArray<UControlRig*>& OutControlRigs) const
+void SControlRigSnapper::GetControlRigs(TArray<UBaseControlRig*>& OutControlRigs) const
 {
 	OutControlRigs.SetNum(0);
 	if (FControlRigEditMode* EditMode = static_cast<FControlRigEditMode*>(GLevelEditorModeTools().GetActiveMode(FControlRigEditMode::ModeName)))
 	{
-		TArrayView<TWeakObjectPtr<UControlRig>> ControlRigPtrs = EditMode->GetControlRigs();
-		for (TWeakObjectPtr<UControlRig>& ControlRigPtr : ControlRigPtrs)
+		TArrayView<TWeakObjectPtr<UBaseControlRig>> ControlRigPtrs = EditMode->GetControlRigs();
+		for (TWeakObjectPtr<UBaseControlRig>& ControlRigPtr : ControlRigPtrs)
 		{
-			if (UControlRig* ControlRig = ControlRigPtr.Get())
+			if (UBaseControlRig* ControlRig = ControlRigPtr.Get())
 			{
 				OutControlRigs.Add(ControlRig);
 			}
diff --git a/Engine/Plugins/Animation/ControlRig/Source/ControlRigEditor/Private/EditMode/SControlRigSnapper.h b/Engine/Plugins/Animation/ControlRig/Source/ControlRigEditor/Private/EditMode/SControlRigSnapper.h
--- a/Engine/Plugins/Animation/ControlRig/Source/ControlRigEditor/Private/EditMode/SControlRigSnapper.h
+++ b/Engine/Plugins/Animation/ControlRig/Source/ControlRigEditor/Private/EditMode/SControlRigSnapper.h
@@ -13,7 +13,7 @@
 #include "IDetailsView.h"
 #include "MovieSceneSequenceID.h"
 
-class UControlRig;
+class UBaseControlRig;
 class ISequencer;
 class AActor;
 
@@ -30,7 +30,7 @@ class SControlRigSnapper : public SCompoundWidget
 
 private:
 
-	void GetControlRigs(TArray<UControlRig*>& OutControlRigs) const;
+	void GetControlRigs(TArray<UBaseControlRig*>& OutControlRigs) const;
 
 	/*
 	* Delegates and Helpers
diff --git a/Engine/Plugins/Animation/ControlRig/Source/ControlRigEditor/Private/EditMode/SControlRigSpacePicker.cpp b/Engine/Plugins/Animation/ControlRig/Source/ControlRigEditor/Private/EditMode/SControlRigSpacePicker.cpp
--- a/Engine/Plugins/Animation/ControlRig/Source/ControlRigEditor/Private/EditMode/SControlRigSpacePicker.cpp
+++ b/Engine/Plugins/Animation/ControlRig/Source/ControlRigEditor/Private/EditMode/SControlRigSpacePicker.cpp
@@ -108,10 +108,10 @@ SControlRigSpacePicker::~SControlRigSpacePicker()
 	//base class handles control rig related cleanup
 }
 
-UControlRig* SControlRigSpacePicker::GetControlRig()
+UBaseControlRig* SControlRigSpacePicker::GetControlRig()
 {
-	TArray<UControlRig*> ControlRigs = GetControlRigs();
-	for(UControlRig* ControlRig: ControlRigs)
+	TArray<UBaseControlRig*> ControlRigs = GetControlRigs();
+	for(UBaseControlRig* ControlRig: ControlRigs)
 	{
 		if (ControlRig)
 		{
@@ -125,10 +125,10 @@ UControlRig* SControlRigSpacePicker::GetControlRig()
 	return nullptr;
 }
 
-void SControlRigSpacePicker::HandleControlSelected(UControlRig* Subject, FRigControlElement* ControlElement, bool bSelected)
+void SControlRigSpacePicker::HandleControlSelected(UBaseControlRig* Subject, FRigControlElement* ControlElement, bool bSelected)
 {
 	FControlRigBaseDockableView::HandleControlSelected(Subject, ControlElement, bSelected);
-	if (UControlRig* ControlRig = GetControlRig())
+	if (UBaseControlRig* ControlRig = GetControlRig())
 	{
 		// get the selected controls
 		TArray<FRigElementKey> SelectedControls = ControlRig->GetHierarchy()->GetSelectedKeys(ERigElementType::Control);
@@ -143,7 +143,7 @@ void SControlRigSpacePicker::HandleControlSelected(UControlRig* Subject, FRigCon
 
 const FRigControlElementCustomization* SControlRigSpacePicker::HandleGetControlElementCustomization(URigHierarchy* InHierarchy, const FRigElementKey& InControlKey)
 {
-	if (UControlRig* ControlRig = GetControlRig())
+	if (UBaseControlRig* ControlRig = GetControlRig())
 	{
 		return ControlRig->GetControlCustomization(InControlKey);
 	}
@@ -156,7 +156,7 @@ void SControlRigSpacePicker::HandleActiveSpaceChanged(URigHierarchy* InHierarchy
 
 	if (ISequencer* Sequencer = GetSequencer())
 	{
-		if (UControlRig* ControlRig = GetControlRig())
+		if (UBaseControlRig* ControlRig = GetControlRig())
 		{
 			FString FailureReason;
 			URigHierarchy::TElementDependencyMap DependencyMap = InHierarchy->GetDependenciesForVM(ControlRig->GetVM());
@@ -194,7 +194,7 @@ void SControlRigSpacePicker::HandleActiveSpaceChanged(URigHierarchy* InHierarchy
 void SControlRigSpacePicker::HandleSpaceListChanged(URigHierarchy* InHierarchy, const FRigElementKey& InControlKey,
 	const TArray<FRigElementKey>& InSpaceList)
 {
-	if (UControlRig* ControlRig = GetControlRig())
+	if (UBaseControlRig* ControlRig = GetControlRig())
 	{
 		if (const FRigControlElement* ControlElement = InHierarchy->Find<FRigControlElement>(InControlKey))
 		{
@@ -253,7 +253,7 @@ FReply SControlRigSpacePicker::OnBakeControlsToNewSpaceButtonClicked()
 	{
 		return FReply::Unhandled();
 	}
-	UControlRig* ControlRig = GetControlRig();
+	UBaseControlRig* ControlRig = GetControlRig();
 
 	FRigSpacePickerBakeSettings Settings;
 	//Find default target space, just use first control and find space at current sequencer time
diff --git a/Engine/Plugins/Animation/ControlRig/Source/ControlRigEditor/Private/EditMode/SControlRigSpacePicker.h b/Engine/Plugins/Animation/ControlRig/Source/ControlRigEditor/Private/EditMode/SControlRigSpacePicker.h
--- a/Engine/Plugins/Animation/ControlRig/Source/ControlRigEditor/Private/EditMode/SControlRigSpacePicker.h
+++ b/Engine/Plugins/Animation/ControlRig/Source/ControlRigEditor/Private/EditMode/SControlRigSpacePicker.h
@@ -29,7 +29,7 @@ class SControlRigSpacePicker : public SCompoundWidget, public FControlRigBaseDoc
 	void Construct(const FArguments& InArgs, FControlRigEditMode& InEditMode);
 
 private:
-	virtual void HandleControlSelected(UControlRig* Subject, FRigControlElement* InControl, bool bSelected) override;
+	virtual void HandleControlSelected(UBaseControlRig* Subject, FRigControlElement* InControl, bool bSelected) override;
 
 	/** Space picker widget*/
 	TSharedPtr<SRigSpacePickerWidget> SpacePickerWidget;
@@ -43,7 +43,7 @@ private:
 	EVisibility GetAddSpaceButtonVisibility() const;
 
 	//for now picker works off of one ControlRig, this function gets the first control rig with a selection
-	UControlRig* GetControlRig();
+	UBaseControlRig* GetControlRig();
 
 };
 
