Control Rig: Rename ControlRig to BaseControlRig and created a specialized ControlRig

diff --git a/Engine/Plugins/Animation/ControlRig/Source/ControlRigEditor/Private/EditMode/SControlRigUpdatePose.cpp b/Engine/Plugins/Animation/ControlRig/Source/ControlRigEditor/Private/EditMode/SControlRigUpdatePose.cpp
--- a/Engine/Plugins/Animation/ControlRig/Source/ControlRigEditor/Private/EditMode/SControlRigUpdatePose.cpp
+++ b/Engine/Plugins/Animation/ControlRig/Source/ControlRigEditor/Private/EditMode/SControlRigUpdatePose.cpp
@@ -81,7 +81,7 @@ private:
 		FControlRigEditMode* ControlRigEditMode = static_cast<FControlRigEditMode*>(GLevelEditorModeTools().GetActiveMode(FControlRigEditMode::ModeName));
 		if (PoseAsset.IsValid() && ControlRigEditMode)
 		{
-			TMap<UControlRig*, TArray<FRigElementKey>> AllSelectedControls;
+			TMap<UBaseControlRig*, TArray<FRigElementKey>> AllSelectedControls;
 			ControlRigEditMode->GetAllSelectedControls(AllSelectedControls);
 			return (AllSelectedControls.Num() > 0);
 		}
@@ -95,11 +95,11 @@ private:
 		FControlRigEditMode* ControlRigEditMode = static_cast<FControlRigEditMode*>(GLevelEditorModeTools().GetActiveMode(FControlRigEditMode::ModeName));
 		if (PoseAsset.IsValid() && ControlRigEditMode)
 		{
-			TMap<UControlRig*, TArray<FRigElementKey>> AllSelectedControls;
+			TMap<UBaseControlRig*, TArray<FRigElementKey>> AllSelectedControls;
 			ControlRigEditMode->GetAllSelectedControls(AllSelectedControls);
 			if (AllSelectedControls.Num() == 1)
 			{
-				TArray<UControlRig*> ControlRigs;
+				TArray<UBaseControlRig*> ControlRigs;
 				AllSelectedControls.GenerateKeyArray(ControlRigs);
 				const FScopedTransaction Transaction(NSLOCTEXT("ControlRig", "UpdatePose", "Update Pose"));
 				PoseAsset->Modify();
diff --git a/Engine/Plugins/Animation/ControlRig/Source/ControlRigEditor/Private/EditMode/SCurveControlContainer.cpp b/Engine/Plugins/Animation/ControlRig/Source/ControlRigEditor/Private/EditMode/SCurveControlContainer.cpp
--- a/Engine/Plugins/Animation/ControlRig/Source/ControlRigEditor/Private/EditMode/SCurveControlContainer.cpp
+++ b/Engine/Plugins/Animation/ControlRig/Source/ControlRigEditor/Private/EditMode/SCurveControlContainer.cpp
@@ -134,7 +134,7 @@ float SCurveControlListRow::GetValue() const
 //////////////////////////////////////////////////////////////////////////
 // SCurveControlContainer
 
-void SCurveControlContainer::Construct(const FArguments& InArgs, UControlRig* InControlRig)
+void SCurveControlContainer::Construct(const FArguments& InArgs, UBaseControlRig* InControlRig)
 {
 	ChildSlot
 	[
@@ -183,7 +183,7 @@ void SCurveControlContainer::Construct(const FArguments& InArgs, UControlRig* In
 }
 
 
-void SCurveControlContainer::SetControlRig(UControlRig* InControlRig)
+void SCurveControlContainer::SetControlRig(UBaseControlRig* InControlRig)
 {
 
 	if (ControlRig.IsValid())
@@ -350,7 +350,7 @@ void SCurveControlContainer::OnSelectionChanged(FDisplayedCurveControlInfoPtr Se
 		}
 	}
 }
-void SCurveControlContainer::OnRigElementSelected(UControlRig* Subject, FRigControlElement* ControlElement, bool bSelected)
+void SCurveControlContainer::OnRigElementSelected(UBaseControlRig* Subject, FRigControlElement* ControlElement, bool bSelected)
 {
 	for(const FDisplayedCurveControlInfoPtr& Item : CurveControlList)
 	{
diff --git a/Engine/Plugins/Animation/ControlRig/Source/ControlRigEditor/Private/EditMode/SCurveControlContainer.h b/Engine/Plugins/Animation/ControlRig/Source/ControlRigEditor/Private/EditMode/SCurveControlContainer.h
--- a/Engine/Plugins/Animation/ControlRig/Source/ControlRigEditor/Private/EditMode/SCurveControlContainer.h
+++ b/Engine/Plugins/Animation/ControlRig/Source/ControlRigEditor/Private/EditMode/SCurveControlContainer.h
@@ -14,7 +14,7 @@
 
 class SInlineEditableTextBlock;
 struct FAssetData;
-class UControlRig;
+class UBaseControlRig;
 
 class FDisplayedCurveControlInfo
 {
@@ -129,7 +129,7 @@ public:
 	* @param InArgs - Arguments passed from Slate
 	*
 	*/
-	void Construct( const FArguments& InArgs, UControlRig* InControlRig);
+	void Construct( const FArguments& InArgs, UBaseControlRig* InControlRig);
 
 	/**
 	* Destructor - resets the animation curve
@@ -173,7 +173,7 @@ public:
 	FText& GetFilterText() { return FilterText; }
 
 	void RefreshCurveList();
-	void SetControlRig(UControlRig* InControlRig);
+	void SetControlRig(UBaseControlRig* InControlRig);
 private:
 
 	/**
@@ -190,7 +190,7 @@ private:
 
 	void OnSelectionChanged(FDisplayedCurveControlInfoPtr Selection, ESelectInfo::Type SelectInfo);
 
-	void OnRigElementSelected(UControlRig* Subject, FRigControlElement* ControlElement, bool bSelected);
+	void OnRigElementSelected(UBaseControlRig* Subject, FRigControlElement* ControlElement, bool bSelected);
 
 	/** Box to filter to a specific morph target name */
 	TSharedPtr<SSearchBox>	NameFilterBox;
@@ -202,7 +202,7 @@ private:
 	TSharedPtr<SCurveControlListType> CurveControlListView;
 
 	/** ControlRig */
-	TWeakObjectPtr<UControlRig> ControlRig;
+	TWeakObjectPtr<UBaseControlRig> ControlRig;
 
 	/** Current text typed into NameFilterBox */
 	FText FilterText;
diff --git a/Engine/Plugins/Animation/ControlRig/Source/ControlRigEditor/Private/Editor/ControlRigContextMenuContext.cpp b/Engine/Plugins/Animation/ControlRig/Source/ControlRigEditor/Private/Editor/ControlRigContextMenuContext.cpp
--- a/Engine/Plugins/Animation/ControlRig/Source/ControlRigEditor/Private/Editor/ControlRigContextMenuContext.cpp
+++ b/Engine/Plugins/Animation/ControlRig/Source/ControlRigEditor/Private/Editor/ControlRigContextMenuContext.cpp
@@ -39,11 +39,11 @@ UControlRigBlueprint* UControlRigContextMenuContext::GetControlRigBlueprint() co
 	return nullptr;
 }
 
-UControlRig* UControlRigContextMenuContext::GetControlRig() const
+UBaseControlRig* UControlRigContextMenuContext::GetControlRig() const
 {
 	if (UControlRigBlueprint* RigBlueprint = GetControlRigBlueprint())
 	{
-		if (UControlRig* ControlRig = Cast<UControlRig>(RigBlueprint->GetObjectBeingDebugged()))
+		if (UBaseControlRig* ControlRig = Cast<UBaseControlRig>(RigBlueprint->GetObjectBeingDebugged()))
 		{
 			return ControlRig;
 		}
diff --git a/Engine/Plugins/Animation/ControlRig/Source/ControlRigEditor/Private/Editor/ControlRigContextMenuContext.h b/Engine/Plugins/Animation/ControlRig/Source/ControlRigEditor/Private/Editor/ControlRigContextMenuContext.h
--- a/Engine/Plugins/Animation/ControlRig/Source/ControlRigEditor/Private/Editor/ControlRigContextMenuContext.h
+++ b/Engine/Plugins/Animation/ControlRig/Source/ControlRigEditor/Private/Editor/ControlRigContextMenuContext.h
@@ -8,7 +8,7 @@
 
 #include "ControlRigContextMenuContext.generated.h"
 
-class UControlRig;
+class UBaseControlRig;
 class UControlRigBlueprint;
 class FControlRigEditor;
 class URigVMGraph;
@@ -121,7 +121,7 @@ public:
 	
 	/** Get the active control rig instance in the viewport */
 	UFUNCTION(BlueprintCallable, Category = ControlRigEditorExtensions)
-    UControlRig* GetControlRig() const;
+    UBaseControlRig* GetControlRig() const;
 
 	/** Returns true if either alt key is down */
 	UFUNCTION(BlueprintCallable, Category = ControlRigEditorExtensions)
diff --git a/Engine/Plugins/Animation/ControlRig/Source/ControlRigEditor/Private/Editor/ControlRigEditor.cpp b/Engine/Plugins/Animation/ControlRig/Source/ControlRigEditor/Private/Editor/ControlRigEditor.cpp
--- a/Engine/Plugins/Animation/ControlRig/Source/ControlRigEditor/Private/Editor/ControlRigEditor.cpp
+++ b/Engine/Plugins/Animation/ControlRig/Source/ControlRigEditor/Private/Editor/ControlRigEditor.cpp
@@ -187,16 +187,16 @@ UControlRigBlueprint* FControlRigEditor::GetControlRigBlueprint() const
 	return Cast<UControlRigBlueprint>(GetRigVMBlueprint());
 }
 
-UControlRig* FControlRigEditor::GetControlRig() const
+UBaseControlRig* FControlRigEditor::GetControlRig() const
 {
-	return Cast<UControlRig>(GetRigVMHost());
+	return Cast<UBaseControlRig>(GetRigVMHost());
 }
 
 URigHierarchy* FControlRigEditor::GetHierarchyBeingDebugged() const
 {
 	if(UControlRigBlueprint* RigBlueprint = GetControlRigBlueprint())
 	{
-		if(UControlRig* RigBeingDebugged = Cast<UControlRig>(RigBlueprint->GetObjectBeingDebugged()))
+		if(UBaseControlRig* RigBeingDebugged = Cast<UBaseControlRig>(RigBlueprint->GetObjectBeingDebugged()))
 		{
 			return RigBeingDebugged->GetHierarchy();
 		}
@@ -771,7 +771,7 @@ void FControlRigEditor::SetEventQueue(TArray<FName> InEventQueue, bool bCompile)
 
 	FRigVMEditor::SetEventQueue(InEventQueue, bCompile);
 
-	if (UControlRig* ControlRig = GetControlRig())
+	if (UBaseControlRig* ControlRig = GetControlRig())
 	{
 		if (InEventQueue.Num() > 0)
 		{
@@ -894,8 +894,8 @@ void FControlRigEditor::HandleSetObjectBeingDebugged(UObject* InObject)
 {
 	FRigVMEditor::HandleSetObjectBeingDebugged(InObject);
 	
-	UControlRig* DebuggedControlRig = Cast<UControlRig>(InObject);
-	if(UControlRig* PreviouslyDebuggedControlRig = Cast<UControlRig>(GetBlueprintObj()->GetObjectBeingDebugged()))
+	UBaseControlRig* DebuggedControlRig = Cast<UBaseControlRig>(InObject);
+	if(UBaseControlRig* PreviouslyDebuggedControlRig = Cast<UBaseControlRig>(GetBlueprintObj()->GetObjectBeingDebugged()))
 	{
 		if(!PreviouslyDebuggedControlRig->HasAnyFlags(RF_BeginDestroyed))
 		{
@@ -1125,7 +1125,7 @@ void FControlRigEditor::Compile()
 				const FTransform Transform = ControlRigBlueprint->Hierarchy->GetInitialLocalTransform(ControlElement->GetIndex());
 
 				/*/
-				if (UControlRig* ControlRig = GetControlRig())
+				if (UBaseControlRig* ControlRig = GetControlRig())
 				{
 					ControlRig->Modify();
 					ControlRig->GetControlHierarchy().SetLocalTransform(Control.Index, Transform);
@@ -1227,7 +1227,7 @@ void FControlRigEditor::HandleVMCompiledEvent(UObject* InCompiledObject, URigVM*
 
 	if(UControlRigBlueprint* ControlRigBlueprint = GetControlRigBlueprint())
 	{
-		if(UControlRig* ControlRig = InVM->GetTypedOuter<UControlRig>())
+		if(UBaseControlRig* ControlRig = InVM->GetTypedOuter<UBaseControlRig>())
 		{
 			ControlRigBlueprint->UpdateElementKeyRedirector(ControlRig);
 		}
@@ -1240,9 +1240,9 @@ void FControlRigEditor::SaveAsset_Execute()
 
 	// Save the new state of the hierarchy in the default object, so that it has the correct values on load
 	UControlRigBlueprint* RigBlueprint = Cast<UControlRigBlueprint>(GetBlueprintObj());
-	if(const UControlRig* ControlRig = GetControlRig())
+	if(const UBaseControlRig* ControlRig = GetControlRig())
 	{
-		UControlRig* CDO = ControlRig->GetClass()->GetDefaultObject<UControlRig>();
+		UBaseControlRig* CDO = ControlRig->GetClass()->GetDefaultObject<UBaseControlRig>();
 		CDO->DynamicHierarchy->CopyHierarchy(RigBlueprint->Hierarchy);
 		RigBlueprint->UpdateElementKeyRedirector(CDO);
 	}
@@ -1258,9 +1258,9 @@ void FControlRigEditor::SaveAssetAs_Execute()
 
 	// Save the new state of the hierarchy in the default object, so that it has the correct values on load
 	UControlRigBlueprint* RigBlueprint = Cast<UControlRigBlueprint>(GetBlueprintObj());
-	if(const UControlRig* ControlRig = GetControlRig())
+	if(const UBaseControlRig* ControlRig = GetControlRig())
 	{
-		UControlRig* CDO = ControlRig->GetClass()->GetDefaultObject<UControlRig>();
+		UBaseControlRig* CDO = ControlRig->GetClass()->GetDefaultObject<UBaseControlRig>();
 		CDO->DynamicHierarchy->CopyHierarchy(RigBlueprint->Hierarchy);
 		RigBlueprint->UpdateElementKeyRedirector(CDO);
 	}
@@ -1347,7 +1347,7 @@ void FControlRigEditor::PostTransaction(bool bSuccess, const FTransaction* Trans
 			GetPersonaToolkit()->SetPreviewMesh(PreviewMesh, true);
 		}
 
-		if (UControlRig* DebuggedControlRig = Cast<UControlRig>(RigBlueprint->GetObjectBeingDebugged()))
+		if (UBaseControlRig* DebuggedControlRig = Cast<UBaseControlRig>(RigBlueprint->GetObjectBeingDebugged()))
 		{
 			if(URigHierarchy* Hierarchy = DebuggedControlRig->GetHierarchy())
 			{
@@ -1479,7 +1479,7 @@ void FControlRigEditor::HandleVMExecutedEvent(URigVMHost* InHost, const FName& I
 		}
 
 		// update transient controls on nodes / pins
-		if(UControlRig* DebuggedControlRig = Cast<UControlRig>(ControlRigBP->GetObjectBeingDebugged()))
+		if(UBaseControlRig* DebuggedControlRig = Cast<UBaseControlRig>(ControlRigBP->GetObjectBeingDebugged()))
 		{
 			if(!DebuggedControlRig->RigUnitManipulationInfos.IsEmpty())
 			{
@@ -1511,7 +1511,7 @@ void FControlRigEditor::HandleVMExecutedEvent(URigVMHost* InHost, const FName& I
 							}
 						}
 				
-						FRigUnit* UnitInstance = UControlRig::GetRigUnitInstanceFromScope(NodeInstance);
+						FRigUnit* UnitInstance = UBaseControlRig::GetRigUnitInstanceFromScope(NodeInstance);
 						UnitInstance->UpdateHierarchyForDirectManipulation(Node, NodeInstance, ExecuteContext, ManipulationInfo);
 						ManipulationInfo->bInitialized = true;
 						UnitInstance->PerformDebugDrawingForDirectManipulation(Node, NodeInstance, ExecuteContext, ManipulationInfo);
@@ -1536,14 +1536,14 @@ void FControlRigEditor::Tick(float DeltaTime)
 	// tick the control rig in case we don't have skeletal mesh
 	if (UControlRigBlueprint* Blueprint = GetControlRigBlueprint())
 	{
-		UControlRig* ControlRig = GetControlRig();
+		UBaseControlRig* ControlRig = GetControlRig();
 		if (Blueprint->GetPreviewMesh() == nullptr && 
 			ControlRig != nullptr && 
 			bExecutionControlRig)
 		{
 			{
 				// prevent transient controls from getting reset
-				UControlRig::FTransientControlPoseScope	PoseScope(ControlRig);
+				UBaseControlRig::FTransientControlPoseScope	PoseScope(ControlRig);
 				// reset transforms here to prevent additive transforms from accumulating to INF
 				ControlRig->GetHierarchy()->ResetPoseToInitial(ERigElementType::Bone);
 			}
@@ -2075,7 +2075,7 @@ void FControlRigEditor::OnToolbarDrawAxesOnSelectionChanged(ECheckBoxState InNew
 
 bool FControlRigEditor::IsToolbarDrawNullsEnabled() const
 {
-	if (const UControlRig* ControlRig = GetControlRig())
+	if (const UBaseControlRig* ControlRig = GetControlRig())
 	{
 		if (!ControlRig->IsConstructionModeEnabled())
 		{
@@ -2184,14 +2184,14 @@ void FControlRigEditor::UpdateRigVMHost()
 	{
 		UControlRigSkeletalMeshComponent* EditorSkelComp = Cast<UControlRigSkeletalMeshComponent>(GetPersonaToolkit()->GetPreviewScene()->GetPreviewMeshComponent());
 		UControlRigLayerInstance* AnimInstance = Cast<UControlRigLayerInstance>(EditorSkelComp->GetAnimInstance());
-		UControlRig* ControlRig = GetControlRig();
+		UBaseControlRig* ControlRig = GetControlRig();
 
 		if (AnimInstance && ControlRig)
 		{
  			PreviewInstance = Cast<UAnimPreviewInstance>(AnimInstance->GetSourceAnimInstance());
 			ControlRig->PreviewInstance = PreviewInstance;
 
-			if (UControlRig* CDO = Cast<UControlRig>(Class->GetDefaultObject()))
+			if (UBaseControlRig* CDO = Cast<UBaseControlRig>(Class->GetDefaultObject()))
 			{
 				CDO->ShapeLibraries = GetControlRigBlueprint()->ShapeLibraries;
 			}
@@ -2255,7 +2255,7 @@ void FControlRigEditor::CacheNameLists()
 			}
 
 			const TArray<TSoftObjectPtr<UControlRigShapeLibrary>>* ShapeLibraries = &ControlRigBP->ShapeLibraries;
-			if(const UControlRig* DebuggedControlRig = Hierarchy->GetTypedOuter<UControlRig>())
+			if(const UBaseControlRig* DebuggedControlRig = Hierarchy->GetTypedOuter<UBaseControlRig>())
 			{
 				ShapeLibraries = &DebuggedControlRig->GetShapeLibraries();
 			}
@@ -2275,7 +2275,7 @@ void FControlRigEditor::HandlePreviewMeshChanged(USkeletalMesh* InOldSkeletalMes
 			ControlRigBP->SetPreviewMesh(InNewSkeletalMesh);
 			UpdateRigVMHost();
 			
-			if(UControlRig* DebuggedControlRig = Cast<UControlRig>(GetBlueprintObj()->GetObjectBeingDebugged()))
+			if(UBaseControlRig* DebuggedControlRig = Cast<UBaseControlRig>(GetBlueprintObj()->GetObjectBeingDebugged()))
 			{
 				DebuggedControlRig->GetHierarchy()->Notify(ERigHierarchyNotification::HierarchyReset, nullptr);
 				DebuggedControlRig->Initialize(true);
@@ -2377,7 +2377,7 @@ void FControlRigEditor::FilterDraggedKeys(TArray<FRigElementKey>& Keys, bool bRe
 			}
 			else
 			{
-				if(const UControlRig* DebuggedControlRig = Cast<UControlRig>(ControlRigBlueprint->GetObjectBeingDebugged()))
+				if(const UBaseControlRig* DebuggedControlRig = Cast<UBaseControlRig>(ControlRigBlueprint->GetObjectBeingDebugged()))
 				{
 					if(!DebuggedControlRig->GetHierarchy()->Contains(Key))
 					{
@@ -2402,7 +2402,7 @@ FTransform FControlRigEditor::GetRigElementTransform(const FRigElementKey& InEle
 
 	if(bOnDebugInstance)
 	{
-		UControlRig* DebuggedControlRig = Cast<UControlRig>(GetBlueprintObj()->GetObjectBeingDebugged());
+		UBaseControlRig* DebuggedControlRig = Cast<UBaseControlRig>(GetBlueprintObj()->GetObjectBeingDebugged());
 		if (DebuggedControlRig == nullptr)
 		{
 			DebuggedControlRig = GetControlRig();
@@ -2781,12 +2781,12 @@ EVisibility FControlRigEditor::GetDirectManipulationVisibility() const
 
 FText FControlRigEditor::GetDirectionManipulationText() const
 {
-	if (UControlRig* DebuggedControlRig = Cast<UControlRig>(GetBlueprintObj()->GetObjectBeingDebugged()))
+	if (UBaseControlRig* DebuggedControlRig = Cast<UBaseControlRig>(GetBlueprintObj()->GetObjectBeingDebugged()))
 	{
 		TArray<FRigControlElement*> TransientControls = DebuggedControlRig->GetHierarchy()->GetTransientControls();
 		for(const FRigControlElement* TransientControl : TransientControls)
 		{
-			const FString Target = UControlRig::GetTargetFromTransientControl(TransientControl->GetKey());
+			const FString Target = UBaseControlRig::GetTargetFromTransientControl(TransientControl->GetKey());
 			if(!Target.IsEmpty())
 			{
 				return FText::FromString(Target);
@@ -2843,7 +2843,7 @@ const TArray<FRigDirectManipulationTarget> FControlRigEditor::GetDirectManipulat
 {
 	if(DirectManipulationSubject.IsValid())
 	{
-		if (UControlRig* DebuggedControlRig = Cast<UControlRig>(GetBlueprintObj()->GetObjectBeingDebugged()))
+		if (UBaseControlRig* DebuggedControlRig = Cast<UBaseControlRig>(GetBlueprintObj()->GetObjectBeingDebugged()))
 		{
 			if(const URigVMUnitNode* Node = DirectManipulationSubject.Get())
 			{
@@ -2852,7 +2852,7 @@ const TArray<FRigDirectManipulationTarget> FControlRigEditor::GetDirectManipulat
 					const TSharedPtr<FStructOnScope> NodeInstance = Node->ConstructLiveStructInstance(DebuggedControlRig);
 					if(NodeInstance.IsValid() && NodeInstance->IsValid())
 					{
-						if(const FRigUnit* UnitInstance = UControlRig::GetRigUnitInstanceFromScope(NodeInstance))
+						if(const FRigUnit* UnitInstance = UBaseControlRig::GetRigUnitInstanceFromScope(NodeInstance))
 						{
 							TArray<FRigDirectManipulationTarget> Targets;
 							if(UnitInstance->GetDirectManipulationTargets(Node, NodeInstance, DebuggedControlRig->GetHierarchy(), Targets, nullptr))
@@ -2908,7 +2908,7 @@ FText FControlRigEditor::GetConnectorWarningText() const
 	{
 		if(Blueprint->IsControlRigModule())
 		{
-			if(UControlRig* ControlRig = GetControlRig())
+			if(UBaseControlRig* ControlRig = GetControlRig())
 			{
 				FString FailureReason;
 				if(!ControlRig->AllConnectorsAreResolved(&FailureReason))
@@ -3269,7 +3269,7 @@ void FControlRigEditor::OnHierarchyModified_AnyThread(ERigHierarchyNotification
 						else
 						{
 							// only clear the details if we are not looking at a transient control
-							if (UControlRig* DebuggedControlRig = Cast<UControlRig>(GetBlueprintObj()->GetObjectBeingDebugged()))
+							if (UBaseControlRig* DebuggedControlRig = Cast<UBaseControlRig>(GetBlueprintObj()->GetObjectBeingDebugged()))
 							{
 								if(DebuggedControlRig->RigUnitManipulationInfos.IsEmpty())
 								{
@@ -3387,7 +3387,7 @@ void FControlRigEditor::SynchronizeViewportBoneSelection()
 
 void FControlRigEditor::UpdateBoneModification(FName BoneName, const FTransform& LocalTransform)
 {
-	if (UControlRig* ControlRig = GetControlRig())
+	if (UBaseControlRig* ControlRig = GetControlRig())
 	{ 
 		if (PreviewInstance)
 		{ 
@@ -3401,7 +3401,7 @@ void FControlRigEditor::UpdateBoneModification(FName BoneName, const FTransform&
 		}
 		
 		TMap<FName, FTransform>* TransformOverrideMap = &ControlRig->TransformOverrideForUserCreatedBones;
-		if (UControlRig* DebuggedControlRig = Cast<UControlRig>(GetBlueprintObj()->GetObjectBeingDebugged()))
+		if (UBaseControlRig* DebuggedControlRig = Cast<UBaseControlRig>(GetBlueprintObj()->GetObjectBeingDebugged()))
 		{
 			TransformOverrideMap = &DebuggedControlRig->TransformOverrideForUserCreatedBones;
 		}
@@ -3415,7 +3415,7 @@ void FControlRigEditor::UpdateBoneModification(FName BoneName, const FTransform&
 
 void FControlRigEditor::RemoveBoneModification(FName BoneName)
 {
-	if (UControlRig* ControlRig = GetControlRig())
+	if (UBaseControlRig* ControlRig = GetControlRig())
 	{
 		if (PreviewInstance)
 		{
@@ -3423,7 +3423,7 @@ void FControlRigEditor::RemoveBoneModification(FName BoneName)
 		}
 
 		TMap<FName, FTransform>* TransformOverrideMap = &ControlRig->TransformOverrideForUserCreatedBones;
-		if (UControlRig* DebuggedControlRig = Cast<UControlRig>(GetBlueprintObj()->GetObjectBeingDebugged()))
+		if (UBaseControlRig* DebuggedControlRig = Cast<UBaseControlRig>(GetBlueprintObj()->GetObjectBeingDebugged()))
 		{
 			TransformOverrideMap = &DebuggedControlRig->TransformOverrideForUserCreatedBones;
 		}
@@ -3434,7 +3434,7 @@ void FControlRigEditor::RemoveBoneModification(FName BoneName)
 
 void FControlRigEditor::ResetAllBoneModification()
 {
-	if (UControlRig* ControlRig = GetControlRig())
+	if (UBaseControlRig* ControlRig = GetControlRig())
 	{
 		if (IsValid(PreviewInstance))
 		{
@@ -3442,7 +3442,7 @@ void FControlRigEditor::ResetAllBoneModification()
 		}
 
 		TMap<FName, FTransform>* TransformOverrideMap = &ControlRig->TransformOverrideForUserCreatedBones;
-		if (UControlRig* DebuggedControlRig = Cast<UControlRig>(GetBlueprintObj()->GetObjectBeingDebugged()))
+		if (UBaseControlRig* DebuggedControlRig = Cast<UBaseControlRig>(GetBlueprintObj()->GetObjectBeingDebugged()))
 		{
 			TransformOverrideMap = &DebuggedControlRig->TransformOverrideForUserCreatedBones;
 		}
@@ -4337,9 +4337,9 @@ void FControlRigEditor::HandleMakeElementGetterSetter(ERigElementGetterSetterTyp
 	}
 }
 
-void FControlRigEditor::HandleOnControlModified(UControlRig* Subject, FRigControlElement* ControlElement, const FRigControlModifiedContext& Context)
+void FControlRigEditor::HandleOnControlModified(UBaseControlRig* Subject, FRigControlElement* ControlElement, const FRigControlModifiedContext& Context)
 {
-	UControlRig* DebuggedControlRig = Cast<UControlRig>(GetBlueprintObj()->GetObjectBeingDebugged());
+	UBaseControlRig* DebuggedControlRig = Cast<UBaseControlRig>(GetBlueprintObj()->GetObjectBeingDebugged());
 	if (Subject != DebuggedControlRig)
 	{
 		return;
@@ -4356,8 +4356,8 @@ void FControlRigEditor::HandleOnControlModified(UControlRig* Subject, FRigContro
 	if (ControlElement->Settings.bIsTransientControl && !GIsTransacting)
 	{
 		const URigVMUnitNode* UnitNode = nullptr;
-		const FString NodeName = UControlRig::GetNodeNameFromTransientControl(ControlElement->GetKey());
-		const FString PoseTarget = UControlRig::GetTargetFromTransientControl(ControlElement->GetKey());
+		const FString NodeName = UBaseControlRig::GetNodeNameFromTransientControl(ControlElement->GetKey());
+		const FString PoseTarget = UBaseControlRig::GetTargetFromTransientControl(ControlElement->GetKey());
 		TSharedPtr<FStructOnScope> NodeInstance;
 		TSharedPtr<FRigDirectManipulationInfo> ManipulationInfo;
 
@@ -4440,7 +4440,7 @@ void FControlRigEditor::HandleOnControlModified(UControlRig* Subject, FRigContro
 		else
 		{
 			FRigControlValue ControlValue = Hierarchy->GetControlValue(ControlElement, ERigControlValueType::Current);
-			const FRigElementKey ElementKey = UControlRig::GetElementKeyFromTransientControl(ControlElement->GetKey());
+			const FRigElementKey ElementKey = UBaseControlRig::GetElementKeyFromTransientControl(ControlElement->GetKey());
 
 			if (ElementKey.Type == ERigElementType::Bone)
 			{
@@ -4517,7 +4517,7 @@ TSharedPtr<FUICommandList> FControlRigEditor::HandleOnViewportContextMenuCommand
 	return TSharedPtr<FUICommandList>();
 }
 
-void FControlRigEditor::OnPreForwardsSolve_AnyThread(UControlRig* InRig, const FName& InEventName)
+void FControlRigEditor::OnPreForwardsSolve_AnyThread(UBaseControlRig* InRig, const FName& InEventName)
 {
 	// if we are debugging a PIE instance, we need to remember the input pose on the
 	// rig so we can perform multiple evaluations. this is to avoid double transforms / double forward solve results.
@@ -4537,7 +4537,7 @@ void FControlRigEditor::OnPreForwardsSolve_AnyThread(UControlRig* InRig, const F
 	}
 }
 
-void FControlRigEditor::OnPreConstructionForUI_AnyThread(UControlRig* InRig, const FName& InEventName)
+void FControlRigEditor::OnPreConstructionForUI_AnyThread(UBaseControlRig* InRig, const FName& InEventName)
 {
 	bIsConstructionEventRunning = true;
 	const TArrayView<const FRigElementKey> Elements;
@@ -4552,7 +4552,7 @@ void FControlRigEditor::OnPreConstructionForUI_AnyThread(UControlRig* InRig, con
 	}
 }
 
-void FControlRigEditor::OnPreConstruction_AnyThread(UControlRig* InRig, const FName& InEventName)
+void FControlRigEditor::OnPreConstruction_AnyThread(UBaseControlRig* InRig, const FName& InEventName)
 {
 	if(UControlRigBlueprint* RigBlueprint = GetControlRigBlueprint())
 	{
@@ -4602,7 +4602,7 @@ void FControlRigEditor::OnPreConstruction_AnyThread(UControlRig* InRig, const FN
 	}
 }
 
-void FControlRigEditor::OnPostConstruction_AnyThread(UControlRig* InRig, const FName& InEventName)
+void FControlRigEditor::OnPostConstruction_AnyThread(UBaseControlRig* InRig, const FName& InEventName)
 {
 	bIsConstructionEventRunning = false;
 
diff --git a/Engine/Plugins/Animation/ControlRig/Source/ControlRigEditor/Private/Editor/ControlRigEditor.h b/Engine/Plugins/Animation/ControlRig/Source/ControlRigEditor/Private/Editor/ControlRigEditor.h
--- a/Engine/Plugins/Animation/ControlRig/Source/ControlRigEditor/Private/Editor/ControlRigEditor.h
+++ b/Engine/Plugins/Animation/ControlRig/Source/ControlRigEditor/Private/Editor/ControlRigEditor.h
@@ -111,7 +111,7 @@ public:
 	UControlRigBlueprint* GetControlRigBlueprint() const;
 
 	// returns the hierarchy being debugged
-	UControlRig* GetControlRig() const;
+	UBaseControlRig* GetControlRig() const;
 
 	// returns the hierarchy being debugged
 	URigHierarchy* GetHierarchyBeingDebugged() const;
@@ -133,7 +133,7 @@ public:
 	FControlRigEditorEditMode* GetEditMode() const;
 
 	// this changes everytime you compile, so don't cache it expecting it will last. 
-	UControlRig* GetInstanceRig() const { return GetControlRig();  }
+	UBaseControlRig* GetInstanceRig() const { return GetControlRig();  }
 
 	void OnCurveContainerChanged();
 
@@ -248,7 +248,7 @@ private:
 	void FilterDraggedKeys(TArray<FRigElementKey>& Keys, bool bRemoveNameSpace);
 	void HandleMakeElementGetterSetter(ERigElementGetterSetterType Type, bool bIsGetter, TArray<FRigElementKey> Keys, UEdGraph* Graph, FVector2D NodePosition);
 
-	void HandleOnControlModified(UControlRig* Subject, FRigControlElement* ControlElement, const FRigControlModifiedContext& Context);
+	void HandleOnControlModified(UBaseControlRig* Subject, FRigControlElement* ControlElement, const FRigControlModifiedContext& Context);
 
 	virtual void HandleRefreshEditorFromBlueprint(URigVMBlueprint* InBlueprint) override;
 
@@ -303,10 +303,10 @@ protected:
 	int32 RigHierarchyTabCount;
 	TWeakObjectPtr<AStaticMeshActor> WeakGroundActorPtr;
 
-	void OnPreForwardsSolve_AnyThread(UControlRig* InRig, const FName& InEventName);
-	void OnPreConstructionForUI_AnyThread(UControlRig* InRig, const FName& InEventName);
-	void OnPreConstruction_AnyThread(UControlRig* InRig, const FName& InEventName);
-	void OnPostConstruction_AnyThread(UControlRig* InRig, const FName& InEventName);
+	void OnPreForwardsSolve_AnyThread(UBaseControlRig* InRig, const FName& InEventName);
+	void OnPreConstructionForUI_AnyThread(UBaseControlRig* InRig, const FName& InEventName);
+	void OnPreConstruction_AnyThread(UBaseControlRig* InRig, const FName& InEventName);
+	void OnPostConstruction_AnyThread(UBaseControlRig* InRig, const FName& InEventName);
 	FRigPose PreConstructionPose;
 	TArray<FRigConnectorInfo> ConnectorInfos;
 
diff --git a/Engine/Plugins/Animation/ControlRig/Source/ControlRigEditor/Private/Editor/ControlRigEditorEditMode.cpp b/Engine/Plugins/Animation/ControlRig/Source/ControlRigEditor/Private/Editor/ControlRigEditorEditMode.cpp
--- a/Engine/Plugins/Animation/ControlRig/Source/ControlRigEditor/Private/Editor/ControlRigEditorEditMode.cpp
+++ b/Engine/Plugins/Animation/ControlRig/Source/ControlRigEditor/Private/Editor/ControlRigEditorEditMode.cpp
@@ -30,13 +30,13 @@ void FControlRigEditorEditMode::Render(const FSceneView* View, FViewport* Viewpo
 	}
 
 	// is there a control rig to edit?
-	TArray<UControlRig*> ControlRigs = GetControlRigsArray(true );
+	TArray<UBaseControlRig*> ControlRigs = GetControlRigsArray(true );
 	if (ControlRigs.IsEmpty())
 	{
 		return;
 	}
 
-	UControlRig* ControlRig = ControlRigs[0]; //just one control rig in the CR asset editor
+	UBaseControlRig* ControlRig = ControlRigs[0]; //just one control rig in the CR asset editor
 	if (!ControlRig) 
 	{
 		return;
diff --git a/Engine/Plugins/Animation/ControlRig/Source/ControlRigEditor/Private/Editor/ControlRigSkeletalMeshComponent.cpp b/Engine/Plugins/Animation/ControlRig/Source/ControlRigEditor/Private/Editor/ControlRigSkeletalMeshComponent.cpp
--- a/Engine/Plugins/Animation/ControlRig/Source/ControlRigEditor/Private/Editor/ControlRigSkeletalMeshComponent.cpp
+++ b/Engine/Plugins/Animation/ControlRig/Source/ControlRigEditor/Private/Editor/ControlRigSkeletalMeshComponent.cpp
@@ -74,7 +74,7 @@ void UControlRigSkeletalMeshComponent::RebuildDebugDrawSkeleton()
 
 	if (ControlRigInstance)
 	{
-		UControlRig* ControlRig = ControlRigInstance->GetFirstAvailableControlRig();
+		UBaseControlRig* ControlRig = ControlRigInstance->GetFirstAvailableControlRig();
 		if (ControlRig)
 		{
 			// we are trying to poke into running instances of Control Rigs
@@ -151,7 +151,7 @@ FTransform UControlRigSkeletalMeshComponent::GetDrawTransform(int32 BoneIndex) c
 
 	if (ControlRigInstance)
 	{
-		UControlRig* ControlRig = ControlRigInstance->GetFirstAvailableControlRig();
+		UBaseControlRig* ControlRig = ControlRigInstance->GetFirstAvailableControlRig();
 		if (ControlRig && DebugDrawBoneIndexInHierarchy.IsValidIndex(BoneIndex))
 		{
 			// just copy it because this is not thread safe
@@ -172,7 +172,7 @@ void UControlRigSkeletalMeshComponent::EnablePreview(bool bEnable, UAnimationAss
 	}
 }
 
-void UControlRigSkeletalMeshComponent::SetControlRigBeingDebugged(UControlRig* InControlRig)
+void UControlRigSkeletalMeshComponent::SetControlRigBeingDebugged(UBaseControlRig* InControlRig)
 {
 	if(ControlRigBeingDebuggedPtr.Get() == InControlRig)
 	{
@@ -181,7 +181,7 @@ void UControlRigSkeletalMeshComponent::SetControlRigBeingDebugged(UControlRig* I
 
 	if(ControlRigBeingDebuggedPtr.IsValid())
 	{
-		if(UControlRig* ControlRigBeingDebugged = ControlRigBeingDebuggedPtr.Get())
+		if(UBaseControlRig* ControlRigBeingDebugged = ControlRigBeingDebuggedPtr.Get())
 		{
 			if(!ControlRigBeingDebugged->HasAnyFlags(RF_BeginDestroyed))
 			{
@@ -298,12 +298,12 @@ void UControlRigSkeletalMeshComponent::OnHierarchyModified_AnyThread(ERigHierarc
 	}
 }
 
-void UControlRigSkeletalMeshComponent::OnPreConstruction_AnyThread(UControlRig* InControlRig, const FName& InEventName)
+void UControlRigSkeletalMeshComponent::OnPreConstruction_AnyThread(UBaseControlRig* InControlRig, const FName& InEventName)
 {
 	bIsConstructionEventRunning = true;
 }
 
-void UControlRigSkeletalMeshComponent::OnPostConstruction_AnyThread(UControlRig* InControlRig, const FName& InEventName)
+void UControlRigSkeletalMeshComponent::OnPostConstruction_AnyThread(UBaseControlRig* InControlRig, const FName& InEventName)
 {
 	bIsConstructionEventRunning = false;
 
diff --git a/Engine/Plugins/Animation/ControlRig/Source/ControlRigEditor/Private/Editor/ControlRigSkeletalMeshComponent.h b/Engine/Plugins/Animation/ControlRig/Source/ControlRigEditor/Private/Editor/ControlRigSkeletalMeshComponent.h
--- a/Engine/Plugins/Animation/ControlRig/Source/ControlRigEditor/Private/Editor/ControlRigSkeletalMeshComponent.h
+++ b/Engine/Plugins/Animation/ControlRig/Source/ControlRigEditor/Private/Editor/ControlRigSkeletalMeshComponent.h
@@ -44,7 +44,7 @@ public:
 	/*
 	 * React to a control rig being debugged
 	 */
-	void SetControlRigBeingDebugged(UControlRig* InControlRig);
+	void SetControlRigBeingDebugged(UBaseControlRig* InControlRig);
 	
 	/*
 	 *	Rebuild debug draw skeleton 
@@ -54,14 +54,14 @@ public:
 	void OnHierarchyModified(ERigHierarchyNotification InNotif, URigHierarchy* InHierarchy, const FRigBaseElement* InElement);
 	void OnHierarchyModified_AnyThread(ERigHierarchyNotification InNotif, URigHierarchy* InHierarchy, const FRigBaseElement* InElement);
 
-	void OnPreConstruction_AnyThread(UControlRig* InControlRig, const FName& InEventName);
-	void OnPostConstruction_AnyThread(UControlRig* InControlRig, const FName& InEventName);
+	void OnPreConstruction_AnyThread(UBaseControlRig* InControlRig, const FName& InEventName);
+	void OnPostConstruction_AnyThread(UBaseControlRig* InControlRig, const FName& InEventName);
 
 private:
 	FReferenceSkeleton DebugDrawSkeleton;
 	TArray<FBoneIndexType> DebugDrawBones;
 	TArray<int32> DebugDrawBoneIndexInHierarchy;
-	TWeakObjectPtr<UControlRig> ControlRigBeingDebuggedPtr;
+	TWeakObjectPtr<UBaseControlRig> ControlRigBeingDebuggedPtr;
 	int32 HierarchyInteractionBracket;
 	bool bRebuildDebugDrawSkeletonRequired;
 	bool bIsConstructionEventRunning;
diff --git a/Engine/Plugins/Animation/ControlRig/Source/ControlRigEditor/Private/Editor/RigValidationTabSummoner.cpp b/Engine/Plugins/Animation/ControlRig/Source/ControlRigEditor/Private/Editor/RigValidationTabSummoner.cpp
--- a/Engine/Plugins/Animation/ControlRig/Source/ControlRigEditor/Private/Editor/RigValidationTabSummoner.cpp
+++ b/Engine/Plugins/Animation/ControlRig/Source/ControlRigEditor/Private/Editor/RigValidationTabSummoner.cpp
@@ -35,7 +35,7 @@ TSharedRef<SWidget> FRigValidationTabSummoner::CreateTabBody(const FWorkflowTabS
 	check(Validator);
 
 	TSharedRef<SControlRigValidationWidget> ValidationWidget = SNew(SControlRigValidationWidget, Validator);
-	Validator->SetControlRig(Cast<UControlRig>(RigBlueprint->GetObjectBeingDebugged()));
+	Validator->SetControlRig(Cast<UBaseControlRig>(RigBlueprint->GetObjectBeingDebugged()));
 	return ValidationWidget;
 }
 
diff --git a/Engine/Plugins/Animation/ControlRig/Source/ControlRigEditor/Private/Editor/SControlRigAnimAttributeView.cpp b/Engine/Plugins/Animation/ControlRig/Source/ControlRigEditor/Private/Editor/SControlRigAnimAttributeView.cpp
--- a/Engine/Plugins/Animation/ControlRig/Source/ControlRigEditor/Private/Editor/SControlRigAnimAttributeView.cpp
+++ b/Engine/Plugins/Animation/ControlRig/Source/ControlRigEditor/Private/Editor/SControlRigAnimAttributeView.cpp
@@ -33,7 +33,7 @@ void SControlRigAnimAttributeView::Construct(const FArguments& InArgs, TSharedRe
 	ControlRigBlueprint = ControlRigEditor.Pin()->GetControlRigBlueprint();
 	ControlRigBlueprint->OnSetObjectBeingDebugged().AddRaw(this, &SControlRigAnimAttributeView::HandleSetObjectBeingDebugged);
 
-	if (UControlRig* ControlRig = Cast<UControlRig>(ControlRigBlueprint->GetObjectBeingDebugged()))
+	if (UBaseControlRig* ControlRig = Cast<UBaseControlRig>(ControlRigBlueprint->GetObjectBeingDebugged()))
 	{
 		StartObservingNewControlRig(ControlRig);
 	}
@@ -48,14 +48,14 @@ void SControlRigAnimAttributeView::HandleSetObjectBeingDebugged(UObject* InObjec
 
 	StopObservingCurrentControlRig();
 	
-	if(UControlRig* ControlRig = Cast<UControlRig>(InObject))
+	if(UBaseControlRig* ControlRig = Cast<UBaseControlRig>(InObject))
 	{
 		StartObservingNewControlRig(ControlRig);
 	}
 }
 
 
-void SControlRigAnimAttributeView::StartObservingNewControlRig(UControlRig* InControlRig) 
+void SControlRigAnimAttributeView::StartObservingNewControlRig(UBaseControlRig* InControlRig) 
 {
 	if (InControlRig)
 	{
@@ -71,7 +71,7 @@ void SControlRigAnimAttributeView::StopObservingCurrentControlRig()
 {
 	if(ControlRigBeingDebuggedPtr.IsValid())
 	{
-		if(UControlRig* ControlRig = ControlRigBeingDebuggedPtr.Get())
+		if(UBaseControlRig* ControlRig = ControlRigBeingDebuggedPtr.Get())
 		{
 			if(!ControlRig->HasAnyFlags(RF_BeginDestroyed))
 			{
@@ -85,7 +85,7 @@ void SControlRigAnimAttributeView::StopObservingCurrentControlRig()
 }
 
 void SControlRigAnimAttributeView::HandleControlRigPostForwardSolve(
-	UControlRig* InControlRig,
+	UBaseControlRig* InControlRig,
 	const FName& InEventName) const
 {
 	if (IsValid(InControlRig) && ensure(InControlRig == ControlRigBeingDebuggedPtr))
diff --git a/Engine/Plugins/Animation/ControlRig/Source/ControlRigEditor/Private/Editor/SControlRigAnimAttributeView.h b/Engine/Plugins/Animation/ControlRig/Source/ControlRigEditor/Private/Editor/SControlRigAnimAttributeView.h
--- a/Engine/Plugins/Animation/ControlRig/Source/ControlRigEditor/Private/Editor/SControlRigAnimAttributeView.h
+++ b/Engine/Plugins/Animation/ControlRig/Source/ControlRigEditor/Private/Editor/SControlRigAnimAttributeView.h
@@ -7,7 +7,7 @@
 
 class FControlRigEditor;
 class UControlRigBlueprint;
-class UControlRig;
+class UBaseControlRig;
 
 class SAnimAttributeView;
 
@@ -25,13 +25,13 @@ class SControlRigAnimAttributeView : public SCompoundWidget
 
 	void HandleSetObjectBeingDebugged(UObject* InObject);
 	
-	void StartObservingNewControlRig(UControlRig* InControlRig);
+	void StartObservingNewControlRig(UBaseControlRig* InControlRig);
 	void StopObservingCurrentControlRig();
 
-	void HandleControlRigPostForwardSolve(UControlRig* InControlRig, const FName& InEventName) const;
+	void HandleControlRigPostForwardSolve(UBaseControlRig* InControlRig, const FName& InEventName) const;
 	
 	TWeakPtr<FControlRigEditor> ControlRigEditor;
 	TWeakObjectPtr<UControlRigBlueprint> ControlRigBlueprint;
-	TWeakObjectPtr<UControlRig> ControlRigBeingDebuggedPtr;
+	TWeakObjectPtr<UBaseControlRig> ControlRigBeingDebuggedPtr;
 	TSharedPtr<SAnimAttributeView> AttributeView;
 };
diff --git a/Engine/Plugins/Animation/ControlRig/Source/ControlRigEditor/Private/Editor/SControlRigValidationWidget.cpp b/Engine/Plugins/Animation/ControlRig/Source/ControlRigEditor/Private/Editor/SControlRigValidationWidget.cpp
--- a/Engine/Plugins/Animation/ControlRig/Source/ControlRigEditor/Private/Editor/SControlRigValidationWidget.cpp
+++ b/Engine/Plugins/Animation/ControlRig/Source/ControlRigEditor/Private/Editor/SControlRigValidationWidget.cpp
@@ -236,7 +236,7 @@ void SControlRigValidationWidget::HandleMessageReported(EMessageSeverity::Type I
 
 void SControlRigValidationWidget::HandleMessageTokenClicked(const TSharedRef<class IMessageToken>& InToken)
 {
-	if (UControlRig* ControlRig = Validator->GetControlRig())
+	if (UBaseControlRig* ControlRig = Validator->GetControlRig())
 	{
 		if (InToken->GetType() == EMessageToken::AssetName)
 		{
diff --git a/Engine/Plugins/Animation/ControlRig/Source/ControlRigEditor/Private/Editor/SRigCurveContainer.cpp b/Engine/Plugins/Animation/ControlRig/Source/ControlRigEditor/Private/Editor/SRigCurveContainer.cpp
--- a/Engine/Plugins/Animation/ControlRig/Source/ControlRigEditor/Private/Editor/SRigCurveContainer.cpp
+++ b/Engine/Plugins/Animation/ControlRig/Source/ControlRigEditor/Private/Editor/SRigCurveContainer.cpp
@@ -709,7 +709,7 @@ URigHierarchy* SRigCurveContainer::GetInstanceHierarchy() const
 {
 	if (ControlRigEditor.IsValid())
 	{
-		UControlRig* ControlRig = ControlRigEditor.Pin()->GetInstanceRig();
+		UBaseControlRig* ControlRig = ControlRigEditor.Pin()->GetInstanceRig();
 		if (ControlRig)
 		{
 			return ControlRig->GetHierarchy();
diff --git a/Engine/Plugins/Animation/ControlRig/Source/ControlRigEditor/Private/Editor/SRigHierarchy.cpp b/Engine/Plugins/Animation/ControlRig/Source/ControlRigEditor/Private/Editor/SRigHierarchy.cpp
--- a/Engine/Plugins/Animation/ControlRig/Source/ControlRigEditor/Private/Editor/SRigHierarchy.cpp
+++ b/Engine/Plugins/Animation/ControlRig/Source/ControlRigEditor/Private/Editor/SRigHierarchy.cpp
@@ -871,7 +871,7 @@ void SRigHierarchy::HandleSetObjectBeingDebugged(UObject* InObject)
 
 	if(ControlRigBeingDebuggedPtr.IsValid())
 	{
-		if(UControlRig* ControlRigBeingDebugged = ControlRigBeingDebuggedPtr.Get())
+		if(UBaseControlRig* ControlRigBeingDebugged = ControlRigBeingDebuggedPtr.Get())
 		{
 			if(!ControlRigBeingDebugged->HasAnyFlags(RF_BeginDestroyed))
 			{
@@ -882,7 +882,7 @@ void SRigHierarchy::HandleSetObjectBeingDebugged(UObject* InObject)
 
 	ControlRigBeingDebuggedPtr.Reset();
 	
-	if(UControlRig* ControlRig = Cast<UControlRig>(InObject))
+	if(UBaseControlRig* ControlRig = Cast<UBaseControlRig>(InObject))
 	{
 		ControlRigBeingDebuggedPtr = ControlRig;
 		if(URigHierarchy* Hierarchy = ControlRig->GetHierarchy())
@@ -900,7 +900,7 @@ void SRigHierarchy::HandleSetObjectBeingDebugged(UObject* InObject)
 	RefreshTreeView();
 }
 
-void SRigHierarchy::OnPreConstruction_AnyThread(UControlRig* InRig, const FName& InEventName)
+void SRigHierarchy::OnPreConstruction_AnyThread(UBaseControlRig* InRig, const FName& InEventName)
 {
 	if(InRig != ControlRigBeingDebuggedPtr.Get())
 	{
@@ -910,7 +910,7 @@ void SRigHierarchy::OnPreConstruction_AnyThread(UControlRig* InRig, const FName&
 	SelectionBeforeConstruction = InRig->GetHierarchy()->GetSelectedKeys();
 }
 
-void SRigHierarchy::OnPostConstruction_AnyThread(UControlRig* InRig, const FName& InEventName)
+void SRigHierarchy::OnPostConstruction_AnyThread(UBaseControlRig* InRig, const FName& InEventName)
 {
 	if(InRig != ControlRigBeingDebuggedPtr.Get())
 	{
@@ -969,7 +969,7 @@ void SRigHierarchy::OnNavigateToFirstConnectorWarning()
 {
 	if(ControlRigEditor.IsValid())
 	{
-		if(UControlRig* ControlRig = ControlRigEditor.Pin()->GetControlRig())
+		if(UBaseControlRig* ControlRig = ControlRigEditor.Pin()->GetControlRig())
 		{
 			FRigElementKey ConnectorKey;
 			if(!ControlRig->AllConnectorsAreResolved(nullptr, &ConnectorKey))
@@ -1271,7 +1271,7 @@ void SRigHierarchy::CreateContextMenu()
 									if(Element->IsProcedural())
 									{
 										const int32 InstructionIndex = Element->GetCreatedAtInstructionIndex();
-										if(const UControlRig* ControlRig = Cast<UControlRig>(CurrentControlRigBlueprint->GetObjectBeingDebugged()))
+										if(const UBaseControlRig* ControlRig = Cast<UBaseControlRig>(CurrentControlRigBlueprint->GetObjectBeingDebugged()))
 										{
 											if(ControlRig->VM)
 											{
@@ -1518,7 +1518,7 @@ void SRigHierarchy::RefreshHierarchy(const FAssetData& InAssetData, bool bOnlyRe
 		// we do this to avoid the editmode / viewport shapes to refresh recursively,
 		// which can add an extreme slowdown depending on the number of bones (n^(n-1))
 		bool bSelectBones = true;
-		if (UControlRig* CurrentRig = StrongEditor->GetControlRig())
+		if (UBaseControlRig* CurrentRig = StrongEditor->GetControlRig())
 		{
 			bSelectBones = !CurrentRig->IsConstructionModeEnabled();
 		}
@@ -1653,7 +1653,7 @@ void SRigHierarchy::ImportHierarchy(const FAssetData& InAssetData)
 		// we do this to avoid the editmode / viewport shapes to refresh recursively,
 		// which can add an extreme slowdown depending on the number of bones (n^(n-1))
 		bool bSelectBones = true;
-		if (const UControlRig* CurrentRig = EditorSharedPtr->GetControlRig())
+		if (const UBaseControlRig* CurrentRig = EditorSharedPtr->GetControlRig())
 		{
 			bSelectBones = !CurrentRig->IsConstructionModeEnabled();
 		}
@@ -2328,14 +2328,14 @@ URigHierarchy* SRigHierarchy::GetHierarchy() const
 {
 	if (ControlRigBlueprint.IsValid())
 	{
-		if (UControlRig* DebuggedRig = ControlRigBeingDebuggedPtr.Get())
+		if (UBaseControlRig* DebuggedRig = ControlRigBeingDebuggedPtr.Get())
 		{
 			return DebuggedRig->GetHierarchy();
 		}
 	}
 	if (ControlRigEditor.IsValid())
 	{
-		if (UControlRig* CurrentRig = ControlRigEditor.Pin()->GetControlRig())
+		if (UBaseControlRig* CurrentRig = ControlRigEditor.Pin()->GetControlRig())
 		{
 			return CurrentRig->GetHierarchy();
 		}
@@ -2881,7 +2881,7 @@ void SRigHierarchy::HandleControlBoneOrSpaceTransform()
 		return;
 	}
 
-	UControlRig* DebuggedControlRig = Cast<UControlRig>(Blueprint->GetObjectBeingDebugged());
+	UBaseControlRig* DebuggedControlRig = Cast<UBaseControlRig>(Blueprint->GetObjectBeingDebugged());
 	if(DebuggedControlRig == nullptr)
 	{
 		return;
diff --git a/Engine/Plugins/Animation/ControlRig/Source/ControlRigEditor/Private/Editor/SRigHierarchy.h b/Engine/Plugins/Animation/ControlRig/Source/ControlRigEditor/Private/Editor/SRigHierarchy.h
--- a/Engine/Plugins/Animation/ControlRig/Source/ControlRigEditor/Private/Editor/SRigHierarchy.h
+++ b/Engine/Plugins/Animation/ControlRig/Source/ControlRigEditor/Private/Editor/SRigHierarchy.h
@@ -17,7 +17,7 @@ class FControlRigEditor;
 class SSearchBox;
 class FUICommandList;
 class URigVMBlueprint;
-class UControlRig;
+class UBaseControlRig;
 struct FAssetData;
 class FMenuBuilder;
 class UToolMenu;
@@ -180,7 +180,7 @@ private:
 	TSharedPtr<SRigHierarchyTreeView> TreeView;
 
 	TWeakObjectPtr<UControlRigBlueprint> ControlRigBlueprint;
-	TWeakObjectPtr<UControlRig> ControlRigBeingDebuggedPtr;
+	TWeakObjectPtr<UBaseControlRig> ControlRigBeingDebuggedPtr;
 	
 	/** Command list we bind to */
 	TSharedPtr<FUICommandList> CommandList;
@@ -233,8 +233,8 @@ private:
 	void OnHierarchyModified_AnyThread(ERigHierarchyNotification InNotif, URigHierarchy* InHierarchy, const FRigBaseElement* InElement);
 	void HandleRefreshEditorFromBlueprint(URigVMBlueprint* InBlueprint);
 	void HandleSetObjectBeingDebugged(UObject* InObject);
-	void OnPreConstruction_AnyThread(UControlRig* InRig, const FName& InEventName);
-	void OnPostConstruction_AnyThread(UControlRig* InRig, const FName& InEventName);
+	void OnPreConstruction_AnyThread(UBaseControlRig* InRig, const FName& InEventName);
+	void OnPostConstruction_AnyThread(UBaseControlRig* InRig, const FName& InEventName);
 	void OnNavigateToFirstConnectorWarning();
 
 	bool bIsConstructionEventRunning;
diff --git a/Engine/Plugins/Animation/ControlRig/Source/ControlRigEditor/Private/Editor/SRigHierarchyTreeView.cpp b/Engine/Plugins/Animation/ControlRig/Source/ControlRigEditor/Private/Editor/SRigHierarchyTreeView.cpp
--- a/Engine/Plugins/Animation/ControlRig/Source/ControlRigEditor/Private/Editor/SRigHierarchyTreeView.cpp
+++ b/Engine/Plugins/Animation/ControlRig/Source/ControlRigEditor/Private/Editor/SRigHierarchyTreeView.cpp
@@ -502,7 +502,7 @@ bool SRigHierarchyTreeView::AddElement(const FRigBaseElement* InElement)
 			if(Hierarchy)
 			{
 				// add the connector as a tag rather than its own element in the tree
-				if(UControlRig* ControlRig = Hierarchy->GetTypedOuter<UControlRig>())
+				if(UBaseControlRig* ControlRig = Hierarchy->GetTypedOuter<UBaseControlRig>())
 				{
 					FRigElementKeyRedirector& Redirector = ControlRig->GetElementKeyRedirector();
 					if(const FCachedRigElement* Cache = Redirector.Find(InElement->GetKey()))
@@ -888,9 +888,9 @@ void SRigHierarchyTreeView::AddConnectorResolveWarningTag(TSharedPtr<FRigTreeEle
 	check(InRigElement);
 	check(InRigElement->GetType() == ERigElementType::Connector);
 
-	if(UControlRig* ControlRig = InHierarchy->GetTypedOuter<UControlRig>())
+	if(UBaseControlRig* ControlRig = InHierarchy->GetTypedOuter<UBaseControlRig>())
 	{
-		TWeakObjectPtr<UControlRig> ControlRigPtr(ControlRig);
+		TWeakObjectPtr<UBaseControlRig> ControlRigPtr(ControlRig);
 		const FRigElementKey ConnectorKey = InRigElement->GetKey();
 		
 		TAttribute<FText> GetTooltipText = TAttribute<FText>::CreateSP(this,
@@ -911,9 +911,9 @@ void SRigHierarchyTreeView::AddConnectorResolveWarningTag(TSharedPtr<FRigTreeEle
 }
 
 FText SRigHierarchyTreeView::GetConnectorWarningMessage(TSharedPtr<FRigTreeElement> InTreeElement,
-	TWeakObjectPtr<UControlRig> InControlRigPtr, const FRigElementKey InConnectorKey) const
+	TWeakObjectPtr<UBaseControlRig> InControlRigPtr, const FRigElementKey InConnectorKey) const
 {
-	if(UControlRig* ControlRig = InControlRigPtr.Get())
+	if(UBaseControlRig* ControlRig = InControlRigPtr.Get())
 	{
 		URigHierarchy* Hierarchy = ControlRig->GetHierarchy();
 				
diff --git a/Engine/Plugins/Animation/ControlRig/Source/ControlRigEditor/Private/Editor/SRigHierarchyTreeView.h b/Engine/Plugins/Animation/ControlRig/Source/ControlRigEditor/Private/Editor/SRigHierarchyTreeView.h
--- a/Engine/Plugins/Animation/ControlRig/Source/ControlRigEditor/Private/Editor/SRigHierarchyTreeView.h
+++ b/Engine/Plugins/Animation/ControlRig/Source/ControlRigEditor/Private/Editor/SRigHierarchyTreeView.h
@@ -317,7 +317,7 @@ public:
 private:
 
 	void AddConnectorResolveWarningTag(TSharedPtr<FRigTreeElement> InTreeElement, const FRigBaseElement* InRigElement, const URigHierarchy* InHierarchy);
-	FText GetConnectorWarningMessage(TSharedPtr<FRigTreeElement> InTreeElement, TWeakObjectPtr<UControlRig> InControlRigPtr, const FRigElementKey InConnectorKey) const;
+	FText GetConnectorWarningMessage(TSharedPtr<FRigTreeElement> InTreeElement, TWeakObjectPtr<UBaseControlRig> InControlRigPtr, const FRigElementKey InConnectorKey) const;
 
 	/** A temporary snapshot of the SparseItemInfos in STreeView, used during RefreshTreeView() */
 	TSparseItemMap OldSparseItemInfos;
diff --git a/Engine/Plugins/Animation/ControlRig/Source/ControlRigEditor/Private/Editor/SRigSpacePickerWidget.cpp b/Engine/Plugins/Animation/ControlRig/Source/ControlRigEditor/Private/Editor/SRigSpacePickerWidget.cpp
--- a/Engine/Plugins/Animation/ControlRig/Source/ControlRigEditor/Private/Editor/SRigSpacePickerWidget.cpp
+++ b/Engine/Plugins/Animation/ControlRig/Source/ControlRigEditor/Private/Editor/SRigSpacePickerWidget.cpp
@@ -683,13 +683,13 @@ FReply SRigSpacePickerWidget::HandleAddElementClicked()
 						URigHierarchy::TElementDependencyMap DependencyMap;
 						FString FailureReason;
 
-						if(UControlRig* ControlRig = StrongHierarchy->GetTypedOuter<UControlRig>())
+						if(UBaseControlRig* ControlRig = StrongHierarchy->GetTypedOuter<UBaseControlRig>())
 						{
 							DependencyMap = StrongHierarchy->GetDependenciesForVM(ControlRig->GetVM()); 
 						}
 						else if(UControlRigBlueprint* RigBlueprint = StrongHierarchy->GetTypedOuter<UControlRigBlueprint>())
 						{
-							if(UControlRig* CDO = Cast<UControlRig>(RigBlueprint->GetRigVMBlueprintGeneratedClass()->GetDefaultObject()))
+							if(UBaseControlRig* CDO = Cast<UBaseControlRig>(RigBlueprint->GetRigVMBlueprintGeneratedClass()->GetDefaultObject()))
 							{
 								DependencyMap = StrongHierarchy->GetDependenciesForVM(CDO->GetVM()); 
 							}
diff --git a/Engine/Plugins/Animation/ControlRig/Source/ControlRigEditor/Private/Graph/ControlRigGraphPanelPinFactory.cpp b/Engine/Plugins/Animation/ControlRig/Source/ControlRigEditor/Private/Graph/ControlRigGraphPanelPinFactory.cpp
--- a/Engine/Plugins/Animation/ControlRig/Source/ControlRigEditor/Private/Graph/ControlRigGraphPanelPinFactory.cpp
+++ b/Engine/Plugins/Animation/ControlRig/Source/ControlRigEditor/Private/Graph/ControlRigGraphPanelPinFactory.cpp
@@ -248,7 +248,7 @@ TSharedPtr<SGraphPin> FControlRigGraphPanelPinFactory::CreatePin_Internal(UEdGra
 								}
 
 								const URigHierarchy* Hierarchy = Blueprint->Hierarchy;
-								if (UControlRig* ControlRig = Cast<UControlRig>(Blueprint->GetObjectBeingDebugged()))
+								if (UBaseControlRig* ControlRig = Cast<UBaseControlRig>(Blueprint->GetObjectBeingDebugged()))
 								{
 									Hierarchy = ControlRig->GetHierarchy();
 								}
@@ -331,7 +331,7 @@ TSharedPtr<SGraphPin> FControlRigGraphPanelPinFactory::CreatePin_Internal(UEdGra
 						{
 							if (const UControlRigBlueprint* Blueprint = RigGraph->GetTypedOuter<UControlRigBlueprint>())
 							{
-								if(UControlRig* ControlRig = Cast<UControlRig>(Blueprint->GetObjectBeingDebugged()))
+								if(UBaseControlRig* ControlRig = Cast<UBaseControlRig>(Blueprint->GetObjectBeingDebugged()))
 								{
 									const FString MapHash = Blueprint->GetPathName();
 									const int32 MetadataVersion = ControlRig->GetHierarchy()->GetMetadataVersion();
@@ -419,7 +419,7 @@ TSharedPtr<SGraphPin> FControlRigGraphPanelPinFactory::CreatePin_Internal(UEdGra
 						{
 							if (const UControlRigBlueprint* Blueprint = RigGraph->GetTypedOuter<UControlRigBlueprint>())
 							{
-								if(UControlRig* ControlRig = Cast<UControlRig>(Blueprint->GetObjectBeingDebugged()))
+								if(UBaseControlRig* ControlRig = Cast<UBaseControlRig>(Blueprint->GetObjectBeingDebugged()))
 								{
 									const FString MapHash = Blueprint->GetPathName();
 									const int32 MetadataTagVersion = ControlRig->GetHierarchy()->GetMetadataTagVersion(); 
diff --git a/Engine/Plugins/Animation/ControlRig/Source/ControlRigEditor/Private/Sequencer/ControlConstraintChannelInterface.cpp b/Engine/Plugins/Animation/ControlRig/Source/ControlRigEditor/Private/Sequencer/ControlConstraintChannelInterface.cpp
--- a/Engine/Plugins/Animation/ControlRig/Source/ControlRigEditor/Private/Sequencer/ControlConstraintChannelInterface.cpp
+++ b/Engine/Plugins/Animation/ControlRig/Source/ControlRigEditor/Private/Sequencer/ControlConstraintChannelInterface.cpp
@@ -62,7 +62,7 @@ UWorld* FControlConstraintChannelInterface::GetHandleWorld(UTransformableHandle*
 	}
 	
 	const UTransformableControlHandle* ControlHandle = static_cast<UTransformableControlHandle*>(InHandle);
-	const UControlRig* ControlRig = ControlHandle->ControlRig.LoadSynchronous();
+	const UBaseControlRig* ControlRig = ControlHandle->ControlRig.LoadSynchronous();
 
 	return ControlRig ? ControlRig->GetWorld() : nullptr;
 }
@@ -114,7 +114,7 @@ bool FControlConstraintChannelInterface::SmartConstraintKey(
 			TGuardValue<bool> CompensateGuard(FMovieSceneConstraintChannelHelper::bDoNotCompensate, true);
 			TGuardValue<bool> RemoveConstraintGuard(FConstraintsManagerController::bDoNotRemoveConstraint, true);
 
-			UControlRig* ControlRig = ControlHandle->ControlRig.Get();
+			UBaseControlRig* ControlRig = ControlHandle->ControlRig.Get();
 			const FName& ControlName = ControlHandle->ControlName;
 				
 			// store the frames to compensate
@@ -228,7 +228,7 @@ UMovieSceneControlRigParameterSection* FControlConstraintChannelInterface::GetCo
 		return nullptr;
 	}
 	
-	const UControlRig* ControlRig = InHandle->ControlRig.LoadSynchronous();
+	const UBaseControlRig* ControlRig = InHandle->ControlRig.LoadSynchronous();
 	if (!ControlRig)
 	{
 		return nullptr;
@@ -242,7 +242,7 @@ UMovieSceneControlRigParameterSection* FControlConstraintChannelInterface::GetCo
 
 	auto GetControlRigTrack = [InHandle, MovieScene]()->UMovieSceneControlRigParameterTrack*
 	{
-		const TWeakObjectPtr<UControlRig> ControlRig = InHandle->ControlRig.LoadSynchronous();
+		const TWeakObjectPtr<UBaseControlRig> ControlRig = InHandle->ControlRig.LoadSynchronous();
 		if (ControlRig.IsValid())
 		{	
 			const TArray<FMovieSceneBinding>& Bindings = MovieScene->GetBindings();
diff --git a/Engine/Plugins/Animation/ControlRig/Source/ControlRigEditor/Private/Sequencer/ControlRigParameterTrackEditor.cpp b/Engine/Plugins/Animation/ControlRig/Source/ControlRigEditor/Private/Sequencer/ControlRigParameterTrackEditor.cpp
--- a/Engine/Plugins/Animation/ControlRig/Source/ControlRigEditor/Private/Sequencer/ControlRigParameterTrackEditor.cpp
+++ b/Engine/Plugins/Animation/ControlRig/Source/ControlRigEditor/Private/Sequencer/ControlRigParameterTrackEditor.cpp
@@ -254,14 +254,14 @@ FControlRigParameterTrackEditor::FControlRigParameterTrackEditor(TSharedRef<ISeq
 				static bool bHasEnteredSilent = false;
 				
 				bool bRequestEvaluate = false;
-				TMap<UControlRig*, UControlRig*> OldToNewControlRigs;
+				TMap<UBaseControlRig*, UBaseControlRig*> OldToNewControlRigs;
 				FControlRigEditMode* ControlRigEditMode = GetEditMode();
 				if (ControlRigEditMode)
 				{
-					TArrayView<TWeakObjectPtr<UControlRig>> ControlRigs = ControlRigEditMode->GetControlRigs();
-					for (TWeakObjectPtr<UControlRig>& ControlRigPtr : ControlRigs)
+					TArrayView<TWeakObjectPtr<UBaseControlRig>> ControlRigs = ControlRigEditMode->GetControlRigs();
+					for (TWeakObjectPtr<UBaseControlRig>& ControlRigPtr : ControlRigs)
 					{
-						if (UControlRig* ControlRig = ControlRigPtr.Get())
+						if (UBaseControlRig* ControlRig = ControlRigPtr.Get())
 						{
 							if (ControlRig->GetObjectBinding() && ControlRig->GetObjectBinding()->GetBoundObject() == nullptr)
 							{
@@ -279,7 +279,7 @@ FControlRigParameterTrackEditor::FControlRigParameterTrackEditor(TSharedRef<ISeq
 						UControlRigComponent* NewControlRigComponent = Cast<UControlRigComponent>(ReplacedObject.Value);
 						if (OldControlRigComponent->GetControlRig())
 						{
-							UControlRig* NewControlRig = nullptr;
+							UBaseControlRig* NewControlRig = nullptr;
 							if (NewControlRigComponent)
 							{
 								NewControlRig = NewControlRigComponent->GetControlRig();
@@ -287,9 +287,9 @@ FControlRigParameterTrackEditor::FControlRigParameterTrackEditor(TSharedRef<ISeq
 							OldToNewControlRigs.Emplace(OldControlRigComponent->GetControlRig(), NewControlRig);
 						}
 					}
-					else if (UControlRig* OldControlRig = Cast<UControlRig>(ReplacedObject.Key))
+					else if (UBaseControlRig* OldControlRig = Cast<UBaseControlRig>(ReplacedObject.Key))
 					{
-						UControlRig* NewControlRig = Cast<UControlRig>(ReplacedObject.Value);
+						UBaseControlRig* NewControlRig = Cast<UBaseControlRig>(ReplacedObject.Value);
 						OldToNewControlRigs.Emplace(OldControlRig, NewControlRig);
 					}
 				}
@@ -304,9 +304,9 @@ FControlRigParameterTrackEditor::FControlRigParameterTrackEditor(TSharedRef<ISeq
 
 					IterateTracks([this, &OldToNewControlRigs, ControlRigEditMode, &bRequestEvaluate](UMovieSceneControlRigParameterTrack* Track)
 					{
-						if (UControlRig* OldControlRig = Track->GetControlRig())
+						if (UBaseControlRig* OldControlRig = Track->GetControlRig())
 						{
-							UControlRig** NewControlRig = OldToNewControlRigs.Find(OldControlRig);
+							UBaseControlRig** NewControlRig = OldToNewControlRigs.Find(OldControlRig);
 							if (NewControlRig)
 							{
 								TArray<FName> SelectedControls = OldControlRig->CurrentControlSelection();
@@ -327,7 +327,7 @@ FControlRigParameterTrackEditor::FControlRigParameterTrackEditor(TSharedRef<ISeq
 								{
 									ControlRigEditMode->ReplaceControlRig(OldControlRig, *NewControlRig);
 
-									UControlRig* PtrNewControlRig = *NewControlRig;
+									UBaseControlRig* PtrNewControlRig = *NewControlRig;
 
 									auto UpdateSelectionDelegate = [this, SelectedControls, PtrNewControlRig, bRequestEvaluate]()
 									{
@@ -404,7 +404,7 @@ FControlRigParameterTrackEditor::FControlRigParameterTrackEditor(TSharedRef<ISeq
 	// Register all modified/selections for control rigs
 	IterateTracks([this](UMovieSceneControlRigParameterTrack* Track)
 	{
-		if (UControlRig* ControlRig = Track->GetControlRig())
+		if (UBaseControlRig* ControlRig = Track->GetControlRig())
 		{
 			BindControlRig(ControlRig);
 		}
@@ -430,7 +430,7 @@ FControlRigParameterTrackEditor::~FControlRigParameterTrackEditor()
 }
 
 
-void FControlRigParameterTrackEditor::BindControlRig(UControlRig* ControlRig)
+void FControlRigParameterTrackEditor::BindControlRig(UBaseControlRig* ControlRig)
 {
 	if (ControlRig && BoundControlRigs.Contains(ControlRig) == false)
 	{
@@ -469,7 +469,7 @@ void FControlRigParameterTrackEditor::BindControlRig(UControlRig* ControlRig)
 		}
 	}
 }
-void FControlRigParameterTrackEditor::UnbindControlRig(UControlRig* ControlRig)
+void FControlRigParameterTrackEditor::UnbindControlRig(UBaseControlRig* ControlRig)
 {
 	if (ControlRig && BoundControlRigs.Contains(ControlRig) == true)
 	{
@@ -496,12 +496,12 @@ void FControlRigParameterTrackEditor::UnbindControlRig(UControlRig* ControlRig)
 void FControlRigParameterTrackEditor::UnbindAllControlRigs()
 {
 	ClearOutAllSpaceAndConstraintDelegates();
-	TArray<TWeakObjectPtr<UControlRig>> ControlRigs = BoundControlRigs;
-	for (TWeakObjectPtr<UControlRig>& ObjectPtr : ControlRigs)
+	TArray<TWeakObjectPtr<UBaseControlRig>> ControlRigs = BoundControlRigs;
+	for (TWeakObjectPtr<UBaseControlRig>& ObjectPtr : ControlRigs)
 	{
 		if (ObjectPtr.IsValid())
 		{
-			UControlRig* ControlRig = ObjectPtr.Get();
+			UBaseControlRig* ControlRig = ObjectPtr.Get();
 			UnbindControlRig(ControlRig);
 		}
 	}
@@ -511,7 +511,7 @@ void FControlRigParameterTrackEditor::UnbindAllControlRigs()
 
 void FControlRigParameterTrackEditor::ObjectImplicitlyAdded(UObject* InObject)
 {
-	UControlRig* ControlRig = Cast<UControlRig>(InObject);
+	UBaseControlRig* ControlRig = Cast<UBaseControlRig>(InObject);
 	if (ControlRig)
 	{
 		BindControlRig(ControlRig);
@@ -520,7 +520,7 @@ void FControlRigParameterTrackEditor::ObjectImplicitlyAdded(UObject* InObject)
 
 void FControlRigParameterTrackEditor::ObjectImplicitlyRemoved(UObject* InObject)
 {
-	UControlRig* ControlRig = Cast<UControlRig>(InObject);
+	UBaseControlRig* ControlRig = Cast<UBaseControlRig>(InObject);
 	if (ControlRig)
 	{
 		UnbindControlRig(ControlRig);
@@ -893,7 +893,7 @@ void FControlRigParameterTrackEditor::BakeToControlRig(UClass* InClass, FGuid Ob
 	FSlateApplication::Get().DismissAllMenus();
 	const TSharedPtr<ISequencer> SequencerParent = GetSequencer();
 
-	if (InClass && InClass->IsChildOf(UControlRig::StaticClass()) && SequencerParent.IsValid())
+	if (InClass && InClass->IsChildOf(UBaseControlRig::StaticClass()) && SequencerParent.IsValid())
 	{
 
 		UMovieSceneSequence* OwnerSequence = GetSequencer()->GetFocusedMovieSceneSequence();
@@ -971,7 +971,7 @@ void FControlRigParameterTrackEditor::BakeToControlRig(UClass* InClass, FGuid Ob
 
 					FString ObjectName = InClass->GetName();
 					ObjectName.RemoveFromEnd(TEXT("_C"));
-					UControlRig* ControlRig = bReuseControlRig ? Track->GetControlRig() : NewObject<UControlRig>(Track, InClass, FName(*ObjectName), RF_Transactional);
+					UBaseControlRig* ControlRig = bReuseControlRig ? Track->GetControlRig() : NewObject<UBaseControlRig>(Track, InClass, FName(*ObjectName), RF_Transactional);
 					if (InClass != UFKControlRig::StaticClass() && !ControlRig->SupportsEvent(FRigUnit_InverseExecution::EventName))
 					{
 						TempAnimSequence->MarkAsGarbage();
@@ -988,7 +988,7 @@ void FControlRigParameterTrackEditor::BakeToControlRig(UClass* InClass, FGuid Ob
 					else
 					{
 						/* mz todo we don't unbind  will test more
-						UControlRig* OldControlRig = ControlRigEditMode->GetControlRig(false);
+						UBaseControlRig* OldControlRig = ControlRigEditMode->GetControlRig(false);
 						if (OldControlRig)
 						{
 							UnbindControlRig(OldControlRig);
@@ -1001,7 +1001,7 @@ void FControlRigParameterTrackEditor::BakeToControlRig(UClass* InClass, FGuid Ob
 						ControlRig->Modify();
 						ControlRig->SetObjectBinding(MakeShared<FControlRigObjectBinding>());
 						ControlRig->GetObjectBinding()->BindToObject(BoundActor);
-						ControlRig->GetDataSourceRegistry()->RegisterDataSource(UControlRig::OwnerComponent, ControlRig->GetObjectBinding()->GetBoundObject());
+						ControlRig->GetDataSourceRegistry()->RegisterDataSource(UBaseControlRig::OwnerComponent, ControlRig->GetObjectBinding()->GetBoundObject());
 						ControlRig->Initialize();
 						ControlRig->RequestInit();
 						ControlRig->SetBoneInitialTransformsFromSkeletalMeshComponent(SkelMeshComp, true);
@@ -1119,7 +1119,7 @@ void FControlRigParameterTrackEditor::IterateTracksInMovieScene(UMovieScene& Mov
 	}	
 }
 
-void FControlRigParameterTrackEditor::BakeInvertedPose(UControlRig* InControlRig, UMovieSceneControlRigParameterTrack* Track)
+void FControlRigParameterTrackEditor::BakeInvertedPose(UBaseControlRig* InControlRig, UMovieSceneControlRigParameterTrack* Track)
 {
 	if (!InControlRig->IsAdditive())
 	{
@@ -1149,7 +1149,7 @@ void FControlRigParameterTrackEditor::BakeInvertedPose(UControlRig* InControlRig
 	FScopedSlowTask Progress(AnimationLength, LOCTEXT("BakingToControlRig_SlowTask", "Baking To Control Rig..."));
 	Progress.MakeDialog(true);
 
-	auto DelegateHandle = InControlRig->OnPreAdditiveValuesApplication_AnyThread().AddLambda([](UControlRig* InControlRig, const FName& InEventName)
+	auto DelegateHandle = InControlRig->OnPreAdditiveValuesApplication_AnyThread().AddLambda([](UBaseControlRig* InControlRig, const FName& InEventName)
 	{
 		InControlRig->InvertInputPose();
 	});
@@ -1376,7 +1376,7 @@ void FControlRigParameterTrackEditor::HandleAddControlRigSubMenu(FMenuBuilder& M
 		AssetPickerConfig.InitialAssetViewType = EAssetViewType::List;
 		AssetPickerConfig.OnShouldFilterAsset = FOnShouldFilterAsset::CreateRaw(this, &FControlRigParameterTrackEditor::ShouldFilterAsset);
 		AssetPickerConfig.Filter.bRecursiveClasses = true;
-		AssetPickerConfig.Filter.ClassPaths.Add((UControlRig::StaticClass())->GetClassPathName());
+		AssetPickerConfig.Filter.ClassPaths.Add((UBaseControlRig::StaticClass())->GetClassPathName());
 		AssetPickerConfig.Filter.TagsAndValues.Add(TEXT("Skeleton"), FAssetData(Skeleton).GetExportTextName());
 		AssetPickerConfig.SaveSettingsName = TEXT("SequencerAssetPicker");
 
@@ -1400,7 +1400,7 @@ void FControlRigParameterTrackEditor::HandleAddControlRigSubMenu(FMenuBuilder& M
 /*
 bool FControlRigParameterTrackEditor::ShouldFilterAsset(const FAssetData& AssetData)
 {
-	if (AssetData.AssetClassPath == UControlRig::StaticClass()->GetClassPathName())
+	if (AssetData.AssetClassPath == UBaseControlRig::StaticClass()->GetClassPathName())
 	{
 		return true;
 	}
@@ -1414,9 +1414,9 @@ void FControlRigParameterTrackEditor::OnControlRigAssetSelected(const FAssetData
 	UObject* SelectedObject = AssetData.GetAsset();todo
 	TSharedPtr<ISequencer> SequencerPtr = GetSequencer();
 
-	if (SelectedObject && SelectedObject->IsA(UControlRig::StaticClass()) && SequencerPtr.IsValid())
+	if (SelectedObject && SelectedObject->IsA(UBaseControlRig::StaticClass()) && SequencerPtr.IsValid())
 	{
-		UControlRig* ControlRig = CastChecked<UControlRig>(AssetData.GetAsset());
+		UBaseControlRig* ControlRig = CastChecked<UBaseControlRig>(AssetData.GetAsset());
 
 		const FScopedTransaction Transaction(LOCTEXT("AddAnimation_Transaction", "Add Animation"));
 
@@ -1429,7 +1429,7 @@ void FControlRigParameterTrackEditor::OnControlRigAssetSelected(const FAssetData
 	}
 */
 
-void FControlRigParameterTrackEditor::AddControlRig(const UClass* InClass, UObject* BoundActor, FGuid ObjectBinding, UControlRig* InExistingControlRig)
+void FControlRigParameterTrackEditor::AddControlRig(const UClass* InClass, UObject* BoundActor, FGuid ObjectBinding, UBaseControlRig* InExistingControlRig)
 {
 	UWorld* World = GCurrentLevelEditingViewportClient ? GCurrentLevelEditingViewportClient->GetWorld() : nullptr;
 	ULevelSequence* LevelSequence = Cast<ULevelSequence>(GetSequencer()->GetFocusedMovieSceneSequence());
@@ -1462,7 +1462,7 @@ void FControlRigParameterTrackEditor::AddControlRigFromComponent(FGuid InGuid)
 		BoundActor->GetComponents(ControlRigComponents);
 		for (UControlRigComponent* ControlRigComponent : ControlRigComponents)
 		{
-			if (UControlRig* CR = ControlRigComponent->GetControlRig())
+			if (UBaseControlRig* CR = ControlRigComponent->GetControlRig())
 			{
 				AddControlRig(CR->GetClass(), BoundActor, InGuid, CR);
 			}
@@ -1496,7 +1496,7 @@ bool FControlRigParameterTrackEditor::CanAddTransformKeysForSelectedObjects() co
 	FControlRigEditMode* ControlRigEditMode = GetEditMode();
 	if (ControlRigEditMode)
 	{
-		TMap<UControlRig*, TArray<FRigElementKey>> SelectedControls;
+		TMap<UBaseControlRig*, TArray<FRigElementKey>> SelectedControls;
 		ControlRigEditMode->GetAllSelectedControls(SelectedControls);
 		return (SelectedControls.Num() > 0);
 		}
@@ -1516,7 +1516,7 @@ void FControlRigParameterTrackEditor::OnAddTransformKeysForSelectedObjects(EMovi
 		return;
 	}
 
-	TMap<UControlRig*, TArray<FRigElementKey>> SelectedControls;
+	TMap<UBaseControlRig*, TArray<FRigElementKey>> SelectedControls;
 	ControlRigEditMode->GetAllSelectedControls(SelectedControls);
 	if (SelectedControls.Num() <= 0)
 	{
@@ -1526,9 +1526,9 @@ void FControlRigParameterTrackEditor::OnAddTransformKeysForSelectedObjects(EMovi
 	FScopedTransaction KeyTransaction(LOCTEXT("SetKeysOnControls", "Set Keys On Controls"), !GIsTransacting);
 
 	static constexpr bool bInConstraintSpace = true;
-	for (const TPair<UControlRig*, TArray<FRigElementKey>>& Selection : SelectedControls)
+	for (const TPair<UBaseControlRig*, TArray<FRigElementKey>>& Selection : SelectedControls)
 	{
-		UControlRig* ControlRig = Selection.Key;
+		UBaseControlRig* ControlRig = Selection.Key;
 		if (const TSharedPtr<IControlRigObjectBinding> ObjectBinding = ControlRig->GetObjectBinding())
 		{
 			if (UObject* Object = ObjectBinding->GetBoundObject())
@@ -1553,7 +1553,7 @@ static void EvaluateThisControl(UMovieSceneControlRigParameterSection* Section,
 	{
 		return;
 	}
-	UControlRig* ControlRig = Section->GetControlRig();
+	UBaseControlRig* ControlRig = Section->GetControlRig();
 	if (!ControlRig)
 	{
 		return;
@@ -1809,7 +1809,7 @@ void FControlRigParameterTrackEditor::AddTrackForComponent(USceneComponent* InCo
 						{
 							if (URigVMBlueprintGeneratedClass* RigClass = BPControlRig->GetRigVMBlueprintGeneratedClass())
 							{
-								if (UControlRig* CDO = Cast<UControlRig>(RigClass->GetDefaultObject(true /* create if needed */)))
+								if (UBaseControlRig* CDO = Cast<UBaseControlRig>(RigClass->GetDefaultObject(true /* create if needed */)))
 								{
 									AddControlRig(CDO->GetClass(), InComponent, Binding);
 								}
@@ -1898,7 +1898,7 @@ void FControlRigParameterTrackEditor::OnActivateSequenceChanged(FMovieSceneSeque
 {
 	IterateTracks([this](UMovieSceneControlRigParameterTrack* Track)
 	{
-		if (UControlRig* ControlRig = Track->GetControlRig())
+		if (UBaseControlRig* ControlRig = Track->GetControlRig())
 		{
 			BindControlRig(ControlRig);
 		}
@@ -1918,9 +1918,9 @@ void FControlRigParameterTrackEditor::OnSequencerDataChanged(EMovieSceneDataChan
 	if (ControlRigEditMode && ControlRigEditMode->GetControlRigsArray(false /*bIsVisible*/).Num() != 0 && MovieScene && (DataChangeType == EMovieSceneDataChangeType::MovieSceneStructureItemRemoved ||
 		DataChangeType == EMovieSceneDataChangeType::Unknown))
 	{
-		TArray<UControlRig*> ControlRigs = ControlRigEditMode->GetControlRigsArray(false /*bIsVisible*/);
+		TArray<UBaseControlRig*> ControlRigs = ControlRigEditMode->GetControlRigsArray(false /*bIsVisible*/);
 		const float FPS = (float)GetSequencer()->GetFocusedDisplayRate().AsDecimal();
-		for (UControlRig* ControlRig : ControlRigs)
+		for (UBaseControlRig* ControlRig : ControlRigs)
 		{
 			if (ControlRig)
 			{
@@ -1954,7 +1954,7 @@ void FControlRigParameterTrackEditor::PostEvaluation(UMovieScene* MovieScene, FF
 	{
 		IterateTracksInMovieScene(*MovieScene, [this](const UMovieSceneControlRigParameterTrack* Track)
 		{
-			if (const UControlRig* ControlRig = Track->GetControlRig())
+			if (const UBaseControlRig* ControlRig = Track->GetControlRig())
 				{
 					if (ControlRig->GetObjectBinding())
 					{
@@ -1986,7 +1986,7 @@ void FControlRigParameterTrackEditor::OnSelectionChanged(TArray<UMovieSceneTrack
 
 	FControlRigEditMode* ControlRigEditMode = GetEditMode();
 	bool bEditModeExisted = ControlRigEditMode != nullptr;
-	UControlRig* ControlRig = nullptr;
+	UBaseControlRig* ControlRig = nullptr;
 
 	TArray<const IKeyArea*> KeyAreas;
 	const bool UseSelectedKeys = CVarSelectedKeysSelectControls.GetValueOnGameThread();
@@ -1996,9 +1996,9 @@ void FControlRigParameterTrackEditor::OnSelectionChanged(TArray<UMovieSceneTrack
 		if (FSlateApplication::Get().GetModifierKeys().IsShiftDown() == false && 
 			FSlateApplication::Get().GetModifierKeys().IsControlDown() == false && ControlRigEditMode)
 		{
-			TMap<UControlRig*, TArray<FRigElementKey>> AllSelectedControls;
+			TMap<UBaseControlRig*, TArray<FRigElementKey>> AllSelectedControls;
 			ControlRigEditMode->GetAllSelectedControls(AllSelectedControls);
-			for (TPair<UControlRig*, TArray<FRigElementKey>>& SelectedControl : AllSelectedControls)
+			for (TPair<UBaseControlRig*, TArray<FRigElementKey>>& SelectedControl : AllSelectedControls)
 			{
 				ControlRig = SelectedControl.Key;
 				if (ControlRig && ControlRig->CurrentControlSelection().Num() > 0)
@@ -2013,7 +2013,7 @@ void FControlRigParameterTrackEditor::OnSelectionChanged(TArray<UMovieSceneTrack
 			UMovieSceneControlRigParameterTrack* CRTrack = Cast<UMovieSceneControlRigParameterTrack>(Track);
 			if (CRTrack)
 			{
-				UControlRig* TrackControlRig = CRTrack->GetControlRig();
+				UBaseControlRig* TrackControlRig = CRTrack->GetControlRig();
 				if (TrackControlRig)
 				{
 					if (ControlRigEditMode)
@@ -2049,7 +2049,7 @@ void FControlRigParameterTrackEditor::OnSelectionChanged(TArray<UMovieSceneTrack
 	
 }
 
-void FControlRigParameterTrackEditor::SelectRigsAndControls(UControlRig* ControlRig, const TArray<const IKeyArea*>& KeyAreas)
+void FControlRigParameterTrackEditor::SelectRigsAndControls(UBaseControlRig* ControlRig, const TArray<const IKeyArea*>& KeyAreas)
 {
 	FControlRigEditMode* ControlRigEditMode = GetEditMode();
 	
@@ -2059,7 +2059,7 @@ void FControlRigParameterTrackEditor::SelectRigsAndControls(UControlRig* Control
 
 	TArray<FString> StringArray;
 	//we have two sets here one to see if selection has really changed that contains the attirbutes, the other to select just the parent
-	TMap<UControlRig*, TSet<FName>> RigsAndControls;
+	TMap<UBaseControlRig*, TSet<FName>> RigsAndControls;
 	for (const IKeyArea* KeyArea : KeyAreas)
 	{
 		UMovieSceneControlRigParameterSection* MovieSection = Cast<UMovieSceneControlRigParameterSection>(KeyArea->GetOwningSection());
@@ -2145,13 +2145,13 @@ void FControlRigParameterTrackEditor::SelectRigsAndControls(UControlRig* Control
 	//only create transaction if selection is really different.
 	bool bEndTransaction = false;
 	
-	TMap<UControlRig*, TArray<FName>> ControlRigsToClearSelection;
+	TMap<UBaseControlRig*, TArray<FName>> ControlRigsToClearSelection;
 	//get current selection which we will clear if different
 	if (ControlRigEditMode)
 	{
-		TMap<UControlRig*, TArray<FRigElementKey>> SelectedControls;
+		TMap<UBaseControlRig*, TArray<FRigElementKey>> SelectedControls;
 		ControlRigEditMode->GetAllSelectedControls(SelectedControls);
-		for (TPair<UControlRig*, TArray<FRigElementKey>>& Selection : SelectedControls)
+		for (TPair<UBaseControlRig*, TArray<FRigElementKey>>& Selection : SelectedControls)
 		{
 			ControlRig = Selection.Key;
 			if (ControlRig)
@@ -2162,7 +2162,7 @@ void FControlRigParameterTrackEditor::SelectRigsAndControls(UControlRig* Control
 		}
 	}
 
-	for (TPair<UControlRig*, TSet<FName>>& Pair : RigsAndControls)
+	for (TPair<UBaseControlRig*, TSet<FName>>& Pair : RigsAndControls)
 	{
 		//check to see if new selection is same als old selection
 		bool bIsSame = true;
@@ -2225,7 +2225,7 @@ void FControlRigParameterTrackEditor::SelectRigsAndControls(UControlRig* Control
 		}
 	}
 	//go through and clear those still not cleared
-	for (TPair<UControlRig*, TArray<FName>>& SelectedPairs : ControlRigsToClearSelection)
+	for (TPair<UBaseControlRig*, TArray<FName>>& SelectedPairs : ControlRigsToClearSelection)
 	{
 		if (!GIsTransacting && bEndTransaction == false)
 		{
@@ -2249,7 +2249,7 @@ void FControlRigParameterTrackEditor::SelectRigsAndControls(UControlRig* Control
 }
 
 
-FMovieSceneTrackEditor::FFindOrCreateHandleResult FControlRigParameterTrackEditor::FindOrCreateHandleToObject(UObject* InObj, UControlRig* InControlRig)
+FMovieSceneTrackEditor::FFindOrCreateHandleResult FControlRigParameterTrackEditor::FindOrCreateHandleToObject(UObject* InObj, UBaseControlRig* InControlRig)
 {
 	const bool bCreateHandleIfMissing = false;
 	FName CreatedFolderName = NAME_None;
@@ -2341,7 +2341,7 @@ void FControlRigParameterTrackEditor::SelectSequencerNodeInSection(UMovieSceneCo
 	}
 }
 
-FMovieSceneTrackEditor::FFindOrCreateTrackResult FControlRigParameterTrackEditor::FindOrCreateControlRigTrackForObject(FGuid ObjectHandle, UControlRig* ControlRig, FName PropertyName, bool bCreateTrackIfMissing)
+FMovieSceneTrackEditor::FFindOrCreateTrackResult FControlRigParameterTrackEditor::FindOrCreateControlRigTrackForObject(FGuid ObjectHandle, UBaseControlRig* ControlRig, FName PropertyName, bool bCreateTrackIfMissing)
 {
 	FFindOrCreateTrackResult Result;
 	bool bTrackExisted = false;
@@ -2368,7 +2368,7 @@ FMovieSceneTrackEditor::FFindOrCreateTrackResult FControlRigParameterTrackEditor
 	return Result;
 }
 
-UMovieSceneControlRigParameterTrack* FControlRigParameterTrackEditor::FindTrack(const UControlRig* InControlRig) const
+UMovieSceneControlRigParameterTrack* FControlRigParameterTrackEditor::FindTrack(const UBaseControlRig* InControlRig) const
 {
 	if (!GetSequencer().IsValid())
 	{
@@ -2499,7 +2499,7 @@ void FControlRigParameterTrackEditor::HandleSpaceKeyMoved(
 	}
 }
 
-void FControlRigParameterTrackEditor::ClearOutAllSpaceAndConstraintDelegates(const UControlRig* InOptionalControlRig) const
+void FControlRigParameterTrackEditor::ClearOutAllSpaceAndConstraintDelegates(const UBaseControlRig* InOptionalControlRig) const
 {
 	UTickableTransformConstraint::GetOnConstraintChanged().RemoveAll(this);
 	
@@ -2538,7 +2538,7 @@ void FControlRigParameterTrackEditor::ClearOutAllSpaceAndConstraintDelegates(con
 
 				if (CRSection->OnConstraintRemovedHandle.IsValid())
 				{
-					if (const UControlRig* ControlRig = CRSection->GetControlRig())
+					if (const UBaseControlRig* ControlRig = CRSection->GetControlRig())
 					{
 						FConstraintsManagerController& Controller = FConstraintsManagerController::Get(ControlRig->GetWorld());
 						Controller.GetNotifyDelegate().Remove(CRSection->OnConstraintRemovedHandle);
@@ -2658,7 +2658,7 @@ void FControlRigParameterTrackEditor::HandleConstraintKeyDeleted(
 {
 	UMovieSceneControlRigParameterSection* Section = Cast<UMovieSceneControlRigParameterSection>(InSection);
 
-	const UControlRig* ControlRig = Section ? Section->GetControlRig() : nullptr;
+	const UBaseControlRig* ControlRig = Section ? Section->GetControlRig() : nullptr;
 	if (!ControlRig || !InConstraintChannel)
 	{
 		return;
@@ -2706,7 +2706,7 @@ void FControlRigParameterTrackEditor::HandleConstraintRemoved(IMovieSceneConstra
 {
 	UMovieSceneControlRigParameterSection* Section = Cast<UMovieSceneControlRigParameterSection>(InSection);
 
-	if (const UControlRig* ControlRig = Section->GetControlRig())
+	if (const UBaseControlRig* ControlRig = Section->GetControlRig())
 	{
 		FConstraintsManagerController& Controller = FConstraintsManagerController::Get(ControlRig->GetWorld());
 		if (!InSection->OnConstraintRemovedHandle.IsValid())
@@ -2801,7 +2801,7 @@ void FControlRigParameterTrackEditor::HandleConstraintPropertyChanged(UTickableT
 			InConstraint, Channel->ActiveChannel, InPropertyChangedEvent, GetSequencer(), Section);
 }
 
-void FControlRigParameterTrackEditor::SetUpEditModeIfNeeded(UControlRig* ControlRig)
+void FControlRigParameterTrackEditor::SetUpEditModeIfNeeded(UBaseControlRig* ControlRig)
 {
 	if (ControlRig)
 	{
@@ -2843,7 +2843,7 @@ void FControlRigParameterTrackEditor::SetUpEditModeIfNeeded(UControlRig* Control
 	}
 }
 
-void FControlRigParameterTrackEditor::HandleControlSelected(UControlRig* Subject, FRigControlElement* ControlElement, bool bSelected)
+void FControlRigParameterTrackEditor::HandleControlSelected(UBaseControlRig* Subject, FRigControlElement* ControlElement, bool bSelected)
 {
 	using namespace UE::Sequencer;
 
@@ -2971,11 +2971,11 @@ void FControlRigParameterTrackEditor::HandleControlSelected(UControlRig* Subject
 	}
 }
 
-void FControlRigParameterTrackEditor::HandleOnPostConstructed(UControlRig* Subject, const FName& InEventName)
+void FControlRigParameterTrackEditor::HandleOnPostConstructed(UBaseControlRig* Subject, const FName& InEventName)
 {
 	if (IsInGameThread())
 	{
-		UControlRig* ControlRig = CastChecked<UControlRig>(Subject);
+		UBaseControlRig* ControlRig = CastChecked<UBaseControlRig>(Subject);
 
 		if (GetSequencer().IsValid())
 		{
@@ -2988,7 +2988,7 @@ void FControlRigParameterTrackEditor::HandleOnPostConstructed(UControlRig* Subje
 	}
 }
 
-void FControlRigParameterTrackEditor::HandleControlModified(UControlRig* ControlRig, FRigControlElement* ControlElement, const FRigControlModifiedContext& Context)
+void FControlRigParameterTrackEditor::HandleControlModified(UBaseControlRig* ControlRig, FRigControlElement* ControlElement, const FRigControlModifiedContext& Context)
 {
 	if(ControlElement == nullptr)
 	{
@@ -3032,7 +3032,7 @@ void FControlRigParameterTrackEditor::HandleControlModified(UControlRig* Control
 	});
 }
 
-void FControlRigParameterTrackEditor::HandleControlUndoBracket(UControlRig* Subject, bool bOpenUndoBracket)
+void FControlRigParameterTrackEditor::HandleControlUndoBracket(UBaseControlRig* Subject, bool bOpenUndoBracket)
 {
 	if(IsInGameThread() && bOpenUndoBracket && ControlUndoBracket == 0)
 	{
@@ -3058,7 +3058,7 @@ void FControlRigParameterTrackEditor::HandleControlUndoBracket(UControlRig* Subj
 	}
 }
 
-void FControlRigParameterTrackEditor::HandleOnControlRigBound(UControlRig* InControlRig)
+void FControlRigParameterTrackEditor::HandleOnControlRigBound(UBaseControlRig* InControlRig)
 {
 	if (!InControlRig)
 	{
@@ -3077,7 +3077,7 @@ void FControlRigParameterTrackEditor::HandleOnControlRigBound(UControlRig* InCon
 	{
 		if (UMovieSceneControlRigParameterSection* Section = Cast< UMovieSceneControlRigParameterSection>(BaseSection))
 		{
-			const UControlRig* ControlRig = Section->GetControlRig();
+			const UBaseControlRig* ControlRig = Section->GetControlRig();
 			if (ControlRig && InControlRig == ControlRig)
 			{
 				if (!Binding->OnControlRigBind().IsBoundToObject(this))
@@ -3092,10 +3092,10 @@ void FControlRigParameterTrackEditor::HandleOnControlRigBound(UControlRig* InCon
 void FControlRigParameterTrackEditor::HandleOnObjectBoundToControlRig(UObject* InObject)
 {
 	//reselect these control rigs since selection may get lost
-	TMap<TWeakObjectPtr<UControlRig>, TArray<FName>> ReselectIfNeeded;
+	TMap<TWeakObjectPtr<UBaseControlRig>, TArray<FName>> ReselectIfNeeded;
 	// look for sections to update
 	TArray<UMovieSceneControlRigParameterSection*> SectionsToUpdate;
-	for (TWeakObjectPtr<UControlRig>& ControlRigPtr : BoundControlRigs)
+	for (TWeakObjectPtr<UBaseControlRig>& ControlRigPtr : BoundControlRigs)
 	{
 		if (ControlRigPtr.IsValid() == false)
 		{
@@ -3130,7 +3130,7 @@ void FControlRigParameterTrackEditor::HandleOnObjectBoundToControlRig(UObject* I
 			{
 				GEditor->GetTimerManager()->SetTimerForNextTick([ReselectIfNeeded]()
 					{
-						for (const TPair <TWeakObjectPtr<UControlRig>, TArray<FName>>& Pair : ReselectIfNeeded)
+						for (const TPair <TWeakObjectPtr<UBaseControlRig>, TArray<FName>>& Pair : ReselectIfNeeded)
 						{
 							if (Pair.Key.IsValid())
 							{
@@ -3149,7 +3149,7 @@ void FControlRigParameterTrackEditor::HandleOnObjectBoundToControlRig(UObject* I
 
 
 void FControlRigParameterTrackEditor::GetControlRigKeys(
-	UControlRig* InControlRig,
+	UBaseControlRig* InControlRig,
 	FName ParameterName,
 	EControlRigContextChannelToKey ChannelsToKey,
 	ESequencerKeyMode KeyMode,
@@ -3414,7 +3414,7 @@ void FControlRigParameterTrackEditor::GetControlRigKeys(
 	}
 }
 
-FKeyPropertyResult FControlRigParameterTrackEditor::AddKeysToControlRigHandle(UObject* InObject, UControlRig* InControlRig,
+FKeyPropertyResult FControlRigParameterTrackEditor::AddKeysToControlRigHandle(UObject* InObject, UBaseControlRig* InControlRig,
 	FGuid ObjectHandle, FFrameNumber KeyTime, FFrameNumber EvaluateTime, FGeneratedTrackKeys& GeneratedKeys,
 	ESequencerKeyMode KeyMode, TSubclassOf<UMovieSceneTrack> TrackClass, FName ControlRigName, FName RigControlName)
 {
@@ -3478,7 +3478,7 @@ FKeyPropertyResult FControlRigParameterTrackEditor::AddKeysToControlRigHandle(UO
 		if (KeyPropertyResult.bKeyCreated)
 		{
 			UMovieSceneControlRigParameterSection* ParamSection = Cast<UMovieSceneControlRigParameterSection>(Track->GetSectionToKey());
-			if (UControlRig* SectionControlRig = ParamSection ? ParamSection->GetControlRig() : nullptr)
+			if (UBaseControlRig* SectionControlRig = ParamSection ? ParamSection->GetControlRig() : nullptr)
 			{
 				TOptional<FFrameNumber> OptionalKeyTime = KeyTime;
 
@@ -3497,7 +3497,7 @@ FKeyPropertyResult FControlRigParameterTrackEditor::AddKeysToControlRigHandle(UO
 }
 
 FKeyPropertyResult FControlRigParameterTrackEditor::AddKeysToControlRig(
-	UObject* InObject, UControlRig* InControlRig, FFrameNumber KeyTime, FFrameNumber EvaluateTime, FGeneratedTrackKeys& GeneratedKeys,
+	UObject* InObject, UBaseControlRig* InControlRig, FFrameNumber KeyTime, FFrameNumber EvaluateTime, FGeneratedTrackKeys& GeneratedKeys,
 	ESequencerKeyMode KeyMode, TSubclassOf<UMovieSceneTrack> TrackClass, FName ControlRigName, FName RigControlName)
 {
 	FKeyPropertyResult KeyPropertyResult;
@@ -3519,7 +3519,7 @@ FKeyPropertyResult FControlRigParameterTrackEditor::AddKeysToControlRig(
 
 void FControlRigParameterTrackEditor::AddControlKeys(
 	UObject* InObject,
-	UControlRig* InControlRig,
+	UBaseControlRig* InControlRig,
 	FName ControlRigName,
 	FName RigControlName,
 	EControlRigContextChannelToKey ChannelsToKey,
@@ -3588,7 +3588,7 @@ void FControlRigParameterTrackEditor::AddControlKeys(
 	EndKeying(); //fine even if we didn't BeginKeying
 }
 
-bool FControlRigParameterTrackEditor::ModifyOurGeneratedKeysByCurrentAndWeight(UObject* Object, UControlRig* InControlRig, FName RigControlName, UMovieSceneTrack* Track, UMovieSceneSection* SectionToKey, FFrameNumber EvaluateTime, FGeneratedTrackKeys& GeneratedTotalKeys, float Weight) const
+bool FControlRigParameterTrackEditor::ModifyOurGeneratedKeysByCurrentAndWeight(UObject* Object, UBaseControlRig* InControlRig, FName RigControlName, UMovieSceneTrack* Track, UMovieSceneSection* SectionToKey, FFrameNumber EvaluateTime, FGeneratedTrackKeys& GeneratedTotalKeys, float Weight) const
 {
 	FFrameRate TickResolution = GetSequencer()->GetFocusedTickResolution();
 	FMovieSceneEvaluationTrack EvalTrack = CastChecked<UMovieSceneControlRigParameterTrack>(Track)->GenerateTrackTemplate(Track);
@@ -3815,7 +3815,7 @@ void FControlRigParameterTrackEditor::BuildTrackContextMenu(FMenuBuilder& MenuBu
 
 		MenuBuilder.AddMenuSeparator();
 	}
-	else if (UControlRig* LayeredRig = Cast<UControlRig>(Track->GetControlRig()))
+	else if (UBaseControlRig* LayeredRig = Cast<UBaseControlRig>(Track->GetControlRig()))
 	{
 		if (LayeredRig->IsAdditive())
 		{
@@ -3888,7 +3888,7 @@ bool FControlRigParameterTrackEditor::HandleAssetAdded(UObject* Asset, const FGu
 	UMovieSceneControlRigParameterTrack* Track = Cast<UMovieSceneControlRigParameterTrack>(MovieScene->FindTrack(UMovieSceneControlRigParameterTrack::StaticClass(), NewGuid, NAME_None));
 	if (Track == nullptr)
 	{
-		UControlRig* CDO = Cast<UControlRig>(RigClass->GetDefaultObject(true /* create if needed */));
+		UBaseControlRig* CDO = Cast<UBaseControlRig>(RigClass->GetDefaultObject(true /* create if needed */));
 		check(CDO);
 
 		AddControlRig(CDO->GetClass(), SpawnedSkeletalMeshActor->GetSkeletalMeshComponent(), NewGuid);
@@ -4485,7 +4485,7 @@ bool FControlRigParameterTrackEditor::CollapseAllLayers(TSharedPtr<ISequencer>&S
 				{
 					FScopedTransaction Transaction(LOCTEXT("CollapseAllSections", "Collapse All Sections"));
 					ParameterSection->Modify();
-					UControlRig* ControlRig = ParameterSection->GetControlRig();
+					UBaseControlRig* ControlRig = ParameterSection->GetControlRig();
 					FMovieSceneSequenceTransform RootToLocalTransform = SequencerPtr->GetFocusedMovieSceneSequenceTransform();
 
 					FFrameNumber StartFrame = InSettings.StartFrame;
@@ -4667,7 +4667,7 @@ void FControlRigParameterSection::BuildSectionContextMenu(FMenuBuilder& MenuBuil
 	UMovieSceneControlRigParameterSection* ParameterSection = CastChecked<UMovieSceneControlRigParameterSection>(WeakSection.Get());
 	TSharedPtr<ISequencer> SequencerPtr = WeakSequencer.Pin();
 
-	UControlRig* ControlRig = ParameterSection->GetControlRig();
+	UBaseControlRig* ControlRig = ParameterSection->GetControlRig();
 
 	if (ControlRig)
 	{
@@ -4932,7 +4932,7 @@ void FControlRigParameterSection::ShowSelectedControlsChannels()
 {
 	UMovieSceneControlRigParameterSection* ParameterSection = CastChecked<UMovieSceneControlRigParameterSection>(WeakSection.Get());
 	TSharedPtr<ISequencer> SequencerPtr = WeakSequencer.Pin();
-	UControlRig* ControlRig = ParameterSection ? ParameterSection->GetControlRig() : nullptr;
+	UBaseControlRig* ControlRig = ParameterSection ? ParameterSection->GetControlRig() : nullptr;
 
 	if (ParameterSection && ControlRig && SequencerPtr.IsValid())
 	{
diff --git a/Engine/Plugins/Animation/ControlRig/Source/ControlRigEditor/Private/Sequencer/ControlRigParameterTrackEditor.h b/Engine/Plugins/Animation/ControlRig/Source/ControlRigEditor/Private/Sequencer/ControlRigParameterTrackEditor.h
--- a/Engine/Plugins/Animation/ControlRig/Source/ControlRigEditor/Private/Sequencer/ControlRigParameterTrackEditor.h
+++ b/Engine/Plugins/Animation/ControlRig/Source/ControlRigEditor/Private/Sequencer/ControlRigParameterTrackEditor.h
@@ -101,7 +101,7 @@ private:
 
 	/** Control Rig Picked */
 	void AddControlRig(UClass* InClass, UObject* BoundActor, FGuid ObjectBinding);
-	void AddControlRig(const UClass* InClass, UObject* BoundActor, FGuid ObjectBinding, UControlRig* InExistingControlRig);
+	void AddControlRig(const UClass* InClass, UObject* BoundActor, FGuid ObjectBinding, UBaseControlRig* InExistingControlRig);
 	void AddControlRigFromComponent(FGuid InGuid);
 	
 	/** Delegate for Selection Changed Event */
@@ -126,15 +126,15 @@ private:
 	void AddTrackForComponent(USceneComponent* Component, FGuid Binding);
 
 	/** Control Rig Delegates*/
-	void HandleControlModified(UControlRig* Subject, FRigControlElement* ControlElement, const FRigControlModifiedContext& Context);
-	void HandleControlSelected(UControlRig* Subject, FRigControlElement* ControlElement, bool bSelected);
-	void HandleControlUndoBracket(UControlRig* Subject, bool bOpenUndoBracket);
-	void HandleOnPostConstructed(UControlRig* Subject, const FName& InEventName);
-	void HandleOnControlRigBound(UControlRig* InControlRig);
+	void HandleControlModified(UBaseControlRig* Subject, FRigControlElement* ControlElement, const FRigControlModifiedContext& Context);
+	void HandleControlSelected(UBaseControlRig* Subject, FRigControlElement* ControlElement, bool bSelected);
+	void HandleControlUndoBracket(UBaseControlRig* Subject, bool bOpenUndoBracket);
+	void HandleOnPostConstructed(UBaseControlRig* Subject, const FName& InEventName);
+	void HandleOnControlRigBound(UBaseControlRig* InControlRig);
 	void HandleOnObjectBoundToControlRig(UObject* InObject);
 
 	//if rig not set then we clear delegates for everyone
-	void ClearOutAllSpaceAndConstraintDelegates(const UControlRig* InOptionalControlRig = nullptr) const;
+	void ClearOutAllSpaceAndConstraintDelegates(const UBaseControlRig* InOptionalControlRig = nullptr) const;
 
 	/** SpaceChannel Delegates*/
 	void HandleOnSpaceAdded(
@@ -166,13 +166,13 @@ private:
 	void HandleConstraintPropertyChanged(UTickableTransformConstraint* InConstraint, const FPropertyChangedEvent& InPropertyChangedEvent) const;
 
 	/** Select control rig if not selected, select controls from key areas */
-	void SelectRigsAndControls(UControlRig* Subject, const TArray<const IKeyArea*>& KeyAreas);
+	void SelectRigsAndControls(UBaseControlRig* Subject, const TArray<const IKeyArea*>& KeyAreas);
 
 	/** Handle Creation for Skeleton, SkelMeshComp or Actor Owner, they may or may not have a binding, use optional control rig to pick one */
-	FMovieSceneTrackEditor::FFindOrCreateHandleResult FindOrCreateHandleToObject(UObject* InObj, UControlRig* InControlRig = nullptr);
+	FMovieSceneTrackEditor::FFindOrCreateHandleResult FindOrCreateHandleToObject(UObject* InObj, UBaseControlRig* InControlRig = nullptr);
 
 	/** Handle Creation for control rig track given the object binding and the control rig */	
-	FMovieSceneTrackEditor::FFindOrCreateTrackResult FindOrCreateControlRigTrackForObject(FGuid ObjectBinding, UControlRig* ControlRig, FName PropertyName = NAME_None, bool bCreateTrackIfMissing = true);
+	FMovieSceneTrackEditor::FFindOrCreateTrackResult FindOrCreateControlRigTrackForObject(FGuid ObjectBinding, UBaseControlRig* ControlRig, FName PropertyName = NAME_None, bool bCreateTrackIfMissing = true);
 
 	/** Import FBX*/
 	void ImportFBX(UMovieSceneControlRigParameterTrack* InTrack, UMovieSceneControlRigParameterSection* InSection, 
@@ -182,7 +182,7 @@ private:
 	void ExportFBX(UMovieSceneControlRigParameterTrack* InTrack, UMovieSceneControlRigParameterSection* InSection);
 	
 	/** Find Track for given ControlRig*/
-	UMovieSceneControlRigParameterTrack* FindTrack(const UControlRig* InControlRig) const;
+	UMovieSceneControlRigParameterTrack* FindTrack(const UBaseControlRig* InControlRig) const;
 
 	/** Select Bones to Animate on FK Rig*/
 	void SelectFKBonesToAnimate(UFKControlRig* FKControlRig, UMovieSceneControlRigParameterTrack* Track);
@@ -197,10 +197,10 @@ private:
 	void BakeToControlRig(UClass* InClass, FGuid ObjectBinding,UObject* BoundObject, USkeletalMeshComponent* SkelMeshComp, USkeleton* Skeleton);
 
 	/** Bake Inersion of Additive Control Rig to Rest Pose*/
-	void BakeInvertedPose(UControlRig* InControlRig, UMovieSceneControlRigParameterTrack* Track);
+	void BakeInvertedPose(UBaseControlRig* InControlRig, UMovieSceneControlRigParameterTrack* Track);
 
 	/** Set Up EditMode for Specified Control Rig*/
-	void SetUpEditModeIfNeeded(UControlRig* ControlRig);
+	void SetUpEditModeIfNeeded(UBaseControlRig* ControlRig);
 
 	/** Helper functions to iterate over UMovieSceneControlRigParameterTracks in the currently focussed MovieScene*/
 	void IterateTracks(TFunctionRef<bool(UMovieSceneControlRigParameterTrack*)> Callback) const;
@@ -212,16 +212,16 @@ private:
 
 public:
 
-	void AddControlKeys(UObject* InObject, UControlRig* InControlRig, FName PropertyName,
+	void AddControlKeys(UObject* InObject, UBaseControlRig* InControlRig, FName PropertyName,
 		FName ParameterName, EControlRigContextChannelToKey ChannelsToKey, ESequencerKeyMode KeyMode,
 		float InLocalTime, const bool bInConstraintSpace = false);
-	void GetControlRigKeys(UControlRig* InControlRig, FName ParameterName, EControlRigContextChannelToKey ChannelsToKey,
+	void GetControlRigKeys(UBaseControlRig* InControlRig, FName ParameterName, EControlRigContextChannelToKey ChannelsToKey,
 		ESequencerKeyMode KeyMode, UMovieSceneControlRigParameterSection* SectionToKey,	FGeneratedTrackKeys& OutGeneratedKeys,
 		const bool bInConstraintSpace = false);
 	FKeyPropertyResult AddKeysToControlRig(
-		UObject* InObject, UControlRig* InControlRig, FFrameNumber KeyTime, FFrameNumber EvaluateTime, FGeneratedTrackKeys& GeneratedKeys,
+		UObject* InObject, UBaseControlRig* InControlRig, FFrameNumber KeyTime, FFrameNumber EvaluateTime, FGeneratedTrackKeys& GeneratedKeys,
 		ESequencerKeyMode KeyMode, TSubclassOf<UMovieSceneTrack> TrackClass, FName ControlRigName, FName RigControlName);
-	FKeyPropertyResult AddKeysToControlRigHandle(UObject* InObject, UControlRig* InControlRig,
+	FKeyPropertyResult AddKeysToControlRigHandle(UObject* InObject, UBaseControlRig* InControlRig,
 		FGuid ObjectHandle, FFrameNumber KeyTime, FFrameNumber EvaluateTime, FGeneratedTrackKeys& GeneratedKeys,
 		ESequencerKeyMode KeyMode, TSubclassOf<UMovieSceneTrack> TrackClass, FName ControlRigName, FName RigControlName);
 	/**
@@ -234,7 +234,7 @@ public:
 	 * @param InOutGeneratedTrackKeys The Keys we need to modify. We change these values.
 	 * @param Weight The weight we need to modify the values by.
 	 */
-	bool ModifyOurGeneratedKeysByCurrentAndWeight(UObject* Object, UControlRig* InControlRig, FName RigControlName, UMovieSceneTrack *Track, UMovieSceneSection* SectionToKey, FFrameNumber EvaluateTime, FGeneratedTrackKeys& InOutGeneratedTotalKeys, float Weight) const;
+	bool ModifyOurGeneratedKeysByCurrentAndWeight(UObject* Object, UBaseControlRig* InControlRig, FName RigControlName, UMovieSceneTrack *Track, UMovieSceneSection* SectionToKey, FFrameNumber EvaluateTime, FGeneratedTrackKeys& InOutGeneratedTotalKeys, float Weight) const;
 
 	//**Function to collapse all layers from this section onto the first absoluate layer.*/
 	static bool CollapseAllLayers(TSharedPtr<ISequencer>& SequencerPtr, UMovieSceneTrack* OwnerTrack, const FBakingAnimationKeySettings& InSettings);
@@ -254,10 +254,10 @@ private:
 	TSet<FName> DisplayedControls;
 	TSet<FName> UnDisplayedControls;
 	bool bCurveDisplayTickIsPending;
-	void BindControlRig(UControlRig* ControlRig);
-	void UnbindControlRig(UControlRig* ControlRig);
+	void BindControlRig(UBaseControlRig* ControlRig);
+	void UnbindControlRig(UBaseControlRig* ControlRig);
 	void UnbindAllControlRigs();
-	TArray<TWeakObjectPtr<UControlRig>> BoundControlRigs;
+	TArray<TWeakObjectPtr<UBaseControlRig>> BoundControlRigs;
 
 private:
 
diff --git a/Engine/Plugins/Animation/ControlRig/Source/ControlRigEditor/Private/Sequencer/ControlRigSequencerFilter.cpp b/Engine/Plugins/Animation/ControlRig/Source/ControlRigEditor/Private/Sequencer/ControlRigSequencerFilter.cpp
--- a/Engine/Plugins/Animation/ControlRig/Source/ControlRigEditor/Private/Sequencer/ControlRigSequencerFilter.cpp
+++ b/Engine/Plugins/Animation/ControlRig/Source/ControlRigEditor/Private/Sequencer/ControlRigSequencerFilter.cpp
@@ -63,7 +63,7 @@ public:
 
 	virtual FString GetName() const override { return TEXT("ControlRigControlsFilter"); }
 	virtual FText GetDisplayName() const override { return LOCTEXT("SequenceTrackFilter_ControlRigControls", "Control Rig Controls"); }
-	virtual FSlateIcon GetIcon() const { return FSlateIconFinder::FindIconForClass(UControlRig::StaticClass()); }
+	virtual FSlateIcon GetIcon() const { return FSlateIconFinder::FindIconForClass(UBaseControlRig::StaticClass()); }
 
 	virtual bool PassesFilterWithDisplayName(FTrackFilterType InItem, const FText& InText) const
 	{
@@ -158,7 +158,7 @@ public:
 
 	virtual FString GetName() const override { return TEXT("ControlRigControlsSelectedFilter"); }
 	virtual FText GetDisplayName() const override { return LOCTEXT("SequenceTrackFilter_ControlRigSelectedControls", "Selected Control Rig Controls"); }
-	virtual FSlateIcon GetIcon() const { return FSlateIconFinder::FindIconForClass(UControlRig::StaticClass()); }
+	virtual FSlateIcon GetIcon() const { return FSlateIconFinder::FindIconForClass(UBaseControlRig::StaticClass()); }
 
 	virtual bool PassesFilter(FTrackFilterType InItem) const override
 	{
@@ -170,7 +170,7 @@ public:
 		const UMovieSceneControlRigParameterTrack* Track = Cast<UMovieSceneControlRigParameterTrack>(InItem);
 		if (Track)
 		{
-			UControlRig *ControlRig = Track->GetControlRig();
+			UBaseControlRig *ControlRig = Track->GetControlRig();
 			if (ControlRig && ControlRig->GetHierarchy())
 			{
 				FName Name(*InText.ToString());
