Control Rig: Rename ControlRig to BaseControlRig and created a specialized ControlRig

diff --git a/Engine/Plugins/Animation/ControlRig/Source/ControlRigEditor/Private/Sequencer/ControlRigSpaceChannelEditors.cpp b/Engine/Plugins/Animation/ControlRig/Source/ControlRigEditor/Private/Sequencer/ControlRigSpaceChannelEditors.cpp
--- a/Engine/Plugins/Animation/ControlRig/Source/ControlRigEditor/Private/Sequencer/ControlRigSpaceChannelEditors.cpp
+++ b/Engine/Plugins/Animation/ControlRig/Source/ControlRigEditor/Private/Sequencer/ControlRigSpaceChannelEditors.cpp
@@ -70,7 +70,7 @@
 //do not compensate via the FControlRigParameterTrackEditor setting a key.
 static bool bDoNotCompensate = false;
 
-static FKeyHandle SequencerOpenSpaceSwitchDialog(UControlRig* ControlRig, TArray<FRigElementKey> SelectedControls,ISequencer* Sequencer, FMovieSceneControlRigSpaceChannel* Channel, UMovieSceneSection* SectionToKey, FFrameNumber Time)
+static FKeyHandle SequencerOpenSpaceSwitchDialog(UBaseControlRig* ControlRig, TArray<FRigElementKey> SelectedControls,ISequencer* Sequencer, FMovieSceneControlRigSpaceChannel* Channel, UMovieSceneSection* SectionToKey, FFrameNumber Time)
 {
 	FKeyHandle Handle = FKeyHandle::Invalid();
 	if (ControlRig == nullptr || Sequencer == nullptr)
@@ -144,7 +144,7 @@ FKeyHandle AddOrUpdateKey(FMovieSceneControlRigSpaceChannel* Channel, UMovieScen
 	/**
 	if (UMovieSceneControlRigParameterSection* Section = Cast<UMovieSceneControlRigParameterSection>(SectionToKey))
 	{
-		if (UControlRig* ControlRig = Section->GetControlRig())
+		if (UBaseControlRig* ControlRig = Section->GetControlRig())
 		{
 			FName ControlName = Section->FindControlNameFromSpaceChannel(Channel);
 			if (ControlName != NAME_None)
@@ -186,7 +186,7 @@ TSharedRef<SWidget> CreateKeyEditor(const TMovieSceneChannelHandle<FMovieSceneCo
 **********************************************************************/
 
 
-FSpaceChannelAndSection FControlRigSpaceChannelHelpers::FindSpaceChannelAndSectionForControl(UControlRig* ControlRig, FName ControlName, ISequencer* Sequencer, bool bCreateIfNeeded)
+FSpaceChannelAndSection FControlRigSpaceChannelHelpers::FindSpaceChannelAndSectionForControl(UBaseControlRig* ControlRig, FName ControlName, ISequencer* Sequencer, bool bCreateIfNeeded)
 {
 	FSpaceChannelAndSection SpaceChannelAndSection;
 	SpaceChannelAndSection.SpaceChannel = nullptr;
@@ -202,7 +202,7 @@ FSpaceChannelAndSection FControlRigSpaceChannelHelpers::FindSpaceChannelAndSecti
 	}
 	const TArray<FMovieSceneBinding>& Bindings = MovieScene->GetBindings();
 	bool bRecreateCurves = false;
-	TArray<TPair<UControlRig*, FName>> ControlRigPairsToReselect;
+	TArray<TPair<UBaseControlRig*, FName>> ControlRigPairsToReselect;
 	for (const FMovieSceneBinding& Binding : Bindings)
 	{
 		UMovieSceneControlRigParameterTrack* ControlRigParameterTrack = Cast<UMovieSceneControlRigParameterTrack>(MovieScene->FindTrack(UMovieSceneControlRigParameterTrack::StaticClass(), Binding.GetObjectGuid(), NAME_None));
@@ -223,7 +223,7 @@ FSpaceChannelAndSection FControlRigSpaceChannelHelpers::FindSpaceChannelAndSecti
 				{
 					if (ControlRig->IsControlSelected(ControlName))
 					{
-						TPair<UControlRig*, FName> Pair;
+						TPair<UBaseControlRig*, FName> Pair;
 						Pair.Key = ControlRig;
 						Pair.Value = ControlName;
 						ControlRigPairsToReselect.Add(Pair);
@@ -243,7 +243,7 @@ FSpaceChannelAndSection FControlRigSpaceChannelHelpers::FindSpaceChannelAndSecti
 	if (bRecreateCurves)
 	{
 		Sequencer->RecreateCurveEditor(); //this will require the curve editor to get recreated so the ordering is correct
-		for (TPair<UControlRig*, FName>& Pair : ControlRigPairsToReselect)
+		for (TPair<UBaseControlRig*, FName>& Pair : ControlRigPairsToReselect)
 		{
 
 			GEditor->GetTimerManager()->SetTimerForNextTick([Pair]()
@@ -256,7 +256,7 @@ FSpaceChannelAndSection FControlRigSpaceChannelHelpers::FindSpaceChannelAndSecti
 	return SpaceChannelAndSection;
 }
 
-static TTuple<FRigControlElement*, FChannelMapInfo*, int32, int32> GetControlAndChannelInfos(UControlRig* ControlRig, UMovieSceneControlRigParameterSection* ControlRigSection, FName ControlName)
+static TTuple<FRigControlElement*, FChannelMapInfo*, int32, int32> GetControlAndChannelInfos(UBaseControlRig* ControlRig, UMovieSceneControlRigParameterSection* ControlRigSection, FName ControlName)
 {
 	FRigControlElement* ControlElement = nullptr;
 	FChannelMapInfo* pChannelIndex = nullptr;
@@ -286,7 +286,7 @@ static TTuple<FRigControlElement*, FChannelMapInfo*, int32, int32> GetControlAnd
 *  4) For 1-3 we always compensate any transform keys to the new space
 */
 
-FKeyHandle FControlRigSpaceChannelHelpers::SequencerKeyControlRigSpaceChannel(UControlRig* ControlRig, ISequencer* Sequencer, FMovieSceneControlRigSpaceChannel* Channel, UMovieSceneSection* SectionToKey, FFrameNumber Time, URigHierarchy* RigHierarchy, const FRigElementKey& ControlKey, const FRigElementKey& SpaceKey)
+FKeyHandle FControlRigSpaceChannelHelpers::SequencerKeyControlRigSpaceChannel(UBaseControlRig* ControlRig, ISequencer* Sequencer, FMovieSceneControlRigSpaceChannel* Channel, UMovieSceneSection* SectionToKey, FFrameNumber Time, URigHierarchy* RigHierarchy, const FRigElementKey& ControlKey, const FRigElementKey& SpaceKey)
 {
 	FKeyHandle Handle = FKeyHandle::Invalid();
 	if (bDoNotCompensate == true || ControlRig == nullptr || Sequencer == nullptr || Sequencer->GetFocusedMovieSceneSequence() == nullptr || Sequencer->GetFocusedMovieSceneSequence()->GetMovieScene() == nullptr)
@@ -516,7 +516,7 @@ FKeyHandle FControlRigSpaceChannelHelpers::SequencerKeyControlRigSpaceChannel(UC
 }
 
 //get the mask of the transform keys at the specified time
-static EControlRigContextChannelToKey GetCurrentTransformKeysAtThisTime(UControlRig* ControlRig, FName ControlName, UMovieSceneSection* SectionToKey, FFrameNumber Time)
+static EControlRigContextChannelToKey GetCurrentTransformKeysAtThisTime(UBaseControlRig* ControlRig, FName ControlName, UMovieSceneSection* SectionToKey, FFrameNumber Time)
 {
 	EControlRigContextChannelToKey ChannelToKey = EControlRigContextChannelToKey::None;
 	if (ControlRig && SectionToKey)
@@ -635,7 +635,7 @@ static EControlRigContextChannelToKey GetCurrentTransformKeysAtThisTime(UControl
 	return ChannelToKey;
 }
 
-void  FControlRigSpaceChannelHelpers::SequencerSpaceChannelKeyDeleted(UControlRig* ControlRig, ISequencer* Sequencer, FName ControlName, FMovieSceneControlRigSpaceChannel* Channel, UMovieSceneControlRigParameterSection* SectionToKey,
+void  FControlRigSpaceChannelHelpers::SequencerSpaceChannelKeyDeleted(UBaseControlRig* ControlRig, ISequencer* Sequencer, FName ControlName, FMovieSceneControlRigSpaceChannel* Channel, UMovieSceneControlRigParameterSection* SectionToKey,
 	FFrameNumber TimeOfDeletion)
 {
 	if (bDoNotCompensate == true)
@@ -717,7 +717,7 @@ void  FControlRigSpaceChannelHelpers::SequencerSpaceChannelKeyDeleted(UControlRi
 	}
 }
 
-void FControlRigSpaceChannelHelpers::DeleteTransformKeysAtThisTime(UControlRig* ControlRig, UMovieSceneControlRigParameterSection* Section, FName ControlName, FFrameNumber Time)
+void FControlRigSpaceChannelHelpers::DeleteTransformKeysAtThisTime(UBaseControlRig* ControlRig, UMovieSceneControlRigParameterSection* Section, FName ControlName, FFrameNumber Time)
 {
 	if (Section && ControlRig)
 	{
@@ -783,7 +783,7 @@ void FControlRigSpaceChannelHelpers::DeleteTransformKeysAtThisTime(UControlRig*
 }
 
 void FControlRigSpaceChannelHelpers::GetFramesInThisSpaceAfterThisTime(
-	UControlRig* ControlRig,
+	UBaseControlRig* ControlRig,
 	FName ControlName,
 	FMovieSceneControlRigSpaceBaseKey CurrentValue,
 	FMovieSceneControlRigSpaceChannel* Channel,
@@ -860,7 +860,7 @@ void FControlRigSpaceChannelHelpers::GetFramesInThisSpaceAfterThisTime(
 	}
 }
 
-void FControlRigSpaceChannelHelpers::SequencerBakeControlInSpace(UControlRig* ControlRig, ISequencer* Sequencer, FMovieSceneControlRigSpaceChannel* Channel, UMovieSceneSection* SectionToKey,
+void FControlRigSpaceChannelHelpers::SequencerBakeControlInSpace(UBaseControlRig* ControlRig, ISequencer* Sequencer, FMovieSceneControlRigSpaceChannel* Channel, UMovieSceneSection* SectionToKey,
 	URigHierarchy* RigHierarchy, const FRigElementKey& ControlKey, FRigSpacePickerBakeSettings Settings)
 {
 	if (bDoNotCompensate == true)
@@ -1176,7 +1176,7 @@ void FControlRigSpaceChannelHelpers::SequencerBakeControlInSpace(UControlRig* Co
 	}
 }
 
-void FControlRigSpaceChannelHelpers::HandleSpaceKeyTimeChanged(UControlRig* ControlRig, FName ControlName, FMovieSceneControlRigSpaceChannel* Channel, UMovieSceneSection* SectionToKey,
+void FControlRigSpaceChannelHelpers::HandleSpaceKeyTimeChanged(UBaseControlRig* ControlRig, FName ControlName, FMovieSceneControlRigSpaceChannel* Channel, UMovieSceneSection* SectionToKey,
 	FFrameNumber CurrentFrame, FFrameNumber NextFrame)
 {
 	if (ControlRig && Channel && SectionToKey && (CurrentFrame != NextFrame))
@@ -1260,7 +1260,7 @@ void FControlRigSpaceChannelHelpers::HandleSpaceKeyTimeChanged(UControlRig* Cont
 
 
 void FControlRigSpaceChannelHelpers::CompensateIfNeeded(
-	UControlRig* ControlRig,
+	UBaseControlRig* ControlRig,
 	ISequencer* Sequencer,
 	UMovieSceneControlRigParameterSection* Section,
 	FName ControlName,
@@ -1421,7 +1421,7 @@ FReply FControlRigSpaceChannelHelpers::OpenBakeDialog(ISequencer* Sequencer, FMo
 	}
 	if (UMovieSceneControlRigParameterSection* Section = Cast<UMovieSceneControlRigParameterSection>(SectionToKey))
 	{
-		if (UControlRig* ControlRig = Section->GetControlRig())
+		if (UBaseControlRig* ControlRig = Section->GetControlRig())
 		{
 			FName ControlName = Section->FindControlNameFromSpaceChannel(Channel);
 
@@ -1530,7 +1530,7 @@ TArray<FKeyBarCurveModel::FBarRange> FControlRigSpaceChannelHelpers::FindRanges(
 }
 
 // NOTE use this function in HandleSpaceKeyTimeChanged, DeleteTransformKeysAtThisTime, ...
-TPair<FRigControlElement*, FChannelMapInfo*> FControlRigSpaceChannelHelpers::GetControlAndChannelInfo(UControlRig* ControlRig, UMovieSceneControlRigParameterSection* Section, FName ControlName)
+TPair<FRigControlElement*, FChannelMapInfo*> FControlRigSpaceChannelHelpers::GetControlAndChannelInfo(UBaseControlRig* ControlRig, UMovieSceneControlRigParameterSection* Section, FName ControlName)
 {
 	FRigControlElement* ControlElement = ControlRig ? ControlRig->FindControl(ControlName) : nullptr;
 	FChannelMapInfo* pChannelIndex = Section ? Section->ControlChannelMap.Find(ControlName) : nullptr;
diff --git a/Engine/Plugins/Animation/ControlRig/Source/ControlRigEditor/Private/Tools/AssetDefinition_ControlRigPose.cpp b/Engine/Plugins/Animation/ControlRig/Source/ControlRigEditor/Private/Tools/AssetDefinition_ControlRigPose.cpp
--- a/Engine/Plugins/Animation/ControlRig/Source/ControlRigEditor/Private/Tools/AssetDefinition_ControlRigPose.cpp
+++ b/Engine/Plugins/Animation/ControlRig/Source/ControlRigEditor/Private/Tools/AssetDefinition_ControlRigPose.cpp
@@ -32,14 +32,14 @@ namespace MenuExtension_AnimationAsset
 			return;
 		}
 		
-		TMap<UControlRig*, TArray<FRigElementKey>> AllSelectedControls;
+		TMap<UBaseControlRig*, TArray<FRigElementKey>> AllSelectedControls;
 		ControlRigEditMode->GetAllSelectedControls(AllSelectedControls);
-		TArray<UControlRig*> ControlRigs;
+		TArray<UBaseControlRig*> ControlRigs;
 		AllSelectedControls.GenerateKeyArray(ControlRigs);
 			
 		for (UControlRigPoseAsset* ControlRigPoseAsset : ControlRigPoseAssets)
 		{
-			for (UControlRig* ControlRig : ControlRigs)
+			for (UBaseControlRig* ControlRig : ControlRigs)
 			{
 				ControlRigPoseAsset->PastePose(ControlRig, false, false);
 			}
@@ -57,14 +57,14 @@ namespace MenuExtension_AnimationAsset
 			return;
 		}
 
-		TMap<UControlRig*, TArray<FRigElementKey>> AllSelectedControls;
+		TMap<UBaseControlRig*, TArray<FRigElementKey>> AllSelectedControls;
 		ControlRigEditMode->GetAllSelectedControls(AllSelectedControls);
-		TArray<UControlRig*> ControlRigs;
+		TArray<UBaseControlRig*> ControlRigs;
 		AllSelectedControls.GenerateKeyArray(ControlRigs);
 
 		for (UControlRigPoseAsset* ControlRigPoseAsset : ControlRigPoseAssets)
 		{
-			for (UControlRig* ControlRig : ControlRigs)
+			for (UBaseControlRig* ControlRig : ControlRigs)
 			{
 				ControlRigPoseAsset->SelectControls(ControlRig);
 			}
@@ -82,14 +82,14 @@ namespace MenuExtension_AnimationAsset
 			return;
 		}
 
-		TMap<UControlRig*, TArray<FRigElementKey>> AllSelectedControls;
+		TMap<UBaseControlRig*, TArray<FRigElementKey>> AllSelectedControls;
 		ControlRigEditMode->GetAllSelectedControls(AllSelectedControls);
-		TArray<UControlRig*> ControlRigs;
+		TArray<UBaseControlRig*> ControlRigs;
 		AllSelectedControls.GenerateKeyArray(ControlRigs);
 
 		for (UControlRigPoseAsset* ControlRigPoseAsset : ControlRigPoseAssets)
 		{
-			for (UControlRig* ControlRig : ControlRigs)
+			for (UBaseControlRig* ControlRig : ControlRigs)
 			{
 				ControlRigPoseAsset->SavePose(ControlRig, false);
 			}
diff --git a/Engine/Plugins/Animation/ControlRig/Source/ControlRigEditor/Private/Tools/ControlRigPoseThumbnailRenderer.cpp b/Engine/Plugins/Animation/ControlRig/Source/ControlRigEditor/Private/Tools/ControlRigPoseThumbnailRenderer.cpp
--- a/Engine/Plugins/Animation/ControlRig/Source/ControlRigEditor/Private/Tools/ControlRigPoseThumbnailRenderer.cpp
+++ b/Engine/Plugins/Animation/ControlRig/Source/ControlRigEditor/Private/Tools/ControlRigPoseThumbnailRenderer.cpp
@@ -67,7 +67,7 @@ bool FControlRigPoseThumbnailScene::SetControlRigPoseAsset(UControlRigPoseAsset*
 	if (PoseAsset)
 	{
 		USkeletalMesh* SkeletalMesh = nullptr; // PoseAsset->GetSkeletalMeshAsset();
-		UControlRig* ControlRig = nullptr; // PoseAsset->GetControlRig();
+		UBaseControlRig* ControlRig = nullptr; // PoseAsset->GetControlRig();
 		PreviewActor->GetSkeletalMeshComponent()->OverrideMaterials.Empty();
 
 
@@ -80,10 +80,10 @@ bool FControlRigPoseThumbnailScene::SetControlRigPoseAsset(UControlRigPoseAsset*
 
 			bSetSucessfully = true;
 			//now set up the control rig and the anim instance fo rit
-			UControlRig* TempControlRig = NewObject<UControlRig>(GetTransientPackage(), ControlRig->GetClass(),NAME_None, RF_Transient);
+			UBaseControlRig* TempControlRig = NewObject<UBaseControlRig>(GetTransientPackage(), ControlRig->GetClass(),NAME_None, RF_Transient);
 			TempControlRig->SetObjectBinding(MakeShared<FControlRigObjectBinding>());
 			TempControlRig->GetObjectBinding()->BindToObject(SkeletalMesh);
-			TempControlRig->GetDataSourceRegistry()->RegisterDataSource(UControlRig::OwnerComponent, TempControlRig->GetObjectBinding()->GetBoundObject());
+			TempControlRig->GetDataSourceRegistry()->RegisterDataSource(UBaseControlRig::OwnerComponent, TempControlRig->GetObjectBinding()->GetBoundObject());
 
 
 			bool bWasCreated;
diff --git a/Engine/Plugins/Animation/ControlRig/Source/ControlRigEditor/Private/Tools/ControlRigSnapper.cpp b/Engine/Plugins/Animation/ControlRig/Source/ControlRigEditor/Private/Tools/ControlRigSnapper.cpp
--- a/Engine/Plugins/Animation/ControlRig/Source/ControlRigEditor/Private/Tools/ControlRigSnapper.cpp
+++ b/Engine/Plugins/Animation/ControlRig/Source/ControlRigEditor/Private/Tools/ControlRigSnapper.cpp
@@ -112,7 +112,7 @@ TWeakPtr<ISequencer> FControlRigSnapper::GetSequencer()
 }
 
 static bool LocalGetControlRigControlTransforms(IMovieScenePlayer* Player, const TOptional<FFrameNumber>& CurrentFrame, UMovieSceneSequence* MovieSceneSequence, FMovieSceneSequenceIDRef Template, FMovieSceneSequenceTransform& RootToLocalTransform,
-	UControlRig* ControlRig, const FName& ControlName,
+	UBaseControlRig* ControlRig, const FName& ControlName,
 	const TArray<FFrameNumber>& Frames, const TArray<FTransform>& ParentTransforms, TArray<FTransform>& OutTransforms)
 {
 	if (Frames.Num() > ParentTransforms.Num())
@@ -160,7 +160,7 @@ static bool LocalGetControlRigControlTransforms(IMovieScenePlayer* Player, const
 	return true;
 }
 
-bool FControlRigSnapper::GetControlRigControlTransforms(ISequencer* Sequencer,  UControlRig* ControlRig, const FName& ControlName,
+bool FControlRigSnapper::GetControlRigControlTransforms(ISequencer* Sequencer,  UBaseControlRig* ControlRig, const FName& ControlName,
 	const TArray<FFrameNumber> &Frames, const TArray<FTransform>& ParentTransforms,TArray<FTransform>& OutTransforms)
 {
 	if (Sequencer->GetFocusedMovieSceneSequence())
@@ -175,7 +175,7 @@ bool FControlRigSnapper::GetControlRigControlTransforms(ISequencer* Sequencer,
 	return false;
 }
 
-bool FControlRigSnapper::GetControlRigControlTransforms(UWorld* World,ULevelSequence* LevelSequence, UControlRig* ControlRig, const FName& ControlName,
+bool FControlRigSnapper::GetControlRigControlTransforms(UWorld* World,ULevelSequence* LevelSequence, UBaseControlRig* ControlRig, const FName& ControlName,
 	const TArray<FFrameNumber>& Frames, const TArray<FTransform>& ParentTransforms, TArray<FTransform>& OutTransforms)
 {
 	if (LevelSequence)
@@ -297,7 +297,7 @@ static bool CalculateWorldTransformsFromParents(ISequencer* Sequencer, const FCo
 	for (const FControlRigForWorldTransforms& ControlRigAndSelection : ParentToSnap.ControlRigs)
 	{
 		//get actor transform...
-		UControlRig* ControlRig = ControlRigAndSelection.ControlRig.Get();
+		UBaseControlRig* ControlRig = ControlRigAndSelection.ControlRig.Get();
 
 		if (ControlRig)
 		{
@@ -341,7 +341,7 @@ static void GetControlRigParents(const FControlRigForWorldTransforms& ControlRig
 	TArray<FActorForWorldTransforms>& OutParentActors, TArray<FControlRigForWorldTransforms>& OutControlRigs)
 {
 
-	if (UControlRig* ControlRig = ControlRigAndSelection.ControlRig.Get())
+	if (UBaseControlRig* ControlRig = ControlRigAndSelection.ControlRig.Get())
 	{
 		URigHierarchy* Hierarchy = ControlRig->GetHierarchy();
 		for (const FName& ControlName : ControlRigAndSelection.ControlNames)
@@ -381,7 +381,7 @@ static void GetControlRigParents(const FControlRigForWorldTransforms& ControlRig
 	}
 }
 
-static UMovieSceneControlRigParameterSection* GetControlRigSection(ISequencer* Sequencer, const UControlRig* ControlRig)
+static UMovieSceneControlRigParameterSection* GetControlRigSection(ISequencer* Sequencer, const UBaseControlRig* ControlRig)
 {
 
 	if (ControlRig == nullptr || Sequencer == nullptr)
@@ -587,7 +587,7 @@ bool FControlRigSnapper::SnapIt(FFrameNumber StartFrame, FFrameNumber EndFrame,c
 		for (const FControlRigForWorldTransforms& ControlRigAndSelection : ActorToSnap.ControlRigs)
 		{
 			//get actor transform...
-			UControlRig* ControlRig = ControlRigAndSelection.ControlRig.Get();
+			UBaseControlRig* ControlRig = ControlRigAndSelection.ControlRig.Get();
 
 			if (ControlRig)
 			{
diff --git a/Engine/Plugins/Animation/ControlRig/Source/ControlRigEditor/Private/Tools/ControlRigTweener.cpp b/Engine/Plugins/Animation/ControlRig/Source/ControlRigEditor/Private/Tools/ControlRigTweener.cpp
--- a/Engine/Plugins/Animation/ControlRig/Source/ControlRigEditor/Private/Tools/ControlRigTweener.cpp
+++ b/Engine/Plugins/Animation/ControlRig/Source/ControlRigEditor/Private/Tools/ControlRigTweener.cpp
@@ -107,7 +107,7 @@ bool FAnimSliderKeySelection::Setup(const TWeakPtr<ISequencer>& InSequencer)
 bool FAnimSliderObjectSelection::Setup(TWeakPtr<ISequencer>& InSequencer, TWeakPtr<FControlRigEditMode>& InEditMode)
 {
 	ChannelsArray.Reset();
-	const TArray<UControlRig*> ControlRigs = GetControlRigs(InEditMode);
+	const TArray<UBaseControlRig*> ControlRigs = GetControlRigs(InEditMode);
 	return Setup(ControlRigs, InSequencer);
 }
 
@@ -183,14 +183,14 @@ void FAnimSliderKeySelection::GetMapOfContiguousKeys()
 * FAnimSliderObjectSelection
 *
 */
-TArray<UControlRig*> FAnimSliderObjectSelection::GetControlRigs(TWeakPtr<FControlRigEditMode>& InEditMode)
+TArray<UBaseControlRig*> FAnimSliderObjectSelection::GetControlRigs(TWeakPtr<FControlRigEditMode>& InEditMode)
 {
-	TArray<UControlRig*> ControlRigs;
+	TArray<UBaseControlRig*> ControlRigs;
 	if (const FControlRigEditMode* EditMode = InEditMode.Pin().Get())
 	{
-		TMap<UControlRig*, TArray<FRigElementKey>> SelectedControls;
+		TMap<UBaseControlRig*, TArray<FRigElementKey>> SelectedControls;
 		EditMode->GetAllSelectedControls(SelectedControls);
-		for (TPair<UControlRig*, TArray<FRigElementKey>>& Selected : SelectedControls)
+		for (TPair<UBaseControlRig*, TArray<FRigElementKey>>& Selected : SelectedControls)
 		{
 			ControlRigs.Add(Selected.Key);
 		}
@@ -273,7 +273,7 @@ void FAnimSliderObjectSelection::SetupChannel(FFrameNumber CurrentFrame, TArray<
 		&& KeyBounds.PreviousIndex != KeyBounds.NextIndex) ? true : false;
 }
 
-bool FAnimSliderObjectSelection::Setup(const TArray<UControlRig*>& SelectedControlRigs, TWeakPtr<ISequencer>& InSequencer)
+bool FAnimSliderObjectSelection::Setup(const TArray<UBaseControlRig*>& SelectedControlRigs, TWeakPtr<ISequencer>& InSequencer)
 {
 	ChannelsArray.Reset();
 	if (InSequencer.IsValid() == false)
@@ -299,7 +299,7 @@ bool FAnimSliderObjectSelection::Setup(const TArray<UControlRig*>& SelectedContr
 			if (Section && Section->IsActive() && Section->GetRange().Contains(CurrentTime.Time.GetFrame()) && !HandledSections.Contains(Section))
 			{
 				HandledSections.Add(Section);				
-				UControlRig* ControlRig = Track->GetControlRig();
+				UBaseControlRig* ControlRig = Track->GetControlRig();
 				TArray<FRigControlElement*> CurrentControls;
 				ControlRig->GetControlsInOrder(CurrentControls);
 
@@ -613,7 +613,7 @@ FText FControlsToTween::GetTooltipText() const
 	return LOCTEXT("TweenControllerTooltip", "Tween between the next and previous keys");
 }
 
-bool FControlsToTween::Setup(const TArray<UControlRig*>& SelectedControlRigs, TWeakPtr<ISequencer>& InSequencer)
+bool FControlsToTween::Setup(const TArray<UBaseControlRig*>& SelectedControlRigs, TWeakPtr<ISequencer>& InSequencer)
 {
 	return ObjectSelection.Setup(SelectedControlRigs, InSequencer);
 }
diff --git a/Engine/Plugins/Animation/ControlRig/Source/ControlRigEditor/Public/ControlRigSequencerEditorLibrary.h b/Engine/Plugins/Animation/ControlRig/Source/ControlRigEditor/Public/ControlRigSequencerEditorLibrary.h
--- a/Engine/Plugins/Animation/ControlRig/Source/ControlRigEditor/Public/ControlRigSequencerEditorLibrary.h
+++ b/Engine/Plugins/Animation/ControlRig/Source/ControlRigEditor/Public/ControlRigSequencerEditorLibrary.h
@@ -24,7 +24,7 @@ class UTickableConstraint;
 class UTickableTransformConstraint;
 class UTransformableHandle;
 struct FBakingAnimationKeySettings;
-class UControlRig;
+class UBaseControlRig;
 
 USTRUCT(BlueprintType)
 struct FControlRigSequencerBindingProxy
@@ -36,7 +36,7 @@ struct FControlRigSequencerBindingProxy
 		, Track(nullptr)
 	{}
 
-	FControlRigSequencerBindingProxy(const FMovieSceneBindingProxy& InProxy, UControlRig* InControlRig, UMovieSceneControlRigParameterTrack* InTrack)
+	FControlRigSequencerBindingProxy(const FMovieSceneBindingProxy& InProxy, UBaseControlRig* InControlRig, UMovieSceneControlRigParameterTrack* InTrack)
 		: Proxy(InProxy)
 		, ControlRig(InControlRig)
 		, Track(InTrack)
@@ -46,7 +46,7 @@ struct FControlRigSequencerBindingProxy
 	FMovieSceneBindingProxy Proxy;
 
 	UPROPERTY(BlueprintReadOnly, Category = ControlRig)
-	TObjectPtr<UControlRig> ControlRig;
+	TObjectPtr<UBaseControlRig> ControlRig;
 
 	UPROPERTY(BlueprintReadOnly, Category = ControlRig)
 	TObjectPtr<UMovieSceneControlRigParameterTrack> Track;
@@ -80,7 +80,7 @@ public:
 	* @return returns list of visible Control Rigs
 	*/
 	UFUNCTION(BlueprintCallable, Category = "Editor Scripting | Sequencer Tools | Control Rig")
-	static TArray<UControlRig*> GetVisibleControlRigs();
+	static TArray<UBaseControlRig*> GetVisibleControlRigs();
 
 	/**
 	* Get all of the control rigs and their bindings in the level sequence
@@ -251,7 +251,7 @@ public:
 	* @param TweenValue The tween value to use, range from -1(blend to previous) to 1(blend to next)
 	*/
 	UFUNCTION(BlueprintCallable, Category = "Editor Scripting | Sequencer Tools | Control Rig")
-	static bool TweenControlRig(ULevelSequence* LevelSequence, UControlRig* ControlRig, float TweenValue);
+	static bool TweenControlRig(ULevelSequence* LevelSequence, UBaseControlRig* ControlRig, float TweenValue);
 	
 	/**
 	* Peform specified blend operation based upon selected keys in the curve editor or selected control rig controls
@@ -336,7 +336,7 @@ public:
 	* @return Returns World Transform
 	*/
 	UFUNCTION(BlueprintCallable, Category = "Editor Scripting | Sequencer Tools | Control Rig")
-	static FTransform GetControlRigWorldTransform(ULevelSequence* LevelSequence, UControlRig* ControlRig, FName ControlName, FFrameNumber Frame,
+	static FTransform GetControlRigWorldTransform(ULevelSequence* LevelSequence, UBaseControlRig* ControlRig, FName ControlName, FFrameNumber Frame,
 		ESequenceTimeUnit TimeUnit = ESequenceTimeUnit::DisplayRate);
 
 	/**
@@ -349,7 +349,7 @@ public:
 	* @return Returns World Transforms
 	*/
 	UFUNCTION(BlueprintCallable, Category = "Editor Scripting | Sequencer Tools | Control Rig")
-	static TArray<FTransform> GetControlRigWorldTransforms(ULevelSequence* LevelSequence, UControlRig* ControlRig, FName ControlName, const TArray<FFrameNumber>& Frames,
+	static TArray<FTransform> GetControlRigWorldTransforms(ULevelSequence* LevelSequence, UBaseControlRig* ControlRig, FName ControlName, const TArray<FFrameNumber>& Frames,
 		ESequenceTimeUnit TimeUnit = ESequenceTimeUnit::DisplayRate);
 
 	/**
@@ -363,7 +363,7 @@ public:
 	* @param bSetKey Whether or not to set a key.
 	*/
 	UFUNCTION(BlueprintCallable, Category = "Editor Scripting | Sequencer Tools | Control Rig")
-	static void SetControlRigWorldTransform(ULevelSequence* LevelSequence, UControlRig* ControlRig, FName ControlName, FFrameNumber Frame, const FTransform& WorldTransform,
+	static void SetControlRigWorldTransform(ULevelSequence* LevelSequence, UBaseControlRig* ControlRig, FName ControlName, FFrameNumber Frame, const FTransform& WorldTransform,
 		ESequenceTimeUnit TimeUnit = ESequenceTimeUnit::DisplayRate, bool bSetKey = true);
 
 	/**
@@ -376,7 +376,7 @@ public:
 	* @param TimeUnit Unit for frame values, either in display rate or tick resolution
 	*/
 	UFUNCTION(BlueprintCallable, Category = "Editor Scripting | Sequencer Tools | Control Rig")
-	static void SetControlRigWorldTransforms(ULevelSequence* LevelSequence, UControlRig* ControlRig, FName ControlName, const TArray<FFrameNumber>& Frames, 
+	static void SetControlRigWorldTransforms(ULevelSequence* LevelSequence, UBaseControlRig* ControlRig, FName ControlName, const TArray<FFrameNumber>& Frames, 
 		 const TArray<FTransform>& WorldTransforms, ESequenceTimeUnit TimeUnit = ESequenceTimeUnit::DisplayRate);
 	
 	/**
@@ -389,7 +389,7 @@ public:
 	* @return Returns Value at that time
 	*/
 	UFUNCTION(BlueprintCallable, Category = "Editor Scripting | Sequencer Tools | Control Rig")
-	static float GetLocalControlRigFloat(ULevelSequence* LevelSequence, UControlRig* ControlRig, FName ControlName, FFrameNumber Frame,
+	static float GetLocalControlRigFloat(ULevelSequence* LevelSequence, UBaseControlRig* ControlRig, FName ControlName, FFrameNumber Frame,
 		ESequenceTimeUnit TimeUnit = ESequenceTimeUnit::DisplayRate);
 
 	/**
@@ -402,7 +402,7 @@ public:
 	* @return Returns Values at those times
 	*/
 	UFUNCTION(BlueprintCallable, Category = "Editor Scripting | Sequencer Tools | Control Rig")
-	static TArray<float> GetLocalControlRigFloats(ULevelSequence* LevelSequence, UControlRig* ControlRig, FName ControlName, const TArray<FFrameNumber>& Frames,
+	static TArray<float> GetLocalControlRigFloats(ULevelSequence* LevelSequence, UBaseControlRig* ControlRig, FName ControlName, const TArray<FFrameNumber>& Frames,
 		ESequenceTimeUnit TimeUnit = ESequenceTimeUnit::DisplayRate);
 
 	/**
@@ -416,7 +416,7 @@ public:
 	* @param bSetKey If True set a key, if not just set the value
 	*/
 	UFUNCTION(BlueprintCallable, Category = "Editor Scripting | Sequencer Tools | Control Rig")
-	static void SetLocalControlRigFloat(ULevelSequence* LevelSequence, UControlRig* ControlRig, FName ControlName, FFrameNumber Frame, float Value,
+	static void SetLocalControlRigFloat(ULevelSequence* LevelSequence, UBaseControlRig* ControlRig, FName ControlName, FFrameNumber Frame, float Value,
 		ESequenceTimeUnit TimeUnit = ESequenceTimeUnit::DisplayRate,bool bSetKey = true);
 
 	/**
@@ -429,7 +429,7 @@ public:
 	* @param TimeUnit Unit for frame values, either in display rate or tick resolution
 	*/
 	UFUNCTION(BlueprintCallable, Category = "Editor Scripting | Sequencer Tools | Control Rig")
-	static void SetLocalControlRigFloats(ULevelSequence* LevelSequence, UControlRig* ControlRig, FName ControlName, const TArray<FFrameNumber>& Frames, const TArray<float> Values,
+	static void SetLocalControlRigFloats(ULevelSequence* LevelSequence, UBaseControlRig* ControlRig, FName ControlName, const TArray<FFrameNumber>& Frames, const TArray<float> Values,
 		ESequenceTimeUnit TimeUnit = ESequenceTimeUnit::DisplayRate);
 
 	/**
@@ -442,7 +442,7 @@ public:
 	* @return Returns Value at that time
 	*/
 	UFUNCTION(BlueprintCallable, Category = "Editor Scripting | Sequencer Tools | Control Rig")
-	static bool GetLocalControlRigBool(ULevelSequence* LevelSequence, UControlRig* ControlRig, FName ControlName, FFrameNumber Frame,
+	static bool GetLocalControlRigBool(ULevelSequence* LevelSequence, UBaseControlRig* ControlRig, FName ControlName, FFrameNumber Frame,
 		ESequenceTimeUnit TimeUnit = ESequenceTimeUnit::DisplayRate);
 
 	/**
@@ -455,7 +455,7 @@ public:
 	* @return Returns Values at those times
 	*/
 	UFUNCTION(BlueprintCallable, Category = "Editor Scripting | Sequencer Tools | Control Rig")
-	static TArray<bool> GetLocalControlRigBools(ULevelSequence* LevelSequence, UControlRig* ControlRig, FName ControlName, const TArray<FFrameNumber>& Frames,
+	static TArray<bool> GetLocalControlRigBools(ULevelSequence* LevelSequence, UBaseControlRig* ControlRig, FName ControlName, const TArray<FFrameNumber>& Frames,
 		ESequenceTimeUnit TimeUnit = ESequenceTimeUnit::DisplayRate);
 
 	/**
@@ -469,7 +469,7 @@ public:
 	* @param bSetKey If True set a key, if not just set the value
 	*/
 	UFUNCTION(BlueprintCallable, Category = "Editor Scripting | Sequencer Tools | Control Rig")
-	static void SetLocalControlRigBool(ULevelSequence* LevelSequence, UControlRig* ControlRig, FName ControlName, FFrameNumber Frame, bool Value,
+	static void SetLocalControlRigBool(ULevelSequence* LevelSequence, UBaseControlRig* ControlRig, FName ControlName, FFrameNumber Frame, bool Value,
 		ESequenceTimeUnit TimeUnit = ESequenceTimeUnit::DisplayRate, bool bSetKey = true);
 
 	/**
@@ -482,7 +482,7 @@ public:
 	* @param TimeUnit Unit for frame values, either in display rate or tick resolution
 */
 	UFUNCTION(BlueprintCallable, Category = "Editor Scripting | Sequencer Tools | Control Rig")
-	static void SetLocalControlRigBools(ULevelSequence* LevelSequence, UControlRig* ControlRig, FName ControlName, const TArray<FFrameNumber>& Frames,
+	static void SetLocalControlRigBools(ULevelSequence* LevelSequence, UBaseControlRig* ControlRig, FName ControlName, const TArray<FFrameNumber>& Frames,
 		const TArray<bool> Values, ESequenceTimeUnit TimeUnit = ESequenceTimeUnit::DisplayRate);
 
 	/**
@@ -495,7 +495,7 @@ public:
 	* @return Returns Value at that time
 	*/
 	UFUNCTION(BlueprintCallable, Category = "Editor Scripting | Sequencer Tools | Control Rig")
-	static int32 GetLocalControlRigInt(ULevelSequence* LevelSequence, UControlRig* ControlRig, FName ControlName, FFrameNumber Frame, 
+	static int32 GetLocalControlRigInt(ULevelSequence* LevelSequence, UBaseControlRig* ControlRig, FName ControlName, FFrameNumber Frame, 
 		ESequenceTimeUnit TimeUnit = ESequenceTimeUnit::DisplayRate);
 
 	/**
@@ -508,7 +508,7 @@ public:
 	* @return Returns Values at those times
 	*/
 	UFUNCTION(BlueprintCallable, Category = "Editor Scripting | Sequencer Tools | Control Rig")
-	static TArray<int32> GetLocalControlRigInts(ULevelSequence* LevelSequence, UControlRig* ControlRig, FName ControlName, const TArray<FFrameNumber>& Frames,
+	static TArray<int32> GetLocalControlRigInts(ULevelSequence* LevelSequence, UBaseControlRig* ControlRig, FName ControlName, const TArray<FFrameNumber>& Frames,
 		ESequenceTimeUnit TimeUnit = ESequenceTimeUnit::DisplayRate);
 
 	/**
@@ -522,7 +522,7 @@ public:
 	* @param bSetKey If True set a key, if not just set the value
 	*/
 	UFUNCTION(BlueprintCallable, Category = "Editor Scripting | Sequencer Tools | Control Rig")
-	static void SetLocalControlRigInt(ULevelSequence* LevelSequence, UControlRig* ControlRig, FName ControlName, FFrameNumber Frame, int32 Value,
+	static void SetLocalControlRigInt(ULevelSequence* LevelSequence, UBaseControlRig* ControlRig, FName ControlName, FFrameNumber Frame, int32 Value,
 		ESequenceTimeUnit TimeUnit = ESequenceTimeUnit::DisplayRate,bool bSetKey = true);
 
 
@@ -536,7 +536,7 @@ public:
 	* @param TimeUnit Unit for frame values, either in display rate or tick resolution
 	*/
 	UFUNCTION(BlueprintCallable, Category = "Editor Scripting | Sequencer Tools | Control Rig")
-	static void SetLocalControlRigInts(ULevelSequence* LevelSequence, UControlRig* ControlRig, FName ControlName, const TArray<FFrameNumber>& Frames, const TArray<int32> Values,
+	static void SetLocalControlRigInts(ULevelSequence* LevelSequence, UBaseControlRig* ControlRig, FName ControlName, const TArray<FFrameNumber>& Frames, const TArray<int32> Values,
 		ESequenceTimeUnit TimeUnit = ESequenceTimeUnit::DisplayRate);
 
 	/**
@@ -549,7 +549,7 @@ public:
 	* @return Returns Value at that time
 	*/
 	UFUNCTION(BlueprintCallable, Category = "Editor Scripting | Sequencer Tools | Control Rig")
-	static FVector2D GetLocalControlRigVector2D(ULevelSequence* LevelSequence, UControlRig* ControlRig, FName ControlName, FFrameNumber Frame,
+	static FVector2D GetLocalControlRigVector2D(ULevelSequence* LevelSequence, UBaseControlRig* ControlRig, FName ControlName, FFrameNumber Frame,
 		ESequenceTimeUnit TimeUnit = ESequenceTimeUnit::DisplayRate);
 
 	/**
@@ -562,7 +562,7 @@ public:
 	* @return Returns Values at those times
 	*/
 	UFUNCTION(BlueprintCallable, Category = "Editor Scripting | Sequencer Tools | Control Rig")
-	static TArray<FVector2D> GetLocalControlRigVector2Ds(ULevelSequence* LevelSequence, UControlRig* ControlRig, FName ControlName, const TArray<FFrameNumber>& Frames,
+	static TArray<FVector2D> GetLocalControlRigVector2Ds(ULevelSequence* LevelSequence, UBaseControlRig* ControlRig, FName ControlName, const TArray<FFrameNumber>& Frames,
 		ESequenceTimeUnit TimeUnit = ESequenceTimeUnit::DisplayRate);
 
 	/**
@@ -576,7 +576,7 @@ public:
 	* @param bSetKey If True set a key, if not just set the value
 	*/
 	UFUNCTION(BlueprintCallable, Category = "Editor Scripting | Sequencer Tools | Control Rig")
-	static void SetLocalControlRigVector2D(ULevelSequence* LevelSequence, UControlRig* ControlRig, FName ControlName, FFrameNumber Frame, FVector2D Value,
+	static void SetLocalControlRigVector2D(ULevelSequence* LevelSequence, UBaseControlRig* ControlRig, FName ControlName, FFrameNumber Frame, FVector2D Value,
 		ESequenceTimeUnit TimeUnit = ESequenceTimeUnit::DisplayRate,bool bSetKey = true);
 
 
@@ -590,7 +590,7 @@ public:
 	* @param TimeUnit Unit for frame values, either in display rate or tick resolution
 	*/
 	UFUNCTION(BlueprintCallable, Category = "Editor Scripting | Sequencer Tools | Control Rig")
-	static void SetLocalControlRigVector2Ds(ULevelSequence* LevelSequence, UControlRig* ControlRig, FName ControlName, const TArray<FFrameNumber>& Frames, const TArray<FVector2D> Values,
+	static void SetLocalControlRigVector2Ds(ULevelSequence* LevelSequence, UBaseControlRig* ControlRig, FName ControlName, const TArray<FFrameNumber>& Frames, const TArray<FVector2D> Values,
 		ESequenceTimeUnit TimeUnit = ESequenceTimeUnit::DisplayRate);
 
 	/**
@@ -603,7 +603,7 @@ public:
 	* @return Returns Value at that time
 	*/
 	UFUNCTION(BlueprintCallable, Category = "Editor Scripting | Sequencer Tools | Control Rig")
-	static FVector GetLocalControlRigPosition(ULevelSequence* LevelSequence, UControlRig* ControlRig, FName ControlName, FFrameNumber Frame,
+	static FVector GetLocalControlRigPosition(ULevelSequence* LevelSequence, UBaseControlRig* ControlRig, FName ControlName, FFrameNumber Frame,
 		ESequenceTimeUnit TimeUnit = ESequenceTimeUnit::DisplayRate);
 
 	/**
@@ -616,7 +616,7 @@ public:
 	* @return Returns Values at those times
 	*/
 	UFUNCTION(BlueprintCallable, Category = "Editor Scripting | Sequencer Tools | Control Rig")
-	static TArray<FVector> GetLocalControlRigPositions(ULevelSequence* LevelSequence, UControlRig* ControlRig, FName ControlName, const TArray<FFrameNumber>& Frames,
+	static TArray<FVector> GetLocalControlRigPositions(ULevelSequence* LevelSequence, UBaseControlRig* ControlRig, FName ControlName, const TArray<FFrameNumber>& Frames,
 		ESequenceTimeUnit TimeUnit = ESequenceTimeUnit::DisplayRate);
 
 	/**
@@ -630,7 +630,7 @@ public:
 	* @param bSetKey If True set a key, if not just set the value
 	*/
 	UFUNCTION(BlueprintCallable, Category = "Editor Scripting | Sequencer Tools | Control Rig")
-	static void SetLocalControlRigPosition(ULevelSequence* LevelSequence, UControlRig* ControlRig, FName ControlName, FFrameNumber Frame, FVector Value,
+	static void SetLocalControlRigPosition(ULevelSequence* LevelSequence, UBaseControlRig* ControlRig, FName ControlName, FFrameNumber Frame, FVector Value,
 		ESequenceTimeUnit TimeUnit = ESequenceTimeUnit::DisplayRate, bool bSetKey = true);
 
 	/**
@@ -643,7 +643,7 @@ public:
 	* @param TimeUnit Unit for frame values, either in display rate or tick resolution
 	*/
 	UFUNCTION(BlueprintCallable, Category = "Editor Scripting | Sequencer Tools | Control Rig")
-	static void SetLocalControlRigPositions(ULevelSequence* LevelSequence, UControlRig* ControlRig, FName ControlName, const TArray<FFrameNumber>& Frames, const TArray<FVector> Values,
+	static void SetLocalControlRigPositions(ULevelSequence* LevelSequence, UBaseControlRig* ControlRig, FName ControlName, const TArray<FFrameNumber>& Frames, const TArray<FVector> Values,
 		ESequenceTimeUnit TimeUnit = ESequenceTimeUnit::DisplayRate);
 
 	/**
@@ -656,7 +656,7 @@ public:
 	* @return Returns Value at that time
 	*/
 	UFUNCTION(BlueprintCallable, Category = "Editor Scripting | Sequencer Tools | Control Rig")
-	static FRotator GetLocalControlRigRotator(ULevelSequence* LevelSequence, UControlRig* ControlRig, FName ControlName, FFrameNumber Frame,
+	static FRotator GetLocalControlRigRotator(ULevelSequence* LevelSequence, UBaseControlRig* ControlRig, FName ControlName, FFrameNumber Frame,
 		ESequenceTimeUnit TimeUnit = ESequenceTimeUnit::DisplayRate);
 
 	/**
@@ -669,7 +669,7 @@ public:
 	* @return Returns Values at those times
 	*/
 	UFUNCTION(BlueprintCallable, Category = "Editor Scripting | Sequencer Tools | Control Rig")
-	static TArray<FRotator> GetLocalControlRigRotators(ULevelSequence* LevelSequence, UControlRig* ControlRig, FName ControlName, const TArray<FFrameNumber>& Frames,
+	static TArray<FRotator> GetLocalControlRigRotators(ULevelSequence* LevelSequence, UBaseControlRig* ControlRig, FName ControlName, const TArray<FFrameNumber>& Frames,
 		ESequenceTimeUnit TimeUnit = ESequenceTimeUnit::DisplayRate);
 
 	/**
@@ -683,7 +683,7 @@ public:
 	* @param bSetKey If True set a key, if not just set the value
 	*/
 	UFUNCTION(BlueprintCallable, Category = "Editor Scripting | Sequencer Tools | Control Rig")
-	static void SetLocalControlRigRotator(ULevelSequence* LevelSequence, UControlRig* ControlRig, FName ControlName, FFrameNumber Frame, FRotator Value, 
+	static void SetLocalControlRigRotator(ULevelSequence* LevelSequence, UBaseControlRig* ControlRig, FName ControlName, FFrameNumber Frame, FRotator Value, 
 		ESequenceTimeUnit TimeUnit = ESequenceTimeUnit::DisplayRate, bool bSetKey = true);
 
 	/**
@@ -696,7 +696,7 @@ public:
 	* @param TimeUnit Unit for frame values, either in display rate or tick resolution
 	*/
 	UFUNCTION(BlueprintCallable, Category = "Editor Scripting | Sequencer Tools | Control Rig")
-	static void SetLocalControlRigRotators(ULevelSequence* LevelSequence, UControlRig* ControlRig, FName ControlName, const TArray<FFrameNumber>& Frames, const TArray<FRotator> Values,
+	static void SetLocalControlRigRotators(ULevelSequence* LevelSequence, UBaseControlRig* ControlRig, FName ControlName, const TArray<FFrameNumber>& Frames, const TArray<FRotator> Values,
 		ESequenceTimeUnit TimeUnit = ESequenceTimeUnit::DisplayRate);
 
 	/**
@@ -709,7 +709,7 @@ public:
 	* @return Returns Value at that time
 	*/
 	UFUNCTION(BlueprintCallable, Category = "Editor Scripting | Sequencer Tools | Control Rig")
-	static FVector GetLocalControlRigScale(ULevelSequence* LevelSequence, UControlRig* ControlRig, FName ControlName, FFrameNumber Frame,
+	static FVector GetLocalControlRigScale(ULevelSequence* LevelSequence, UBaseControlRig* ControlRig, FName ControlName, FFrameNumber Frame,
 		ESequenceTimeUnit TimeUnit = ESequenceTimeUnit::DisplayRate);
 
 	/**
@@ -722,7 +722,7 @@ public:
 	* @return Returns Values at those times
 	*/
 	UFUNCTION(BlueprintCallable, Category = "Editor Scripting | Sequencer Tools | Control Rig")
-	static TArray<FVector> GetLocalControlRigScales(ULevelSequence* LevelSequence, UControlRig* ControlRig, FName ControlName, const TArray<FFrameNumber>& Frames,
+	static TArray<FVector> GetLocalControlRigScales(ULevelSequence* LevelSequence, UBaseControlRig* ControlRig, FName ControlName, const TArray<FFrameNumber>& Frames,
 		ESequenceTimeUnit TimeUnit = ESequenceTimeUnit::DisplayRate);
 
 	/**
@@ -736,7 +736,7 @@ public:
 	* @param bSetKey If True set a key, if not just set the value
 */
 	UFUNCTION(BlueprintCallable, Category = "Editor Scripting | Sequencer Tools | Control Rig")
-	static void SetLocalControlRigScale(ULevelSequence* LevelSequence, UControlRig* ControlRig, FName ControlName, FFrameNumber Frame, FVector Value, 
+	static void SetLocalControlRigScale(ULevelSequence* LevelSequence, UBaseControlRig* ControlRig, FName ControlName, FFrameNumber Frame, FVector Value, 
 		ESequenceTimeUnit TimeUnit = ESequenceTimeUnit::DisplayRate, bool bSetKey = true);
 
 	/**
@@ -749,7 +749,7 @@ public:
 	* @param TimeUnit Unit for frame values, either in display rate or tick resolution
 	*/
 	UFUNCTION(BlueprintCallable, Category = "Editor Scripting | Sequencer Tools | Control Rig")
-	static void SetLocalControlRigScales(ULevelSequence* LevelSequence, UControlRig* ControlRig, FName ControlName, const TArray<FFrameNumber>& Frames, const TArray<FVector> Values,
+	static void SetLocalControlRigScales(ULevelSequence* LevelSequence, UBaseControlRig* ControlRig, FName ControlName, const TArray<FFrameNumber>& Frames, const TArray<FVector> Values,
 		ESequenceTimeUnit TimeUnit = ESequenceTimeUnit::DisplayRate);
 
 
@@ -763,7 +763,7 @@ public:
 	* @return Returns Value at that time
 	*/
 	UFUNCTION(BlueprintCallable, Category = "Editor Scripting | Sequencer Tools | Control Rig")
-	static FEulerTransform GetLocalControlRigEulerTransform(ULevelSequence* LevelSequence, UControlRig* ControlRig, FName ControlName, FFrameNumber Frame,
+	static FEulerTransform GetLocalControlRigEulerTransform(ULevelSequence* LevelSequence, UBaseControlRig* ControlRig, FName ControlName, FFrameNumber Frame,
 		ESequenceTimeUnit TimeUnit = ESequenceTimeUnit::DisplayRate);
 
 	/**
@@ -776,7 +776,7 @@ public:
 	* @return Returns Values at those times
 	*/
 	UFUNCTION(BlueprintCallable, Category = "Editor Scripting | Sequencer Tools | Control Rig")
-	static TArray<FEulerTransform> GetLocalControlRigEulerTransforms(ULevelSequence* LevelSequence, UControlRig* ControlRig, FName ControlName, const TArray<FFrameNumber>& Frames,
+	static TArray<FEulerTransform> GetLocalControlRigEulerTransforms(ULevelSequence* LevelSequence, UBaseControlRig* ControlRig, FName ControlName, const TArray<FFrameNumber>& Frames,
 		ESequenceTimeUnit TimeUnit = ESequenceTimeUnit::DisplayRate);
 
 	/**
@@ -790,7 +790,7 @@ public:
 	* @param bSetKey If True set a key, if not just set the value
 	*/
 	UFUNCTION(BlueprintCallable, Category = "Editor Scripting | Sequencer Tools | Control Rig")
-	static void SetLocalControlRigEulerTransform(ULevelSequence* LevelSequence, UControlRig* ControlRig, FName ControlName, FFrameNumber Frame, FEulerTransform Value,
+	static void SetLocalControlRigEulerTransform(ULevelSequence* LevelSequence, UBaseControlRig* ControlRig, FName ControlName, FFrameNumber Frame, FEulerTransform Value,
 		ESequenceTimeUnit TimeUnit = ESequenceTimeUnit::DisplayRate, bool bSetKey = true);
 
 	/**
@@ -803,7 +803,7 @@ public:
 	* @param TimeUnit Unit for frame values, either in display rate or tick resolution
 	*/
 	UFUNCTION(BlueprintCallable, Category = "Editor Scripting | Sequencer Tools | Control Rig")
-	static void SetLocalControlRigEulerTransforms(ULevelSequence* LevelSequence, UControlRig* ControlRig, FName ControlName, const TArray<FFrameNumber>& Frames, const TArray<FEulerTransform> Values,
+	static void SetLocalControlRigEulerTransforms(ULevelSequence* LevelSequence, UBaseControlRig* ControlRig, FName ControlName, const TArray<FFrameNumber>& Frames, const TArray<FEulerTransform> Values,
 		ESequenceTimeUnit TimeUnit = ESequenceTimeUnit::DisplayRate);
 
 	/**
@@ -816,7 +816,7 @@ public:
 	* @return Returns Value at that time
 	*/
 	UFUNCTION(BlueprintCallable, Category = "Editor Scripting | Sequencer Tools | Control Rig")
-	static FTransformNoScale GetLocalControlRigTransformNoScale(ULevelSequence* LevelSequence, UControlRig* ControlRig, FName ControlName, FFrameNumber Frame,
+	static FTransformNoScale GetLocalControlRigTransformNoScale(ULevelSequence* LevelSequence, UBaseControlRig* ControlRig, FName ControlName, FFrameNumber Frame,
 		ESequenceTimeUnit TimeUnit = ESequenceTimeUnit::DisplayRate);
 
 	/**
@@ -829,7 +829,7 @@ public:
 	* @return Returns Values at those times
 	*/
 	UFUNCTION(BlueprintCallable, Category = "Editor Scripting | Sequencer Tools | Control Rig")
-	static TArray<FTransformNoScale> GetLocalControlRigTransformNoScales(ULevelSequence* LevelSequence, UControlRig* ControlRig, FName ControlName, const TArray<FFrameNumber>& Frames,
+	static TArray<FTransformNoScale> GetLocalControlRigTransformNoScales(ULevelSequence* LevelSequence, UBaseControlRig* ControlRig, FName ControlName, const TArray<FFrameNumber>& Frames,
 		ESequenceTimeUnit TimeUnit = ESequenceTimeUnit::DisplayRate);
 
 	/**
@@ -843,7 +843,7 @@ public:
 	* @param bSetKey If True set a key, if not just set the value
 	*/
 	UFUNCTION(BlueprintCallable, Category = "Editor Scripting | Sequencer Tools | Control Rig")
-	static void SetLocalControlRigTransformNoScale(ULevelSequence* LevelSequence, UControlRig* ControlRig, FName ControlName, FFrameNumber Frame, FTransformNoScale Value, 
+	static void SetLocalControlRigTransformNoScale(ULevelSequence* LevelSequence, UBaseControlRig* ControlRig, FName ControlName, FFrameNumber Frame, FTransformNoScale Value, 
 		ESequenceTimeUnit TimeUnit = ESequenceTimeUnit::DisplayRate, bool bSetKey = true);
 
 	/**
@@ -856,7 +856,7 @@ public:
 	* @param TimeUnit Unit for frame values, either in display rate or tick resolution
 	*/
 	UFUNCTION(BlueprintCallable, Category = "Editor Scripting | Sequencer Tools | Control Rig")
-	static void SetLocalControlRigTransformNoScales(ULevelSequence* LevelSequence, UControlRig* ControlRig, FName ControlName, const TArray<FFrameNumber>& Frames, const TArray<FTransformNoScale> Values,
+	static void SetLocalControlRigTransformNoScales(ULevelSequence* LevelSequence, UBaseControlRig* ControlRig, FName ControlName, const TArray<FFrameNumber>& Frames, const TArray<FTransformNoScale> Values,
 		ESequenceTimeUnit TimeUnit = ESequenceTimeUnit::DisplayRate);
 
 	/**
@@ -869,7 +869,7 @@ public:
 	* @return Returns Value at that time
 	*/
 	UFUNCTION(BlueprintCallable, Category = "Editor Scripting | Sequencer Tools | Control Rig")
-	static FTransform GetLocalControlRigTransform(ULevelSequence* LevelSequence, UControlRig* ControlRig, FName ControlName, FFrameNumber Frame, 
+	static FTransform GetLocalControlRigTransform(ULevelSequence* LevelSequence, UBaseControlRig* ControlRig, FName ControlName, FFrameNumber Frame, 
 		ESequenceTimeUnit TimeUnit = ESequenceTimeUnit::DisplayRate);
 
 	/**
@@ -882,7 +882,7 @@ public:
 	* @return Returns Values at those times
 	*/
 	UFUNCTION(BlueprintCallable, Category = "Editor Scripting | Sequencer Tools | Control Rig")
-	static TArray<FTransform> GetLocalControlRigTransforms(ULevelSequence* LevelSequence, UControlRig* ControlRig, FName ControlName, const TArray<FFrameNumber>& Frames,
+	static TArray<FTransform> GetLocalControlRigTransforms(ULevelSequence* LevelSequence, UBaseControlRig* ControlRig, FName ControlName, const TArray<FFrameNumber>& Frames,
 		ESequenceTimeUnit TimeUnit = ESequenceTimeUnit::DisplayRate);
 
 	/**
@@ -896,7 +896,7 @@ public:
 	* @param bSetKey If True set a key, if not just set the value
 	*/
 	UFUNCTION(BlueprintCallable, Category = "Editor Scripting | Sequencer Tools | Control Rig")
-	static void SetLocalControlRigTransform(ULevelSequence* LevelSequence, UControlRig* ControlRig, FName ControlName, FFrameNumber Frame, FTransform Value, 
+	static void SetLocalControlRigTransform(ULevelSequence* LevelSequence, UBaseControlRig* ControlRig, FName ControlName, FFrameNumber Frame, FTransform Value, 
 		ESequenceTimeUnit TimeUnit = ESequenceTimeUnit::DisplayRate, bool bSetKey = true);
 
 	/**
@@ -909,7 +909,7 @@ public:
 	* @param TimeUnit Unit for frame values, either in display rate or tick resolution
 	*/
 	UFUNCTION(BlueprintCallable, Category = "Editor Scripting | Sequencer Tools | Control Rig")
-	static void SetLocalControlRigTransforms(ULevelSequence* LevelSequence, UControlRig* ControlRig, FName ControlName, const TArray<FFrameNumber>& Frames, const TArray<FTransform> Values,
+	static void SetLocalControlRigTransforms(ULevelSequence* LevelSequence, UBaseControlRig* ControlRig, FName ControlName, const TArray<FFrameNumber>& Frames, const TArray<FTransform> Values,
 		ESequenceTimeUnit TimeUnit = ESequenceTimeUnit::DisplayRate);
 
 	/*
@@ -980,7 +980,7 @@ public:
 	 * @param TimeUnit Unit for the InTime, either in display rate or tick resolution
 	 */
 	UFUNCTION(BlueprintCallable, Category = "Editor Scripting | Sequencer Tools | Control Rig")
-	static bool SetControlRigSpace(ULevelSequence* InSequence, UControlRig* InControlRig, FName InControlName, const FRigElementKey& InSpaceKey, FFrameNumber InTime,  ESequenceTimeUnit TimeUnit = ESequenceTimeUnit::DisplayRate);
+	static bool SetControlRigSpace(ULevelSequence* InSequence, UBaseControlRig* InControlRig, FName InControlName, const FRigElementKey& InSpaceKey, FFrameNumber InTime,  ESequenceTimeUnit TimeUnit = ESequenceTimeUnit::DisplayRate);
 
 	/** Bake specified Control Rig Controls to a specified Space based upon the current settings
 	* @param InSequence Sequence to bake
@@ -990,7 +990,7 @@ public:
 	* @param TimeUnit Unit for the start and end times in the InSettings parameter.
 	*/
 	UFUNCTION(BlueprintCallable, Category = "Editor Scripting | Sequencer Tools | Control Rig")
-	static bool BakeControlRigSpace(ULevelSequence* InSequence, UControlRig* InControlRig, const TArray<FName>& InControlNames, FRigSpacePickerBakeSettings InSettings, ESequenceTimeUnit TimeUnit = ESequenceTimeUnit::DisplayRate);
+	static bool BakeControlRigSpace(ULevelSequence* InSequence, UBaseControlRig* InControlRig, const TArray<FName>& InControlNames, FRigSpacePickerBakeSettings InSettings, ESequenceTimeUnit TimeUnit = ESequenceTimeUnit::DisplayRate);
 	
 	/** Delete the Control Rig Space Key for the Control at the specified time. This will delete any attached Control Rig keys at this time and will perform any needed compensation to the new space.
 	*
@@ -1002,7 +1002,7 @@ public:
 	* @return Will return false if function fails,  for example if there is no key at this time it will fail.
 	*/
 	UFUNCTION(BlueprintCallable, Category = "Editor Scripting | Sequencer Tools | Control Rig")
-	static bool DeleteControlRigSpace(ULevelSequence* InSequence, UControlRig* InControlRig, FName InControlName, FFrameNumber InTime, ESequenceTimeUnit TimeUnit = ESequenceTimeUnit::DisplayRate);
+	static bool DeleteControlRigSpace(ULevelSequence* InSequence, UBaseControlRig* InControlRig, FName InControlName, FFrameNumber InTime, ESequenceTimeUnit TimeUnit = ESequenceTimeUnit::DisplayRate);
 
 	/** Move the Control Rig Space Key for the Control at the specified time to the new time. This will also move any Control Rig keys at this space switch boundary.
 	*
@@ -1015,7 +1015,7 @@ public:
 	* @return Will return false if function fails, for example if there is no key at this time it will fail, or if the new time is invalid it could fail also
 	*/
 	UFUNCTION(BlueprintCallable, Category = "Editor Scripting | Sequencer Tools | Control Rig")
-	static bool MoveControlRigSpace(ULevelSequence* InSequence, UControlRig* InControlRig, FName InControlName, FFrameNumber InTime, FFrameNumber InNewTime, ESequenceTimeUnit TimeUnit = ESequenceTimeUnit::DisplayRate);
+	static bool MoveControlRigSpace(ULevelSequence* InSequence, UBaseControlRig* InControlRig, FName InControlName, FFrameNumber InTime, FFrameNumber InNewTime, ESequenceTimeUnit TimeUnit = ESequenceTimeUnit::DisplayRate);
 
 	/** Rename the Control Rig Channels in Sequencer to the specified new control names, which should be present on the Control Rig
 	* @param InSequence Sequence to rename controls
@@ -1025,7 +1025,7 @@ public:
 	* @return Return true if the function succeeds, false if it doesn't which can happen if the name arrays don't match in size or any of the new Control Names aren't valid
 	*/
 	UFUNCTION(BlueprintCallable, Category = "Editor Scripting | Sequencer Tools | Control Rig")
-	static bool RenameControlRigControlChannels(ULevelSequence* InSequence, UControlRig* InControlRig, const TArray<FName>& InOldControlNames, const TArray<FName>& InNewControlNames);
+	static bool RenameControlRigControlChannels(ULevelSequence* InSequence, UBaseControlRig* InControlRig, const TArray<FName>& InOldControlNames, const TArray<FName>& InNewControlNames);
 
 	/** Get the controls mask for the given ControlName */
 	UFUNCTION(BlueprintPure, Category = "Editor Scripting | Sequencer Tools | Control Rig")
@@ -1047,20 +1047,20 @@ public:
 	@param InControlRig Rig to test to see if FK Control Rig
 	**/
 	UFUNCTION(BlueprintCallable, Category = "Editor Scripting | Sequencer Tools | Control Rig")
-	static bool IsFKControlRig(UControlRig* InControlRig);
+	static bool IsFKControlRig(UBaseControlRig* InControlRig);
 
 	/**	Whether or not the control rig is an Additive Control Rig.
 	@param InControlRig Rig to test to see if Additive Control Rig
 	**/
 	UFUNCTION(BlueprintCallable, Category = "Editor Scripting | Sequencer Tools | Control Rig")
-	static bool IsAdditiveControlRig(UControlRig* InControlRig);
+	static bool IsAdditiveControlRig(UBaseControlRig* InControlRig);
 
 	/**	Get FKControlRig Apply Mode.
 	@param InControlRig Rig to test
 	@return The EControlRigFKRigExecuteMode mode it is in, either Replace,Additive or Direct
 	**/
 	UFUNCTION(BlueprintCallable, Category = "Editor Scripting | Sequencer Tools | Control Rig")
-	static EControlRigFKRigExecuteMode  GetFKControlRigApplyMode(UControlRig* InControlRig);
+	static EControlRigFKRigExecuteMode  GetFKControlRigApplyMode(UBaseControlRig* InControlRig);
 
 	/**	Set the FK Control Rig to apply mode
 	@param InControlRig Rig to set 
@@ -1068,5 +1068,5 @@ public:
 	@return returns True if the mode was set, may not be set if the Control Rig doesn't support these modes currently only FKControlRig's do.
 	**/
 	UFUNCTION(BlueprintCallable, Category = "Editor Scripting | Sequencer Tools | Control Rig")
-	static bool SetControlRigApplyMode(UControlRig* InControlRig, EControlRigFKRigExecuteMode InApplyMode);
+	static bool SetControlRigApplyMode(UBaseControlRig* InControlRig, EControlRigFKRigExecuteMode InApplyMode);
 };
diff --git a/Engine/Plugins/Animation/ControlRig/Source/ControlRigEditor/Public/ControlRigSpaceChannelEditors.h b/Engine/Plugins/Animation/ControlRig/Source/ControlRigEditor/Public/ControlRigSpaceChannelEditors.h
--- a/Engine/Plugins/Animation/ControlRig/Source/ControlRigEditor/Public/ControlRigSpaceChannelEditors.h
+++ b/Engine/Plugins/Animation/ControlRig/Source/ControlRigEditor/Public/ControlRigSpaceChannelEditors.h
@@ -21,7 +21,7 @@
 
 struct FKeyHandle;
 struct FKeyDrawParams;
-class UControlrig;
+class UBaseControlRig;
 class ISequencer;
 class UMovieSceneSection;
 class URigHierarchy;
@@ -44,23 +44,23 @@ struct FSpaceChannelAndSection
 */
 struct FControlRigSpaceChannelHelpers
 {
-	static FKeyHandle SequencerKeyControlRigSpaceChannel(UControlRig* ControlRig, ISequencer* Sequencer, FMovieSceneControlRigSpaceChannel* Channel, UMovieSceneSection* SectionToKey, FFrameNumber Time, URigHierarchy* RigHierarchy, const FRigElementKey& ControlKey, const FRigElementKey& SpaceKey);
-	static void SequencerSpaceChannelKeyDeleted(UControlRig* ControlRig, ISequencer* Sequencer, FName ControlName, FMovieSceneControlRigSpaceChannel* Channel, UMovieSceneControlRigParameterSection* SectionToKey, FFrameNumber TimeOfDeletion);
-	static void CompensateIfNeeded(UControlRig* ControlRig, ISequencer* Sequencer, UMovieSceneControlRigParameterSection* Section, FName ControlName, TOptional<FFrameNumber>& Time);
-	static FSpaceChannelAndSection FindSpaceChannelAndSectionForControl(UControlRig* ControlRig, FName ControlName, ISequencer* Sequencer, bool bCreateIfNeeded);
-	static void SequencerBakeControlInSpace(UControlRig* ControlRig, ISequencer* Sequencer, FMovieSceneControlRigSpaceChannel* Channel, UMovieSceneSection* SectionToKey,
+	static FKeyHandle SequencerKeyControlRigSpaceChannel(UBaseControlRig* ControlRig, ISequencer* Sequencer, FMovieSceneControlRigSpaceChannel* Channel, UMovieSceneSection* SectionToKey, FFrameNumber Time, URigHierarchy* RigHierarchy, const FRigElementKey& ControlKey, const FRigElementKey& SpaceKey);
+	static void SequencerSpaceChannelKeyDeleted(UBaseControlRig* ControlRig, ISequencer* Sequencer, FName ControlName, FMovieSceneControlRigSpaceChannel* Channel, UMovieSceneControlRigParameterSection* SectionToKey, FFrameNumber TimeOfDeletion);
+	static void CompensateIfNeeded(UBaseControlRig* ControlRig, ISequencer* Sequencer, UMovieSceneControlRigParameterSection* Section, FName ControlName, TOptional<FFrameNumber>& Time);
+	static FSpaceChannelAndSection FindSpaceChannelAndSectionForControl(UBaseControlRig* ControlRig, FName ControlName, ISequencer* Sequencer, bool bCreateIfNeeded);
+	static void SequencerBakeControlInSpace(UBaseControlRig* ControlRig, ISequencer* Sequencer, FMovieSceneControlRigSpaceChannel* Channel, UMovieSceneSection* SectionToKey,
 		URigHierarchy* RigHierarchy, const FRigElementKey& ControlKey, FRigSpacePickerBakeSettings InSettings);
-	static void GetFramesInThisSpaceAfterThisTime(UControlRig* ControlRig, FName ControlName, FMovieSceneControlRigSpaceBaseKey CurrentValue,
+	static void GetFramesInThisSpaceAfterThisTime(UBaseControlRig* ControlRig, FName ControlName, FMovieSceneControlRigSpaceBaseKey CurrentValue,
 		FMovieSceneControlRigSpaceChannel* Channel, UMovieSceneSection* SectionToKey,
 		FFrameNumber Time, TSortedMap<FFrameNumber,FFrameNumber>& OutMoreFrames);
-	static void HandleSpaceKeyTimeChanged(UControlRig* ControlRig, FName ControlName,FMovieSceneControlRigSpaceChannel* Channel, UMovieSceneSection* SectionToKey,
+	static void HandleSpaceKeyTimeChanged(UBaseControlRig* ControlRig, FName ControlName,FMovieSceneControlRigSpaceChannel* Channel, UMovieSceneSection* SectionToKey,
 		FFrameNumber CurrentFrame, FFrameNumber NextFrame);
-	static void DeleteTransformKeysAtThisTime(UControlRig* ControlRig, UMovieSceneControlRigParameterSection* Section, FName ControlName, FFrameNumber Time);
+	static void DeleteTransformKeysAtThisTime(UBaseControlRig* ControlRig, UMovieSceneControlRigParameterSection* Section, FName ControlName, FFrameNumber Time);
 	static FLinearColor GetColor(const FMovieSceneControlRigSpaceBaseKey& Key);
 	static FReply OpenBakeDialog(ISequencer* Sequencer, FMovieSceneControlRigSpaceChannel* Channel, int32 KeyIndex, UMovieSceneSection* SectionToKey);
 	static TArray<FKeyBarCurveModel::FBarRange> FindRanges(FMovieSceneControlRigSpaceChannel* Channel, const UMovieSceneSection* Section);
 	// retrieve the control and the channel infos for that ControlRig/Section.
-	static TPair<FRigControlElement*, FChannelMapInfo*> GetControlAndChannelInfo(UControlRig* ControlRig, UMovieSceneControlRigParameterSection* Section, FName ControlName);
+	static TPair<FRigControlElement*, FChannelMapInfo*> GetControlAndChannelInfo(UBaseControlRig* ControlRig, UMovieSceneControlRigParameterSection* Section, FName ControlName);
 	// retrieve the number of float channels based on the control type.
 	static int32 GetNumFloatChannels(const ERigControlType InControlType);
 };
diff --git a/Engine/Plugins/Animation/ControlRig/Source/ControlRigEditor/Public/EditMode/ControlRigEditMode.h b/Engine/Plugins/Animation/ControlRig/Source/ControlRigEditor/Public/EditMode/ControlRigEditMode.h
--- a/Engine/Plugins/Animation/ControlRig/Source/ControlRigEditor/Public/EditMode/ControlRigEditMode.h
+++ b/Engine/Plugins/Animation/ControlRig/Source/ControlRigEditor/Public/EditMode/ControlRigEditMode.h
@@ -22,7 +22,7 @@ class FRigVMEditor;
 class FViewport;
 class UActorFactory;
 struct FViewportClick;
-class UControlRig;
+class UBaseControlRig;
 class FControlRigInteractionScope;
 class ISequencer;
 class UControlManipulator;
@@ -44,8 +44,8 @@ class UToolMenu;
 DECLARE_DELEGATE_RetVal_ThreeParams(FTransform, FOnGetRigElementTransform, const FRigElementKey& /*RigElementKey*/, bool /*bLocal*/, bool /*bOnDebugInstance*/);
 DECLARE_DELEGATE_ThreeParams(FOnSetRigElementTransform, const FRigElementKey& /*RigElementKey*/, const FTransform& /*Transform*/, bool /*bLocal*/);
 DECLARE_DELEGATE_RetVal(TSharedPtr<FUICommandList>, FNewMenuCommandsDelegate);
-DECLARE_MULTICAST_DELEGATE_TwoParams(FControlRigAddedOrRemoved, UControlRig*, bool /*true if added, false if removed*/);
-DECLARE_MULTICAST_DELEGATE_ThreeParams(FControlRigSelected, UControlRig*, const FRigElementKey& /*RigElementKey*/,const bool /*bIsSelected*/);
+DECLARE_MULTICAST_DELEGATE_TwoParams(FControlRigAddedOrRemoved, UBaseControlRig*, bool /*true if added, false if removed*/);
+DECLARE_MULTICAST_DELEGATE_ThreeParams(FControlRigSelected, UBaseControlRig*, const FRigElementKey& /*RigElementKey*/,const bool /*bIsSelected*/);
 DECLARE_DELEGATE_RetVal(UToolMenu*, FOnGetContextMenu);
 
 class FControlRigEditMode;
@@ -91,16 +91,16 @@ public:
 
 	/** Set the Control Rig Object to be active in the edit mode. You set both the Control Rig and a possible binding together with an optional Sequencer
 	 This will remove all other control rigs present and should be called for stand alone editors, like the Control Rig Editor*/
-	void SetObjects(UControlRig* InControlRig, UObject* BindingObject, TWeakPtr<ISequencer> InSequencer);
+	void SetObjects(UBaseControlRig* InControlRig, UObject* BindingObject, TWeakPtr<ISequencer> InSequencer);
 
 	/** Add a Control Rig object if it doesn't exist, will return true if it was added, false if it wasn't since it's already there. You can also set the Sequencer.*/
-	bool AddControlRigObject(UControlRig* InControlRig, TWeakPtr<ISequencer> InSequencer);
+	bool AddControlRigObject(UBaseControlRig* InControlRig, TWeakPtr<ISequencer> InSequencer);
 
 	/* Remove control rig */
-	void RemoveControlRig(UControlRig* InControlRig);
+	void RemoveControlRig(UBaseControlRig* InControlRig);
 
 	/*Replace old Control Rig with the New Control Rig, perhaps from a recompile in the level editor*/
-	void ReplaceControlRig(UControlRig* OldControlRig, UControlRig* NewControlRig);
+	void ReplaceControlRig(UBaseControlRig* OldControlRig, UBaseControlRig* NewControlRig);
 
 	/** This edit mode is re-used between the level editor and the asset editors (control rig editor etc.). Calling this indicates which context we are in */
 	virtual bool IsInLevelEditor() const;
@@ -166,9 +166,9 @@ public:
 	// callback that gets called when rig element is selected in other view
 	void OnHierarchyModified(ERigHierarchyNotification InNotif, URigHierarchy* InHierarchy, const FRigBaseElement* InElement);
 	void OnHierarchyModified_AnyThread(ERigHierarchyNotification InNotif, URigHierarchy* InHierarchy, const FRigBaseElement* InElement);
-	void OnControlModified(UControlRig* Subject, FRigControlElement* InControlElement, const FRigControlModifiedContext& Context);
-	void OnPreConstruction_AnyThread(UControlRig* InRig, const FName& InEventName);
-	void OnPostConstruction_AnyThread(UControlRig* InRig, const FName& InEventName);
+	void OnControlModified(UBaseControlRig* Subject, FRigControlElement* InControlElement, const FRigControlModifiedContext& Context);
+	void OnPreConstruction_AnyThread(UBaseControlRig* InRig, const FName& InEventName);
+	void OnPostConstruction_AnyThread(UBaseControlRig* InRig, const FName& InEventName);
 
 	/** return true if it can be removed from preview scene 
 	- this is to ensure preview scene doesn't remove shape actors */
@@ -178,7 +178,7 @@ public:
 
 	/** Requests to recreate the shape actors in the next tick. Will recreate only the ones for the specified
 	Control Rig, otherwise will recreate all of them*/
-	void RequestToRecreateControlShapeActors(UControlRig* ControlRig = nullptr); 
+	void RequestToRecreateControlShapeActors(UBaseControlRig* ControlRig = nullptr); 
 
 	static uint32 ValidControlTypeMask()
 	{
@@ -189,7 +189,7 @@ protected:
 
 	// shape related functions wrt enable/selection
 	/** Get the node name from the property path */
-	AControlRigShapeActor* GetControlShapeFromControlName(UControlRig* InControlRig,const FName& InControlName) const;
+	AControlRigShapeActor* GetControlShapeFromControlName(UBaseControlRig* InControlRig,const FName& InControlName) const;
 
 protected:
 	/** Helper function: set ControlRigs array to the details panel */
@@ -290,7 +290,7 @@ protected:
 	FGuid LastMovieSceneSig;
 
 	/** The scope for the interaction, one per manipulated Control rig */
-	TMap<UControlRig*,FControlRigInteractionScope*> InteractionScopes;
+	TMap<UBaseControlRig*,FControlRigInteractionScope*> InteractionScopes;
 
 	/** True if there's tracking going on right now */
 	bool bIsTracking;
@@ -305,10 +305,10 @@ protected:
 	bool bSelectionChanged;
 
 	/** Cached transform of pivot point for selected objects for each Control Rig */
-	TMap<UControlRig*,FTransform> PivotTransforms;
+	TMap<UBaseControlRig*,FTransform> PivotTransforms;
 
 	/** Previous cached transforms, need this to check on tick if any transform changed, gizmo may have changed*/
-	TMap<UControlRig*, FTransform> LastPivotTransforms;
+	TMap<UBaseControlRig*, FTransform> LastPivotTransforms;
 
 	/** Command bindings for keyboard shortcuts */
 	TSharedPtr<FUICommandList> CommandBindings;
@@ -328,23 +328,23 @@ protected:
 	FControlRigSelected OnControlRigSelectedDelegate;
 
 	/** GetSelectedRigElements, if InControlRig is nullptr get the first one */
-	TArray<FRigElementKey> GetSelectedRigElements(UControlRig* InControlRig = nullptr) const;
+	TArray<FRigElementKey> GetSelectedRigElements(UBaseControlRig* InControlRig = nullptr) const;
 
 	/* Flag to recreate shapes during tick */
 	ERecreateControlRigShape RecreateControlShapesRequired;
 	/* List of Control Rigs we should recreate*/
-	TArray<UControlRig*> ControlRigsToRecreate;
+	TArray<UBaseControlRig*> ControlRigsToRecreate;
 
 	/* Flag to temporarily disable handling notifs from the hierarchy */
 	bool bSuspendHierarchyNotifs;
 
 	/** Shape actors */
-	TMap<UControlRig*,TArray<TObjectPtr<AControlRigShapeActor>>> ControlRigShapeActors;
+	TMap<UBaseControlRig*,TArray<TObjectPtr<AControlRigShapeActor>>> ControlRigShapeActors;
 	TObjectPtr<UControlRigDetailPanelControlProxies> ControlProxy;
 
 	/** Utility functions for UI/Some other viewport manipulation*/
 	bool IsControlSelected() const;
-	bool AreRigElementSelectedAndMovable(UControlRig* InControlRig) const;
+	bool AreRigElementSelectedAndMovable(UBaseControlRig* InControlRig) const;
 	
 	/** Set initial transform handlers */
 	void OpenContextMenu(FEditorViewportClient* InViewportClient);
@@ -362,26 +362,26 @@ public:
 	void ClearRigElementSelection(uint32 InTypes);
 
 	/** Set a RigElement's selection state */
-	void SetRigElementSelection(UControlRig* ControlRig, ERigElementType Type, const FName& InRigElementName, bool bSelected);
+	void SetRigElementSelection(UBaseControlRig* ControlRig, ERigElementType Type, const FName& InRigElementName, bool bSelected);
 
 	/** Set multiple RigElement's selection states */
-	void SetRigElementSelection(UControlRig* ControlRig, ERigElementType Type, const TArray<FName>& InRigElementNames, bool bSelected);
+	void SetRigElementSelection(UBaseControlRig* ControlRig, ERigElementType Type, const TArray<FName>& InRigElementNames, bool bSelected);
 
 	/** Check if any RigElements are selected */
-	bool AreRigElementsSelected(uint32 InTypes, UControlRig* InControlRig) const;
+	bool AreRigElementsSelected(uint32 InTypes, UBaseControlRig* InControlRig) const;
 
 	/** Get the number of selected RigElements */
-	int32 GetNumSelectedRigElements(uint32 InTypes, UControlRig* InControlRig) const;
+	int32 GetNumSelectedRigElements(uint32 InTypes, UBaseControlRig* InControlRig) const;
 
 	/** Get all of the selected Controls*/
-	void GetAllSelectedControls(TMap<UControlRig*, TArray<FRigElementKey>>& OutSelectedControls) const;
+	void GetAllSelectedControls(TMap<UBaseControlRig*, TArray<FRigElementKey>>& OutSelectedControls) const;
 
 	/** Get all of the ControlRigs, maybe not valid anymore */
-	TArrayView<const TWeakObjectPtr<UControlRig>> GetControlRigs() const;
-	TArrayView<TWeakObjectPtr<UControlRig>> GetControlRigs();
+	TArrayView<const TWeakObjectPtr<UBaseControlRig>> GetControlRigs() const;
+	TArrayView<TWeakObjectPtr<UBaseControlRig>> GetControlRigs();
 	/* Get valid  Control Rigs possibly just visible*/
-	TArray<UControlRig*> GetControlRigsArray(bool bIsVisible);
-	TArray<const UControlRig*> GetControlRigsArray(bool bIsVisible) const;
+	TArray<UBaseControlRig*> GetControlRigsArray(bool bIsVisible);
+	TArray<const UBaseControlRig*> GetControlRigsArray(bool bIsVisible) const;
 
 	/** Get the detail proxies control rig*/
 	UControlRigDetailPanelControlProxies* GetDetailProxies() { return ControlProxy; }
@@ -399,7 +399,7 @@ private:
 	/** Whether or not Pivot Transforms have changed, in which case we need to redraw viewport*/
 	bool HasPivotTransformsChanged() const;
 	/** Set a RigElement's selection state */
-	void SetRigElementSelectionInternal(UControlRig* ControlRig, ERigElementType Type, const FName& InRigElementName, bool bSelected);
+	void SetRigElementSelectionInternal(UBaseControlRig* ControlRig, ERigElementType Type, const FName& InRigElementName, bool bSelected);
 	
 	FEditorViewportClient* CurrentViewportClient;
 	TArray<UE::Widget::EWidgetMode> RequestedWidgetModes;
@@ -418,15 +418,15 @@ public:
 	bool GetOnlySelectRigControls()const;
 
 private:
-	TSet<FName> GetActiveControlsFromSequencer(UControlRig* ControlRig);
+	TSet<FName> GetActiveControlsFromSequencer(UBaseControlRig* ControlRig);
 	bool SetSequencer(TWeakPtr<ISequencer> InSequencer);
 
 	/** Create/Delete for the specified ControlRig*/
-	void CreateShapeActors(UControlRig* InControlRig);
-	void DestroyShapesActors(UControlRig* InControlRig);
+	void CreateShapeActors(UBaseControlRig* InControlRig);
+	void DestroyShapesActors(UBaseControlRig* InControlRig);
 
 	/*Internal function for adding ControlRig*/
-	void AddControlRigInternal(UControlRig* InControlRig);
+	void AddControlRigInternal(UBaseControlRig* InControlRig);
 	void TickManipulatableObjects(float DeltaTime);
 
 	/* Check on tick to see if movie scene has changed, returns true if it has*/
@@ -449,19 +449,19 @@ protected:
 	
 	/** Get bindings to a runtime object */
 	//If the passed in ControlRig is nullptr we use the first Control Rig(this can happen from the BP Editors).
-	USceneComponent* GetHostingSceneComponent(const UControlRig* ControlRig = nullptr) const;
-	FTransform	GetHostingSceneComponentTransform(const UControlRig* ControlRig =  nullptr) const;
+	USceneComponent* GetHostingSceneComponent(const UBaseControlRig* ControlRig = nullptr) const;
+	FTransform	GetHostingSceneComponentTransform(const UBaseControlRig* ControlRig =  nullptr) const;
 
 	//Get if the hosted component is visible
-	bool IsControlRigSkelMeshVisible(UControlRig* ControlRig) const;
+	bool IsControlRigSkelMeshVisible(UBaseControlRig* ControlRig) const;
 
 private:
 
 	// Post pose update handler
 	void OnPoseInitialized();
 	void PostPoseUpdate();
-	void NotifyDrivenControls(UControlRig* InControlRig, const FRigElementKey& InKey);
-	void UpdateSelectabilityOnSkeletalMeshes(UControlRig* InControlRig, bool bEnabled);
+	void NotifyDrivenControls(UBaseControlRig* InControlRig, const FRigElementKey& InKey);
+	void UpdateSelectabilityOnSkeletalMeshes(UBaseControlRig* InControlRig, bool bEnabled);
 
 	// world clean up handlers
 	FDelegateHandle OnWorldCleanupHandle;
@@ -470,8 +470,8 @@ private:
 
 	void OnEditorClosed();
 	
-	TArray<TWeakObjectPtr<UControlRig>> RuntimeControlRigs;
-	TMap<UControlRig*,TStrongObjectPtr<UControlRigEditModeDelegateHelper>> DelegateHelpers;
+	TArray<TWeakObjectPtr<UBaseControlRig>> RuntimeControlRigs;
+	TMap<UBaseControlRig*,TStrongObjectPtr<UControlRigEditModeDelegateHelper>> DelegateHelpers;
 
 	//hack since we can't get the viewport client from the viewport, so in the tick we set the gameview bool and then in render/tickcontrolshapes we use it.
 	TMap<FViewport*, bool>  ViewportToGameView;
diff --git a/Engine/Plugins/Animation/ControlRig/Source/ControlRigEditor/Public/Tools/ControlRigSnapper.h b/Engine/Plugins/Animation/ControlRig/Source/ControlRigEditor/Public/Tools/ControlRigSnapper.h
--- a/Engine/Plugins/Animation/ControlRig/Source/ControlRigEditor/Public/Tools/ControlRigSnapper.h
+++ b/Engine/Plugins/Animation/ControlRig/Source/ControlRigEditor/Public/Tools/ControlRigSnapper.h
@@ -9,7 +9,7 @@
 #include "ControlRigSnapper.generated.h"
 
 class AActor;
-class UControlRig;
+class UBaseControlRig;
 class ISequencer;
 class UMovieScene;
 class ULevelSequence;
@@ -23,7 +23,7 @@ struct FControlRigForWorldTransforms
 
 	FControlRigForWorldTransforms() : ControlRig(nullptr) {};
 	UPROPERTY(BlueprintReadWrite, EditAnywhere, Category = "Control Rig")
-	TWeakObjectPtr<UControlRig> ControlRig;
+	TWeakObjectPtr<UBaseControlRig> ControlRig;
 	UPROPERTY(BlueprintReadWrite, EditAnywhere, Category = "Control Rig")
 	TArray<FName> ControlNames;
 };
@@ -75,7 +75,7 @@ struct FControlRigSnapper
 	* @param ParentTransforms List of Parent Transforms for each time that the resulting transforms will be concatenated to.
 	* @param OutTransforms  Results
 	*/
-	bool GetControlRigControlTransforms(ISequencer* Sequencer, UControlRig* ControlRig, const FName& ControlName,
+	bool GetControlRigControlTransforms(ISequencer* Sequencer, UBaseControlRig* ControlRig, const FName& ControlName,
 		const TArray<FFrameNumber>& Frames, const TArray<FTransform>& ParentTransforms, TArray<FTransform>& OutTransforms);
 
 
@@ -90,7 +90,7 @@ struct FControlRigSnapper
 	* @param OutTransforms  Results
 	*/
 	bool GetControlRigControlTransforms(UWorld* World, ULevelSequence* LevelSequence,
-		UControlRig* ControlRig, const FName& ControlName,const TArray<FFrameNumber>& Frames, const TArray<FTransform>& ParentTransforms, TArray<FTransform>& OutTransforms);
+		UBaseControlRig* ControlRig, const FName& ControlName,const TArray<FFrameNumber>& Frames, const TArray<FTransform>& ParentTransforms, TArray<FTransform>& OutTransforms);
 
 };
 
diff --git a/Engine/Plugins/Animation/ControlRig/Source/ControlRigEditor/Public/Tools/ControlRigTweener.h b/Engine/Plugins/Animation/ControlRig/Source/ControlRigEditor/Public/Tools/ControlRigTweener.h
--- a/Engine/Plugins/Animation/ControlRig/Source/ControlRigEditor/Public/Tools/ControlRigTweener.h
+++ b/Engine/Plugins/Animation/ControlRig/Source/ControlRigEditor/Public/Tools/ControlRigTweener.h
@@ -10,7 +10,7 @@
 
 struct FMovieSceneFloatChannel;
 struct FMovieSceneDoubleChannel;
-class UControlRig;
+class UBaseControlRig;
 class ISequencer;
 struct FRigControl;
 class UMovieSceneSection;
@@ -85,7 +85,7 @@ struct FAnimSliderObjectSelection
 	};
 
 	bool Setup(TWeakPtr<ISequencer>& InSequencer, TWeakPtr<class FControlRigEditMode>& InEditMode);
-	bool Setup(const TArray<UControlRig*>& SelectedControlRigs, TWeakPtr<ISequencer>& InSequencer);
+	bool Setup(const TArray<UBaseControlRig*>& SelectedControlRigs, TWeakPtr<ISequencer>& InSequencer);
 
 	TArray<FObjectChannels>  ChannelsArray;
 
@@ -95,7 +95,7 @@ private:
 	void SetupChannel(FFrameNumber CurrentFrame, TArray<FFrameNumber>& KeyTimes, TArray<FKeyHandle>& Handles, FMovieSceneFloatChannel* FloatChannel,
 		FMovieSceneDoubleChannel* DoubleChannel, FChannelKeyBounds& KeyBounds);
 
-	TArray<UControlRig*> GetControlRigs(TWeakPtr<class FControlRigEditMode>& InEditMode);
+	TArray<UBaseControlRig*> GetControlRigs(TWeakPtr<class FControlRigEditMode>& InEditMode);
 };
 
 /*
@@ -230,7 +230,7 @@ struct FControlsToTween :public FBasicBlendSlider
 	//BasicBlendSlider overrides
 	virtual double DoBlend(const FBlendStruct& BlendStruct) override;
 
-	bool Setup(const TArray<UControlRig*>& SelectedControlRigs, TWeakPtr<ISequencer>& InSequencer);
+	bool Setup(const TArray<UBaseControlRig*>& SelectedControlRigs, TWeakPtr<ISequencer>& InSequencer);
 
 };
 
diff --git a/Engine/Plugins/Animation/RigLogic/Source/RigLogicModule/Private/RigUnit_RigLogic.cpp b/Engine/Plugins/Animation/RigLogic/Source/RigLogicModule/Private/RigUnit_RigLogic.cpp
--- a/Engine/Plugins/Animation/RigLogic/Source/RigLogicModule/Private/RigUnit_RigLogic.cpp
+++ b/Engine/Plugins/Animation/RigLogic/Source/RigLogicModule/Private/RigUnit_RigLogic.cpp
@@ -390,7 +390,7 @@ FRigUnit_RigLogic_Execute()
 			//const double startTime = FPlatformTime::Seconds();
 			if (!Data.SkelMeshComponent.IsValid())
 			{
-				Data.SkelMeshComponent = ExecuteContext.UnitContext.DataSourceRegistry->RequestSource<USkeletalMeshComponent>(UControlRig::OwnerComponent);
+				Data.SkelMeshComponent = ExecuteContext.UnitContext.DataSourceRegistry->RequestSource<USkeletalMeshComponent>(UBaseControlRig::OwnerComponent);
 				// In normal execution, Data.SkelMeshComponent will be nullptr at the beginning
 				// however, during unit testing we cannot fetch it from DataSourceRegistry 
 				// in that case, a mock version will be inserted into Data by unit test beforehand
diff --git a/Engine/Plugins/Editor/SequencerAnimTools/Source/SequencerAnimTools/Private/Sequencer/MovieSceneTransformTrail.cpp b/Engine/Plugins/Editor/SequencerAnimTools/Source/SequencerAnimTools/Private/Sequencer/MovieSceneTransformTrail.cpp
--- a/Engine/Plugins/Editor/SequencerAnimTools/Source/SequencerAnimTools/Private/Sequencer/MovieSceneTransformTrail.cpp
+++ b/Engine/Plugins/Editor/SequencerAnimTools/Source/SequencerAnimTools/Private/Sequencer/MovieSceneTransformTrail.cpp
@@ -686,7 +686,7 @@ bool FMovieSceneControlRigTransformTrail::StartTracking()
 	{
 		return false;
 	}
-	UControlRig* ControlRig = Section->GetControlRig();
+	UBaseControlRig* ControlRig = Section->GetControlRig();
 	if (!ControlRig)
 	{
 		return false ;
@@ -720,7 +720,7 @@ bool FMovieSceneControlRigTransformTrail::ApplyDelta(const FVector& Pos, const F
 	{
 		return false;
 	}
-	UControlRig* ControlRig = Section->GetControlRig();
+	UBaseControlRig* ControlRig = Section->GetControlRig();
 	if (!ControlRig)
 	{
 		return false;
@@ -793,7 +793,7 @@ bool FMovieSceneControlRigTransformTrail::EndTracking()
 	{
 		return false;
 	}
-	UControlRig* ControlRig = Section->GetControlRig();
+	UBaseControlRig* ControlRig = Section->GetControlRig();
 	if (!ControlRig)
 	{
 		return false;
@@ -824,7 +824,7 @@ bool FMovieSceneControlRigTransformTrail::HandleAltClick(FEditorViewportClient*
 	{
 		return false;
 	}
-	UControlRig* ControlRig = Section->GetControlRig();
+	UBaseControlRig* ControlRig = Section->GetControlRig();
 	if (!ControlRig)
 	{
 		return false;
diff --git a/Engine/Plugins/Editor/SequencerAnimTools/Source/SequencerAnimTools/Private/Sequencer/SequencerTrailHierarchy.cpp b/Engine/Plugins/Editor/SequencerAnimTools/Source/SequencerAnimTools/Private/Sequencer/SequencerTrailHierarchy.cpp
--- a/Engine/Plugins/Editor/SequencerAnimTools/Source/SequencerAnimTools/Private/Sequencer/SequencerTrailHierarchy.cpp
+++ b/Engine/Plugins/Editor/SequencerAnimTools/Source/SequencerAnimTools/Private/Sequencer/SequencerTrailHierarchy.cpp
@@ -197,7 +197,7 @@ void FSequencerTrailHierarchy::RemoveTrail(const FGuid& Key)
 				return;
 			}
 		}
-		for (TPair<UControlRig*, TMap<FName, FGuid>>& CompMapPair : ControlsTracked)
+		for (TPair<UBaseControlRig*, TMap<FName, FGuid>>& CompMapPair : ControlsTracked)
 		{
 			if (const FName* FoundControl = CompMapPair.Value.FindKey(Key))
 			{
@@ -214,7 +214,7 @@ struct FTrailControlTransforms
 	TArray<FTransform> Transforms;
 };
 
-void FSequencerTrailHierarchy::UpdateControlRig(const TArray<FFrameNumber> &Frames,UControlRig* ControlRig, TMap<FName, FGuid >& CompMapPair, bool bUseEditedTimes)
+void FSequencerTrailHierarchy::UpdateControlRig(const TArray<FFrameNumber> &Frames,UBaseControlRig* ControlRig, TMap<FName, FGuid >& CompMapPair, bool bUseEditedTimes)
 {
 	if (TSharedPtr<IControlRigObjectBinding> ObjectBinding = ControlRig->GetObjectBinding())
 	{
@@ -299,7 +299,7 @@ void FSequencerTrailHierarchy::UpdateControlRig(const TArray<FFrameNumber> &Fram
 	}
 }
 
-void FSequencerTrailHierarchy::UpdateControlRig(const FTrailEvaluateTimes& EvaluateTimes,UControlRig* ControlRig, TMap<FName, FGuid > &CompMapPair)
+void FSequencerTrailHierarchy::UpdateControlRig(const FTrailEvaluateTimes& EvaluateTimes,UBaseControlRig* ControlRig, TMap<FName, FGuid > &CompMapPair)
 {
 	if (TSharedPtr<IControlRigObjectBinding> ObjectBinding = ControlRig->GetObjectBinding())
 	{
@@ -340,7 +340,7 @@ void FSequencerTrailHierarchy::Update()
 	FTrailHierarchy::Update();  //calculates EvalTimesArr
 	FTrailEvaluateTimes EvalTimes = FTrailEvaluateTimes(EvalTimesArr, SecondsPerSegment);
 
-	for (TPair<UControlRig*, TMap<FName, FGuid>>& CompMapPair : ControlsTracked)
+	for (TPair<UBaseControlRig*, TMap<FName, FGuid>>& CompMapPair : ControlsTracked)
 	{
 		bool bNeedToUpdateControlRig = false;
 		FTrail* ForceTrail = nullptr;
@@ -594,7 +594,7 @@ void FSequencerTrailHierarchy::UpdateSequencerBindings(const TArray<FGuid>& Sequ
 				if (!ControlRigDelegateHandles.Contains(CRParameterTrack))
 				{
 					RegisterControlRigDelegates(BoundComponent, CRParameterTrack);
-					UControlRig* ControlRig = CRParameterTrack->GetControlRig();
+					UBaseControlRig* ControlRig = CRParameterTrack->GetControlRig();
 					if (ControlRig)
 					{
 						TArray<FName> Selected = ControlRig->CurrentControlSelection();
@@ -702,7 +702,7 @@ void FSequencerTrailHierarchy::AddSkeletonToHierarchy(class USkeletalMeshCompone
 	TimingStats.Add("FSequencerTrailHierarchy::AddSkeletonToHierarchy", Timespan);
 }
 
-void FSequencerTrailHierarchy::AddControlRigTrail(USkeletalMeshComponent* Component, UControlRig* ControlRig, UMovieSceneControlRigParameterTrack* CRParameterTrack,const FName& ControlName)
+void FSequencerTrailHierarchy::AddControlRigTrail(USkeletalMeshComponent* Component, UBaseControlRig* ControlRig, UMovieSceneControlRigParameterTrack* CRParameterTrack,const FName& ControlName)
 {
 	if (ControlsTracked.Find(ControlRig) == nullptr)
 	{
@@ -737,10 +737,10 @@ void FSequencerTrailHierarchy::AddControlRigTrail(USkeletalMeshComponent* Compon
 
 void FSequencerTrailHierarchy::RegisterControlRigDelegates(USkeletalMeshComponent* Component, UMovieSceneControlRigParameterTrack* CRParameterTrack)
 {
-	UControlRig* ControlRig = CRParameterTrack->GetControlRig();
+	UBaseControlRig* ControlRig = CRParameterTrack->GetControlRig();
 	URigHierarchy* RigHierarchy = ControlRig->GetHierarchy();
 	FControlRigDelegateHandles& DelegateHandles = ControlRigDelegateHandles.Add(CRParameterTrack);
-	DelegateHandles.OnControlSelected = ControlRig->ControlSelected().AddLambda([this, Component,CRParameterTrack](UControlRig* ControlRig, FRigControlElement* ControlElement, bool bSelected)
+	DelegateHandles.OnControlSelected = ControlRig->ControlSelected().AddLambda([this, Component,CRParameterTrack](UBaseControlRig* ControlRig, FRigControlElement* ControlElement, bool bSelected)
 		{
 			if (ControlElement->Settings.ControlType != ERigControlType::Transform &&
 				ControlElement->Settings.ControlType != ERigControlType::TransformNoScale &&
@@ -769,9 +769,9 @@ void FSequencerTrailHierarchy::RegisterControlRigDelegates(USkeletalMeshComponen
 
 			//check to see if the seleced control rig is sill selected
 			TArray<FGuid> TrailsToRemove;
-			for (TPair<UControlRig*, TMap<FName, FGuid>>& CompMapPair : ControlsTracked)
+			for (TPair<UBaseControlRig*, TMap<FName, FGuid>>& CompMapPair : ControlsTracked)
 			{
-				UControlRig* TrackedControlRig = CompMapPair.Key;
+				UBaseControlRig* TrackedControlRig = CompMapPair.Key;
 				for (TPair<FName, FGuid> NameGuid : CompMapPair.Value)
 				{
 					if (TrackedControlRig->IsControlSelected(NameGuid.Key) == false)
diff --git a/Engine/Plugins/Editor/SequencerAnimTools/Source/SequencerAnimTools/Private/Sequencer/SequencerTrailHierarchy.h b/Engine/Plugins/Editor/SequencerAnimTools/Source/SequencerAnimTools/Private/Sequencer/SequencerTrailHierarchy.h
--- a/Engine/Plugins/Editor/SequencerAnimTools/Source/SequencerAnimTools/Private/Sequencer/SequencerTrailHierarchy.h
+++ b/Engine/Plugins/Editor/SequencerAnimTools/Source/SequencerAnimTools/Private/Sequencer/SequencerTrailHierarchy.h
@@ -17,7 +17,7 @@ class UMovieSceneSection;
 class UMovieScene3DTransformTrack;
 class UMovieSceneControlRigParameterTrack;
 class UMovieSceneControlRigParameterSection;
-class UControlRig;
+class UBaseControlRig;
 struct FRigHierarchyContainer;
 
 namespace UE
@@ -62,7 +62,7 @@ public:
 
 	const TMap<UObject*, FGuid>& GetObjectsTracked() const { return ObjectsTracked; }
 	const TMap<USkeletalMeshComponent*, TMap<FName, FGuid>>& GetBonesTracked() const { return BonesTracked; }
-	const TMap<UControlRig*, TMap<FName, FGuid>>& GetControlsTracked() const { return ControlsTracked; }
+	const TMap<UBaseControlRig*, TMap<FName, FGuid>>& GetControlsTracked() const { return ControlsTracked; }
 
 	void OnBoneVisibilityChanged(USkeleton* Skeleton, const FName& BoneName, const bool bIsVisible);
 	void OnBindingVisibilityStateChanged(UObject* BoundObject, const EBindingVisibilityState VisibilityState);
@@ -71,10 +71,10 @@ public:
 	void OnActorsChangedSomehow(TArray<AActor*>& InActors);	
 
 private:
-	void AddControlRigTrail(USkeletalMeshComponent* Component,UControlRig* ControlRig, UMovieSceneControlRigParameterTrack* CRTrack, const FName& ControlName);
+	void AddControlRigTrail(USkeletalMeshComponent* Component,UBaseControlRig* ControlRig, UMovieSceneControlRigParameterTrack* CRTrack, const FName& ControlName);
 
-	void UpdateControlRig(const TArray<FFrameNumber>& Frames, UControlRig* ControlRig, TMap<FName, FGuid >& CompMapPair, bool bUseEditedTimes);
-	void UpdateControlRig(const FTrailEvaluateTimes& EvalTimes, UControlRig* ControlRig, TMap<FName, FGuid >& CompMapPair);
+	void UpdateControlRig(const TArray<FFrameNumber>& Frames, UBaseControlRig* ControlRig, TMap<FName, FGuid >& CompMapPair, bool bUseEditedTimes);
+	void UpdateControlRig(const FTrailEvaluateTimes& EvalTimes, UBaseControlRig* ControlRig, TMap<FName, FGuid >& CompMapPair);
 
 	void UpdateSequencerBindings(const TArray<FGuid>& SequencerBindings, TFunctionRef<void(UObject*, FTrail*, FGuid)> OnUpdated);
 	void UpdateViewAndEvalRange();
@@ -87,7 +87,7 @@ private:
 	TWeakPtr<ISequencer> WeakSequencer;
 	TMap<UObject*, FGuid> ObjectsTracked;
 	TMap<USkeletalMeshComponent*, TMap<FName, FGuid>> BonesTracked; 
-	TMap<UControlRig*, TMap<FName, FGuid>> ControlsTracked; 
+	TMap<UBaseControlRig*, TMap<FName, FGuid>> ControlsTracked; 
 
 	// TODO: components can have multiple rigs so make this a map from sections to controls instead. However, this is only part of a larger problem of handling blending
 
diff --git a/Engine/Plugins/Editor/SequencerAnimTools/Source/SequencerAnimTools/Private/SequencerAnimEditPivotTool.cpp b/Engine/Plugins/Editor/SequencerAnimTools/Source/SequencerAnimTools/Private/SequencerAnimEditPivotTool.cpp
--- a/Engine/Plugins/Editor/SequencerAnimTools/Source/SequencerAnimTools/Private/SequencerAnimEditPivotTool.cpp
+++ b/Engine/Plugins/Editor/SequencerAnimTools/Source/SequencerAnimTools/Private/SequencerAnimEditPivotTool.cpp
@@ -49,7 +49,7 @@ FSavedMappings USequencerPivotTool::SavedPivotLocations;
 FLastSelectedObjects USequencerPivotTool::LastSelectedObjects;
 
 
-static void GetControlRigsAndSequencer(TArray<TWeakObjectPtr<UControlRig>>& ControlRigs, TWeakPtr<ISequencer>& SequencerPtr, ULevelSequence** LevelSequence)
+static void GetControlRigsAndSequencer(TArray<TWeakObjectPtr<UBaseControlRig>>& ControlRigs, TWeakPtr<ISequencer>& SequencerPtr, ULevelSequence** LevelSequence)
 {
 	*LevelSequence = ULevelSequenceEditorBlueprintLibrary::GetFocusedLevelSequence();
 	//if getting sequencer from level sequence need to use the current(master), not the focused
@@ -61,9 +61,9 @@ static void GetControlRigsAndSequencer(TArray<TWeakObjectPtr<UControlRig>>& Cont
 		SequencerPtr = LevelSequenceEditor ? LevelSequenceEditor->GetSequencer() : nullptr;
 		if (SequencerPtr.IsValid())
 		{
-			TArray<UControlRig*> TempControlRigs;
+			TArray<UBaseControlRig*> TempControlRigs;
 			TempControlRigs = UControlRigSequencerEditorLibrary::GetVisibleControlRigs();
-			for (UControlRig* ControlRig : TempControlRigs)
+			for (UBaseControlRig* ControlRig : TempControlRigs)
 			{
 				ControlRigs.Add(ControlRig);
 			}
@@ -79,11 +79,11 @@ bool USequencerPivotToolBuilder::CanBuildTool(const FToolBuilderState& SceneStat
 		return true;
 	}
 	ULevelSequence* LevelSequence;
-	TArray<TWeakObjectPtr<UControlRig>> ControlRigs;
+	TArray<TWeakObjectPtr<UBaseControlRig>> ControlRigs;
 	TWeakPtr<ISequencer> SequencerPtr;
 	GetControlRigsAndSequencer(ControlRigs, SequencerPtr, &LevelSequence);
 
-	for (TWeakObjectPtr<UControlRig> ControlRig : ControlRigs)
+	for (TWeakObjectPtr<UBaseControlRig> ControlRig : ControlRigs)
 	{
 		if (ControlRig.IsValid() && ControlRig->CurrentControlSelection().Num() > 0)
 		{
@@ -289,7 +289,7 @@ void USequencerPivotTool::Setup()
 	UpdateGizmoTransform();
 	UpdateGizmoVisibility();
 	//we get delegates last since we may select something above
-	for (TWeakObjectPtr<UControlRig>& ControlRig : ControlRigs)
+	for (TWeakObjectPtr<UBaseControlRig>& ControlRig : ControlRigs)
 	{
 		if (ControlRig.IsValid())
 		{
@@ -353,7 +353,7 @@ void USequencerPivotTool::SaveLastSelected()
 	LastSelectedObjects.LastSelectedControlRigs.SetNum(0);
 	LastSelectedObjects.LastSelectedActors.SetNum(0);
 
-	for (TWeakObjectPtr<UControlRig>& ControlRig : ControlRigs)
+	for (TWeakObjectPtr<UBaseControlRig>& ControlRig : ControlRigs)
 	{
 		if (ControlRig.IsValid())
 		{
@@ -428,7 +428,7 @@ bool USequencerPivotTool::SetGizmoBasedOnSelection(bool bUseSaved)
 	bool bHaveSomethingSelected = false;
 	FVector AverageLocation(0.0f);
 	int NumLocations = 0;
-	for (TWeakObjectPtr<UControlRig>& ControlRig : ControlRigs)
+	for (TWeakObjectPtr<UBaseControlRig>& ControlRig : ControlRigs)
 	{
 		if (ControlRig.IsValid() == false)
 		{
@@ -582,14 +582,14 @@ void USequencerPivotTool::OnEditorSelectionChanged(UObject* NewSelection)
 	*/
 }
 
-void USequencerPivotTool::HandleControlSelected(UControlRig* Subject, FRigControlElement* InControl, bool bSelected)
+void USequencerPivotTool::HandleControlSelected(UBaseControlRig* Subject, FRigControlElement* InControl, bool bSelected)
 {
 	DeactivateMe();
 }
 
 void USequencerPivotTool::RemoveDelegates()
 {
-	for (TWeakObjectPtr<UControlRig>& ControlRig : ControlRigs)
+	for (TWeakObjectPtr<UBaseControlRig>& ControlRig : ControlRigs)
 	{
 		if (ControlRig.IsValid())
 		{
@@ -614,7 +614,7 @@ void USequencerPivotTool::GizmoTransformStarted(UTransformProxy* Proxy)
 		const FFrameTime FrameTime = Sequencer->GetLocalTime().ConvertTo(TickResolution);
 		const FFrameNumber FrameNumber = FrameTime.RoundToFrame();
 
-		for (TWeakObjectPtr<UControlRig>& ControlRig : ControlRigs)
+		for (TWeakObjectPtr<UBaseControlRig>& ControlRig : ControlRigs)
 		{
 			if (ControlRig.IsValid() == false)
 			{
@@ -693,7 +693,7 @@ void USequencerPivotTool::GizmoTransformStarted(UTransformProxy* Proxy)
 	InteractionScopes.Reset();
 	if(bInPivotMode)
 	{
-		TMap<UControlRig*, int32> RigToScopeIndex;
+		TMap<UBaseControlRig*, int32> RigToScopeIndex;
 		for(int32 IndexA = 0; IndexA < ControlRigDrags.Num(); IndexA++)
 		{
 			const FControlRigSelectionDuringDrag& ControlRigDrag = ControlRigDrags[IndexA];
diff --git a/Engine/Plugins/Editor/SequencerAnimTools/Source/SequencerAnimTools/Public/SequencerAnimEditPivotTool.h b/Engine/Plugins/Editor/SequencerAnimTools/Source/SequencerAnimTools/Public/SequencerAnimEditPivotTool.h
--- a/Engine/Plugins/Editor/SequencerAnimTools/Source/SequencerAnimTools/Public/SequencerAnimEditPivotTool.h
+++ b/Engine/Plugins/Editor/SequencerAnimTools/Source/SequencerAnimTools/Public/SequencerAnimEditPivotTool.h
@@ -13,7 +13,7 @@
 
 class FControlRigInteractionScope;
 class SWidget;
-class UControlRig;
+class UBaseControlRig;
 class UInteractiveGizmoManager;
 class UTransformProxy;
 enum class EToolShutdownType : uint8;
@@ -56,7 +56,7 @@ struct FSelectionDuringDrag
 };
 struct FControlRigSelectionDuringDrag : public FSelectionDuringDrag
 {
-	UControlRig* ControlRig;
+	UBaseControlRig* ControlRig;
 	FName ControlName;
 };
 
@@ -67,7 +67,7 @@ struct FActorSelectonDuringDrag : public FSelectionDuringDrag
 
 struct FControlRigMappings
 {
-	TWeakObjectPtr<UControlRig> ControlRig;
+	TWeakObjectPtr<UBaseControlRig> ControlRig;
 
 	FTransform GetParentTransform() const;
 	TOptional<FTransform> GetWorldTransform(const FName& Name) const;
@@ -92,7 +92,7 @@ private:
 
 struct FSavedMappings
 {
-	TMap<TWeakObjectPtr<UControlRig>, FControlRigMappings> ControlRigMappings;
+	TMap<TWeakObjectPtr<UBaseControlRig>, FControlRigMappings> ControlRigMappings;
 	TMap<TWeakObjectPtr<AActor>, FActorMappings> ActorMappings;
 };
 
@@ -195,7 +195,7 @@ protected:
 
 	//since we are selection based we can cache this
 	ULevelSequence* LevelSequence;
-	TArray<TWeakObjectPtr<UControlRig>> ControlRigs;
+	TArray<TWeakObjectPtr<UBaseControlRig>> ControlRigs;
 	TWeakPtr<ISequencer> SequencerPtr;
 	TArray<TWeakObjectPtr<AActor>> Actors;
 
@@ -220,7 +220,7 @@ protected:
 	// selection delegates
 	void DeactivateMe();
 	void RemoveDelegates();
-	void HandleControlSelected(UControlRig* Subject, FRigControlElement* InControl, bool bSelected);
+	void HandleControlSelected(UBaseControlRig* Subject, FRigControlElement* InControl, bool bSelected);
 	void OnEditorSelectionChanged(UObject* NewSelection);
 	FDelegateHandle OnEditorSelectionChangedHandle;
 
diff --git a/Engine/Plugins/Experimental/AnimNext/Source/AnimNextUncookedOnly/Private/Graph/AnimGraph/AnimGraphNode_AnimNextGraph.cpp b/Engine/Plugins/Experimental/AnimNext/Source/AnimNextUncookedOnly/Private/Graph/AnimGraph/AnimGraphNode_AnimNextGraph.cpp
--- a/Engine/Plugins/Experimental/AnimNext/Source/AnimNextUncookedOnly/Private/Graph/AnimGraph/AnimGraphNode_AnimNextGraph.cpp
+++ b/Engine/Plugins/Experimental/AnimNext/Source/AnimNextUncookedOnly/Private/Graph/AnimGraph/AnimGraphNode_AnimNextGraph.cpp
@@ -157,7 +157,7 @@ void UAnimGraphNode_AnimNextGraph::GetVariables(bool bInput, TMap<FName, FRigVME
 		if (UControlRigBlueprint* RigBlueprint = Cast<UControlRigBlueprint>(TargetClass->ClassGeneratedBy))
 		{
 			//RigBlueprint->CleanupVariables();
-			UControlRig* ControlRig = TargetClass->GetDefaultObject<UControlRig>();
+			UBaseControlRig* ControlRig = TargetClass->GetDefaultObject<UBaseControlRig>();
 			if (ControlRig)
 			{
 				const TArray<FRigVMExternalVariable>& PublicVariables = ControlRig->GetPublicVariables();
diff --git a/Engine/Plugins/Importers/USDImporter/Source/USDStage/Private/USDLevelSequenceHelper.cpp b/Engine/Plugins/Importers/USDImporter/Source/USDStage/Private/USDLevelSequenceHelper.cpp
--- a/Engine/Plugins/Importers/USDImporter/Source/USDStage/Private/USDLevelSequenceHelper.cpp
+++ b/Engine/Plugins/Importers/USDImporter/Source/USDStage/Private/USDLevelSequenceHelper.cpp
@@ -302,7 +302,7 @@ namespace UsdLevelSequenceHelperImpl
 		{
 			FString ObjectName = InClass->GetName();
 			ObjectName.RemoveFromEnd(TEXT("_C"));
-			UControlRig* ControlRig = NewObject<UControlRig>(Track, InClass, FName(*ObjectName), RF_Transactional);
+			UBaseControlRig* ControlRig = NewObject<UBaseControlRig>(Track, InClass, FName(*ObjectName), RF_Transactional);
 			if (InClass != UFKControlRig::StaticClass() && !ControlRig->SupportsEvent(TEXT("Backwards Solve")))
 			{
 				MovieScene->RemoveTrack(*Track);
@@ -312,7 +312,7 @@ namespace UsdLevelSequenceHelperImpl
 			ControlRig->Modify();
 			ControlRig->SetObjectBinding(MakeShared<FControlRigObjectBinding>());
 			ControlRig->GetObjectBinding()->BindToObject(SkeletalMeshComp);
-			ControlRig->GetDataSourceRegistry()->RegisterDataSource(UControlRig::OwnerComponent, ControlRig->GetObjectBinding()->GetBoundObject());
+			ControlRig->GetDataSourceRegistry()->RegisterDataSource(UBaseControlRig::OwnerComponent, ControlRig->GetObjectBinding()->GetBoundObject());
 			ControlRig->Initialize();
 			ControlRig->RequestInit();
 			ControlRig->SetBoneInitialTransformsFromSkeletalMeshComponent(SkeletalMeshComp, true);
diff --git a/Engine/Plugins/Importers/USDImporter/Source/USDUtilities/Private/USDSkeletalDataConversion.cpp b/Engine/Plugins/Importers/USDImporter/Source/USDUtilities/Private/USDSkeletalDataConversion.cpp
--- a/Engine/Plugins/Importers/USDImporter/Source/USDUtilities/Private/USDSkeletalDataConversion.cpp
+++ b/Engine/Plugins/Importers/USDImporter/Source/USDUtilities/Private/USDSkeletalDataConversion.cpp
@@ -3278,7 +3278,7 @@ bool UnrealToUsd::ConvertControlRigSection(
 		return false;
 	}
 
-	UControlRig* ControlRig = InSection->GetControlRig();
+	UBaseControlRig* ControlRig = InSection->GetControlRig();
 	if ( !ControlRig )
 	{
 		return false;
