Control Rig: Fix compilation issue in MovieSceneControlRigParameterTemplate
#rnx

diff --git a/Engine/Plugins/Animation/ControlRig/Source/ControlRig/Private/Sequencer/MovieSceneControlRigParameterTemplate.cpp b/Engine/Plugins/Animation/ControlRig/Source/ControlRig/Private/Sequencer/MovieSceneControlRigParameterTemplate.cpp
--- a/Engine/Plugins/Animation/ControlRig/Source/ControlRig/Private/Sequencer/MovieSceneControlRigParameterTemplate.cpp
+++ b/Engine/Plugins/Animation/ControlRig/Source/ControlRig/Private/Sequencer/MovieSceneControlRigParameterTemplate.cpp
@@ -692,7 +692,7 @@ void FControlRigBindingHelper::UnBindFromSequencerInstance(UBaseControlRig* Cont
 
 struct FControlRigSkeletalMeshComponentBindingTokenProducer : IMovieScenePreAnimatedTokenProducer
 {
-	FControlRigSkeletalMeshComponentBindingTokenProducer(FMovieSceneSequenceIDRef InSequenceID, UControlRig* InControlRig)
+	FControlRigSkeletalMeshComponentBindingTokenProducer(FMovieSceneSequenceIDRef InSequenceID, UBaseControlRig* InControlRig)
 		: SequenceID(InSequenceID), ControlRigUniqueID(InControlRig->GetUniqueID())
 	{}
 
