Control Rig: Introduce the socket element

diff --git a/Engine/Plugins/Animation/ControlRig/Source/ControlRig/Private/Rigs/RigHierarchy.cpp b/Engine/Plugins/Animation/ControlRig/Source/ControlRig/Private/Rigs/RigHierarchy.cpp
--- a/Engine/Plugins/Animation/ControlRig/Source/ControlRig/Private/Rigs/RigHierarchy.cpp
+++ b/Engine/Plugins/Animation/ControlRig/Source/ControlRig/Private/Rigs/RigHierarchy.cpp
@@ -458,6 +458,7 @@ void URigHierarchy::CopyHierarchy(URigHierarchy* InHierarchy)
 		sizeof(FRigRigidBodyElement),
 		sizeof(FRigReferenceElement),
 		sizeof(FRigConnectorElement),
+		sizeof(FRigSocketElement),
 	}; 
 
 	if(bReallocateElements)
@@ -4777,6 +4778,15 @@ FRigBaseElement* URigHierarchy::MakeElement(ERigElementType InElementType, int32
 			Element = NewElement<FRigConnectorElement>(InCount);
 			break;
 		}
+		case ERigElementType::Socket:
+		{
+			if(OutStructureSize)
+			{
+				*OutStructureSize = sizeof(FRigSocketElement);
+			}
+			Element = NewElement<FRigSocketElement>(InCount);
+			break;
+		}
 		default:
 		{
 			ensure(false);
@@ -4865,6 +4875,15 @@ void URigHierarchy::DestroyElement(FRigBaseElement*& InElement)
 			}
 			break;
 		}
+		case ERigElementType::Socket:
+		{
+			FRigSocketElement* ExistingElements = Cast<FRigSocketElement>(InElement);
+			for(int32 Index=0;Index<Count;Index++)
+			{
+				ExistingElements[Index].~FRigSocketElement(); 
+			}
+			break;
+		}
 		default:
 		{
 			ensure(false);
diff --git a/Engine/Plugins/Animation/ControlRig/Source/ControlRig/Private/Rigs/RigHierarchyController.cpp b/Engine/Plugins/Animation/ControlRig/Source/ControlRig/Private/Rigs/RigHierarchyController.cpp
--- a/Engine/Plugins/Animation/ControlRig/Source/ControlRig/Private/Rigs/RigHierarchyController.cpp
+++ b/Engine/Plugins/Animation/ControlRig/Source/ControlRig/Private/Rigs/RigHierarchyController.cpp
@@ -632,6 +632,67 @@ FRigElementKey URigHierarchyController::AddConnector(FName InName, FTransform In
 	return NewElement->Key;
 }
 
+FRigElementKey URigHierarchyController::AddSocket(FName InName, FRigElementKey InParent, FTransform InTransform,
+	bool bTransformInGlobal, bool bSetupUndo, bool bPrintPythonCommand)
+{
+	if(!IsValid())
+	{
+		return FRigElementKey();
+	}
+
+#if WITH_EDITOR
+	TSharedPtr<FScopedTransaction> TransactionPtr;
+	if(bSetupUndo)
+	{
+		TransactionPtr = MakeShared<FScopedTransaction>(NSLOCTEXT("RigHierarchyController", "Add Socket", "Add Socket"));
+		Hierarchy->Modify();
+	}
+#endif
+
+	FRigSocketElement* NewElement = MakeElement<FRigSocketElement>();
+	{
+		TGuardValue<bool> DisableCacheValidityChecks(Hierarchy->bEnableCacheValidityCheck, false);
+		NewElement->Key.Type = ERigElementType::Socket;
+		NewElement->Key.Name = GetSafeNewName(InName, NewElement->Key.Type);
+		AddElement(NewElement, Hierarchy->Get(Hierarchy->GetIndex(InParent)), true, InName);
+
+		if(bTransformInGlobal)
+		{
+			Hierarchy->SetTransform(NewElement, InTransform, ERigTransformType::InitialGlobal, true, false);
+			Hierarchy->SetTransform(NewElement, InTransform, ERigTransformType::CurrentGlobal, true, false);
+		}
+		else
+		{
+			Hierarchy->SetTransform(NewElement, InTransform, ERigTransformType::InitialLocal, true, false);
+			Hierarchy->SetTransform(NewElement, InTransform, ERigTransformType::CurrentLocal, true, false);
+		}
+
+		NewElement->Pose.Current = NewElement->Pose.Initial;
+	}
+
+#if WITH_EDITOR
+	TransactionPtr.Reset();
+
+	if (bPrintPythonCommand && !bSuspendPythonPrinting)
+	{
+		UBlueprint* Blueprint = GetTypedOuter<UBlueprint>();
+		if (Blueprint)
+		{
+			TArray<FString> Commands = GetAddSocketPythonCommands(NewElement);
+			for (const FString& Command : Commands)
+			{			
+				RigVMPythonUtils::Print(Blueprint->GetFName().ToString(),
+					FString::Printf(TEXT("%s"), *Command));
+			}
+		}
+	}
+#endif
+
+	Hierarchy->EnsureCacheValidity();
+		
+	return NewElement->Key;
+}
+
 FRigControlSettings URigHierarchyController::GetControlSettings(FRigElementKey InKey) const
 {
 	if(!IsValid())
@@ -1104,6 +1165,12 @@ FString URigHierarchyController::ExportToText(TArray<FRigElementKey> InKeys) con
 				FRigConnectorElement::StaticStruct()->ExportText(PerElementData.Content, Element, &DefaultElement, nullptr, PPF_None, nullptr);
 				break;
 			}
+			case ERigElementType::Socket:
+			{
+				FRigSocketElement DefaultElement;
+				FRigSocketElement::StaticStruct()->ExportText(PerElementData.Content, Element, &DefaultElement, nullptr, PPF_None, nullptr);
+				break;
+			}
 			default:
 			{
 				ensure(false);
@@ -1233,11 +1300,17 @@ TArray<FRigElementKey> URigHierarchyController::ImportFromText(FString InContent
 				break;
 			}
 			case ERigElementType::Connector:
-					{
+			{
 				NewElement = MakeElement<FRigConnectorElement>();
 				FRigConnectorElement::StaticStruct()->ImportText(*PerElementData.Content, NewElement, nullptr, EPropertyPortFlags::PPF_None, &ErrorPipe, FRigConnectorElement::StaticStruct()->GetName(), true);
 				break;
 			}
+			case ERigElementType::Socket:
+			{
+				NewElement = MakeElement<FRigSocketElement>();
+				FRigSocketElement::StaticStruct()->ImportText(*PerElementData.Content, NewElement, nullptr, EPropertyPortFlags::PPF_None, &ErrorPipe, FRigSocketElement::StaticStruct()->GetName(), true);
+				break;
+			}
 			default:
 			{
 				ensure(false);
@@ -1550,6 +1623,10 @@ TArray<FString> URigHierarchyController::GetAddElementPythonCommands(FRigBaseEle
 	{
 		return GetAddConnectorPythonCommands(ConnectorElement);
 	}
+	else if(FRigSocketElement* SocketElement = Cast<FRigSocketElement>(Element))
+	{
+		return GetAddSocketPythonCommands(SocketElement);
+	}
 	return TArray<FString>();
 }
 
@@ -1711,6 +1788,27 @@ TArray<FString> URigHierarchyController::GetAddConnectorPythonCommands(FRigConne
 	return Commands;
 }
 
+TArray<FString> URigHierarchyController::GetAddSocketPythonCommands(FRigSocketElement* Socket) const
+{
+	TArray<FString> Commands;
+	FString TransformStr = RigVMPythonUtils::TransformToPythonString(Socket->Pose.Initial.Local.Transform);
+
+	FString ParentKeyStr = "''";
+	if (Socket->ParentElement)
+	{
+		ParentKeyStr = Socket->ParentElement->GetKey().ToPythonString();
+	}
+
+	// AddSocket(FName InName, FRigElementKey InParent, FTransform InTransform, bool bTransformInGlobal = true, bool bSetupUndo = false);
+	Commands.Add(FString::Printf(TEXT("hierarchy_controller.add_socket('%s', %s, %s, False)"),
+		*Socket->GetName(),
+		*ParentKeyStr,
+		*TransformStr
+	));
+
+	return Commands;
+}
+
 TArray<FString> URigHierarchyController::GetSetControlValuePythonCommands(const FRigControlElement* Control, const FRigControlValue& Value,
                                                                           const ERigControlValueType& Type) const
 {
@@ -2382,6 +2480,12 @@ bool URigHierarchyController::AddParent(FRigElementKey InChild, FRigElementKey I
 		return false;
 	}
 
+	if(InParent.Type == ERigElementType::Socket)
+	{
+		ReportWarningf(TEXT("Cannot parent Child '%s' under a Socket parent."), *InChild.ToString());
+		return false;
+	}
+
 	FRigBaseElement* Child = Hierarchy->Find(InChild);
 	if(Child == nullptr)
 	{
@@ -2425,6 +2529,11 @@ bool URigHierarchyController::AddParent(FRigBaseElement* InChild, FRigBaseElemen
 		return false;
 	}
 
+	if(InParent->GetType() == ERigElementType::Socket)
+	{
+		return false;
+	}
+
 	// single parent children can't be parented multiple times
 	if(FRigSingleParentElement* SingleParentElement = Cast<FRigSingleParentElement>(InChild))
 	{
@@ -2862,6 +2971,12 @@ bool URigHierarchyController::SetParent(FRigElementKey InChild, FRigElementKey I
 		return false;
 	}
 
+	if(InParent.Type == ERigElementType::Socket)
+	{
+		ReportWarningf(TEXT("Cannot parent Child '%s' under a Socket parent."), *InChild.ToString());
+		return false;
+	}
+
 	FRigBaseElement* Child = Hierarchy->Find(InChild);
 	if(Child == nullptr)
 	{
diff --git a/Engine/Plugins/Animation/ControlRig/Source/ControlRig/Private/Rigs/RigHierarchyElements.cpp b/Engine/Plugins/Animation/ControlRig/Source/ControlRig/Private/Rigs/RigHierarchyElements.cpp
--- a/Engine/Plugins/Animation/ControlRig/Source/ControlRig/Private/Rigs/RigHierarchyElements.cpp
+++ b/Engine/Plugins/Animation/ControlRig/Source/ControlRig/Private/Rigs/RigHierarchyElements.cpp
@@ -80,6 +80,10 @@ UScriptStruct* FRigBaseElement::GetElementStruct() const
 		{
 			return FRigConnectorElement::StaticStruct();
 		}
+		case ERigElementType::Socket:
+		{
+			return FRigSocketElement::StaticStruct();
+		}
 		default:
 		{
 			break;
@@ -1705,3 +1709,24 @@ void FRigConnectorElement::CopyFrom(URigHierarchy* InHierarchy, FRigBaseElement*
 	const FRigConnectorElement* Source = CastChecked<FRigConnectorElement>(InOther);
 	Settings = Source->Settings;
 }
+
+////////////////////////////////////////////////////////////////////////////////
+// FRigSocketElement
+////////////////////////////////////////////////////////////////////////////////
+
+const FRigBaseElement::EElementIndex FRigSocketElement::ElementTypeIndex = SocketElement;
+
+void FRigSocketElement::Save(FArchive& Ar, URigHierarchy* Hierarchy, ESerializationPhase SerializationPhase)
+{
+	Super::Save(Ar, Hierarchy, SerializationPhase);
+}
+
+void FRigSocketElement::Load(FArchive& Ar, URigHierarchy* Hierarchy, ESerializationPhase SerializationPhase)
+{
+	Super::Load(Ar, Hierarchy, SerializationPhase);
+}
+
+void FRigSocketElement::CopyFrom(URigHierarchy* InHierarchy, FRigBaseElement* InOther, URigHierarchy* InOtherHierarchy)
+{
+	Super::CopyFrom(InHierarchy, InOther, InOtherHierarchy);
+}
\ No newline at end of file
diff --git a/Engine/Plugins/Animation/ControlRig/Source/ControlRig/Private/Units/Execution/RigUnit_DynamicHierarchy.cpp b/Engine/Plugins/Animation/ControlRig/Source/ControlRig/Private/Units/Execution/RigUnit_DynamicHierarchy.cpp
--- a/Engine/Plugins/Animation/ControlRig/Source/ControlRig/Private/Units/Execution/RigUnit_DynamicHierarchy.cpp
+++ b/Engine/Plugins/Animation/ControlRig/Source/ControlRig/Private/Units/Execution/RigUnit_DynamicHierarchy.cpp
@@ -1221,4 +1221,25 @@ FRigUnit_HierarchySetShapeSettings_Execute()
 			UE_CONTROLRIG_RIGUNIT_REPORT_ERROR(TEXT("Control '%s' does not exist."), *Item.ToString());
 		}
 	}
-}
\ No newline at end of file
+}
+
+FRigUnit_HierarchyAddSocket_Execute()
+{
+	FString ErrorMessage;
+	if(!FRigUnit_DynamicHierarchyBase::IsValidToRunInContext(ExecuteContext, true, &ErrorMessage))
+	{
+		if(!ErrorMessage.IsEmpty())
+		{
+			UE_CONTROLRIG_RIGUNIT_REPORT_ERROR(TEXT("%s"), *ErrorMessage);
+		}
+		return;
+	}
+
+	Item.Reset();
+
+	if(URigHierarchyController* Controller = ExecuteContext.Hierarchy->GetController(true))
+	{
+		FRigHierarchyControllerInstructionBracket InstructionBracket(Controller, ExecuteContext.GetInstructionIndex());
+		Item = Controller->AddSocket(Name, Parent, Transform, Space == ERigVMTransformSpace::GlobalSpace, false, false);
+	}
+}
diff --git a/Engine/Plugins/Animation/ControlRig/Source/ControlRig/Private/Units/Execution/RigUnit_DynamicHierarchy.h b/Engine/Plugins/Animation/ControlRig/Source/ControlRig/Private/Units/Execution/RigUnit_DynamicHierarchy.h
--- a/Engine/Plugins/Animation/ControlRig/Source/ControlRig/Private/Units/Execution/RigUnit_DynamicHierarchy.h
+++ b/Engine/Plugins/Animation/ControlRig/Source/ControlRig/Private/Units/Execution/RigUnit_DynamicHierarchy.h
@@ -1438,3 +1438,37 @@ struct CONTROLRIG_API FRigUnit_HierarchySetShapeSettings : public FRigUnit_Dynam
 	RIGVM_METHOD()
 	virtual void Execute() override;
 };
+
+/**
+ * Adds a new socket to the hierarchy
+ * Note: This node only runs as part of the construction event.
+ */
+USTRUCT(meta=(DisplayName="Spawn Socket", Keywords="Construction,Create,New,Locator,Group", Varying))
+struct CONTROLRIG_API FRigUnit_HierarchyAddSocket : public FRigUnit_HierarchyAddElement
+{
+	GENERATED_BODY()
+
+	FRigUnit_HierarchyAddSocket()
+	{
+		Name = TEXT("NewSocket");
+		Transform = FTransform::Identity;
+		Space = ERigVMTransformSpace::LocalSpace;
+	}
+
+	virtual ERigElementType GetElementTypeToSpawn() const override { return ERigElementType::Socket; }
+
+	/*
+	 * The initial transform of the new element
+	 */
+	UPROPERTY(meta = (Input))
+	FTransform Transform;
+
+	/**
+	 * Defines if the transform should be interpreted in local or global space
+	 */ 
+	UPROPERTY(meta = (Input))
+	ERigVMTransformSpace Space;
+
+	RIGVM_METHOD()
+	virtual void Execute() override;
+};
diff --git a/Engine/Plugins/Animation/ControlRig/Source/ControlRig/Public/Rigs/RigHierarchy.h b/Engine/Plugins/Animation/ControlRig/Source/ControlRig/Public/Rigs/RigHierarchy.h
--- a/Engine/Plugins/Animation/ControlRig/Source/ControlRig/Public/Rigs/RigHierarchy.h
+++ b/Engine/Plugins/Animation/ControlRig/Source/ControlRig/Public/Rigs/RigHierarchy.h
@@ -986,7 +986,26 @@ public:
 	 */
 	UFUNCTION(BlueprintCallable, Category = URigHierarchy)
 	TArray<FRigElementKey> RestoreConnectorsFromInfos(TArray<FRigConnectorInfo> InInfos, bool bSetupUndoRedo = false);
-	
+
+	/**
+	 * Returns all Socket elements
+	 * @param bTraverse Returns the elements in order of a depth first traversal
+	 */
+	TArray<FRigSocketElement*> GetSockets(bool bTraverse = false) const
+	{
+		return GetElementsOfType<FRigSocketElement>(bTraverse);
+	}
+
+	/**
+	 * Returns all Socket elements
+	 * @param bTraverse Returns the elements in order of a depth first traversal
+	 */
+	UFUNCTION(BlueprintCallable, Category = URigHierarchy, meta = (DisplayName = "Get Sockets", ScriptName = "GetSockets"))
+	TArray<FRigElementKey> GetSocketKeys(bool bTraverse = true) const
+	{
+		return GetKeysOfType<FRigSocketElement>(bTraverse);
+	}
+
 	/**
 	 * Returns all root elements
 	 */
@@ -4314,10 +4333,14 @@ protected:
 			{
 				return 6;
 			}
-			case ERigElementType::Last:
+			case ERigElementType::Socket:
 			{
 				return 7;
 			}
+			case ERigElementType::Last:
+			{
+				return 8;
+			}
 			case ERigElementType::All:
 			default:
 			{
@@ -4362,6 +4385,10 @@ protected:
 				return ERigElementType::Connector;
 			}
 			case 7:
+			{
+				return ERigElementType::Socket;
+			}
+			case 8:
 			{
 				return ERigElementType::Last;
 			}
diff --git a/Engine/Plugins/Animation/ControlRig/Source/ControlRig/Public/Rigs/RigHierarchyController.h b/Engine/Plugins/Animation/ControlRig/Source/ControlRig/Public/Rigs/RigHierarchyController.h
--- a/Engine/Plugins/Animation/ControlRig/Source/ControlRig/Public/Rigs/RigHierarchyController.h
+++ b/Engine/Plugins/Animation/ControlRig/Source/ControlRig/Public/Rigs/RigHierarchyController.h
@@ -250,6 +250,19 @@ public:
 	UFUNCTION(BlueprintCallable, Category = URigHierarchyController)
 	FRigElementKey AddConnector(FName InName, FTransform InTransform, bool bTransformInGlobal = true, FRigConnectorSettings InSettings = FRigConnectorSettings(), bool bSetupUndo = false, bool bPrintPythonCommand = false);
 
+	/**
+	 * Adds a socket to the hierarchy
+	 * @param InName The suggested name of the new socket - will eventually be corrected by the namespace
+	 * @param InParent The (optional) parent of the new null. If you don't need a parent, pass FRigElementKey()
+	 * @param InTransform The transform for the new socket - either in local or global space, based on bTransformInGlobal
+	 * @param bTransformInGlobal Set this to true if the Transform passed is expressed in global space, false for local space.
+	 * @param bSetupUndo If set to true the stack will record the change for undo / redo
+	 * @param bPrintPythonCommand If set to true a python command equivalent to this call will be printed out
+	 * @return The key for the newly created bone.
+	 */
+	UFUNCTION(BlueprintCallable, Category = URigHierarchyController)
+	FRigElementKey AddSocket(FName InName, FRigElementKey InParent, FTransform InTransform, bool bTransformInGlobal = true, bool bSetupUndo = false, bool bPrintPythonCommand = false);
+
 	/**
 	 * Returns the control settings of a given control
 	 * @param InKey The key of the control to receive the settings for
@@ -591,7 +604,9 @@ public:
 
 	TArray<FString> GetAddRigidBodyPythonCommands(FRigRigidBodyElement* RigidBody) const;
 
-	TArray<FString> GetAddConnectorPythonCommands(FRigConnectorElement* Control) const;
+	TArray<FString> GetAddConnectorPythonCommands(FRigConnectorElement* Connector) const;
+
+	TArray<FString> GetAddSocketPythonCommands(FRigSocketElement* Socket) const;
 
 	TArray<FString> GetSetControlValuePythonCommands(const FRigControlElement* Control, const FRigControlValue& Value, const ERigControlValueType& Type) const;
 	
diff --git a/Engine/Plugins/Animation/ControlRig/Source/ControlRig/Public/Rigs/RigHierarchyDefines.h b/Engine/Plugins/Animation/ControlRig/Source/ControlRig/Public/Rigs/RigHierarchyDefines.h
--- a/Engine/Plugins/Animation/ControlRig/Source/ControlRig/Public/Rigs/RigHierarchyDefines.h
+++ b/Engine/Plugins/Animation/ControlRig/Source/ControlRig/Public/Rigs/RigHierarchyDefines.h
@@ -18,9 +18,9 @@ class URigHierarchy;
  * This can be used as a mask so supported as a bitfield
  */
 UENUM(BlueprintType)
-enum class ERigElementType : uint8
+enum class ERigElementType
 {
-	None,
+	None = 0,
 	Bone = 0x001,
 	Null = 0x002,
 	Space = Null UMETA(Hidden),
@@ -29,9 +29,10 @@ enum class ERigElementType : uint8
 	RigidBody = 0x010 UMETA(Hidden), 
 	Reference = 0x020,
 	Connector = 0x040,
-	Last = 0x080 UMETA(Hidden),
-	All = Bone | Null | Control | Curve | RigidBody | Reference | Connector,
-	ToResetAfterConstructionEvent = Bone | Control | Curve UMETA(Hidden),
+	Socket = 0x080,
+	Last = 0x160 UMETA(Hidden),
+	All = Bone | Null | Control | Curve | RigidBody | Reference | Connector | Socket,
+	ToResetAfterConstructionEvent = Bone | Control | Curve | Socket UMETA(Hidden),
 };
 
 UENUM(BlueprintType)
@@ -1572,6 +1573,10 @@ public:
 			{
 				return FString::Printf(TEXT("Connector(%s)"), *Name.ToString());
 			}
+			case ERigElementType::Socket:
+			{
+				return FString::Printf(TEXT("Socket(%s)"), *Name.ToString());
+			}
 		}
 		
 		return FString();
diff --git a/Engine/Plugins/Animation/ControlRig/Source/ControlRig/Public/Rigs/RigHierarchyElements.h b/Engine/Plugins/Animation/ControlRig/Source/ControlRig/Public/Rigs/RigHierarchyElements.h
--- a/Engine/Plugins/Animation/ControlRig/Source/ControlRig/Public/Rigs/RigHierarchyElements.h
+++ b/Engine/Plugins/Animation/ControlRig/Source/ControlRig/Public/Rigs/RigHierarchyElements.h
@@ -569,6 +569,7 @@ public:
 		RigidBodyElement,
 		ReferenceElement,
 		ConnectorElement,
+		SocketElement,
 
 		Max
 	};
@@ -831,7 +832,8 @@ protected:
 			InElement->GetType() == ERigElementType::Control ||
 			InElement->GetType() == ERigElementType::RigidBody ||
 			InElement->GetType() == ERigElementType::Reference ||
-			InElement->GetType() == ERigElementType::Connector;
+			InElement->GetType() == ERigElementType::Connector ||
+			InElement->GetType() == ERigElementType::Socket;
 	}
 
 public:
@@ -880,7 +882,8 @@ protected:
 		return InElement->GetType() == ERigElementType::Bone ||
 			InElement->GetType() == ERigElementType::RigidBody ||
 			InElement->GetType() == ERigElementType::Reference ||
-			InElement->GetType() == ERigElementType::Connector;
+			InElement->GetType() == ERigElementType::Connector ||
+			InElement->GetType() == ERigElementType::Socket;
 	}
 
 	friend struct FRigBaseElement;
@@ -1754,6 +1757,45 @@ protected:
 	friend class URigHierarchyController;
 };
 
+USTRUCT(BlueprintType)
+struct CONTROLRIG_API FRigSocketElement : public FRigSingleParentElement
+{
+public:
+	
+	GENERATED_BODY()
+	DECLARE_RIG_ELEMENT_METHODS(FRigSocketElement)
+
+	static const EElementIndex ElementTypeIndex;
+
+	FRigSocketElement()
+		: FRigSingleParentElement()
+	{
+		Key.Type = ERigElementType::Socket; 
+	}
+
+	virtual ~FRigSocketElement(){}
+	
+	virtual void Save(FArchive& A, URigHierarchy* Hierarchy, ESerializationPhase SerializationPhase) override;
+	virtual void Load(FArchive& Ar, URigHierarchy* Hierarchy, ESerializationPhase SerializationPhase) override;
+
+private:
+
+	virtual void CopyFrom(URigHierarchy* InHierarchy, FRigBaseElement* InOther, URigHierarchy* InOtherHierarchy) override;
+
+protected:
+
+	static bool IsClassOf(const FRigBaseElement* InElement)
+	{
+		return InElement->GetType() == ERigElementType::Socket;
+	}
+
+protected:
+
+	friend struct FRigBaseElement;
+	friend class URigHierarchy;
+	friend class URigHierarchyController;
+};
+
 USTRUCT()
 struct CONTROLRIG_API FRigHierarchyCopyPasteContentPerElement
 {
diff --git a/Engine/Plugins/Animation/ControlRig/Source/ControlRigDeveloper/Private/ControlRigBlueprint.cpp b/Engine/Plugins/Animation/ControlRig/Source/ControlRigDeveloper/Private/ControlRigBlueprint.cpp
--- a/Engine/Plugins/Animation/ControlRig/Source/ControlRigDeveloper/Private/ControlRigBlueprint.cpp
+++ b/Engine/Plugins/Animation/ControlRig/Source/ControlRigDeveloper/Private/ControlRigBlueprint.cpp
@@ -455,6 +455,7 @@ TArray<URigVMNode*> UControlRigBlueprint::ConvertHierarchyElementsToSpawnerNodes
 
 			URigVMNode* AddNullNode = GraphController->AddUnitNode(FRigUnit_HierarchyAddNull::StaticStruct(), FRigUnit::GetMethodName(), NodePosition);
 			NodePosition += NodePositionIncrement;
+			AddParentItemLink(Key, AddNullNode);
 			SpawnerNodes.Add(AddNullNode);
 
 			if(LastPin)
@@ -584,6 +585,30 @@ TArray<URigVMNode*> UControlRigBlueprint::ConvertHierarchyElementsToSpawnerNodes
 				}
 			}
 		}
+		else if(Key.Type == ERigElementType::Socket)
+		{
+			FString ParentDefault, TransformDefault;
+			GetParentAndTransformDefaults(Key, ParentDefault, TransformDefault);
+
+			URigVMNode* AddSocketNode = GraphController->AddUnitNode(FRigUnit_HierarchyAddSocket::StaticStruct(), FRigUnit::GetMethodName(), NodePosition);
+			NodePosition += NodePositionIncrement;
+			AddParentItemLink(Key, AddSocketNode);
+			SpawnerNodes.Add(AddSocketNode);
+
+			if(LastPin)
+			{
+				if(const URigVMPin* NextPin = AddSocketNode->FindPin(GET_MEMBER_NAME_STRING_CHECKED(FRigUnit_HierarchyAddBone, ExecuteContext)))
+				{
+					GraphController->AddLink(LastPin->GetPinPath(), NextPin->GetPinPath(), true);
+					LastPin = NextPin;
+				}
+			}
+
+			GraphController->SetPinDefaultValue(AddSocketNode->FindPin(GET_MEMBER_NAME_STRING_CHECKED(FRigUnit_HierarchyAddElement, Name))->GetPinPath(), Key.Name.ToString(), true, true);
+			GraphController->SetPinDefaultValue(AddSocketNode->FindPin(GET_MEMBER_NAME_STRING_CHECKED(FRigUnit_HierarchyAddElement, Parent))->GetPinPath(), ParentDefault, true, true);
+			GraphController->SetPinDefaultValue(AddSocketNode->FindPin(GET_MEMBER_NAME_STRING_CHECKED(FRigUnit_HierarchyAddNull, Space))->GetPinPath(), TEXT("LocalSpace"), true, true);
+			GraphController->SetPinDefaultValue(AddSocketNode->FindPin(GET_MEMBER_NAME_STRING_CHECKED(FRigUnit_HierarchyAddNull, Transform))->GetPinPath(), TransformDefault, true, true);
+		}
 	}
 
 	if(bRemoveElements && InHierarchy)
diff --git a/Engine/Plugins/Animation/ControlRig/Source/ControlRigEditor/Private/ControlRigElementDetails.h b/Engine/Plugins/Animation/ControlRig/Source/ControlRigEditor/Private/ControlRigElementDetails.h
--- a/Engine/Plugins/Animation/ControlRig/Source/ControlRigEditor/Private/ControlRigElementDetails.h
+++ b/Engine/Plugins/Animation/ControlRig/Source/ControlRigEditor/Private/ControlRigElementDetails.h
@@ -756,6 +756,17 @@ public:
 	void CustomizeSettings(IDetailLayoutBuilder& DetailBuilder);
 };
 
+class FRigSocketElementDetails : public FRigNullElementDetails
+{
+public:
+
+	// Makes a new instance of this detail layout class for a specific detail view requesting it
+	static TSharedRef<IDetailCustomization> MakeInstance()
+	{
+		return MakeShareable(new FRigSocketElementDetails);
+	}
+};
+
 class FRigConnectionRuleDetails : public IPropertyTypeCustomization
 {
 public:
diff --git a/Engine/Plugins/Animation/ControlRig/Source/ControlRigEditor/Private/ControlRigHierarchyCommands.cpp b/Engine/Plugins/Animation/ControlRig/Source/ControlRigEditor/Private/ControlRigHierarchyCommands.cpp
--- a/Engine/Plugins/Animation/ControlRig/Source/ControlRigEditor/Private/ControlRigHierarchyCommands.cpp
+++ b/Engine/Plugins/Animation/ControlRig/Source/ControlRigEditor/Private/ControlRigHierarchyCommands.cpp
@@ -11,6 +11,7 @@ void FControlRigHierarchyCommands::RegisterCommands()
 	UI_COMMAND(AddAnimationChannelItem, "New Animation Channel", "Add new animation channel.", EUserInterfaceActionType::Button, FInputChord());
 	UI_COMMAND(AddNullItem, "New Null", "Add new null at the origin (0, 0, 0) to the hierarchy.", EUserInterfaceActionType::Button, FInputChord());
 	UI_COMMAND(AddConnectorItem, "New Connector", "Add new connector at the origin (0, 0, 0) to the hierarchy.", EUserInterfaceActionType::Button, FInputChord());
+	UI_COMMAND(AddSocketItem, "New Socket", "Add new socket at the origin (0, 0, 0) to the hierarchy.", EUserInterfaceActionType::Button, FInputChord());
 	UI_COMMAND(DuplicateItem, "Duplicate", "Duplicate the selected items in the hierarchy.", EUserInterfaceActionType::Button, FInputChord(EKeys::D, EModifierKey::Control));
 	UI_COMMAND(MirrorItem, "Mirror", "Mirror the selected items in the hierarchy.", EUserInterfaceActionType::Button, FInputChord());
 	UI_COMMAND(DeleteItem, "Delete", "Deletes the selected items from the hierarchy.", EUserInterfaceActionType::Button, FInputChord(EKeys::Delete));
@@ -37,6 +38,7 @@ void FControlRigHierarchyCommands::RegisterCommands()
 	UI_COMMAND(ShowNulls, "Show Nulls", "Whether to show or hide nulls", EUserInterfaceActionType::ToggleButton, FInputChord());
 	UI_COMMAND(ShowRigidBodies, "Show RigidBodies", "Whether to show or hide rigidbodies", EUserInterfaceActionType::ToggleButton, FInputChord());
 	UI_COMMAND(ShowReferences, "Show References", "Whether to show or hide references", EUserInterfaceActionType::ToggleButton, FInputChord());
+	UI_COMMAND(ShowSockets, "Show Sockets", "Whether to show or hide sockets", EUserInterfaceActionType::ToggleButton, FInputChord());
 	UI_COMMAND(ToggleControlShapeTransformEdit, "Toggle Shape Transform Edit", "Toggle Editing Selected Control's Shape Transform", EUserInterfaceActionType::ToggleButton, FInputChord(EKeys::Period, EModifierKey::Control)); 
 	UI_COMMAND(SpaceSwitching, "Space Switching", "Space switching on the control", EUserInterfaceActionType::Button, FInputChord(EKeys::Tab)); 
 	UI_COMMAND(ShowIconColors, "Show Icon Colors", "Whether to tint the icons with the element color", EUserInterfaceActionType::ToggleButton, FInputChord());
diff --git a/Engine/Plugins/Animation/ControlRig/Source/ControlRigEditor/Private/ControlRigHierarchyCommands.h b/Engine/Plugins/Animation/ControlRig/Source/ControlRigEditor/Private/ControlRigHierarchyCommands.h
--- a/Engine/Plugins/Animation/ControlRig/Source/ControlRigEditor/Private/ControlRigHierarchyCommands.h
+++ b/Engine/Plugins/Animation/ControlRig/Source/ControlRigEditor/Private/ControlRigHierarchyCommands.h
@@ -33,6 +33,9 @@ public:
 	/** Add Item at origin */
 	TSharedPtr< FUICommandInfo > AddConnectorItem;
 
+	/** Add Item at origin */
+	TSharedPtr< FUICommandInfo > AddSocketItem;
+
 	/** Duplicate currently selected items */
 	TSharedPtr< FUICommandInfo > DuplicateItem;
 
@@ -111,6 +114,9 @@ public:
 	/* Show references */
 	TSharedPtr< FUICommandInfo > ShowReferences;
 
+	/* Show sockets */
+	TSharedPtr< FUICommandInfo > ShowSockets;
+
 	/** Toggle Shape Transform Edit*/
 	TSharedPtr< FUICommandInfo > ToggleControlShapeTransformEdit;
 
diff --git a/Engine/Plugins/Animation/ControlRig/Source/ControlRigEditor/Private/EditMode/ControlRigEditMode.cpp b/Engine/Plugins/Animation/ControlRig/Source/ControlRigEditor/Private/EditMode/ControlRigEditMode.cpp
--- a/Engine/Plugins/Animation/ControlRig/Source/ControlRigEditor/Private/EditMode/ControlRigEditMode.cpp
+++ b/Engine/Plugins/Animation/ControlRig/Source/ControlRigEditor/Private/EditMode/ControlRigEditMode.cpp
@@ -3807,6 +3807,7 @@ void FControlRigEditMode::OnHierarchyModified(ERigHierarchyNotification InNotif,
             	case ERigElementType::RigidBody:
             	case ERigElementType::Reference:
             	case ERigElementType::Connector:
+            	case ERigElementType::Socket:
 				{
 					const bool bSelected = InNotif == ERigHierarchyNotification::ElementSelected;
 					// users may select gizmo and control rig units, so we have to let them go through both of them if they do
diff --git a/Engine/Plugins/Animation/ControlRig/Source/ControlRigEditor/Private/EditMode/SControlRigEditModeTools.cpp b/Engine/Plugins/Animation/ControlRig/Source/ControlRigEditor/Private/EditMode/SControlRigEditModeTools.cpp
--- a/Engine/Plugins/Animation/ControlRig/Source/ControlRigEditor/Private/EditMode/SControlRigEditModeTools.cpp
+++ b/Engine/Plugins/Animation/ControlRig/Source/ControlRigEditor/Private/EditMode/SControlRigEditModeTools.cpp
@@ -179,6 +179,7 @@ void SControlRigEditModeTools::Construct(const FArguments& InArgs, TSharedPtr<FC
 	DisplaySettings.bShowControls = true;
 	DisplaySettings.bShowNulls = false;
 	DisplaySettings.bShowReferences = false;
+	DisplaySettings.bShowSockets = false;
 	DisplaySettings.bShowRigidBodies = false;
 	DisplaySettings.bHideParentsOnFilter = true;
 	DisplaySettings.bFlattenHierarchyOnFilter = true;
diff --git a/Engine/Plugins/Animation/ControlRig/Source/ControlRigEditor/Private/EditMode/SControlRigOutliner.cpp b/Engine/Plugins/Animation/ControlRig/Source/ControlRigEditor/Private/EditMode/SControlRigOutliner.cpp
--- a/Engine/Plugins/Animation/ControlRig/Source/ControlRigEditor/Private/EditMode/SControlRigOutliner.cpp
+++ b/Engine/Plugins/Animation/ControlRig/Source/ControlRigEditor/Private/EditMode/SControlRigOutliner.cpp
@@ -531,6 +531,14 @@ bool SMultiRigHierarchyTreeView::AddElement(UBaseControlRig* InControlRig, const
 				}
 				break;
 			}
+			case ERigElementType::Socket:
+			{
+				if (!Settings.bShowSockets)
+				{
+					return false;
+				}
+				break;
+			}
 			case ERigElementType::Curve:
 			{
 				return false;
@@ -953,6 +961,7 @@ void SControlRigOutliner::Construct(const FArguments& InArgs, FControlRigEditMod
 	DisplaySettings.bShowControls = true;
 	DisplaySettings.bShowNulls = false;
 	DisplaySettings.bShowReferences = false;
+	DisplaySettings.bShowSockets = false;
 	DisplaySettings.bShowRigidBodies = false;
 	DisplaySettings.bHideParentsOnFilter = true;
 	DisplaySettings.bFlattenHierarchyOnFilter = true;
diff --git a/Engine/Plugins/Animation/ControlRig/Source/ControlRigEditor/Private/Editor/ControlRigEditor.cpp b/Engine/Plugins/Animation/ControlRig/Source/ControlRigEditor/Private/Editor/ControlRigEditor.cpp
--- a/Engine/Plugins/Animation/ControlRig/Source/ControlRigEditor/Private/Editor/ControlRigEditor.cpp
+++ b/Engine/Plugins/Animation/ControlRig/Source/ControlRigEditor/Private/Editor/ControlRigEditor.cpp
@@ -2437,6 +2437,7 @@ void FControlRigEditor::SetRigElementTransform(const FRigElementKey& InElement,
 	{
 		case ERigElementType::Bone:
 		case ERigElementType::Connector:
+		case ERigElementType::Socket:
 		{
 			FTransform Transform = InTransform;
 			if (bLocal)
diff --git a/Engine/Plugins/Animation/ControlRig/Source/ControlRigEditor/Private/Editor/ControlRigWrapperObject.cpp b/Engine/Plugins/Animation/ControlRig/Source/ControlRigEditor/Private/Editor/ControlRigWrapperObject.cpp
--- a/Engine/Plugins/Animation/ControlRig/Source/ControlRigEditor/Private/Editor/ControlRigWrapperObject.cpp
+++ b/Engine/Plugins/Animation/ControlRig/Source/ControlRigEditor/Private/Editor/ControlRigWrapperObject.cpp
@@ -40,6 +40,10 @@ UClass* UControlRigWrapperObject::GetClassForStruct(UScriptStruct* InStruct, boo
 			{
 				PropertyEditorModule.RegisterCustomClassLayout(WrapperClassName, FOnGetDetailCustomizationInstance::CreateStatic(&FRigConnectorElementDetails::MakeInstance));
 			}
+			else if(InStruct == FRigSocketElement::StaticStruct())
+			{
+				PropertyEditorModule.RegisterCustomClassLayout(WrapperClassName, FOnGetDetailCustomizationInstance::CreateStatic(&FRigSocketElementDetails::MakeInstance));
+			}
 		}
 	}
 #endif
diff --git a/Engine/Plugins/Animation/ControlRig/Source/ControlRigEditor/Private/Editor/SRigHierarchy.cpp b/Engine/Plugins/Animation/ControlRig/Source/ControlRigEditor/Private/Editor/SRigHierarchy.cpp
--- a/Engine/Plugins/Animation/ControlRig/Source/ControlRigEditor/Private/Editor/SRigHierarchy.cpp
+++ b/Engine/Plugins/Animation/ControlRig/Source/ControlRigEditor/Private/Editor/SRigHierarchy.cpp
@@ -332,6 +332,10 @@ void SRigHierarchy::BindCommands()
 		FExecuteAction::CreateSP(this, &SRigHierarchy::HandleNewItem, ERigElementType::Connector, false),
 		FCanExecuteAction::CreateSP(this, &SRigHierarchy::IsNonProceduralElementSelected));
 
+	CommandList->MapAction(Commands.AddSocketItem,
+		FExecuteAction::CreateSP(this, &SRigHierarchy::HandleNewItem, ERigElementType::Socket, false),
+		FCanExecuteAction::CreateSP(this, &SRigHierarchy::IsNonProceduralElementSelected));
+
 	CommandList->MapAction(Commands.DuplicateItem,
 		FExecuteAction::CreateSP(this, &SRigHierarchy::HandleDuplicateItem),
 		FCanExecuteAction::CreateSP(this, &SRigHierarchy::CanDuplicateItem));
@@ -454,6 +458,12 @@ void SRigHierarchy::BindCommands()
 		FCanExecuteAction(),
 		FIsActionChecked::CreateLambda([this]() { return DisplaySettings.bShowReferences; }));
 
+	CommandList->MapAction(
+		Commands.ShowSockets,
+		FExecuteAction::CreateLambda([this]() { DisplaySettings.bShowSockets = !DisplaySettings.bShowSockets; RefreshTreeView(); }),
+		FCanExecuteAction(),
+		FIsActionChecked::CreateLambda([this]() { return DisplaySettings.bShowSockets; }));
+
 	CommandList->MapAction(
 		Commands.ToggleControlShapeTransformEdit,
 		FExecuteAction::CreateLambda([this]()
@@ -501,7 +511,8 @@ EVisibility SRigHierarchy::IsSearchbarVisible() const
 		if ((Hierarchy->Num(ERigElementType::Bone) +
 			Hierarchy->Num(ERigElementType::Null) +
 			Hierarchy->Num(ERigElementType::Control) +
-			Hierarchy->Num(ERigElementType::Connector)) > 0)
+			Hierarchy->Num(ERigElementType::Connector) +
+			Hierarchy->Num(ERigElementType::Socket)) > 0)
 		{
 			return EVisibility::Visible;
 		}
@@ -1243,6 +1254,7 @@ void SRigHierarchy::CreateContextMenu()
 							if(CVarControlRigHierarchyEnableModules.GetValueOnAnyThread())
 							{
 								DefaultSection.AddMenuEntry(Commands.AddConnectorItem);
+								DefaultSection.AddMenuEntry(Commands.AddSocketItem);
 							}
 						})
 					);
@@ -1998,6 +2010,11 @@ void SRigHierarchy::HandleNewItem(ERigElementType InElementType, bool bIsAnimati
 					(void)ResolveConnector(NewItemKey, ParentKey);
 					break;
 				}
+				case ERigElementType::Socket:
+				{
+					NewItemKey = Controller->AddSocket(NewElementName, ParentKey, ParentTransform, true, true, true);
+					break;
+				}
 				default:
 				{
 					return;
@@ -2480,6 +2497,7 @@ TOptional<EItemDropZone> SRigHierarchy::OnCanAcceptDrop(const FDragDropEvent& Dr
 						case ERigElementType::RigidBody:
 						case ERigElementType::Reference:
 						case ERigElementType::Connector:
+						case ERigElementType::Socket:
 						{
 							break;
 						}
@@ -2498,6 +2516,11 @@ TOptional<EItemDropZone> SRigHierarchy::OnCanAcceptDrop(const FDragDropEvent& Dr
 				ReturnDropZone = DropZone;
 				break;
 			}
+			case ERigElementType::Socket:
+			{
+				// You cannot parent anything under a socket
+				return InvalidDropZone;
+			}
 			default:
 			{
 				ReturnDropZone = DropZone;
diff --git a/Engine/Plugins/Animation/ControlRig/Source/ControlRigEditor/Private/Editor/SRigHierarchyTreeView.cpp b/Engine/Plugins/Animation/ControlRig/Source/ControlRigEditor/Private/Editor/SRigHierarchyTreeView.cpp
--- a/Engine/Plugins/Animation/ControlRig/Source/ControlRigEditor/Private/Editor/SRigHierarchyTreeView.cpp
+++ b/Engine/Plugins/Animation/ControlRig/Source/ControlRigEditor/Private/Editor/SRigHierarchyTreeView.cpp
@@ -538,6 +538,14 @@ bool SRigHierarchyTreeView::AddElement(const FRigBaseElement* InElement)
 			}
 			break;
 		}
+		case ERigElementType::Socket:
+		{
+			if(!Settings.bShowSockets)
+			{
+				return false;
+			}
+			break;
+		}
 		default:
 		{
 			break;
@@ -1036,6 +1044,7 @@ TPair<const FSlateBrush*, FSlateColor> SRigHierarchyItem::GetBrushForElementType
 			break;
 		}
 		case ERigElementType::Reference:
+		case ERigElementType::Socket:
 		{
 			Brush = FControlRigEditorStyle::Get().GetBrush("ControlRig.Tree.Socket");
 			break;
diff --git a/Engine/Plugins/Animation/ControlRig/Source/ControlRigEditor/Private/Editor/SRigHierarchyTreeView.h b/Engine/Plugins/Animation/ControlRig/Source/ControlRigEditor/Private/Editor/SRigHierarchyTreeView.h
--- a/Engine/Plugins/Animation/ControlRig/Source/ControlRigEditor/Private/Editor/SRigHierarchyTreeView.h
+++ b/Engine/Plugins/Animation/ControlRig/Source/ControlRigEditor/Private/Editor/SRigHierarchyTreeView.h
@@ -26,6 +26,7 @@ struct CONTROLRIGEDITOR_API FRigTreeDisplaySettings
 		bShowNulls = true;
 		bShowRigidBodies = true;
 		bShowReferences = true;
+		bShowSockets = true;
 		bShowIconColors = true;
 	}
 	
@@ -55,6 +56,9 @@ struct CONTROLRIGEDITOR_API FRigTreeDisplaySettings
 	/** Whether or not to show references in the hierarchy */
 	bool bShowReferences;
 
+	/** Whether or not to show sockets in the hierarchy */
+	bool bShowSockets;
+
 	/** Whether to tint the icons with the element color */
 	bool bShowIconColors;
 };
