Move FindControlRigs back to UControlRig to fix existing blueprints that call this function.

diff --git a/Engine/Plugins/Animation/ControlRig/Source/ControlRig/Private/BaseControlRig.cpp b/Engine/Plugins/Animation/ControlRig/Source/ControlRig/Private/BaseControlRig.cpp
--- a/Engine/Plugins/Animation/ControlRig/Source/ControlRig/Private/BaseControlRig.cpp
+++ b/Engine/Plugins/Animation/ControlRig/Source/ControlRig/Private/BaseControlRig.cpp
@@ -1425,7 +1425,7 @@ void UBaseControlRig::PostReinstanceCallback(const UBaseControlRig* Old)
 
 #endif // WITH_EDITORONLY_DATA
 
-TArray<UBaseControlRig*> UBaseControlRig::FindControlRigs(UObject* Outer, TSubclassOf<UBaseControlRig> OptionalClass)
+TArray<UBaseControlRig*> UBaseControlRig::FindBaseControlRigs(UObject* Outer, TSubclassOf<UBaseControlRig> OptionalClass)
 {
 	TArray<UBaseControlRig*> Result;
 	
diff --git a/Engine/Plugins/Animation/ControlRig/Source/ControlRig/Private/ControlRig.cpp b/Engine/Plugins/Animation/ControlRig/Source/ControlRig/Private/ControlRig.cpp
--- a/Engine/Plugins/Animation/ControlRig/Source/ControlRig/Private/ControlRig.cpp
+++ b/Engine/Plugins/Animation/ControlRig/Source/ControlRig/Private/ControlRig.cpp
@@ -8,6 +8,7 @@
 #include "ControlRigObjectBinding.h"
 #include "Rigs/RigHierarchyController.h"
 #include "ControlRigComponent.h"
+#include "UObject/UObjectIterator.h"
 
 #define LOCTEXT_NAMESPACE "ControlRig"
 
@@ -16,6 +17,59 @@ UControlRig::UControlRig(const FObjectInitializer& ObjectInitializer)
 {
 }
 
+TArray<UControlRig*> UControlRig::FindControlRigs(UObject* Outer, TSubclassOf<UControlRig> OptionalClass)
+{
+	TArray<UControlRig*> Result;
+	
+	if(Outer == nullptr)
+	{
+		return Result; 
+	}
+	
+	AActor* OuterActor = Cast<AActor>(Outer);
+	if(OuterActor == nullptr)
+	{
+		OuterActor = Outer->GetTypedOuter<AActor>();
+	}
+	
+	for (TObjectIterator<UControlRig> Itr; Itr; ++Itr)
+	{
+		UControlRig* RigInstance = *Itr;
+		const UClass* RigInstanceClass = RigInstance ? RigInstance->GetClass() : nullptr;
+		if (OptionalClass == nullptr || (RigInstanceClass && RigInstanceClass->IsChildOf(OptionalClass)))
+		{
+			if(RigInstance->IsInOuter(Outer))
+			{
+				Result.Add(RigInstance);
+				continue;
+			}
+
+			if(OuterActor)
+			{
+				if(RigInstance->IsInOuter(OuterActor))
+				{
+					Result.Add(RigInstance);
+					continue;
+				}
+
+				if (TSharedPtr<IControlRigObjectBinding> Binding = RigInstance->GetObjectBinding())
+				{
+					if (AActor* Actor = Binding->GetHostingActor())
+					{
+						if (Actor == OuterActor)
+						{
+							Result.Add(RigInstance);
+							continue;
+						}
+					}
+				}
+			}
+		}
+	}
+
+	return Result;
+}
+
 bool UControlRig::Execute_Internal(const FName& InEventName)
 {
 	if(IsRigModule())
diff --git a/Engine/Plugins/Animation/ControlRig/Source/ControlRig/Private/ControlRigControlActor.cpp b/Engine/Plugins/Animation/ControlRig/Source/ControlRig/Private/ControlRigControlActor.cpp
--- a/Engine/Plugins/Animation/ControlRig/Source/ControlRig/Private/ControlRigControlActor.cpp
+++ b/Engine/Plugins/Animation/ControlRig/Source/ControlRig/Private/ControlRigControlActor.cpp
@@ -156,7 +156,7 @@ void AControlRigControlActor::Refresh()
 
 	if (ControlRig == nullptr)
 	{
-		TArray<UBaseControlRig*> Rigs = UBaseControlRig::FindControlRigs(ActorToTrack, ControlRigClass);
+		TArray<UBaseControlRig*> Rigs = UBaseControlRig::FindBaseControlRigs(ActorToTrack, ControlRigClass);
 		if(Rigs.Num() > 0)
 		{
 			ControlRig = Rigs[0];
diff --git a/Engine/Plugins/Animation/ControlRig/Source/ControlRig/Public/BaseControlRig.h b/Engine/Plugins/Animation/ControlRig/Source/ControlRig/Public/BaseControlRig.h
--- a/Engine/Plugins/Animation/ControlRig/Source/ControlRig/Public/BaseControlRig.h
+++ b/Engine/Plugins/Animation/ControlRig/Source/ControlRig/Public/BaseControlRig.h
@@ -77,7 +77,7 @@ public:
 	static const FName OwnerComponent;
 
 	UFUNCTION(BlueprintCallable, Category = ControlRig)
-	static TArray<UBaseControlRig*> FindControlRigs(UObject* Outer, TSubclassOf<UBaseControlRig> OptionalClass);
+	static TArray<UBaseControlRig*> FindBaseControlRigs(UObject* Outer, TSubclassOf<UBaseControlRig> OptionalClass);
 
 public:
 	virtual UWorld* GetWorld() const override;
diff --git a/Engine/Plugins/Animation/ControlRig/Source/ControlRig/Public/ControlRig.h b/Engine/Plugins/Animation/ControlRig/Source/ControlRig/Public/ControlRig.h
--- a/Engine/Plugins/Animation/ControlRig/Source/ControlRig/Public/ControlRig.h
+++ b/Engine/Plugins/Animation/ControlRig/Source/ControlRig/Public/ControlRig.h
@@ -34,6 +34,9 @@ public:
 		return nullptr;
 	}
 
+	UFUNCTION(BlueprintCallable, Category = ControlRig)
+	static TArray<UControlRig*> FindControlRigs(UObject* Outer, TSubclassOf<UControlRig> OptionalClass);
+
 	UE_DEPRECATED(5.4, "InteractionRig is no longer used")
 	UFUNCTION(BlueprintSetter, meta = (DeprecatedFunction, DeprecationMessage = "InteractionRig is no longer used"))
 	void SetInteractionRig(UControlRig* InInteractionRig) {}
