Sequencer: Control Rig: CIS Optimus crash, need to make sure mask is correct which get's setup when the proxy is created.

diff --git a/Engine/Plugins/Animation/ControlRig/Source/ControlRig/Private/Sequencer/MovieSceneControlRigParameterTemplate.cpp b/Engine/Plugins/Animation/ControlRig/Source/ControlRig/Private/Sequencer/MovieSceneControlRigParameterTemplate.cpp
--- a/Engine/Plugins/Animation/ControlRig/Source/ControlRig/Private/Sequencer/MovieSceneControlRigParameterTemplate.cpp
+++ b/Engine/Plugins/Animation/ControlRig/Source/ControlRig/Private/Sequencer/MovieSceneControlRigParameterTemplate.cpp
@@ -283,7 +283,7 @@ struct FEvaluatedControlRigParameterSectionChannelMasks : IPersistentEvaluationD
 	TBitArray<> ColorCurveMask;
 	TBitArray<> TransformCurveMask;
 
-	void Initialize(const UMovieSceneControlRigParameterSection* Section,
+	void Initialize(UMovieSceneControlRigParameterSection* Section,
 		TArrayView<const FScalarParameterNameAndCurve> Scalars,
 		TArrayView<const FBoolParameterNameAndCurve> Bools,
 		TArrayView<const FIntegerParameterNameAndCurve> Integers,
@@ -1814,7 +1814,8 @@ void FMovieSceneControlRigParameterTemplate::Evaluate(const FMovieSceneEvaluatio
 			// Naughty const_cast here, but we can't create this inside Initialize because of hotfix restrictions
 			// The cast is ok because we actually do not have any threading involved
 			ChannelMasks = &const_cast<FPersistentEvaluationData&>(PersistentData).GetOrAddSectionData<FEvaluatedControlRigParameterSectionChannelMasks>();
-			ChannelMasks->Initialize(Section, Scalars, Bools, Integers, Enums, Vector2Ds, Vectors, Colors, Transforms);
+			UMovieSceneControlRigParameterSection* NonConstSection = const_cast<UMovieSceneControlRigParameterSection*>(Section);
+			ChannelMasks->Initialize(NonConstSection, Scalars, Bools, Integers, Enums, Vector2Ds, Vectors, Colors, Transforms);
 		}
 
 		UMovieSceneTrack* Track = Cast<UMovieSceneTrack>(Section->GetOuter());
@@ -2316,7 +2317,8 @@ void FMovieSceneControlRigParameterTemplate::Interrogate(const FMovieSceneContex
 	if (Section && Section->GetControlRig() && MovieSceneHelpers::IsSectionKeyable(Section))
 	{
 		FEvaluatedControlRigParameterSectionChannelMasks ChannelMasks;
-		ChannelMasks.Initialize(Section, Scalars, Bools, Integers, Enums, Vector2Ds, Vectors, Colors, Transforms);
+		UMovieSceneControlRigParameterSection* NonConstSection = const_cast<UMovieSceneControlRigParameterSection*>(Section);
+		ChannelMasks.Initialize(NonConstSection, Scalars, Bools, Integers, Enums, Vector2Ds, Vectors, Colors, Transforms);
 
 		//Do blended tokens
 		FEvaluatedControlRigParameterSectionValues Values;
diff --git a/Engine/Plugins/Animation/ControlRig/Source/ControlRig/Public/Sequencer/MovieSceneControlRigParameterSection.h b/Engine/Plugins/Animation/ControlRig/Source/ControlRig/Public/Sequencer/MovieSceneControlRigParameterSection.h
--- a/Engine/Plugins/Animation/ControlRig/Source/ControlRig/Public/Sequencer/MovieSceneControlRigParameterSection.h
+++ b/Engine/Plugins/Animation/ControlRig/Source/ControlRig/Public/Sequencer/MovieSceneControlRigParameterSection.h
@@ -317,13 +317,21 @@ public:
 	//Will return false if fails or is canceled
 	virtual bool LoadAnimSequenceIntoThisSection(UAnimSequence* Sequence, UMovieScene* MovieScene, UObject* BoundObject, bool bKeyReduce, float Tolerance, FFrameNumber InStartFrame = 0, EMovieSceneKeyInterpolation InInterolation = EMovieSceneKeyInterpolation::SmartAuto);
 #endif
-	const TArray<bool>& GetControlsMask() const
+	const TArray<bool>& GetControlsMask() 
 	{
+		if (ChannelProxy.IsValid() == false)
+		{
+			CacheChannelProxy();
+		}
 		return ControlsMask;
 	}
 
-	bool GetControlsMask(int32 Index) const
+	bool GetControlsMask(int32 Index)  
 	{
+		if (ChannelProxy.IsValid() == false)
+		{
+			CacheChannelProxy();
+		}
 		if (Index >= 0 && Index < ControlsMask.Num())
 		{
 			return ControlsMask[Index];
