Control Rig: Support animation channels when converting a rig to a module

diff --git a/Engine/Plugins/Animation/ControlRig/Source/ControlRig/Private/Units/Execution/RigUnit_DynamicHierarchy.cpp b/Engine/Plugins/Animation/ControlRig/Source/ControlRig/Private/Units/Execution/RigUnit_DynamicHierarchy.cpp
--- a/Engine/Plugins/Animation/ControlRig/Source/ControlRig/Private/Units/Execution/RigUnit_DynamicHierarchy.cpp
+++ b/Engine/Plugins/Animation/ControlRig/Source/ControlRig/Private/Units/Execution/RigUnit_DynamicHierarchy.cpp
@@ -929,6 +929,7 @@ FRigUnit_HierarchyAddAnimationChannelFloat_Execute()
 		FRigControlSettings ControlSettings;
 		ControlSettings.ControlType = ERigControlType::Float;
 		ControlSettings.SetupLimitArrayForType(true, true, true);
+		ControlSettings.LimitEnabled[0] = LimitsEnabled.Enabled;
 		ControlSettings.MinimumValue = FRigControlValue::Make<float>(MinimumValue);
 		ControlSettings.MaximumValue = FRigControlValue::Make<float>(MaximumValue);
 		ControlSettings.DisplayName = Controller->GetHierarchy()->GetSafeNewDisplayName(Parent, Name);
@@ -968,6 +969,7 @@ FRigUnit_HierarchyAddAnimationChannelScaleFloat_Execute()
 		FRigControlSettings ControlSettings;
 		ControlSettings.ControlType = ERigControlType::ScaleFloat;
 		ControlSettings.SetupLimitArrayForType(true, true, true);
+		ControlSettings.LimitEnabled[0] = LimitsEnabled.Enabled;
 		ControlSettings.MinimumValue = FRigControlValue::Make<float>(MinimumValue);
 		ControlSettings.MaximumValue = FRigControlValue::Make<float>(MaximumValue);
 		ControlSettings.DisplayName = Controller->GetHierarchy()->GetSafeNewDisplayName(Parent, Name);
@@ -1007,6 +1009,7 @@ FRigUnit_HierarchyAddAnimationChannelInteger_Execute()
 		FRigControlSettings ControlSettings;
 		ControlSettings.ControlType = ERigControlType::Integer;
 		ControlSettings.SetupLimitArrayForType(true, true, true);
+		ControlSettings.LimitEnabled[0] = LimitsEnabled.Enabled;
 		ControlSettings.MinimumValue = FRigControlValue::Make<int32>(MinimumValue);
 		ControlSettings.MaximumValue = FRigControlValue::Make<int32>(MaximumValue);
 		ControlSettings.DisplayName = Controller->GetHierarchy()->GetSafeNewDisplayName(Parent, Name);
@@ -1046,6 +1049,8 @@ FRigUnit_HierarchyAddAnimationChannelVector2D_Execute()
 		FRigControlSettings ControlSettings;
 		ControlSettings.ControlType = ERigControlType::Vector2D;
 		ControlSettings.SetupLimitArrayForType(true, true, true);
+		ControlSettings.LimitEnabled[0] = LimitsEnabled.X;
+		ControlSettings.LimitEnabled[1] = LimitsEnabled.Y;
 		ControlSettings.MinimumValue = FRigControlValue::Make<FVector3f>(FVector3f(MinimumValue.X, MinimumValue.Y, 0.f));
 		ControlSettings.MaximumValue = FRigControlValue::Make<FVector3f>(FVector3f(MaximumValue.X, MaximumValue.Y, 0.f));
 		ControlSettings.DisplayName = Controller->GetHierarchy()->GetSafeNewDisplayName(Parent, Name);
@@ -1085,6 +1090,9 @@ FRigUnit_HierarchyAddAnimationChannelVector_Execute()
 		FRigControlSettings ControlSettings;
 		ControlSettings.ControlType = ERigControlType::Position;
 		ControlSettings.SetupLimitArrayForType(true, true, true);
+		ControlSettings.LimitEnabled[0] = LimitsEnabled.X;
+		ControlSettings.LimitEnabled[1] = LimitsEnabled.Y;
+		ControlSettings.LimitEnabled[2] = LimitsEnabled.Z;
 		ControlSettings.MinimumValue = FRigControlValue::Make<FVector3f>(FVector3f(MinimumValue));
 		ControlSettings.MaximumValue = FRigControlValue::Make<FVector3f>(FVector3f(MaximumValue));
 		ControlSettings.DisplayName = Controller->GetHierarchy()->GetSafeNewDisplayName(Parent, Name);
@@ -1124,6 +1132,9 @@ FRigUnit_HierarchyAddAnimationChannelScaleVector_Execute()
 		FRigControlSettings ControlSettings;
 		ControlSettings.ControlType = ERigControlType::Scale;
 		ControlSettings.SetupLimitArrayForType(true, true, true);
+		ControlSettings.LimitEnabled[0] = LimitsEnabled.X;
+		ControlSettings.LimitEnabled[1] = LimitsEnabled.Y;
+		ControlSettings.LimitEnabled[2] = LimitsEnabled.Z;
 		ControlSettings.MinimumValue = FRigControlValue::Make<FVector3f>(FVector3f(MinimumValue));
 		ControlSettings.MaximumValue = FRigControlValue::Make<FVector3f>(FVector3f(MaximumValue));
 		ControlSettings.DisplayName = Controller->GetHierarchy()->GetSafeNewDisplayName(Parent, Name);
@@ -1163,6 +1174,9 @@ FRigUnit_HierarchyAddAnimationChannelRotator_Execute()
 		FRigControlSettings ControlSettings;
 		ControlSettings.ControlType = ERigControlType::Rotator;
 		ControlSettings.SetupLimitArrayForType(true, true, true);
+		ControlSettings.LimitEnabled[0] = LimitsEnabled.Pitch;
+		ControlSettings.LimitEnabled[1] = LimitsEnabled.Yaw;
+		ControlSettings.LimitEnabled[2] = LimitsEnabled.Roll;
 		ControlSettings.MinimumValue = FRigControlValue::Make<FVector3f>(FVector3f(MinimumValue.Euler()));
 		ControlSettings.MaximumValue = FRigControlValue::Make<FVector3f>(FVector3f(MaximumValue.Euler()));
 		ControlSettings.DisplayName = Controller->GetHierarchy()->GetSafeNewDisplayName(Parent, Name);
diff --git a/Engine/Plugins/Animation/ControlRig/Source/ControlRig/Private/Units/Execution/RigUnit_DynamicHierarchy.h b/Engine/Plugins/Animation/ControlRig/Source/ControlRig/Private/Units/Execution/RigUnit_DynamicHierarchy.h
--- a/Engine/Plugins/Animation/ControlRig/Source/ControlRig/Private/Units/Execution/RigUnit_DynamicHierarchy.h
+++ b/Engine/Plugins/Animation/ControlRig/Source/ControlRig/Private/Units/Execution/RigUnit_DynamicHierarchy.h
@@ -1068,6 +1068,15 @@ struct CONTROLRIG_API FRigUnit_HierarchyAddControlTransform : public FRigUnit_Hi
 	virtual void Execute() override;
 };
 
+USTRUCT(BlueprintType)
+struct CONTROLRIG_API FRigUnit_HierarchyAddAnimationChannelEmptyLimitSettings
+{
+	GENERATED_BODY();
+	
+	FRigUnit_HierarchyAddAnimationChannelEmptyLimitSettings()
+	{}
+};
+
 /**
  * Adds a new animation channel to the hierarchy
  * Note: This node only runs as part of the construction event.
@@ -1103,10 +1112,29 @@ struct CONTROLRIG_API FRigUnit_HierarchyAddAnimationChannelBool : public FRigUni
 	UPROPERTY(meta = (Input))
 	bool MaximumValue;
 
+	/*
+	 * The enable settings for the limits
+	 */
+	UPROPERTY(meta = (Input))
+	FRigUnit_HierarchyAddAnimationChannelEmptyLimitSettings LimitsEnabled;
+
 	RIGVM_METHOD()
 	virtual void Execute() override;
 };
 
+USTRUCT(BlueprintType)
+struct CONTROLRIG_API FRigUnit_HierarchyAddAnimationChannelSingleLimitSettings : public FRigUnit_HierarchyAddAnimationChannelEmptyLimitSettings
+{
+	GENERATED_BODY();
+	
+	FRigUnit_HierarchyAddAnimationChannelSingleLimitSettings()
+		: Enabled(true)
+	{}
+
+	UPROPERTY(EditAnywhere, BlueprintReadWrite, Category = "Settings")
+	FRigControlLimitEnabled Enabled;
+};
+
 /**
  * Adds a new animation channel to the hierarchy
  * Note: This node only runs as part of the construction event.
@@ -1142,6 +1170,12 @@ struct CONTROLRIG_API FRigUnit_HierarchyAddAnimationChannelFloat : public FRigUn
 	UPROPERTY(meta = (Input))
 	float MaximumValue;
 
+	/*
+	 * The enable settings for the limits
+	 */
+	UPROPERTY(meta = (Input))
+	FRigUnit_HierarchyAddAnimationChannelSingleLimitSettings LimitsEnabled;
+
 	RIGVM_METHOD()
 	virtual void Execute() override;
 };
@@ -1181,6 +1215,12 @@ struct CONTROLRIG_API FRigUnit_HierarchyAddAnimationChannelScaleFloat : public F
 	UPROPERTY(meta = (Input))
 	float MaximumValue;
 
+	/*
+	 * The enable settings for the limits
+	 */
+	UPROPERTY(meta = (Input))
+	FRigUnit_HierarchyAddAnimationChannelSingleLimitSettings LimitsEnabled;
+
 	RIGVM_METHOD()
 	virtual void Execute() override;
 };
@@ -1220,10 +1260,33 @@ struct CONTROLRIG_API FRigUnit_HierarchyAddAnimationChannelInteger : public FRig
 	UPROPERTY(meta = (Input))
 	int32 MaximumValue;
 
+	/*
+	 * The enable settings for the limits
+	 */
+	UPROPERTY(meta = (Input))
+	FRigUnit_HierarchyAddAnimationChannelSingleLimitSettings LimitsEnabled;
+
 	RIGVM_METHOD()
 	virtual void Execute() override;
 };
 
+USTRUCT(BlueprintType)
+struct CONTROLRIG_API FRigUnit_HierarchyAddAnimationChannel2DLimitSettings : public FRigUnit_HierarchyAddAnimationChannelEmptyLimitSettings
+{
+	GENERATED_BODY();
+	
+	FRigUnit_HierarchyAddAnimationChannel2DLimitSettings()
+		: X(true)
+		, Y(true)
+	{}
+
+	UPROPERTY(EditAnywhere, BlueprintReadWrite, Category = "Settings")
+	FRigControlLimitEnabled X;
+
+	UPROPERTY(EditAnywhere, BlueprintReadWrite, Category = "Settings")
+	FRigControlLimitEnabled Y;
+};
+
 /**
  * Adds a new animation channel to the hierarchy
  * Note: This node only runs as part of the construction event.
@@ -1259,10 +1322,37 @@ struct CONTROLRIG_API FRigUnit_HierarchyAddAnimationChannelVector2D : public FRi
 	UPROPERTY(meta = (Input))
 	FVector2D MaximumValue;
 
+	/*
+	 * The enable settings for the limits
+	 */
+	UPROPERTY(meta = (Input))
+	FRigUnit_HierarchyAddAnimationChannel2DLimitSettings LimitsEnabled;
+
 	RIGVM_METHOD()
 	virtual void Execute() override;
 };
 
+USTRUCT(BlueprintType)
+struct CONTROLRIG_API FRigUnit_HierarchyAddAnimationChannelVectorLimitSettings : public FRigUnit_HierarchyAddAnimationChannelEmptyLimitSettings
+{
+	GENERATED_BODY();
+	
+	FRigUnit_HierarchyAddAnimationChannelVectorLimitSettings()
+		: X(true)
+		, Y(true)
+		, Z(true)
+	{}
+
+	UPROPERTY(EditAnywhere, BlueprintReadWrite, Category = "Settings")
+	FRigControlLimitEnabled X;
+
+	UPROPERTY(EditAnywhere, BlueprintReadWrite, Category = "Settings")
+	FRigControlLimitEnabled Y;
+
+	UPROPERTY(EditAnywhere, BlueprintReadWrite, Category = "Settings")
+	FRigControlLimitEnabled Z;
+};
+
 /**
  * Adds a new animation channel to the hierarchy
  * Note: This node only runs as part of the construction event.
@@ -1298,6 +1388,12 @@ struct CONTROLRIG_API FRigUnit_HierarchyAddAnimationChannelVector : public FRigU
 	UPROPERTY(meta = (Input))
 	FVector MaximumValue;
 
+	/*
+	 * The enable settings for the limits
+	 */
+	UPROPERTY(meta = (Input))
+	FRigUnit_HierarchyAddAnimationChannelVectorLimitSettings LimitsEnabled;
+
 	RIGVM_METHOD()
 	virtual void Execute() override;
 };
@@ -1337,10 +1433,38 @@ struct CONTROLRIG_API FRigUnit_HierarchyAddAnimationChannelScaleVector : public
 	UPROPERTY(meta = (Input))
 	FVector MaximumValue;
 
+	/*
+	 * The enable settings for the limits
+	 */
+	UPROPERTY(meta = (Input))
+	FRigUnit_HierarchyAddAnimationChannelVectorLimitSettings LimitsEnabled;
+
 	RIGVM_METHOD()
 	virtual void Execute() override;
 };
 
+
+USTRUCT(BlueprintType)
+struct CONTROLRIG_API FRigUnit_HierarchyAddAnimationChannelRotatorLimitSettings : public FRigUnit_HierarchyAddAnimationChannelEmptyLimitSettings
+{
+	GENERATED_BODY();
+	
+	FRigUnit_HierarchyAddAnimationChannelRotatorLimitSettings()
+		: Pitch(true)
+		, Yaw(true)
+		, Roll(true)
+	{}
+
+	UPROPERTY(EditAnywhere, BlueprintReadWrite, Category = "Settings")
+	FRigControlLimitEnabled Pitch;
+
+	UPROPERTY(EditAnywhere, BlueprintReadWrite, Category = "Settings")
+	FRigControlLimitEnabled Yaw;
+
+	UPROPERTY(EditAnywhere, BlueprintReadWrite, Category = "Settings")
+	FRigControlLimitEnabled Roll;
+};
+
 /**
  * Adds a new animation channel to the hierarchy
  * Note: This node only runs as part of the construction event.
@@ -1376,6 +1500,12 @@ struct CONTROLRIG_API FRigUnit_HierarchyAddAnimationChannelRotator : public FRig
 	UPROPERTY(meta = (Input))
 	FRotator MaximumValue;
 
+	/*
+	 * The enable settings for the limits
+	 */
+	UPROPERTY(meta = (Input))
+	FRigUnit_HierarchyAddAnimationChannelRotatorLimitSettings LimitsEnabled;
+
 	RIGVM_METHOD()
 	virtual void Execute() override;
 };
diff --git a/Engine/Plugins/Animation/ControlRig/Source/ControlRigDeveloper/Private/ControlRigBlueprint.cpp b/Engine/Plugins/Animation/ControlRig/Source/ControlRigDeveloper/Private/ControlRigBlueprint.cpp
--- a/Engine/Plugins/Animation/ControlRig/Source/ControlRigDeveloper/Private/ControlRigBlueprint.cpp
+++ b/Engine/Plugins/Animation/ControlRig/Source/ControlRigDeveloper/Private/ControlRigBlueprint.cpp
@@ -485,7 +485,174 @@ TArray<URigVMNode*> UControlRigBlueprint::ConvertHierarchyElementsToSpawnerNodes
 			
 			if(ControlElement->Settings.AnimationType == ERigControlAnimationType::AnimationChannel)
 			{
-				// todo
+				UScriptStruct* UnitNodeStruct = nullptr;
+				TRigVMTypeIndex TypeIndex = INDEX_NONE;
+				FString InitialValue, MinimumValue, MaximumValue, SettingsValue;
+				switch(ControlElement->Settings.ControlType)
+				{
+					case ERigControlType::Bool:
+					{
+						UnitNodeStruct = FRigUnit_HierarchyAddAnimationChannelBool::StaticStruct();
+						TypeIndex = RigVMTypeUtils::TypeIndex::Bool;
+						InitialValue = InHierarchy->GetControlValue(Key, ERigControlValueType::Initial).ToString<float>();
+						MinimumValue = InHierarchy->GetControlValue(Key, ERigControlValueType::Minimum).ToString<float>();
+						MaximumValue = InHierarchy->GetControlValue(Key, ERigControlValueType::Maximum).ToString<float>();
+						break;
+					}
+					case ERigControlType::Float:
+					{
+						UnitNodeStruct = FRigUnit_HierarchyAddAnimationChannelFloat::StaticStruct();
+						TypeIndex = RigVMTypeUtils::TypeIndex::Float;
+						InitialValue = InHierarchy->GetControlValue(Key, ERigControlValueType::Initial).ToString<float>();
+						MinimumValue = InHierarchy->GetControlValue(Key, ERigControlValueType::Minimum).ToString<float>();
+						MaximumValue = InHierarchy->GetControlValue(Key, ERigControlValueType::Maximum).ToString<float>();
+
+						if(ControlElement->Settings.LimitEnabled.Num() == 1)
+						{
+							FRigUnit_HierarchyAddAnimationChannelSingleLimitSettings Settings;
+							Settings.Enabled = ControlElement->Settings.LimitEnabled[0];
+							FRigUnit_HierarchyAddAnimationChannelSingleLimitSettings::StaticStruct()->ExportText(SettingsValue, &Settings, &Settings, nullptr, PPF_None, nullptr);
+						}
+						break;
+					}
+					case ERigControlType::ScaleFloat:
+					{
+						UnitNodeStruct = FRigUnit_HierarchyAddAnimationChannelScaleFloat::StaticStruct();
+						TypeIndex = RigVMTypeUtils::TypeIndex::Float;
+						InitialValue = InHierarchy->GetControlValue(Key, ERigControlValueType::Initial).ToString<float>();
+						MinimumValue = InHierarchy->GetControlValue(Key, ERigControlValueType::Minimum).ToString<float>();
+						MaximumValue = InHierarchy->GetControlValue(Key, ERigControlValueType::Maximum).ToString<float>();
+
+						if(ControlElement->Settings.LimitEnabled.Num() == 1)
+						{
+							FRigUnit_HierarchyAddAnimationChannelSingleLimitSettings Settings;
+							Settings.Enabled = ControlElement->Settings.LimitEnabled[0];
+							FRigUnit_HierarchyAddAnimationChannelSingleLimitSettings::StaticStruct()->ExportText(SettingsValue, &Settings, &Settings, nullptr, PPF_None, nullptr);
+						}
+						break;
+					}
+					case ERigControlType::Integer:
+					{
+						UnitNodeStruct = FRigUnit_HierarchyAddAnimationChannelInteger::StaticStruct();
+						TypeIndex = RigVMTypeUtils::TypeIndex::Int32; 
+						InitialValue = InHierarchy->GetControlValue(Key, ERigControlValueType::Initial).ToString<int32>();
+						MinimumValue = InHierarchy->GetControlValue(Key, ERigControlValueType::Minimum).ToString<int32>();
+						MaximumValue = InHierarchy->GetControlValue(Key, ERigControlValueType::Maximum).ToString<int32>();
+
+						if(ControlElement->Settings.LimitEnabled.Num() == 1)
+						{
+							FRigUnit_HierarchyAddAnimationChannelSingleLimitSettings Settings;
+							Settings.Enabled = ControlElement->Settings.LimitEnabled[0];
+							FRigUnit_HierarchyAddAnimationChannelSingleLimitSettings::StaticStruct()->ExportText(SettingsValue, &Settings, &Settings, nullptr, PPF_None, nullptr);
+						}
+						break;
+					}
+					case ERigControlType::Vector2D:
+					{
+						UnitNodeStruct = FRigUnit_HierarchyAddAnimationChannelVector2D::StaticStruct();
+						TypeIndex = FRigVMRegistry::Get().GetTypeIndex<FVector2D>(); 
+						InitialValue = InHierarchy->GetControlValue(Key, ERigControlValueType::Initial).ToString<FVector2D>();
+						MinimumValue = InHierarchy->GetControlValue(Key, ERigControlValueType::Minimum).ToString<FVector2D>();
+						MaximumValue = InHierarchy->GetControlValue(Key, ERigControlValueType::Maximum).ToString<FVector2D>();
+
+						if(ControlElement->Settings.LimitEnabled.Num() == 2)
+						{
+							FRigUnit_HierarchyAddAnimationChannel2DLimitSettings Settings;
+							Settings.X = ControlElement->Settings.LimitEnabled[0];
+							Settings.Y = ControlElement->Settings.LimitEnabled[1];
+							FRigUnit_HierarchyAddAnimationChannel2DLimitSettings::StaticStruct()->ExportText(SettingsValue, &Settings, &Settings, nullptr, PPF_None, nullptr);
+						}
+						break;
+					}
+					case ERigControlType::Position:
+					{
+						UnitNodeStruct = FRigUnit_HierarchyAddAnimationChannelVector::StaticStruct();
+						TypeIndex = FRigVMRegistry::Get().GetTypeIndex<FVector>(); 
+						InitialValue = InHierarchy->GetControlValue(Key, ERigControlValueType::Initial).ToString<FVector>();
+						MinimumValue = InHierarchy->GetControlValue(Key, ERigControlValueType::Minimum).ToString<FVector>();
+						MaximumValue = InHierarchy->GetControlValue(Key, ERigControlValueType::Maximum).ToString<FVector>();
+
+						if(ControlElement->Settings.LimitEnabled.Num() == 3)
+						{
+							FRigUnit_HierarchyAddAnimationChannelVectorLimitSettings Settings;
+							Settings.X = ControlElement->Settings.LimitEnabled[0];
+							Settings.Y = ControlElement->Settings.LimitEnabled[1];
+							Settings.Z = ControlElement->Settings.LimitEnabled[2];
+							FRigUnit_HierarchyAddAnimationChannelVectorLimitSettings::StaticStruct()->ExportText(SettingsValue, &Settings, &Settings, nullptr, PPF_None, nullptr);
+						}
+						break;
+					}
+					case ERigControlType::Scale:
+					{
+						UnitNodeStruct = FRigUnit_HierarchyAddAnimationChannelScaleVector::StaticStruct();
+						TypeIndex = FRigVMRegistry::Get().GetTypeIndex<FVector>(); 
+						InitialValue = InHierarchy->GetControlValue(Key, ERigControlValueType::Initial).ToString<FVector>();
+						MinimumValue = InHierarchy->GetControlValue(Key, ERigControlValueType::Minimum).ToString<FVector>();
+						MaximumValue = InHierarchy->GetControlValue(Key, ERigControlValueType::Maximum).ToString<FVector>();
+
+						if(ControlElement->Settings.LimitEnabled.Num() == 3)
+						{
+							FRigUnit_HierarchyAddAnimationChannelVectorLimitSettings Settings;
+							Settings.X = ControlElement->Settings.LimitEnabled[0];
+							Settings.Y = ControlElement->Settings.LimitEnabled[1];
+							Settings.Z = ControlElement->Settings.LimitEnabled[2];
+							FRigUnit_HierarchyAddAnimationChannelVectorLimitSettings::StaticStruct()->ExportText(SettingsValue, &Settings, &Settings, nullptr, PPF_None, nullptr);
+						}
+						break;
+					}
+					case ERigControlType::Rotator:
+					{
+						UnitNodeStruct = FRigUnit_HierarchyAddAnimationChannelRotator::StaticStruct();
+						TypeIndex = FRigVMRegistry::Get().GetTypeIndex<FRotator>(); 
+						InitialValue = InHierarchy->GetControlValue(Key, ERigControlValueType::Initial).ToString<FRotator>();
+						MinimumValue = InHierarchy->GetControlValue(Key, ERigControlValueType::Minimum).ToString<FRotator>();
+						MaximumValue = InHierarchy->GetControlValue(Key, ERigControlValueType::Maximum).ToString<FRotator>();
+
+						if(ControlElement->Settings.LimitEnabled.Num() == 3)
+						{
+							FRigUnit_HierarchyAddAnimationChannelRotatorLimitSettings Settings;
+							Settings.Pitch = ControlElement->Settings.LimitEnabled[0];
+							Settings.Yaw = ControlElement->Settings.LimitEnabled[1];
+							Settings.Roll = ControlElement->Settings.LimitEnabled[2];
+							FRigUnit_HierarchyAddAnimationChannelRotatorLimitSettings::StaticStruct()->ExportText(SettingsValue, &Settings, &Settings, nullptr, PPF_None, nullptr);
+						}
+						break;
+					}
+					default:
+					{
+						break;
+					}
+				}
+
+				if(UnitNodeStruct == nullptr)
+				{
+					continue;
+				}
+
+				URigVMNode* AddControlNode = GraphController->AddUnitNode(UnitNodeStruct, FRigUnit::GetMethodName(), NodePosition);
+				NodePosition += NodePositionIncrement;
+				AddParentItemLink(Key, AddControlNode);
+
+				if(LastPin)
+				{
+					if(const URigVMPin* NextPin = AddControlNode->FindPin(GET_MEMBER_NAME_STRING_CHECKED(FRigUnit_HierarchyAddBone, ExecuteContext)))
+					{
+						GraphController->AddLink(LastPin->GetPinPath(), NextPin->GetPinPath(), true);
+						LastPin = NextPin;
+					}
+				}
+
+				GraphController->ResolveWildCardPin(AddControlNode->FindPin(GET_MEMBER_NAME_STRING_CHECKED(FRigUnit_HierarchyAddAnimationChannelFloat, InitialValue))->GetPinPath(), TypeIndex, true);
+				GraphController->SetPinDefaultValue(AddControlNode->FindPin(GET_MEMBER_NAME_STRING_CHECKED(FRigUnit_HierarchyAddAnimationChannelFloat, Name))->GetPinPath(), Key.Name.ToString(), true, true);
+				GraphController->SetPinDefaultValue(AddControlNode->FindPin(GET_MEMBER_NAME_STRING_CHECKED(FRigUnit_HierarchyAddAnimationChannelFloat, Parent))->GetPinPath(), ParentDefault, true, true);
+				GraphController->SetPinDefaultValue(AddControlNode->FindPin(GET_MEMBER_NAME_STRING_CHECKED(FRigUnit_HierarchyAddAnimationChannelFloat, InitialValue))->GetPinPath(), InitialValue, true, true);
+				GraphController->SetPinDefaultValue(AddControlNode->FindPin(GET_MEMBER_NAME_STRING_CHECKED(FRigUnit_HierarchyAddAnimationChannelFloat, MinimumValue))->GetPinPath(), MinimumValue, true, true);
+				GraphController->SetPinDefaultValue(AddControlNode->FindPin(GET_MEMBER_NAME_STRING_CHECKED(FRigUnit_HierarchyAddAnimationChannelFloat, MaximumValue))->GetPinPath(), MaximumValue, true, true);
+
+				if(!SettingsValue.IsEmpty())
+				{
+					GraphController->SetPinDefaultValue(AddControlNode->FindPin(GET_MEMBER_NAME_STRING_CHECKED(FRigUnit_HierarchyAddAnimationChannelFloat, LimitsEnabled))->GetPinPath(), SettingsValue, true, true);
+				}
 			}
 			else
 			{
