New TRS Gizmo: support for parent space transformations (note that this is only exposed/available if the new trs gizmos are enabled).

    - added ECoordSystem::COORD_Parent
    - the mode manager's GetCustomDrawingCoordinateSystem overrides calculates the parent space transformation if required.

For the moment, there's no need to add anything at the ITF level, as everything is redirected to the VPC/Mode Manager thru the UEditorTransformProxy.
If the UTransformGizmo were to manage it locally, we'd probably have to add something similar to UCombinedTransformGizmo::bUseContextCoordinateSystem to make the TransformProxy responsible for managing the "local/parent" coordinate system.

diff --git a/Engine/Plugins/Animation/ControlRig/Source/ControlRigEditor/Private/EditMode/ControlRigEditMode.cpp b/Engine/Plugins/Animation/ControlRig/Source/ControlRigEditor/Private/EditMode/ControlRigEditMode.cpp
--- a/Engine/Plugins/Animation/ControlRig/Source/ControlRigEditor/Private/EditMode/ControlRigEditMode.cpp
+++ b/Engine/Plugins/Animation/ControlRig/Source/ControlRigEditor/Private/EditMode/ControlRigEditMode.cpp
@@ -65,9 +65,7 @@
 #include "Constraints/MovieSceneConstraintChannelHelper.h"
 #include "Editor/EditorPerProjectUserSettings.h"
 #include "TransformConstraint.h"
-#include "PersonaModule.h"
 #include "Animation/DebugSkelMeshComponent.h"
-#include "AnimationEditorPreviewActor.h"
 
 #include UE_INLINE_GENERATED_CPP_BY_NAME(ControlRigEditMode)
 
@@ -526,11 +524,16 @@ void FControlRigEditMode::Tick(FEditorViewportClient* ViewportClient, float Delt
 		HandleSelectionChanged();
 		bSelectionChanged = false;
 	}
+	else
+	{
+		// HandleSelectionChanged() will already update the pivots 
+		UpdatePivotTransforms();
+	}
+	
 	if (!AreEditingControlRigDirectly() == false)
 	{
 		ViewportClient->Invalidate();
 	}
-	RecalcPivotTransform();
 
 	// check if the settings for xray rendering are different for any of the control shape actors
 	const UControlRigEditModeSettings* Settings = GetDefault<UControlRigEditModeSettings>();
@@ -1955,7 +1958,7 @@ bool FControlRigEditMode::InputDelta(FEditorViewportClient* InViewportClient, FV
 		}
 	}
 
-	RecalcPivotTransform();
+	UpdatePivotTransforms();
 
 	if (bManipulatorMadeChange)
 	{
@@ -2168,119 +2171,186 @@ bool FControlRigEditMode::CanRemoveFromPreviewScene(const USceneComponent* InCom
 	return true;
 }
 
-void FControlRigEditMode::RecalcPivotTransform()
+ECoordSystem FControlRigEditMode::GetCoordSystemSpace() const
 {
-	PivotTransforms.Reset();
-	for (TWeakObjectPtr<UBaseControlRig>& RuntimeRigPtr : RuntimeControlRigs)
+	const UControlRigEditModeSettings* Settings = GetDefault<UControlRigEditModeSettings>();
+	if (Settings && Settings->bCoordSystemPerWidgetMode)
 	{
-		if (UBaseControlRig* ControlRig = RuntimeRigPtr.Get())
+		const int32 WidgetMode = static_cast<int32>(GetModeManager()->GetWidgetMode());
+		if (CoordSystemPerWidgetMode.IsValidIndex(WidgetMode))
 		{
-			FTransform PivotTransform = FTransform::Identity;
-			// Use average location as pivot location
-			FVector PivotLocation = FVector::ZeroVector;
+			return CoordSystemPerWidgetMode[WidgetMode];
+		}
+	}
 
-			TArray<FRigElementKey> SelectedRigElements = GetSelectedRigElements(ControlRig);
-			if (AreRigElementsSelected(ValidControlTypeMask(),ControlRig))
-			{
-				FTransform LastTransform = FTransform::Identity;
+	return GetModeManager()->GetCoordSystem();	
+}
 
-				// recalc coord system too
-				FTransform ComponentTransform = GetHostingSceneComponentTransform(ControlRig);
+void FControlRigEditMode::UpdatePivotFromEditedShape(UBaseControlRig* InControlRig)
+{
+	const URigHierarchy* Hierarchy = InControlRig ? InControlRig->GetHierarchy() : nullptr;
+	if (!Hierarchy)
+	{
+		return;
+	}
 
+	if (!ensure(bIsChangingControlShapeTransform))
+	{
+		return;
+	}
+	
+	FTransform PivotTransform = FTransform::Identity;
+	
+	if (auto* ShapeActors = ControlRigShapeActors.Find(InControlRig))
+	{
+		// we just want to change the shape transform of one single control.
+		const int32 Index = ShapeActors->IndexOfByPredicate([](const TObjectPtr<AControlRigShapeActor>& ShapeActor)
+		{
+			return IsValid(ShapeActor) && ShapeActor->IsSelected();
+		});
 
-				int32 NumSelectedControls = 0;
-				for (int32 Index = 0; Index < SelectedRigElements.Num(); ++Index)
-				{
-					if (SelectedRigElements[Index].Type == ERigElementType::Control)
-					{
-						// todo?
-					}
-				}
+		if (Index != INDEX_NONE)
+		{
+			if (FRigControlElement* ControlElement = InControlRig->FindControl((*ShapeActors)[Index]->ControlName))
+			{
+				PivotTransform = Hierarchy->GetControlShapeTransform(ControlElement, ERigTransformType::CurrentGlobal);
+			}				
+		}
+	}
+	
+	PivotTransforms.Add(InControlRig, MoveTemp(PivotTransform));
+}
 
-				if (bIsChangingControlShapeTransform)
-				{
-					if (auto* ShapeActors = ControlRigShapeActors.Find(ControlRig))
-					{
-						for (const AControlRigShapeActor* ShapeActor : *ShapeActors)
-						{
-							if (ShapeActor->IsSelected())
-							{
-								if (FRigControlElement* ControlElement = ControlRig->GetHierarchy()->Find<FRigControlElement>(FRigElementKey(ShapeActor->ControlName, ERigElementType::Control)))
-								{
-									PivotTransform = ControlRig->GetHierarchy()->GetControlShapeTransform(ControlElement, ERigTransformType::CurrentGlobal);
-								}
+void FControlRigEditMode::UpdatePivotFromShapeActors(UBaseControlRig* InControlRig, const bool bEachLocalSpace, const bool bIsParentSpace)
+{
+	if (!ensure(!bIsChangingControlShapeTransform))
+	{
+		return;
+	}
+	
+	const URigHierarchy* Hierarchy = InControlRig ? InControlRig->GetHierarchy() : nullptr;
+	if (!Hierarchy)
+	{
+		return;
+	}
+	const FTransform ComponentTransform = GetHostingSceneComponentTransform(InControlRig);
 
-								// break here since we don't want to change the shape transform of multiple controls.
-								break;
-							}
-						}
-					}
-				}
-				else
+	FTransform LastTransform = FTransform::Identity, PivotTransform = FTransform::Identity;
+
+	if (const auto* ShapeActors = ControlRigShapeActors.Find(InControlRig))
+	{
+		// if in local just use the first selected actor transform
+		// otherwise, compute the average location as pivot location
+		
+		int32 NumSelectedControls = 0;
+		FVector PivotLocation = FVector::ZeroVector;
+		for (const TObjectPtr<AControlRigShapeActor>& ShapeActor : *ShapeActors)
+		{
+			if (IsValid(ShapeActor) && ShapeActor->IsSelected())
+			{
+				const FRigElementKey ControlKey = ShapeActor->GetElementKey();
+				bool bGetParentTransform = bIsParentSpace && Hierarchy->GetNumberOfParents(ControlKey);
+	
+				const FTransform ShapeTransform = ShapeActor->GetActorTransform().GetRelativeTransform(ComponentTransform);
+				LastTransform = bGetParentTransform ? Hierarchy->GetParentTransform(ControlKey) : ShapeTransform;
+				PivotLocation += ShapeTransform.GetLocation();
+				
+				++NumSelectedControls;
+				if (!bEachLocalSpace)
 				{
-					const UControlRigEditModeSettings* Settings = GetDefault<UControlRigEditModeSettings>();
+					break;
+				}
+			}
+		}
 
-					if (auto* ShapeActors = ControlRigShapeActors.Find(ControlRig))
-					{
-						for (const AControlRigShapeActor* ShapeActor : *ShapeActors)
-						{
-							if (ShapeActor->IsSelected())
-							{
-								LastTransform = ShapeActor->GetActorTransform().GetRelativeTransform(ComponentTransform);
-								PivotLocation += LastTransform.GetLocation();
-								++NumSelectedControls;
-								if (Settings && Settings->bLocalTransformsInEachLocalSpace) //if in local just use first actors transform
-								{
-									break;
-								}
+		if (NumSelectedControls > 1)
+		{
+			PivotLocation /= static_cast<double>(NumSelectedControls);
+		}
+		PivotTransform.SetLocation(PivotLocation);
+	}
 
-							}
-						}
-					}
+	// Use the last transform's rotation as pivot rotation
+	const FTransform WorldTransform = LastTransform * ComponentTransform;
+	PivotTransform.SetRotation(WorldTransform.GetRotation());
+	
+	PivotTransforms.Add(InControlRig, MoveTemp(PivotTransform));
+}
 
-					PivotLocation /= (float)FMath::Max(1, NumSelectedControls);
-					PivotTransform.SetLocation(PivotLocation);
+void FControlRigEditMode::UpdatePivotFromElements(UBaseControlRig* InControlRig)
+{
+	if (!ensure(!bIsChangingControlShapeTransform))
+	{
+		return;
+	}
+	
+	const URigHierarchy* Hierarchy = InControlRig ? InControlRig->GetHierarchy() : nullptr;
+	if (!Hierarchy)
+	{
+		return;
+	}
+	
+	const FTransform ComponentTransform = GetHostingSceneComponentTransform(InControlRig);
+	
+	int32 NumSelection = 0;
+	FTransform LastTransform = FTransform::Identity, PivotTransform = FTransform::Identity;
+	FVector PivotLocation = FVector::ZeroVector;
+	const TArray<FRigElementKey> SelectedRigElements = GetSelectedRigElements(InControlRig);
+	
+	for (int32 Index = 0; Index < SelectedRigElements.Num(); ++Index)
+	{
+		if (SelectedRigElements[Index].Type == ERigElementType::Control)
+		{
+			LastTransform = OnGetRigElementTransformDelegate.Execute(SelectedRigElements[Index], false, true);
+			PivotLocation += LastTransform.GetLocation();
+			++NumSelection;
+		}
+	}
 
-					// just use last rotation
-					FTransform WorldTransform = LastTransform * ComponentTransform;
-					PivotTransform.SetRotation(WorldTransform.GetRotation());
-				}
-				PivotTransforms.Add(ControlRig, PivotTransform);
-			}
-			else if (AreRigElementSelectedAndMovable(ControlRig))
-			{
-				// recalc coord system too
-				FTransform ComponentTransform = GetHostingSceneComponentTransform(ControlRig);
+	if (NumSelection == 1)
+	{
+		// A single control just uses its own transform
+		const FTransform WorldTransform = LastTransform * ComponentTransform;
+		PivotTransform.SetRotation(WorldTransform.GetRotation());
+	}
+	else if (NumSelection > 1)
+	{
+		PivotLocation /= static_cast<double>(NumSelection);
+		PivotTransform.SetRotation(ComponentTransform.GetRotation());
+	}
+		
+	PivotTransform.SetLocation(PivotLocation);
+	PivotTransforms.Add(InControlRig, MoveTemp(PivotTransform));
+}
 
-				// Use average location as pivot location
-				PivotLocation = FVector::ZeroVector;
-				int32 NumSelection = 0;
-				FTransform LastTransform = FTransform::Identity;
-				for (int32 Index = 0; Index < SelectedRigElements.Num(); ++Index)
-				{
-					if (SelectedRigElements[Index].Type == ERigElementType::Control)
-					{
-						LastTransform = OnGetRigElementTransformDelegate.Execute(SelectedRigElements[Index], false, true);
-						PivotLocation += LastTransform.GetLocation();
-						++NumSelection;
-					}
-				}
+void FControlRigEditMode::UpdatePivotTransforms()
+{
+	const UControlRigEditModeSettings* Settings = GetDefault<UControlRigEditModeSettings>();
+	const bool bEachLocalSpace = Settings && Settings->bLocalTransformsInEachLocalSpace;
+	const bool bIsParentSpace = GetCoordSystemSpace() == COORD_Parent;
 
-				PivotLocation /= (float)FMath::Max(1, NumSelection);
-				PivotTransform.SetLocation(PivotLocation);
+	PivotTransforms.Reset();
 
-				if (NumSelection == 1)
+	for (const TWeakObjectPtr<UBaseControlRig>& RuntimeRigPtr : RuntimeControlRigs)
+	{
+		if (UBaseControlRig* ControlRig = RuntimeRigPtr.Get())
+		{
+			if (AreRigElementsSelected(ValidControlTypeMask(), ControlRig))
+			{
+				if (bIsChangingControlShapeTransform)
 				{
-					// A single Bone just uses its own transform
-					FTransform WorldTransform = LastTransform * ComponentTransform;
-					PivotTransform.SetRotation(WorldTransform.GetRotation());
+					UpdatePivotFromEditedShape(ControlRig);
 				}
-				else if (NumSelection > 1)
+				else
 				{
-					// If we have more than one Bone selected, use the coordinate space of the component
-					PivotTransform.SetRotation(ComponentTransform.GetRotation());
+					UpdatePivotFromShapeActors(ControlRig, bEachLocalSpace, bIsParentSpace);			
 				}
-				PivotTransforms.Add(ControlRig, PivotTransform);
+			}
+			else if (AreRigElementSelectedAndMovable(ControlRig))
+			{
+				// do we even get in here ?!
+				// we will enter the if first as AreRigElementsSelected will return true before AreRigElementSelectedAndMovable does...
+				UpdatePivotFromElements(ControlRig);
 			}
 		}
 	}
@@ -2343,9 +2413,9 @@ bool FControlRigEditMode::HasPivotTransformsChanged() const
 
 void FControlRigEditMode::HandleSelectionChanged()
 {
-	for (auto& ShapeActors : ControlRigShapeActors)
+	for (const auto& ShapeActors : ControlRigShapeActors)
 	{
-		for (AControlRigShapeActor* ShapeActor : ShapeActors.Value)
+		for (const TObjectPtr<AControlRigShapeActor>& ShapeActor : ShapeActors.Value)
 		{
 			TInlineComponentArray<UPrimitiveComponent*> PrimitiveComponents;
 			ShapeActor->GetComponents(PrimitiveComponents, true);
@@ -2366,7 +2436,8 @@ void FControlRigEditMode::HandleSelectionChanged()
 	}
 
 	// update the pivot transform of our selected objects (they could be animating)
-	RecalcPivotTransform();
+	UpdatePivotTransforms();
+	
 	//need to force the redraw also
 	if (!AreEditingControlRigDirectly())
 	{
@@ -3319,7 +3390,7 @@ void FControlRigEditMode::PostUndo()
 		{
 			//due to tick ordering need to manually make sure we get everything done in correct order.
 			PostPoseUpdate();
-			RecalcPivotTransform();
+			UpdatePivotTransforms();
 			GEditor->RedrawLevelEditingViewports(true);
 		});
 	}
diff --git a/Engine/Plugins/Animation/ControlRig/Source/ControlRigEditor/Public/EditMode/ControlRigEditMode.h b/Engine/Plugins/Animation/ControlRig/Source/ControlRigEditor/Public/EditMode/ControlRigEditMode.h
--- a/Engine/Plugins/Animation/ControlRig/Source/ControlRigEditor/Public/EditMode/ControlRigEditMode.h
+++ b/Engine/Plugins/Animation/ControlRig/Source/ControlRigEditor/Public/EditMode/ControlRigEditMode.h
@@ -198,9 +198,15 @@ protected:
 	/** Set up Details Panel based upon Selected Objects*/
 	void SetUpDetailPanel();
 
-	/** Updates cached pivot transform */
-	void RecalcPivotTransform();
-
+	/** Updates cached pivot transforms */
+	void UpdatePivotTransforms();
+	void UpdatePivotFromEditedShape(UBaseControlRig* InControlRig);
+	void UpdatePivotFromShapeActors(UBaseControlRig* InControlRig, const bool bEachLocalSpace, const bool bIsParentSpace);
+	void UpdatePivotFromElements(UBaseControlRig* InControlRig);
+	
+	/** Get the current coordinate system space */
+	ECoordSystem GetCoordSystemSpace() const;
+	
 	/** Helper function for box/frustum intersection */
 	bool IntersectSelect(bool InSelect, const TFunctionRef<bool(const AControlRigShapeActor*, const FTransform&)>& Intersects);
 
diff --git a/Engine/Source/Editor/EditorFramework/Public/UnrealWidgetFwd.h b/Engine/Source/Editor/EditorFramework/Public/UnrealWidgetFwd.h
--- a/Engine/Source/Editor/EditorFramework/Public/UnrealWidgetFwd.h
+++ b/Engine/Source/Editor/EditorFramework/Public/UnrealWidgetFwd.h
@@ -10,6 +10,7 @@ enum ECoordSystem
 	COORD_None = -1,
 	COORD_World,
 	COORD_Local,
+	COORD_Parent,
 	COORD_Max,
 };
 
diff --git a/Engine/Source/Editor/Experimental/EditorInteractiveToolsFramework/Private/EditorGizmos/EditorTransformGizmoSource.cpp b/Engine/Source/Editor/Experimental/EditorInteractiveToolsFramework/Private/EditorGizmos/EditorTransformGizmoSource.cpp
--- a/Engine/Source/Editor/Experimental/EditorInteractiveToolsFramework/Private/EditorGizmos/EditorTransformGizmoSource.cpp
+++ b/Engine/Source/Editor/Experimental/EditorInteractiveToolsFramework/Private/EditorGizmos/EditorTransformGizmoSource.cpp
@@ -57,12 +57,8 @@ EAxisList::Type UEditorTransformGizmoSource::GetGizmoAxisToDraw(EGizmoTransformM
 EToolContextCoordinateSystem UEditorTransformGizmoSource::GetGizmoCoordSystemSpace() const
 {
 	const FEditorViewportClient* ViewportClient = GetViewportClient();
-	if (ViewportClient && ViewportClient->GetWidgetCoordSystemSpace() == ECoordSystem::COORD_Local)
-	{
-		return EToolContextCoordinateSystem::Local;
-	}
-	
-	return EToolContextCoordinateSystem::World;
+	const ECoordSystem Space = ViewportClient ? ViewportClient->GetWidgetCoordSystemSpace() : COORD_World;
+	return Space == COORD_World ? EToolContextCoordinateSystem::World : EToolContextCoordinateSystem::Local;
 }
 
 float UEditorTransformGizmoSource::GetGizmoScale() const
diff --git a/Engine/Source/Editor/Persona/Private/AnimationEditorViewportClient.cpp b/Engine/Source/Editor/Persona/Private/AnimationEditorViewportClient.cpp
--- a/Engine/Source/Editor/Persona/Private/AnimationEditorViewportClient.cpp
+++ b/Engine/Source/Editor/Persona/Private/AnimationEditorViewportClient.cpp
@@ -1341,9 +1341,9 @@ FVector FAnimationViewportClient::GetWidgetLocation() const
 
 FMatrix FAnimationViewportClient::GetWidgetCoordSystem() const
 {
-	const bool bIsLocal = GetWidgetCoordSystemSpace() == COORD_Local;
-
-	if( bIsLocal )
+	const ECoordSystem Space = GetWidgetCoordSystemSpace();
+	const bool bIsLocalOrParent = Space == COORD_Local || Space == COORD_Parent;
+	if( bIsLocalOrParent )
 	{
 		return ModeTools->GetCustomInputCoordinateSystem();
 	}
diff --git a/Engine/Source/Editor/Persona/Private/SkeletonSelectionEditMode.cpp b/Engine/Source/Editor/Persona/Private/SkeletonSelectionEditMode.cpp
--- a/Engine/Source/Editor/Persona/Private/SkeletonSelectionEditMode.cpp
+++ b/Engine/Source/Editor/Persona/Private/SkeletonSelectionEditMode.cpp
@@ -531,27 +531,50 @@ FTransform FSkeletonSelectionEditMode::GetSocketTransform(const USkeletalMeshSoc
 
 bool FSkeletonSelectionEditMode::GetCustomDrawingCoordinateSystem(FMatrix& InMatrix, void* InData)
 {
+	const bool bIsParentMode = Owner ? Owner->GetCoordSystem() == COORD_Parent : false;
+	
 	const UDebugSkelMeshComponent* PreviewMeshComponent = GetAnimPreviewScene().GetPreviewMeshComponent();
 	if(PreviewMeshComponent)
 	{
 		// Retrieve reference skeleton from either current USkeletalMesh, or USkeleton if no mesh is set
 		const FReferenceSkeleton& ReferenceSkeleton = GetReferenceSkeletonForComponent(PreviewMeshComponent);
 		
-		const int32 BoneIndex = GetAnimPreviewScene().GetSelectedBoneIndex();
+		int32 BoneIndex = GetAnimPreviewScene().GetSelectedBoneIndex();
 		if (BoneIndex != INDEX_NONE && ReferenceSkeleton.IsValidIndex(BoneIndex))
 		{
+			if (bIsParentMode)
+			{
+				const int32 ParentIndex = ReferenceSkeleton.GetParentIndex(BoneIndex);
+				if (ParentIndex != INDEX_NONE)
+				{
+					BoneIndex = ParentIndex;
+				}
+			}
 			const FTransform BoneMatrix = GetBoneTransform(BoneIndex);
 			InMatrix = BoneMatrix.ToMatrixNoScale().RemoveTranslation();
 			return true;
 		}
-		else if (GetAnimPreviewScene().GetSelectedSocket().IsValid())
+		
+		if (GetAnimPreviewScene().GetSelectedSocket().IsValid())
 		{
 			const USkeletalMeshSocket* Socket = GetAnimPreviewScene().GetSelectedSocket().Socket;
+			if (bIsParentMode)
+			{
+				const int32 SocketBoneIndex = PreviewMeshComponent->GetBoneIndex(Socket->BoneName);
+				if (SocketBoneIndex != INDEX_NONE && ReferenceSkeleton.IsValidIndex(SocketBoneIndex))
+				{
+					const FTransform SocketBoneMatrix = GetBoneTransform(SocketBoneIndex);
+					InMatrix = SocketBoneMatrix.ToMatrixNoScale().RemoveTranslation();
+					return true;
+				}
+			}
+			
 			const FTransform SocketMatrix = GetSocketTransform(Socket);
 			InMatrix = SocketMatrix.ToMatrixNoScale().RemoveTranslation();
 			return true;
 		}
-		else if (const AActor* SelectedActor = GetAnimPreviewScene().GetSelectedActor())
+		
+		if (const AActor* SelectedActor = GetAnimPreviewScene().GetSelectedActor())
 		{
 			InMatrix = SelectedActor->GetTransform().ToMatrixNoScale().RemoveTranslation();
 			return true;
diff --git a/Engine/Source/Editor/UnrealEd/Private/EditorModeManager.cpp b/Engine/Source/Editor/UnrealEd/Private/EditorModeManager.cpp
--- a/Engine/Source/Editor/UnrealEd/Private/EditorModeManager.cpp
+++ b/Engine/Source/Editor/UnrealEd/Private/EditorModeManager.cpp
@@ -19,23 +19,14 @@
 #include "Framework/Application/SlateApplication.h"
 #include "Framework/Notifications/NotificationManager.h"
 #include "Widgets/Notifications/SNotificationList.h"
-#include "Widgets/Input/SButton.h"
-#include "Engine/LevelStreaming.h"
 #include "Editor/EditorEngine.h"
-#include "UnrealEdGlobals.h"
 #include "Editor/UnrealEdEngine.h"
-#include "Widgets/Docking/SDockTab.h"
-#include "Widgets/Layout/SWidgetSwitcher.h"
-#include "Styling/AppStyle.h"
 #include "Framework/Commands/UICommandList.h"
-#include "Framework/MultiBox/MultiBoxBuilder.h"
 #include "Toolkits/BaseToolkit.h"
 #include "Subsystems/AssetEditorSubsystem.h"
 #include "Subsystems/BrushEditingSubsystem.h"
 #include "Tools/UEdMode.h"
-#include "Widgets/Images/SImage.h"
 #include "InputRouter.h"
-#include "InteractiveGizmoManager.h"
 #include "EdModeInteractiveToolsContext.h"
 #include "Tools/LegacyEdModeInterfaces.h"
 #include "CanvasTypes.h"
@@ -134,10 +125,22 @@ void FEditorModeTools::LoadConfig(void)
 	GConfig->GetBool(TEXT("FEditorModeTools"),TEXT("ShowWidget"),bShowWidget,
 		GEditorPerProjectIni);
 
-	const bool bGetRawValue = true;
+	static constexpr bool bGetRawValue = true;
 	int32 CoordSystemAsInt = (int32)GetCoordSystem(bGetRawValue);
 	GConfig->GetInt(TEXT("FEditorModeTools"),TEXT("CoordSystem"), CoordSystemAsInt,
 		GEditorPerProjectIni);
+
+	if (static_cast<ECoordSystem>(CoordSystemAsInt) == COORD_Parent)
+	{
+		// parent mode is only supported with new trs gizmos for now
+		static IConsoleVariable* UseLegacyWidgetCVar = IConsoleManager::Get().FindConsoleVariable(TEXT("Gizmos.UseLegacyWidget"));
+		const bool bUseNewGizmo = UseLegacyWidgetCVar ? UseLegacyWidgetCVar->GetInt() < 1 : false;
+		if (!bUseNewGizmo)
+		{
+			CoordSystemAsInt = static_cast<int32>(COORD_Local);
+		}
+	}
+	
 	SetCoordSystem((ECoordSystem)CoordSystemAsInt);
 
 	LoadWidgetSettings();
@@ -147,7 +150,7 @@ void FEditorModeTools::SaveConfig(void)
 {
 	GConfig->SetBool(TEXT("FEditorModeTools"), TEXT("ShowWidget"), bShowWidget, GEditorPerProjectIni);
 
-	const bool bGetRawValue = true;
+	static constexpr bool bGetRawValue = true;
 	GConfig->SetInt(TEXT("FEditorModeTools"), TEXT("CoordSystem"), (int32)GetCoordSystem(bGetRawValue), GEditorPerProjectIni);
 
 	SaveWidgetSettings();
@@ -671,16 +674,14 @@ void FEditorModeTools::SetPivotLocation( const FVector& Location, const bool bIn
 	}
 }
 
-ECoordSystem FEditorModeTools::GetCoordSystem(bool bGetRawValue)
+ECoordSystem FEditorModeTools::GetCoordSystem(bool bGetRawValue) const
 {
 	if (!bGetRawValue && (GetWidgetMode() == UE::Widget::WM_Scale))
 	{
 		return COORD_Local;
 	}
-	else
-	{
-		return CoordSystem;
-	}
+	
+	return CoordSystem;
 }
 
 void FEditorModeTools::SetCoordSystem(ECoordSystem NewCoordSystem)
@@ -977,16 +978,18 @@ UTexture2D* FEditorModeTools::GetVertexTexture() const
 	return GEngine->DefaultBSPVertexTexture;
 }
 
-FMatrix FEditorModeTools::GetCustomDrawingCoordinateSystem()
+FMatrix FEditorModeTools::GetCustomDrawingCoordinateSystem() const
 {
 	FMatrix Matrix = FMatrix::Identity;
 
 	switch (GetCoordSystem())
 	{
 		case COORD_Local:
-		{
 			Matrix = GetLocalCoordinateSystem();
-		}
+		break;
+		
+		case COORD_Parent:
+			Matrix = GetParentSpaceCoordinateSystem();
 		break;
 
 		case COORD_World:
@@ -999,12 +1002,38 @@ FMatrix FEditorModeTools::GetCustomDrawingCoordinateSystem()
 	return Matrix;
 }
 
-FMatrix FEditorModeTools::GetCustomInputCoordinateSystem()
+FMatrix FEditorModeTools::GetCustomInputCoordinateSystem() const
 {
 	return GetCustomDrawingCoordinateSystem();
 }
 
-FMatrix FEditorModeTools::GetLocalCoordinateSystem()
+FMatrix FEditorModeTools::GetLocalCoordinateSystem() const
+{
+	return GetCustomCoordinateSystem([](const TTypedElement<ITypedElementWorldInterface>& InElement, FTransform& OutTransform)
+	{
+		InElement.GetWorldTransform(OutTransform);
+	});
+}
+
+FMatrix FEditorModeTools::GetParentSpaceCoordinateSystem() const
+{
+	return GetCustomCoordinateSystem([](const TTypedElement<ITypedElementWorldInterface>& InElement, FTransform& OutTransform)
+	{
+		if (InElement.GetWorldTransform(OutTransform))
+		{
+			FTransform RelativeTransform;
+			if (InElement.GetRelativeTransform(RelativeTransform))
+			{
+				const FTransform ParentWorld = RelativeTransform.Inverse() * OutTransform;
+				OutTransform.SetRotation(ParentWorld.GetRotation());
+			}
+			return true;
+		}
+		return false;
+	});
+}
+
+FMatrix FEditorModeTools::GetCustomCoordinateSystem(TUniqueFunction<void(const TTypedElement<ITypedElementWorldInterface>&, FTransform&)>&& InGetTransformFunc) const
 {
 	FMatrix Matrix = FMatrix::Identity;
 	// Let the current mode have a shot at setting the local coordinate system.
@@ -1033,9 +1062,9 @@ FMatrix FEditorModeTools::GetLocalCoordinateSystem()
 		
 		if (LastSelected)
 		{
-			FTransform LocalToWorldTransform;
-			LastSelected.GetWorldTransform(LocalToWorldTransform);
-			Matrix = FQuatRotationMatrix(LocalToWorldTransform.GetRotation());
+			FTransform CustomToWorldTransform;
+			InGetTransformFunc(LastSelected, CustomToWorldTransform);
+			Matrix = FQuatRotationMatrix(CustomToWorldTransform.GetRotation());
 		}
 	}
 
diff --git a/Engine/Source/Editor/UnrealEd/Private/SEditorViewport.cpp b/Engine/Source/Editor/UnrealEd/Private/SEditorViewport.cpp
--- a/Engine/Source/Editor/UnrealEd/Private/SEditorViewport.cpp
+++ b/Engine/Source/Editor/UnrealEd/Private/SEditorViewport.cpp
@@ -764,9 +764,14 @@ void SEditorViewport::OnCycleCoordinateSystem()
 	int32 CoordSystemAsInt = Client->GetWidgetCoordSystemSpace();
 
 	++CoordSystemAsInt;
-	if( CoordSystemAsInt == COORD_Max )
+
+	static IConsoleVariable* UseLegacyWidgetCVar = IConsoleManager::Get().FindConsoleVariable(TEXT("Gizmos.UseLegacyWidget"));
+	const bool bUseNewGizmos = UseLegacyWidgetCVar ? UseLegacyWidgetCVar->GetInt() < 1 : false;
+	// parent mode is only supported with new trs gizmos for now
+	const int CoordMax = bUseNewGizmos ? COORD_Max : COORD_Parent;
+	if( CoordSystemAsInt >= CoordMax )
 	{
-		CoordSystemAsInt -= COORD_Max;
+		CoordSystemAsInt = COORD_World;
 	}
 
 	Client->SetWidgetCoordSystemSpace( (ECoordSystem)CoordSystemAsInt );
diff --git a/Engine/Source/Editor/UnrealEd/Private/STransformViewportToolbar.cpp b/Engine/Source/Editor/UnrealEd/Private/STransformViewportToolbar.cpp
--- a/Engine/Source/Editor/UnrealEd/Private/STransformViewportToolbar.cpp
+++ b/Engine/Source/Editor/UnrealEd/Private/STransformViewportToolbar.cpp
@@ -32,8 +32,6 @@
 #include "LevelEditorActions.h"
 #include "Styling/ToolBarStyle.h"
 #include "SEditorViewportToolBarMenu.h"
-#include "SEditorViewportToolBarButton.h"
-#include "Widgets/Layout/SSpacer.h"
 
 #define LOCTEXT_NAMESPACE "TransformToolBar"
 
@@ -211,10 +209,22 @@ TSharedRef< SWidget > STransformViewportToolBar::MakeTransformToolBar( const TSh
 
 		ToolbarBuilder.SetIsFocusable( true );
 
+		TAttribute<FText> CoordSystemToolTip = TAttribute<FText>::CreateLambda([]
+		{
+			static IConsoleVariable* UseLegacyWidgetCVar = IConsoleManager::Get().FindConsoleVariable(TEXT("Gizmos.UseLegacyWidget"));
+			const bool bUseNewGizmo = UseLegacyWidgetCVar ? UseLegacyWidgetCVar->GetInt() < 1 : false;
+			if (bUseNewGizmo)
+			{
+				return LOCTEXT(	"CycleTransformGizmoCoordSystemWithParent_ToolTip",
+								"Cycles the transform gizmo coordinate systems between world, local and parent space");
+			}
+			return FEditorViewportCommands::Get().CycleTransformGizmoCoordSystem->GetDescription();
+		});
+		
 		ToolbarBuilder.AddToolBarButton( FEditorViewportCommands::Get().CycleTransformGizmoCoordSystem,
 			NAME_None,
 			TAttribute<FText>(),
-			TAttribute<FText>(),
+			CoordSystemToolTip,
 			TAttribute<FSlateIcon>(this, &STransformViewportToolBar::GetLocalToWorldIcon),
 			FName(TEXT("CycleTransformGizmoCoordSystem")),
 
@@ -545,6 +555,12 @@ FSlateIcon STransformViewportToolBar::GetLocalToWorldIcon() const
 		return FSlateIcon(FAppStyle::GetAppStyleSetName(), WorldIcon);
 	}
 
+	if( Viewport.IsValid() && Viewport.Pin()->IsCoordSystemActive(COORD_Parent) )
+	{
+		static const FName ParentIcon("Icons.ConstraintManager.ParentHierarchy");
+		return FSlateIcon(FAppStyle::GetAppStyleSetName(), ParentIcon);
+	}
+	
 	static FName LocalIcon("Icons.Transform");
 	return FSlateIcon(FAppStyle::GetAppStyleSetName(), LocalIcon);
 }
diff --git a/Engine/Source/Editor/UnrealEd/Public/EditorModeManager.h b/Engine/Source/Editor/UnrealEd/Public/EditorModeManager.h
--- a/Engine/Source/Editor/UnrealEd/Public/EditorModeManager.h
+++ b/Engine/Source/Editor/UnrealEd/Public/EditorModeManager.h
@@ -13,6 +13,7 @@
 
 #include "Elements/Framework/TypedElementSelectionSet.h"
 
+class ITypedElementWorldInterface;
 class FCanvas;
 class FEditorViewportClient;
 class FEdMode;
@@ -131,9 +132,10 @@ public:
 	 */
 	UNREALED_API bool EnsureNotInMode(FEditorModeID ModeID, const FText& ErrorMsg = FText::GetEmpty(), bool bNotifyUser = false) const;
 
-	UNREALED_API FMatrix GetCustomDrawingCoordinateSystem();
-	UNREALED_API FMatrix GetCustomInputCoordinateSystem();
-	UNREALED_API FMatrix GetLocalCoordinateSystem();
+	UNREALED_API FMatrix GetCustomDrawingCoordinateSystem() const;
+	UNREALED_API FMatrix GetCustomInputCoordinateSystem() const;
+	UNREALED_API FMatrix GetLocalCoordinateSystem() const;
+	UNREALED_API FMatrix GetParentSpaceCoordinateSystem() const;
 	
 	/** 
 	 * Returns true if the passed in editor mode is active 
@@ -395,7 +397,7 @@ public:
 	 * 
 	 * @param bGetRawValue true when you want the actual value of CoordSystem, not the value modified by the state.
 	 */
-	UNREALED_API ECoordSystem GetCoordSystem(bool bGetRawValue = false);
+	UNREALED_API ECoordSystem GetCoordSystem(bool bGetRawValue = false) const;
 
 	/** Sets the current CoordSystem */
 	UNREALED_API void SetCoordSystem(ECoordSystem NewCoordSystem);
@@ -561,6 +563,9 @@ protected:
 		});
 	}
 
+	/** Returns the custom coordinate matrix using a callback-type request. */
+	FMatrix GetCustomCoordinateSystem(TUniqueFunction<void(const TTypedElement<ITypedElementWorldInterface>&, FTransform&)>&& InGetTransformFunc) const;
+
 	UNREALED_API void ExitAllModesPendingDeactivate();
 
 	/** List of default modes for this tool.  These must all be compatible with each other. */
