Control Rig: Fix crash when selecting socket in rig hierarchy view
#rnx

diff --git a/Engine/Plugins/Animation/ControlRig/Source/ControlRigDeveloper/Private/Graph/ControlRigGraph.cpp b/Engine/Plugins/Animation/ControlRig/Source/ControlRigDeveloper/Private/Graph/ControlRigGraph.cpp
--- a/Engine/Plugins/Animation/ControlRig/Source/ControlRigDeveloper/Private/Graph/ControlRigGraph.cpp
+++ b/Engine/Plugins/Animation/ControlRig/Source/ControlRigDeveloper/Private/Graph/ControlRigGraph.cpp
@@ -85,6 +85,7 @@ void UControlRigGraph::CacheNameLists(URigHierarchy* InHierarchy, const FRigVMDr
 		ElementNameLists.FindOrAdd(ERigElementType::RigidBody);
 		ElementNameLists.FindOrAdd(ERigElementType::Reference);
 		ElementNameLists.FindOrAdd(ERigElementType::Connector);
+		ElementNameLists.FindOrAdd(ERigElementType::Socket);
 
 		TArray<TSharedPtr<FRigVMStringWithTag>>& AllNameList = ElementNameLists.FindChecked(ERigElementType::All);
 		TArray<TSharedPtr<FRigVMStringWithTag>>& BoneNameList = ElementNameLists.FindChecked(ERigElementType::Bone);
@@ -93,6 +94,8 @@ void UControlRigGraph::CacheNameLists(URigHierarchy* InHierarchy, const FRigVMDr
 		TArray<TSharedPtr<FRigVMStringWithTag>>& CurveNameList = ElementNameLists.FindChecked(ERigElementType::Curve);
 		TArray<TSharedPtr<FRigVMStringWithTag>>& RigidBodyNameList = ElementNameLists.FindChecked(ERigElementType::RigidBody);
 		TArray<TSharedPtr<FRigVMStringWithTag>>& ReferenceNameList = ElementNameLists.FindChecked(ERigElementType::Reference);
+		TArray<TSharedPtr<FRigVMStringWithTag>>& ConnectorNameList = ElementNameLists.FindChecked(ERigElementType::Connector);
+		TArray<TSharedPtr<FRigVMStringWithTag>>& SocketNameList = ElementNameLists.FindChecked(ERigElementType::Socket);
 		
 		CacheNameListForHierarchy<FRigBaseElement>(ControlRig, InHierarchy, AllNameList, false);
 		CacheNameListForHierarchy<FRigBoneElement>(ControlRig, InHierarchy, BoneNameList, false);
@@ -102,6 +105,8 @@ void UControlRigGraph::CacheNameLists(URigHierarchy* InHierarchy, const FRigVMDr
 		CacheNameListForHierarchy<FRigCurveElement>(ControlRig, InHierarchy, CurveNameList, false);
 		CacheNameListForHierarchy<FRigRigidBodyElement>(ControlRig, InHierarchy, RigidBodyNameList, false);
 		CacheNameListForHierarchy<FRigReferenceElement>(ControlRig, InHierarchy, ReferenceNameList, false);
+		CacheNameListForHierarchy<FRigConnectorElement>(ControlRig, InHierarchy, ConnectorNameList, false);
+		CacheNameListForHierarchy<FRigSocketElement>(ControlRig, InHierarchy, SocketNameList, false);
 
 		LastHierarchyTopologyVersion = InHierarchy->GetTopologyVersion();
 	}
diff --git a/Engine/Plugins/Animation/ControlRig/Source/ControlRigDeveloper/Public/Graph/ControlRigGraph.h b/Engine/Plugins/Animation/ControlRig/Source/ControlRigDeveloper/Public/Graph/ControlRigGraph.h
--- a/Engine/Plugins/Animation/ControlRig/Source/ControlRigDeveloper/Public/Graph/ControlRigGraph.h
+++ b/Engine/Plugins/Animation/ControlRig/Source/ControlRigDeveloper/Public/Graph/ControlRigGraph.h
@@ -57,6 +57,10 @@ public:
 	{
 		return GetElementNameList(ERigElementType::Connector);
 	}
+	const TArray<TSharedPtr<FRigVMStringWithTag>>* GetSocketNameList(URigVMPin* InPin = nullptr) const
+	{
+		return GetElementNameList(ERigElementType::Socket);
+	}
 
 	virtual const TArray<TSharedPtr<FRigVMStringWithTag>>* GetNameListForWidget(const FString& InWidgetName) const override;
 
