New TRS Gizmos (follow-up to 28836836)
    - fixed gizmo local space computation
#rnx

diff --git a/Engine/Plugins/Animation/ControlRig/Source/ControlRigEditor/Private/EditMode/ControlRigEditMode.cpp b/Engine/Plugins/Animation/ControlRig/Source/ControlRigEditor/Private/EditMode/ControlRigEditMode.cpp
--- a/Engine/Plugins/Animation/ControlRig/Source/ControlRigEditor/Private/EditMode/ControlRigEditMode.cpp
+++ b/Engine/Plugins/Animation/ControlRig/Source/ControlRigEditor/Private/EditMode/ControlRigEditMode.cpp
@@ -2256,7 +2256,7 @@ void FControlRigEditMode::UpdatePivotFromShapeActors(UBaseControlRig* InControlR
 				PivotLocation += ShapeTransform.GetLocation();
 				
 				++NumSelectedControls;
-				if (!bEachLocalSpace)
+				if (bEachLocalSpace)
 				{
 					break;
 				}
