Control Rig: Undo change of ERigElementType and amend the lost of element Last type
#rnx

diff --git a/Engine/Plugins/Animation/ControlRig/Source/ControlRig/Private/Rigs/RigHierarchy.cpp b/Engine/Plugins/Animation/ControlRig/Source/ControlRig/Private/Rigs/RigHierarchy.cpp
--- a/Engine/Plugins/Animation/ControlRig/Source/ControlRig/Private/Rigs/RigHierarchy.cpp
+++ b/Engine/Plugins/Animation/ControlRig/Source/ControlRig/Private/Rigs/RigHierarchy.cpp
@@ -361,8 +361,13 @@ void URigHierarchy::Reset_Impl(bool bResetElements)
 		}
 		Elements.Reset();
 		ElementsPerType.Reset();
-		for(int32 TypeIndex=0;TypeIndex<RigElementTypeToFlatIndex(ERigElementType::Last);TypeIndex++)
+		const UEnum* ElementTypeEnum = StaticEnum<ERigElementType>();
+		for(int32 TypeIndex=0; ;TypeIndex++)
 		{
+			if ((ERigElementType)ElementTypeEnum->GetValueByIndex(TypeIndex) == ERigElementType::All)
+			{
+				break;
+			}
 			ElementsPerType.Add(TArray<FRigBaseElement*>());
 		}
 		IndexLookup.Reset();
@@ -422,8 +427,13 @@ void URigHierarchy::CopyHierarchy(URigHierarchy* InHierarchy)
 	bReallocateElements = Elements.Num() < InHierarchy->Elements.Num();
 	if(!bReallocateElements)
 	{
-		for(int32 ElementTypeIndex = 0; ElementTypeIndex < RigElementTypeToFlatIndex(ERigElementType::Last); ElementTypeIndex++)
+		const UEnum* ElementTypeEnum = StaticEnum<ERigElementType>();
+		for(int32 ElementTypeIndex = 0; ; ElementTypeIndex++)
 		{
+			if ((ERigElementType)ElementTypeEnum->GetValueByIndex(ElementTypeIndex) == ERigElementType::All)
+			{
+				break;
+			}
 			check(ElementsPerType.IsValidIndex(ElementTypeIndex));
 			check(InHierarchy->ElementsPerType.IsValidIndex(ElementTypeIndex));
 			if(ElementsPerType[ElementTypeIndex].Num() < InHierarchy->ElementsPerType[ElementTypeIndex].Num())
@@ -524,8 +534,13 @@ void URigHierarchy::CopyHierarchy(URigHierarchy* InHierarchy)
 
 		// shrink the containers accordingly
 		Elements.SetNum(InHierarchy->Elements.Num());
-		for(int32 ElementTypeIndex = 0; ElementTypeIndex < RigElementTypeToFlatIndex(ERigElementType::Last); ElementTypeIndex++)
+		const UEnum* ElementTypeEnum = StaticEnum<ERigElementType>();
+		for(int32 ElementTypeIndex = 0; ; ElementTypeIndex++)
 		{
+			if ((ERigElementType)ElementTypeEnum->GetValueByIndex(ElementTypeIndex) == ERigElementType::All)
+			{
+				break;
+			}
 			ElementsPerType[ElementTypeIndex].SetNum(InHierarchy->ElementsPerType[ElementTypeIndex].Num());
 		}
 
diff --git a/Engine/Plugins/Animation/ControlRig/Source/ControlRig/Private/Units/Highlevel/Hierarchy/RigUnit_ModifyTransforms.cpp b/Engine/Plugins/Animation/ControlRig/Source/ControlRig/Private/Units/Highlevel/Hierarchy/RigUnit_ModifyTransforms.cpp
--- a/Engine/Plugins/Animation/ControlRig/Source/ControlRig/Private/Units/Highlevel/Hierarchy/RigUnit_ModifyTransforms.cpp
+++ b/Engine/Plugins/Animation/ControlRig/Source/ControlRig/Private/Units/Highlevel/Hierarchy/RigUnit_ModifyTransforms.cpp
@@ -207,8 +207,13 @@ int32 FRigUnit_ModifyTransforms::GetIndexFromTarget(const FString& InTarget) con
 		static TArray<FString> DisplayNames;
 		if(DisplayNames.IsEmpty())
 		{
-			for(int64 Index = 0; Index < (int64)ERigElementType::Last; Index++)
+			const UEnum* ElementTypeEnum = StaticEnum<ERigElementType>();
+			for(int64 Index = 0; ; Index++)
 			{
+				if ((ERigElementType)ElementTypeEnum->GetValueByIndex(Index) == ERigElementType::All)
+				{
+					break;
+				}
 				DisplayNames.Add(TypeEnum->GetDisplayNameTextByValue(Index).ToString());
 			}
 		}
diff --git a/Engine/Plugins/Animation/ControlRig/Source/ControlRig/Public/Rigs/RigHierarchy.h b/Engine/Plugins/Animation/ControlRig/Source/ControlRig/Public/Rigs/RigHierarchy.h
--- a/Engine/Plugins/Animation/ControlRig/Source/ControlRig/Public/Rigs/RigHierarchy.h
+++ b/Engine/Plugins/Animation/ControlRig/Source/ControlRig/Public/Rigs/RigHierarchy.h
@@ -4337,14 +4337,9 @@ protected:
 			{
 				return 7;
 			}
-			case ERigElementType::Last:
-			{
-				return 8;
-			}
 			case ERigElementType::All:
 			default:
 			{
-				checkNoEntry();
 				break;
 			}
 		}
@@ -4388,13 +4383,8 @@ protected:
 			{
 				return ERigElementType::Socket;
 			}
-			case 8:
-			{
-				return ERigElementType::Last;
-			}
 			default:
 			{
-				checkNoEntry();
 				break;
 			}
 		}
diff --git a/Engine/Plugins/Animation/ControlRig/Source/ControlRig/Public/Rigs/RigHierarchyDefines.h b/Engine/Plugins/Animation/ControlRig/Source/ControlRig/Public/Rigs/RigHierarchyDefines.h
--- a/Engine/Plugins/Animation/ControlRig/Source/ControlRig/Public/Rigs/RigHierarchyDefines.h
+++ b/Engine/Plugins/Animation/ControlRig/Source/ControlRig/Public/Rigs/RigHierarchyDefines.h
@@ -18,7 +18,7 @@ class URigHierarchy;
  * This can be used as a mask so supported as a bitfield
  */
 UENUM(BlueprintType)
-enum class ERigElementType
+enum class ERigElementType : uint8
 {
 	None = 0,
 	Bone = 0x001,
@@ -30,7 +30,6 @@ enum class ERigElementType
 	Reference = 0x020,
 	Connector = 0x040,
 	Socket = 0x080,
-	Last = 0x160 UMETA(Hidden),
 	All = Bone | Null | Control | Curve | RigidBody | Reference | Connector | Socket,
 	ToResetAfterConstructionEvent = Bone | Control | Curve | Socket UMETA(Hidden),
 };
