Sequencer: Control Rig: Always get cached proxy via function to make sure it's set up.

diff --git a/Engine/Plugins/Animation/ControlRig/Source/ControlRig/Private/Sequencer/MovieSceneControlRigParameterSection.cpp b/Engine/Plugins/Animation/ControlRig/Source/ControlRig/Private/Sequencer/MovieSceneControlRigParameterSection.cpp
--- a/Engine/Plugins/Animation/ControlRig/Source/ControlRig/Private/Sequencer/MovieSceneControlRigParameterSection.cpp
+++ b/Engine/Plugins/Animation/ControlRig/Source/ControlRig/Private/Sequencer/MovieSceneControlRigParameterSection.cpp
@@ -946,7 +946,7 @@ void UMovieSceneControlRigParameterSection::SetBlendType(EMovieSceneBlendType In
 			const FChannelMapInfo* ChannelInfo = nullptr;
 
 			// Set Defaults based upon Type
-			TArrayView<FMovieSceneFloatChannel*> FloatChannels = ChannelProxy->GetChannels<FMovieSceneFloatChannel>();
+			TArrayView<FMovieSceneFloatChannel*> FloatChannels = GetChannelProxy().GetChannels<FMovieSceneFloatChannel>();
 			TArray<FRigControlElement*> Controls = ControlRig->AvailableControls();
 
 			for (FRigControlElement* ControlElement : Controls)
@@ -3516,7 +3516,7 @@ TOptional<float> UMovieSceneControlRigParameterSection::EvaluateScalarParameter(
 	TOptional<float> OptValue;	
 	if (const FChannelMapInfo* ChannelInfo = ControlChannelMap.Find(InParameterName))
 	{
-		TArrayView<FMovieSceneFloatChannel*> FloatChannels = ChannelProxy->GetChannels<FMovieSceneFloatChannel>();
+		TArrayView<FMovieSceneFloatChannel*> FloatChannels = GetChannelProxy().GetChannels<FMovieSceneFloatChannel>();
 		float Value = 0.0f;
 		FloatChannels[ChannelInfo->ChannelIndex]->Evaluate(InTime, Value);
 		OptValue = Value;
@@ -3568,7 +3568,7 @@ TOptional<FVector> UMovieSceneControlRigParameterSection::EvaluateVectorParamete
 	TOptional<FVector> OptValue;
 	if (const FChannelMapInfo* ChannelInfo = ControlChannelMap.Find(InParameterName))
 	{
-		TArrayView<FMovieSceneFloatChannel*> FloatChannels = ChannelProxy->GetChannels<FMovieSceneFloatChannel>();
+		TArrayView<FMovieSceneFloatChannel*> FloatChannels = GetChannelProxy().GetChannels<FMovieSceneFloatChannel>();
 		FVector3f Value(0.0f, 0.0f, 0.0f);
 		FloatChannels[ChannelInfo->ChannelIndex]->Evaluate(InTime, Value.X);
 		FloatChannels[ChannelInfo->ChannelIndex + 1]->Evaluate(InTime, Value.Y);
@@ -3583,7 +3583,7 @@ TOptional<FVector2D> UMovieSceneControlRigParameterSection::EvaluateVector2DPara
 	TOptional<FVector2D> OptValue;
 	if (const FChannelMapInfo* ChannelInfo = ControlChannelMap.Find(InParameterName))
 	{
-		TArrayView<FMovieSceneFloatChannel*> FloatChannels = ChannelProxy->GetChannels<FMovieSceneFloatChannel>();
+		TArrayView<FMovieSceneFloatChannel*> FloatChannels = GetChannelProxy().GetChannels<FMovieSceneFloatChannel>();
 		FVector2f Value(0.0f, 0.0f);
 		FloatChannels[ChannelInfo->ChannelIndex]->Evaluate(InTime, Value.X);
 		FloatChannels[ChannelInfo->ChannelIndex + 1]->Evaluate(InTime, Value.Y);
@@ -3597,7 +3597,7 @@ TOptional<FLinearColor>UMovieSceneControlRigParameterSection:: EvaluateColorPara
 	TOptional<FLinearColor> OptValue;
 	if (const FChannelMapInfo* ChannelInfo = ControlChannelMap.Find(InParameterName))
 	{
-		TArrayView<FMovieSceneFloatChannel*> FloatChannels = ChannelProxy->GetChannels<FMovieSceneFloatChannel>();
+		TArrayView<FMovieSceneFloatChannel*> FloatChannels = GetChannelProxy().GetChannels<FMovieSceneFloatChannel>();
 		FLinearColor Value(0.0f, 0.0f, 0.0f, 1.0f);
 		FloatChannels[ChannelInfo->ChannelIndex]->Evaluate(InTime, Value.R);
 		FloatChannels[ChannelInfo->ChannelIndex + 1]->Evaluate(InTime, Value.G);		
@@ -3613,7 +3613,7 @@ TOptional<FEulerTransform> UMovieSceneControlRigParameterSection::EvaluateTransf
 	TOptional<FEulerTransform> OptValue;
 	if (const FChannelMapInfo* ChannelInfo = ControlChannelMap.Find(InParameterName))
 	{
-		TArrayView<FMovieSceneFloatChannel*> FloatChannels = ChannelProxy->GetChannels<FMovieSceneFloatChannel>();
+		TArrayView<FMovieSceneFloatChannel*> FloatChannels = GetChannelProxy().GetChannels<FMovieSceneFloatChannel>();
 		FEulerTransform Value = FEulerTransform::Identity;
 		FVector3f Translation(ForceInitToZero), Scale(FVector3f::OneVector);
 		FRotator3f Rotator(0.0f, 0.0f, 0.0f);
