Control Rig: Fix issue when creating a new control rig
#rnx

diff --git a/Engine/Plugins/Animation/ControlRig/Source/ControlRigEditor/Private/ControlRigBlueprintFactory.cpp b/Engine/Plugins/Animation/ControlRig/Source/ControlRigEditor/Private/ControlRigBlueprintFactory.cpp
--- a/Engine/Plugins/Animation/ControlRig/Source/ControlRigEditor/Private/ControlRigBlueprintFactory.cpp
+++ b/Engine/Plugins/Animation/ControlRig/Source/ControlRigEditor/Private/ControlRigBlueprintFactory.cpp
@@ -42,7 +42,7 @@ public:
 	void Construct( const FArguments& InArgs )
 	{
 		bOkClicked = false;
-		ParentClass = UBaseControlRig::StaticClass();
+		ParentClass = UControlRig::StaticClass(); // default to control rig
 
 		ChildSlot
 		[
@@ -271,7 +271,7 @@ UControlRigBlueprintFactory::UControlRigBlueprintFactory()
 	bCreateNew = true;
 	bEditAfterNew = true;
 	SupportedClass = UControlRigBlueprint::StaticClass();
-	ParentClass = UBaseControlRig::StaticClass();
+	ParentClass = UControlRig::StaticClass(); // default to control rig
 }
 
 bool UControlRigBlueprintFactory::ConfigureProperties()
diff --git a/Engine/Plugins/Animation/ControlRig/Source/ControlRigEditor/Private/ControlRigBlueprintFactory.h b/Engine/Plugins/Animation/ControlRig/Source/ControlRigEditor/Private/ControlRigBlueprintFactory.h
--- a/Engine/Plugins/Animation/ControlRig/Source/ControlRigEditor/Private/ControlRigBlueprintFactory.h
+++ b/Engine/Plugins/Animation/ControlRig/Source/ControlRigEditor/Private/ControlRigBlueprintFactory.h
@@ -21,7 +21,7 @@ public:
 
 	// The parent class of the created blueprint
 	UPROPERTY(EditAnywhere, Category="Control Rig Factory", meta=(AllowAbstract = ""))
-	TSubclassOf<UControlRig> ParentClass;
+	TSubclassOf<UBaseControlRig> ParentClass;
 
 	// UFactory Interface
 	virtual bool ConfigureProperties() override;
