Control Rig: Make sure to evaluate inputs for Control Rig Anim Node during cachebones

diff --git a/Engine/Plugins/Animation/ControlRig/Source/ControlRig/Private/AnimNode_ControlRig.cpp b/Engine/Plugins/Animation/ControlRig/Source/ControlRig/Private/AnimNode_ControlRig.cpp
--- a/Engine/Plugins/Animation/ControlRig/Source/ControlRig/Private/AnimNode_ControlRig.cpp
+++ b/Engine/Plugins/Animation/ControlRig/Source/ControlRig/Private/AnimNode_ControlRig.cpp
@@ -142,6 +142,9 @@ void FAnimNode_ControlRig::CacheBones_AnyThread(const FAnimationCacheBonesContex
 {
 	DECLARE_SCOPE_HIERARCHICAL_COUNTER_FUNC()
 
+	// make sure the inputs on the node are evaluated before propagating the inputs
+	GetEvaluateGraphExposedInputs().Execute(Context);
+
 	// we also need access to the properties when running construction event
 	PropagateInputProperties(Context.AnimInstanceProxy->GetAnimInstanceObject());
 
