Control Rig: Rebase Connector Element to Base Element

* Remove transform

diff --git a/Engine/Plugins/Animation/ControlRig/Source/ControlRig/Private/Rigs/RigHierarchy.cpp b/Engine/Plugins/Animation/ControlRig/Source/ControlRig/Private/Rigs/RigHierarchy.cpp
--- a/Engine/Plugins/Animation/ControlRig/Source/ControlRig/Private/Rigs/RigHierarchy.cpp
+++ b/Engine/Plugins/Animation/ControlRig/Source/ControlRig/Private/Rigs/RigHierarchy.cpp
@@ -998,12 +998,11 @@ TArray<FRigElementKey> URigHierarchy::RestoreConnectorsFromInfos(TArray<FRigConn
 
 		if(const FRigConnectorElement* Connector = Find<FRigConnectorElement>(Key))
 		{
-			SetInitialLocalTransform(Connector->Index, Info.LocalTransform, true, bSetupUndoRedo, false);
 			SetConnectorSettings(Key, Info.Settings, bSetupUndoRedo, false, false);
 		}
 		else
 		{
-			Key = GetController()->AddConnector(Info.Name, Info.LocalTransform, false, Info.Settings, bSetupUndoRedo, false);
+			Key = GetController()->AddConnector(Info.Name, Info.Settings, bSetupUndoRedo, false);
 		}
 
 		Keys.Add(Key);
diff --git a/Engine/Plugins/Animation/ControlRig/Source/ControlRig/Private/Rigs/RigHierarchyController.cpp b/Engine/Plugins/Animation/ControlRig/Source/ControlRig/Private/Rigs/RigHierarchyController.cpp
--- a/Engine/Plugins/Animation/ControlRig/Source/ControlRig/Private/Rigs/RigHierarchyController.cpp
+++ b/Engine/Plugins/Animation/ControlRig/Source/ControlRig/Private/Rigs/RigHierarchyController.cpp
@@ -570,8 +570,7 @@ FRigElementKey URigHierarchyController::AddReference(FName InName, FRigElementKe
 	return NewElement->Key;
 }
 
-FRigElementKey URigHierarchyController::AddConnector(FName InName, FTransform InTransform,
-	bool bTransformInGlobal, FRigConnectorSettings InSettings, bool bSetupUndo, bool bPrintPythonCommand)
+FRigElementKey URigHierarchyController::AddConnector(FName InName, FRigConnectorSettings InSettings, bool bSetupUndo, bool bPrintPythonCommand)
 {
 	if(!IsValid())
 	{
@@ -594,19 +593,6 @@ FRigElementKey URigHierarchyController::AddConnector(FName InName, FTransform In
 		NewElement->Key.Name = GetSafeNewName(InName, NewElement->Key.Type);
 		NewElement->Settings = InSettings;
 		AddElement(NewElement, nullptr, true, InName);
-
-		if(bTransformInGlobal)
-		{
-			Hierarchy->SetTransform(NewElement, InTransform, ERigTransformType::InitialGlobal, true, false);
-			Hierarchy->SetTransform(NewElement, InTransform, ERigTransformType::CurrentGlobal, true, false);
-		}
-		else
-		{
-			Hierarchy->SetTransform(NewElement, InTransform, ERigTransformType::InitialLocal, true, false);
-			Hierarchy->SetTransform(NewElement, InTransform, ERigTransformType::CurrentLocal, true, false);
-		}
-
-		NewElement->Pose.Current = NewElement->Pose.Initial;
 	}
 
 #if WITH_EDITOR
@@ -1760,13 +1746,6 @@ TArray<FString> URigHierarchyController::GetAddRigidBodyPythonCommands(FRigRigid
 TArray<FString> URigHierarchyController::GetAddConnectorPythonCommands(FRigConnectorElement* Connector) const
 {
 	TArray<FString> Commands;
-	FString TransformStr = RigVMPythonUtils::TransformToPythonString(Connector->Pose.Initial.Local.Transform);
-
-	FString ParentKeyStr = "''";
-	if (Connector->ParentElement)
-	{
-		ParentKeyStr = Connector->ParentElement->GetKey().ToPythonString();
-	}
 
 	FRigConnectorSettings& Settings = Connector->Settings;
 	FString SettingsStr;
@@ -1778,10 +1757,9 @@ TArray<FString> URigHierarchyController::GetAddConnectorPythonCommands(FRigConne
 		Commands.Append(URigHierarchy::ConnectorSettingsToPythonCommands(Settings, SettingsStr));	
 	}
 
-	// AddConnector(FName InName, FRigElementKey InParent, FTransform InTransform, bool bTransformInGlobal = true, FRigConnectorSettings InSettings = FRigConnectorSettings(), bool bSetupUndo = false);
-	Commands.Add(FString::Printf(TEXT("hierarchy_controller.add_connector('%s', %s, False, %s)"),
+	// AddConnector(FName InName, FRigConnectorSettings InSettings = FRigConnectorSettings(), bool bSetupUndo = false);
+	Commands.Add(FString::Printf(TEXT("hierarchy_controller.add_connector('%s', %s)"),
 		*Connector->GetName(),
-		*TransformStr,
 		*SettingsStr
 	));
 
diff --git a/Engine/Plugins/Animation/ControlRig/Source/ControlRig/Private/Rigs/RigHierarchyElements.cpp b/Engine/Plugins/Animation/ControlRig/Source/ControlRig/Private/Rigs/RigHierarchyElements.cpp
--- a/Engine/Plugins/Animation/ControlRig/Source/ControlRig/Private/Rigs/RigHierarchyElements.cpp
+++ b/Engine/Plugins/Animation/ControlRig/Source/ControlRig/Private/Rigs/RigHierarchyElements.cpp
@@ -1697,7 +1697,6 @@ FRigConnectorInfo FRigConnectorElement::GetConnectorInfo(const URigHierarchy* In
 	FRigConnectorInfo Info;
 	Info.Name = Key.Name;
 	Info.ResolvedTarget = InHierarchy->GetResolvedTarget(Key);
-	Info.LocalTransform = InHierarchy->GetInitialLocalTransform(Key);
 	Info.Settings = Settings;
 	return Info;
 }
diff --git a/Engine/Plugins/Animation/ControlRig/Source/ControlRig/Public/Rigs/RigHierarchyController.h b/Engine/Plugins/Animation/ControlRig/Source/ControlRig/Public/Rigs/RigHierarchyController.h
--- a/Engine/Plugins/Animation/ControlRig/Source/ControlRig/Public/Rigs/RigHierarchyController.h
+++ b/Engine/Plugins/Animation/ControlRig/Source/ControlRig/Public/Rigs/RigHierarchyController.h
@@ -240,15 +240,13 @@ public:
 	/**
 	 * Adds a connector to the hierarchy
 	 * @param InName The suggested name of the new connector - will eventually be corrected by the namespace
-	 * @param InTransform The transform for the new connector - either in local or global space, based on bTransformInGlobal
-	 * @param bTransformInGlobal Set this to true if the Transform passed is expressed in global space, false for local space.
 	 * @param InSettings All of the connector's settings
 	 * @param bSetupUndo If set to true the stack will record the change for undo / redo
 	 * @param bPrintPythonCommand If set to true a python command equivalent to this call will be printed out
 	 * @return The key for the newly created bone.
 	 */
 	UFUNCTION(BlueprintCallable, Category = URigHierarchyController)
-	FRigElementKey AddConnector(FName InName, FTransform InTransform, bool bTransformInGlobal = true, FRigConnectorSettings InSettings = FRigConnectorSettings(), bool bSetupUndo = false, bool bPrintPythonCommand = false);
+	FRigElementKey AddConnector(FName InName, FRigConnectorSettings InSettings = FRigConnectorSettings(), bool bSetupUndo = false, bool bPrintPythonCommand = false);
 
 	/**
 	 * Adds a socket to the hierarchy
diff --git a/Engine/Plugins/Animation/ControlRig/Source/ControlRig/Public/Rigs/RigHierarchyElements.h b/Engine/Plugins/Animation/ControlRig/Source/ControlRig/Public/Rigs/RigHierarchyElements.h
--- a/Engine/Plugins/Animation/ControlRig/Source/ControlRig/Public/Rigs/RigHierarchyElements.h
+++ b/Engine/Plugins/Animation/ControlRig/Source/ControlRig/Public/Rigs/RigHierarchyElements.h
@@ -832,7 +832,6 @@ protected:
 			InElement->GetType() == ERigElementType::Control ||
 			InElement->GetType() == ERigElementType::RigidBody ||
 			InElement->GetType() == ERigElementType::Reference ||
-			InElement->GetType() == ERigElementType::Connector ||
 			InElement->GetType() == ERigElementType::Socket;
 	}
 
@@ -882,7 +881,6 @@ protected:
 		return InElement->GetType() == ERigElementType::Bone ||
 			InElement->GetType() == ERigElementType::RigidBody ||
 			InElement->GetType() == ERigElementType::Reference ||
-			InElement->GetType() == ERigElementType::Connector ||
 			InElement->GetType() == ERigElementType::Socket;
 	}
 
@@ -1692,7 +1690,6 @@ struct CONTROLRIG_API FRigConnectorInfo
 	FRigConnectorInfo()
 		: Name(NAME_None)
 		, ResolvedTarget()
-		, LocalTransform(FTransform::Identity)
 		, Settings()
 	{}
 	
@@ -1702,15 +1699,12 @@ struct CONTROLRIG_API FRigConnectorInfo
 	UPROPERTY(BlueprintReadWrite, EditAnywhere, Category = Connector)
 	FRigElementKey ResolvedTarget;
 
-	UPROPERTY(BlueprintReadWrite, EditAnywhere, Category = Connector)
-	FTransform LocalTransform;
-
 	UPROPERTY(BlueprintReadWrite, EditAnywhere, Category = Connector)
 	FRigConnectorSettings Settings;
 };
 
 USTRUCT(BlueprintType)
-struct CONTROLRIG_API FRigConnectorElement : public FRigSingleParentElement
+struct CONTROLRIG_API FRigConnectorElement : public FRigBaseElement
 {
 public:
 	
@@ -1720,7 +1714,7 @@ public:
 	static const EElementIndex ElementTypeIndex;
 
 	FRigConnectorElement()
-		: FRigSingleParentElement()
+		: FRigBaseElement()
 	{
 		Key.Type = ERigElementType::Connector; 
 	}
diff --git a/Engine/Plugins/Animation/ControlRig/Source/ControlRigDeveloper/Private/ControlRigBlueprint.cpp b/Engine/Plugins/Animation/ControlRig/Source/ControlRigDeveloper/Private/ControlRigBlueprint.cpp
--- a/Engine/Plugins/Animation/ControlRig/Source/ControlRigDeveloper/Private/ControlRigBlueprint.cpp
+++ b/Engine/Plugins/Animation/ControlRig/Source/ControlRigDeveloper/Private/ControlRigBlueprint.cpp
@@ -290,7 +290,7 @@ bool UControlRigBlueprint::TurnIntoControlRigModule(bool InAutoConvertHierarchy,
 			if(Hierarchy->Num(ERigElementType::Connector) == 0)
 			{
 				static const FName RootName = TEXT("Root");
-				Controller->AddConnector(RootName, FTransform::Identity);
+				Controller->AddConnector(RootName);
 			}
 		}
 	}
diff --git a/Engine/Plugins/Animation/ControlRig/Source/ControlRigEditor/Private/Editor/SRigHierarchy.cpp b/Engine/Plugins/Animation/ControlRig/Source/ControlRigEditor/Private/Editor/SRigHierarchy.cpp
--- a/Engine/Plugins/Animation/ControlRig/Source/ControlRigEditor/Private/Editor/SRigHierarchy.cpp
+++ b/Engine/Plugins/Animation/ControlRig/Source/ControlRigEditor/Private/Editor/SRigHierarchy.cpp
@@ -2006,7 +2006,7 @@ void SRigHierarchy::HandleNewItem(ERigElementType InElementType, bool bIsAnimati
 						}
 					}
 
-					NewItemKey = Controller->AddConnector(NewElementName, ParentTransform, true, FRigConnectorSettings(), true);
+					NewItemKey = Controller->AddConnector(NewElementName, FRigConnectorSettings(), true);
 					(void)ResolveConnector(NewItemKey, ParentKey);
 					break;
 				}
