Sequencer: Control Rig, fix selection dropping do to details panel proxy getting Modified.

diff --git a/Engine/Plugins/Animation/ControlRig/Source/ControlRigEditor/Private/EditMode/ControlRigControlsProxy.cpp b/Engine/Plugins/Animation/ControlRig/Source/ControlRigEditor/Private/EditMode/ControlRigControlsProxy.cpp
--- a/Engine/Plugins/Animation/ControlRig/Source/ControlRigEditor/Private/EditMode/ControlRigControlsProxy.cpp
+++ b/Engine/Plugins/Animation/ControlRig/Source/ControlRigEditor/Private/EditMode/ControlRigControlsProxy.cpp
@@ -219,7 +219,6 @@ void UControlRigTransformControlProxy::ValueChanged()
 	FRigControlElement* ControlElement = GetControlElement();
 	if (ControlElement && ControlRig.IsValid())
 	{
-		Modify();
 		const FName PropertyName("Transform");
 		FTrackInstancePropertyBindings Binding(PropertyName, PropertyName.ToString());
 		const FTransform NewTransform = ControlRig.Get()->GetControlValue(ControlElement, ERigControlValueType::Current).Get<FRigControlValue::FTransform_Float>().ToTransform();
@@ -299,7 +298,6 @@ void UControlRigTransformNoScaleControlProxy::ValueChanged()
 	FRigControlElement* ControlElement = GetControlElement();
 	if (ControlElement)
 	{
-		Modify();
 		const FName PropertyName("Transform");
 		FTrackInstancePropertyBindings Binding(PropertyName, PropertyName.ToString());
 		const FTransformNoScale NewTransform = ControlRig.Get()->GetControlValue(ControlElement, ERigControlValueType::Current).Get<FRigControlValue::FTransformNoScale_Float>().ToTransform();
@@ -376,7 +374,6 @@ void UControlRigEulerTransformControlProxy::ValueChanged()
 	FRigControlElement* ControlElement = GetControlElement();
 	if (ControlElement)
 	{
-		Modify();
 		const FName PropertyName("Transform");
 		FTrackInstancePropertyBindings Binding(PropertyName, PropertyName.ToString());
 		FEulerTransform NewTransform = ControlRig.Get()->GetControlValue(ControlElement, ERigControlValueType::Current).Get<FRigControlValue::FEulerTransform_Float>().ToTransform();
@@ -446,7 +443,6 @@ void UControlRigFloatControlProxy::ValueChanged()
 	FRigControlElement* ControlElement = GetControlElement();
 	if (ControlElement)
 	{
-		Modify();
 		const FName PropertyName("Float");
 		FTrackInstancePropertyBindings Binding(PropertyName, PropertyName.ToString());
 		const float Val = ControlRig.Get()->GetControlValue(ControlElement, ERigControlValueType::Current).Get<float>();
@@ -496,7 +492,6 @@ void UControlRigIntegerControlProxy::ValueChanged()
 	FRigControlElement* ControlElement = GetControlElement();
 	if (ControlElement)
 	{
-		Modify();
 		const FName PropertyName("Integer");
 		FTrackInstancePropertyBindings Binding(PropertyName, PropertyName.ToString());
 		const int32 Val = ControlRig.Get()->GetControlValue(ControlElement, ERigControlValueType::Current).Get<int32>();
@@ -546,7 +541,6 @@ void UControlRigEnumControlProxy::ValueChanged()
 	FRigControlElement* ControlElement = GetControlElement();
 	if (ControlElement)
 	{
-		Modify();
 		const FName PropertyName("Enum");
 		FTrackInstancePropertyBindings Binding(PropertyName, PropertyName.ToString());
 
@@ -613,7 +607,6 @@ void UControlRigVectorControlProxy::ValueChanged()
 	FRigControlElement* ControlElement = GetControlElement();
 	if (ControlElement)
 	{
-		Modify();
 		const FName PropertyName("Vector");
 		FTrackInstancePropertyBindings Binding(PropertyName, PropertyName.ToString());
 		FVector3f Val = ControlRig.Get()->GetControlValue(ControlElement, ERigControlValueType::Current).Get<FVector3f>();
@@ -734,7 +727,6 @@ void UControlRigVector2DControlProxy::ValueChanged()
 	FRigControlElement* ControlElement = GetControlElement();
 	if (ControlElement)
 	{
-		Modify();
 		const FName PropertyName("Vector2D");
 		FTrackInstancePropertyBindings Binding(PropertyName, PropertyName.ToString());
 		const FVector3f TempValue = ControlRig.Get()->GetControlValue(ControlElement, ERigControlValueType::Current).Get<FVector3f>();
@@ -796,7 +788,6 @@ void UControlRigBoolControlProxy::ValueChanged()
 	FRigControlElement* ControlElement = GetControlElement();
 	if (ControlElement)
 	{
-		Modify();
 		const FName PropertyName("Bool");
 		FTrackInstancePropertyBindings Binding(PropertyName, PropertyName.ToString());
 		const bool Val = ControlRig.Get()->GetControlValue(ControlElement, ERigControlValueType::Current).Get<bool>();
@@ -920,6 +911,7 @@ void UControlRigDetailPanelControlProxies::AddProxy(UBaseControlRig* ControlRig,
 			Proxy->SetFlags(RF_Transactional);
 			Proxy->SetName(Name);
 			Proxy->ControlRig = ControlRig;
+			Proxy->Modify();
 			Proxy->ValueChanged();
 
 			FControlToProxyMap* ControlRigProxies = AllProxies.Find(ControlRig);
@@ -1006,14 +998,18 @@ void UControlRigDetailPanelControlProxies::RecreateAllProxies(UBaseControlRig* C
 	}
 }
 
-void UControlRigDetailPanelControlProxies::ProxyChanged(UBaseControlRig* ControlRig, const FName& Name)
+void UControlRigDetailPanelControlProxies::ProxyChanged(UBaseControlRig* ControlRig, const FName& Name, bool bModify )
 {
 	if (IsInGameThread())
 	{
 		UControlRigControlsProxy* Proxy = FindProxy(ControlRig,Name);
 		if (Proxy)
 		{
-			Modify();
+			if (bModify)
+			{
+				Modify();
+				Proxy->Modify();
+			}
 			Proxy->ValueChanged();
 		}
 	}
diff --git a/Engine/Plugins/Animation/ControlRig/Source/ControlRigEditor/Private/EditMode/ControlRigControlsProxy.h b/Engine/Plugins/Animation/ControlRig/Source/ControlRigEditor/Private/EditMode/ControlRigControlsProxy.h
--- a/Engine/Plugins/Animation/ControlRig/Source/ControlRigEditor/Private/EditMode/ControlRigControlsProxy.h
+++ b/Engine/Plugins/Animation/ControlRig/Source/ControlRigEditor/Private/EditMode/ControlRigControlsProxy.h
@@ -326,10 +326,11 @@ protected:
 
 
 public:
+	void SelectProxy(UControlRig* InControlRig, const FName& Name, bool bSelected);
 	UControlRigControlsProxy* FindProxy(UBaseControlRig* InControlRig, const FName& Name) const;
 	void AddProxy(UBaseControlRig* InControlRig, const FName& Name,  FRigControlElement* ControlElement);
 	void RemoveProxy(UBaseControlRig* InControlRig, const FName& Name );
-	void ProxyChanged(UBaseControlRig* InControlRig, const FName& Name);
+	void ProxyChanged(UBaseControlRig* InControlRig, const FName& Name, bool bModify = true);
 	void RemoveAllProxies(UBaseControlRig* InControlRig);
 	void RecreateAllProxies(UBaseControlRig* InControlRig);
 	void SelectProxy(UBaseControlRig* InControlRig, const FName& Name, bool bSelected);
diff --git a/Engine/Plugins/Animation/ControlRig/Source/ControlRigEditor/Private/EditMode/ControlRigEditMode.cpp b/Engine/Plugins/Animation/ControlRig/Source/ControlRigEditor/Private/EditMode/ControlRigEditMode.cpp
--- a/Engine/Plugins/Animation/ControlRig/Source/ControlRigEditor/Private/EditMode/ControlRigEditMode.cpp
+++ b/Engine/Plugins/Animation/ControlRig/Source/ControlRigEditor/Private/EditMode/ControlRigEditMode.cpp
@@ -4043,7 +4043,8 @@ void FControlRigEditMode::OnHierarchyModified_AnyThread(ERigHierarchyNotificatio
 void FControlRigEditMode::OnControlModified(UBaseControlRig* Subject, FRigControlElement* InControlElement, const FRigControlModifiedContext& Context)
 {
 	//this makes sure the details panel ui get's updated, don't remove
-	ControlProxy->ProxyChanged(Subject,InControlElement->GetFName());
+	const bool bModify = Context.SetKey != EControlRigSetKey::Never;
+	ControlProxy->ProxyChanged(Subject,InControlElement->GetFName(), bModify);
 
 	/*
 	FScopedTransaction ScopedTransaction(LOCTEXT("ModifyControlTransaction", "Modify Control"),!GIsTransacting && Context.SetKey != EControlRigSetKey::Never);
