Sequencer: CIS FIx

diff --git a/Engine/Plugins/Animation/ControlRig/Source/ControlRigEditor/Private/EditMode/ControlRigControlsProxy.h b/Engine/Plugins/Animation/ControlRig/Source/ControlRigEditor/Private/EditMode/ControlRigControlsProxy.h
--- a/Engine/Plugins/Animation/ControlRig/Source/ControlRigEditor/Private/EditMode/ControlRigControlsProxy.h
+++ b/Engine/Plugins/Animation/ControlRig/Source/ControlRigEditor/Private/EditMode/ControlRigControlsProxy.h
@@ -326,7 +326,7 @@ protected:
 
 
 public:
-	void SelectProxy(UControlRig* InControlRig, const FName& Name, bool bSelected);
+	void SelectProxy(UBaseControlRig* InControlRig, const FName& Name, bool bSelected);
 	UControlRigControlsProxy* FindProxy(UBaseControlRig* InControlRig, const FName& Name) const;
 	void AddProxy(UBaseControlRig* InControlRig, const FName& Name,  FRigControlElement* ControlElement);
 	void RemoveProxy(UBaseControlRig* InControlRig, const FName& Name );
