CIS fix pt 2

diff --git a/Engine/Plugins/Animation/ControlRig/Source/ControlRigEditor/Private/EditMode/ControlRigControlsProxy.h b/Engine/Plugins/Animation/ControlRig/Source/ControlRigEditor/Private/EditMode/ControlRigControlsProxy.h
--- a/Engine/Plugins/Animation/ControlRig/Source/ControlRigEditor/Private/EditMode/ControlRigControlsProxy.h
+++ b/Engine/Plugins/Animation/ControlRig/Source/ControlRigEditor/Private/EditMode/ControlRigControlsProxy.h
@@ -333,7 +333,6 @@ public:
 	void ProxyChanged(UBaseControlRig* InControlRig, const FName& Name, bool bModify = true);
 	void RemoveAllProxies(UBaseControlRig* InControlRig);
 	void RecreateAllProxies(UBaseControlRig* InControlRig);
-	void SelectProxy(UBaseControlRig* InControlRig, const FName& Name, bool bSelected);
 	const TArray<UControlRigControlsProxy*>& GetSelectedProxies() const { return SelectedProxies;}
 	bool IsSelected(UBaseControlRig* InControlRig, const FName& Name) const;
 
