Control Rig: Allow runtime changes to control rig class in AnimBP

diff --git a/Engine/Plugins/Animation/ControlRig/Source/ControlRig/Private/AnimNode_ControlRig.cpp b/Engine/Plugins/Animation/ControlRig/Source/ControlRig/Private/AnimNode_ControlRig.cpp
--- a/Engine/Plugins/Animation/ControlRig/Source/ControlRig/Private/AnimNode_ControlRig.cpp
+++ b/Engine/Plugins/Animation/ControlRig/Source/ControlRig/Private/AnimNode_ControlRig.cpp
@@ -10,6 +10,7 @@
 #include "AnimationRuntime.h"
 #include "Animation/AnimCurveUtils.h"
 #include "Animation/AnimStats.h"
+#include "Units/Execution/RigUnit_PrepareForExecution.h"
 
 #include UE_INLINE_GENERATED_CPP_BY_NAME(AnimNode_ControlRig)
 
@@ -47,14 +48,14 @@ void FAnimNode_ControlRig::OnInitializeAnimInstance(const FAnimInstanceProxy* In
 {
 	DECLARE_SCOPE_HIERARCHICAL_COUNTER_FUNC()
 
-	if (ControlRigClass)
+	ControlRigPerClass.Reset();
+	if(DefaultControlRigClass)
+	{
+		ControlRigClass = nullptr;
+	}
+	
+	if(UpdateControlRigIfNeeded(InAnimInstance))
 	{
-		ControlRig = NewObject<UBaseControlRig>(InAnimInstance->GetOwningComponent(), ControlRigClass);
-		ControlRig->Initialize(true);
-		ControlRig->RequestInit();
-		RefPoseSetterHash.Reset();
-		ControlRig->OnInitialized_AnyThread().AddRaw(this, &FAnimNode_ControlRig::HandleOnInitialized_AnyThread);
-
 		UpdateControlRigRefPoseIfNeeded(InProxy);
 	}
 
@@ -109,6 +110,11 @@ void FAnimNode_ControlRig::Update_AnyThread(const FAnimationUpdateContext& Conte
 		InternalBlendAlpha = 0.f;
 	}
 
+	if(const UAnimInstance* AnimInstance = Cast<UAnimInstance>(Context.GetAnimInstanceObject()))
+	{
+		(void)UpdateControlRigIfNeeded(AnimInstance);
+	}
+
 	UpdateControlRigRefPoseIfNeeded(Context.AnimInstanceProxy);
 	FAnimNode_ControlRigBase::Update_AnyThread(Context);
 
@@ -148,6 +154,12 @@ void FAnimNode_ControlRig::CacheBones_AnyThread(const FAnimationCacheBonesContex
 	// we also need access to the properties when running construction event
 	PropagateInputProperties(Context.AnimInstanceProxy->GetAnimInstanceObject());
 
+	// update the control rig instance just in case the dynamic control rig class has changed
+	if(const UAnimInstance* AnimInstance = Cast<UAnimInstance>(Context.GetAnimInstanceObject()))
+	{
+		(void)UpdateControlRigIfNeeded(AnimInstance);
+	}
+
 	FAnimNode_ControlRigBase::CacheBones_AnyThread(Context);
 
 	FBoneContainer& RequiredBones = Context.AnimInstanceProxy->GetRequiredBones();
@@ -210,6 +222,15 @@ void FAnimNode_ControlRig::PostSerialize(const FArchive& Ar)
 	DECLARE_SCOPE_HIERARCHICAL_COUNTER_FUNC()
 }
 
+UClass* FAnimNode_ControlRig::GetTargetClass() const
+{
+	if(ControlRigClass)
+	{
+		return ControlRigClass;
+	}
+	return DefaultControlRigClass;
+}
+
 void FAnimNode_ControlRig::UpdateInput(UBaseControlRig* InControlRig, const FPoseContext& InOutput)
 {
 	DECLARE_SCOPE_HIERARCHICAL_COUNTER_FUNC()
@@ -261,6 +282,199 @@ void FAnimNode_ControlRig::UpdateOutput(UBaseControlRig* InControlRig, FPoseCont
 	}
 }
 
+void FAnimNode_ControlRig::SetControlRigClass(TSubclassOf<UBaseControlRig> InControlRigClass)
+{
+	if(DefaultControlRigClass == nullptr)
+	{
+		DefaultControlRigClass = ControlRigClass;
+	}
+
+	// this may be setting an invalid runtime rig class,
+	// which will be validated during UpdateControlRigIfNeeded
+	ControlRigClass = InControlRigClass;
+}
+
+bool FAnimNode_ControlRig::UpdateControlRigIfNeeded(const UAnimInstance* InAnimInstance)
+{
+	if (UClass* ExpectedClass = GetTargetClass())
+	{
+		if(ControlRig != nullptr)
+		{
+			if(ControlRig->GetClass() != ExpectedClass)
+			{
+				UBaseControlRig* NewControlRig = nullptr;
+
+				auto ReportErrorAndSwitchToDefaultRig = [this, InAnimInstance, ExpectedClass](const FString& InMessage) -> bool
+				{
+					static constexpr TCHAR Format[] =  TEXT("[%s] Cannot switch to runtime rig class '%s' - reverting to default. %s");
+					UE_LOG(LogControlRig, Warning, Format, *InAnimInstance->GetPathName(), *ExpectedClass->GetName(), *InMessage);
+
+					// mark the class to be known - and nullptr - indicating that it is not supported.
+					ControlRigPerClass.FindOrAdd(ExpectedClass, nullptr);
+
+					// fall back to the default control rig and switch to that
+					ControlRigClass = nullptr;
+					
+					return UpdateControlRigIfNeeded(InAnimInstance);
+				};
+
+				// if we are reacting to a programmatic change
+				// we need to perform validation between the two control rigs (old and new)
+				if((ControlRigClass == ExpectedClass) &&
+					DefaultControlRigClass &&
+					(ExpectedClass != DefaultControlRigClass))
+				{
+					// check if we already created this before
+					if(const TObjectPtr<UBaseControlRig>* ExistingControlRig = ControlRigPerClass.Find(ExpectedClass))
+					{
+						NewControlRig = *ExistingControlRig;
+
+						// the existing control rig is nullptr indicates that the class is not supported.
+						// the warning will have been logged before - so it's not required to log it again.
+						if(NewControlRig == nullptr)
+						{
+							// fall back to the default control rig and switch to that
+							ControlRigClass = nullptr;
+							return UpdateControlRigIfNeeded(InAnimInstance);
+						}
+					}
+					else
+					{
+						if(ExpectedClass->IsNative())
+						{
+							static constexpr TCHAR Format[] = TEXT("Class '%s' is not supported (is it native).");
+							return ReportErrorAndSwitchToDefaultRig(FString::Printf(Format, *ExpectedClass->GetName()));
+						}
+						
+						// compare the two classes and make sure that the expected class is a super set in terms of
+						// user defined properties.
+						for (TFieldIterator<FProperty> PropertyIt(ControlRigClass); PropertyIt; ++PropertyIt)
+						{
+							const FProperty* OldProperty = *PropertyIt;
+							if(OldProperty->IsNative())
+							{
+								continue;
+							}
+
+							const FProperty* NewProperty = ExpectedClass->FindPropertyByName(OldProperty->GetFName());
+							if(NewProperty == nullptr)
+							{
+								static constexpr TCHAR Format[] = TEXT("Property / Variable '%s' is missing.");
+								return ReportErrorAndSwitchToDefaultRig(FString::Printf(Format, *OldProperty->GetName()));
+							}
+
+							if(!NewProperty->SameType(OldProperty))
+							{
+								FString OldExtendedCPPType, NewExtendedCPPType;
+								const FString OldCPPType = OldProperty->GetCPPType(&OldExtendedCPPType);
+								const FString NewCPPType = NewProperty->GetCPPType(&NewExtendedCPPType);
+								static constexpr TCHAR Format[] = TEXT("Property / Variable '%s' has incorrect type (is '%s', expected '%s').");
+								return ReportErrorAndSwitchToDefaultRig(FString::Printf(Format, *NewProperty->GetName(), *(NewCPPType + NewExtendedCPPType), *(OldCPPType + OldExtendedCPPType)));
+							}
+						}
+						
+						// create a new control rig using the new class
+						NewControlRig = NewObject<UBaseControlRig>(InAnimInstance->GetOwningComponent(), ExpectedClass);
+						NewControlRig->Initialize(true);
+						NewControlRig->RequestInit();
+
+						// temporarily set the new control rig to be the target instance
+						TGuardValue<TObjectPtr<UObject>> TargetInstanceGuard(TargetInstance, NewControlRig);
+
+						// propagate all variable inputs
+						PropagateInputProperties(InAnimInstance);
+
+						// run construction on the rig
+						NewControlRig->Execute(FRigUnit_PrepareForExecution::EventName);
+
+						const URigHierarchy* OldHierarchy = ControlRig->GetHierarchy();
+						const URigHierarchy* NewHierarchy = NewControlRig->GetHierarchy();
+
+						// now compare the two rigs - we need to check bone hierarchy compatibility.
+						const TArray<FRigElementKey> OldBoneKeys = OldHierarchy->GetBoneKeys(false);
+						const TArray<FRigElementKey> NewBoneKeys = NewHierarchy->GetBoneKeys(false);
+						for(const FRigElementKey& BoneKey : OldBoneKeys)
+						{
+							if(!NewBoneKeys.Contains(BoneKey))
+							{
+								static constexpr TCHAR Format[] = TEXT("Bone '%s' is missing from the rig.");
+								return ReportErrorAndSwitchToDefaultRig(FString::Printf(Format, *BoneKey.Name.ToString()));
+							}
+						}
+
+						// we also need to check curve hierarchy compatibility.
+						const TArray<FRigElementKey> OldCurveKeys = OldHierarchy->GetCurveKeys();
+						const TArray<FRigElementKey> NewCurveKeys = NewHierarchy->GetCurveKeys();
+						for(const FRigElementKey& CurveKey : OldCurveKeys)
+						{
+							if(!NewCurveKeys.Contains(CurveKey))
+							{
+								static constexpr TCHAR Format[] = TEXT("Curve '%s' is missing from the rig.");
+								return ReportErrorAndSwitchToDefaultRig(FString::Printf(Format, *CurveKey.Name.ToString()));
+							}
+						}
+
+						// we also need to check that potentially exposed controls match
+						for (int32 PropIdx = 0; PropIdx < DestPropertyNames.Num(); ++PropIdx)
+						{
+							if(const FRigControlElement* OldControlElement = ControlRig->FindControl(DestPropertyNames[PropIdx]))
+							{
+								const FRigControlElement* NewControlElement = NewControlRig->FindControl(DestPropertyNames[PropIdx]);
+								if(NewControlElement == nullptr)
+								{
+									static constexpr TCHAR Format[] = TEXT("Control '%s' is missing from the rig.");
+									return ReportErrorAndSwitchToDefaultRig(FString::Printf(Format, *DestPropertyNames[PropIdx].ToString()));
+								}
+
+								if(NewControlElement->Settings.ControlType != OldControlElement->Settings.ControlType)
+								{
+									static const UEnum* ControlTypeEnum = StaticEnum<ERigControlType>();
+									const FString OldType = ControlTypeEnum->GetDisplayNameTextByValue((int64)OldControlElement->Settings.ControlType).ToString();
+									const FString NewType = ControlTypeEnum->GetDisplayNameTextByValue((int64)NewControlElement->Settings.ControlType).ToString();
+									static constexpr TCHAR Format[] = TEXT("Control '%s' has the incorrect type (is '%s', expected '%s').");
+									return ReportErrorAndSwitchToDefaultRig(FString::Printf(Format, *DestPropertyNames[PropIdx].ToString(), *NewType, *OldType));
+								}
+							}
+						}
+
+						// fall through: we have a compatible new control rig, let's just use that.
+						ControlRigPerClass.FindOrAdd(NewControlRig->GetClass(), NewControlRig);
+					}
+				}
+
+				// stop listening to the rig, store it for reuse
+				ControlRig->OnInitialized_AnyThread().RemoveAll(this);
+				ControlRigPerClass.FindOrAdd(ControlRig->GetClass(), ControlRig);
+				ControlRig = nullptr;
+
+				if(NewControlRig)
+				{
+					Swap(NewControlRig, ControlRig);
+				}
+			}
+			else
+			{
+				// we have a control rig of the right class
+				return false;
+			}
+		}
+
+		if(ControlRig == nullptr)
+		{
+			ControlRig = NewObject<UBaseControlRig>(InAnimInstance->GetOwningComponent(), ExpectedClass);
+			ControlRig->Initialize(true);
+			ControlRig->RequestInit();
+		}
+		RefPoseSetterHash.Reset();
+		ControlRig->OnInitialized_AnyThread().AddRaw(this, &FAnimNode_ControlRig::HandleOnInitialized_AnyThread);
+
+		UpdateInputOutputMappingIfRequired(ControlRig, InAnimInstance->GetRequiredBones());
+
+		return true;
+	}
+	return false;
+}
+
 void FAnimNode_ControlRig::UpdateControlRigRefPoseIfNeeded(const FAnimInstanceProxy* InProxy, bool bIncludePoseInHash)
 {
 	if(!bSetRefPoseFromSkeleton)
diff --git a/Engine/Plugins/Animation/ControlRig/Source/ControlRig/Private/AnimNode_ControlRigBase.cpp b/Engine/Plugins/Animation/ControlRig/Source/ControlRig/Private/AnimNode_ControlRigBase.cpp
--- a/Engine/Plugins/Animation/ControlRig/Source/ControlRig/Private/AnimNode_ControlRigBase.cpp
+++ b/Engine/Plugins/Animation/ControlRig/Source/ControlRig/Private/AnimNode_ControlRigBase.cpp
@@ -558,12 +558,6 @@ void FAnimNode_ControlRigBase::CacheBones_AnyThread(const FAnimationCacheBonesCo
 
 	if (UBaseControlRig* ControlRig = GetControlRig())
 	{
-		URigHierarchy* Hierarchy = ControlRig->GetHierarchy();
-		if(Hierarchy == nullptr)
-		{
-			return;
-		}
-
 		// fill up node names
 		FBoneContainer& RequiredBones = Context.AnimInstanceProxy->GetRequiredBones();
 
@@ -579,183 +573,194 @@ void FAnimNode_ControlRigBase::CacheBones_AnyThread(const FAnimationCacheBonesCo
 			bControlRigRequiresInitialization = false;
 		}
 
-		ControlRigBoneInputMappingByIndex.Reset();
-		ControlRigBoneOutputMappingByIndex.Reset();
-		ControlRigCurveMappingByIndex.Reset();
-		ControlRigBoneInputMappingByName.Reset();
-		ControlRigBoneOutputMappingByName.Reset();
-		ControlRigCurveMappingByName.Reset();
+		UpdateInputOutputMappingIfRequired(ControlRig, RequiredBones);
 
-		if(RequiredBones.IsValid())
+		if(bControlRigRequiresInitialization)
 		{
-			const TArray<FBoneIndexType>& RequiredBonesArray = RequiredBones.GetBoneIndicesArray();
-			const int32 NumBones = RequiredBonesArray.Num();
+			// re-init only if this is the first run
+			// and restore control values
+			ControlRig->RequestInit();
+			bControlRigRequiresInitialization = false;
+		}
+		
+		LastBonesSerialNumberForCacheBones = BonesSerialNumber;
+	}
+}
 
-			const FReferenceSkeleton& RefSkeleton = RequiredBones.GetReferenceSkeleton();
+void FAnimNode_ControlRigBase::UpdateInputOutputMappingIfRequired(UBaseControlRig* InControlRig, const FBoneContainer& InRequiredBones)
+{
+	const URigHierarchy* Hierarchy = InControlRig->GetHierarchy();
+	if(Hierarchy == nullptr)
+	{
+		return;
+	}
 
-			// @todo: thread-safe? probably not in editor, but it may not be a big issue in editor
-			if (NodeMappingContainer.IsValid())
-			{
-				// get target to source mapping table - this is reversed mapping table
-				TMap<FName, FName> TargetToSourceMappingTable;
-				NodeMappingContainer->GetTargetToSourceMappingTable(TargetToSourceMappingTable);
+	ControlRigBoneInputMappingByIndex.Reset();
+	ControlRigBoneOutputMappingByIndex.Reset();
+	ControlRigCurveMappingByIndex.Reset();
+	ControlRigBoneInputMappingByName.Reset();
+	ControlRigBoneOutputMappingByName.Reset();
+	ControlRigCurveMappingByName.Reset();
+
+	if(InRequiredBones.IsValid())
+	{
+		const TArray<FBoneIndexType>& RequiredBonesArray = InRequiredBones.GetBoneIndicesArray();
+		const int32 NumBones = RequiredBonesArray.Num();
+
+		const FReferenceSkeleton& RefSkeleton = InRequiredBones.GetReferenceSkeleton();
 
-				// now fill up node name
-				for (uint16 Index = 0; Index < NumBones; ++Index)
+		// @todo: thread-safe? probably not in editor, but it may not be a big issue in editor
+		if (NodeMappingContainer.IsValid())
+		{
+			// get target to source mapping table - this is reversed mapping table
+			TMap<FName, FName> TargetToSourceMappingTable;
+			NodeMappingContainer->GetTargetToSourceMappingTable(TargetToSourceMappingTable);
+
+			// now fill up node name
+			for (uint16 Index = 0; Index < NumBones; ++Index)
+			{
+				// get bone name, and find reverse mapping
+				FName TargetNodeName = RefSkeleton.GetBoneName(RequiredBonesArray[Index]);
+				FName* SourceName = TargetToSourceMappingTable.Find(TargetNodeName);
+				if (SourceName)
 				{
-					// get bone name, and find reverse mapping
-					FName TargetNodeName = RefSkeleton.GetBoneName(RequiredBonesArray[Index]);
-					FName* SourceName = TargetToSourceMappingTable.Find(TargetNodeName);
-					if (SourceName)
-					{
-						ControlRigBoneInputMappingByName.Add(*SourceName, Index);
-					}
+					ControlRigBoneInputMappingByName.Add(*SourceName, Index);
 				}
 			}
-			else
-			{
-				TArray<FName> NodeNames;
-				TArray<FNodeItem> NodeItems;
-				ControlRig->GetMappableNodeData(NodeNames, NodeItems);
+		}
+		else
+		{
+			TArray<FName> NodeNames;
+			TArray<FNodeItem> NodeItems;
+			InControlRig->GetMappableNodeData(NodeNames, NodeItems);
 
-				// even if not mapped, we map only node that exists in the controlrig
-				for (uint16 Index = 0; Index < NumBones; ++Index)
+			// even if not mapped, we map only node that exists in the controlrig
+			for (uint16 Index = 0; Index < NumBones; ++Index)
+			{
+				const FName& BoneName = RefSkeleton.GetBoneName(RequiredBonesArray[Index]);
+				if (NodeNames.Contains(BoneName))
 				{
-					const FName& BoneName = RefSkeleton.GetBoneName(RequiredBonesArray[Index]);
-					if (NodeNames.Contains(BoneName))
-					{
-						ControlRigBoneInputMappingByName.Add(BoneName, Index);
-					}
+					ControlRigBoneInputMappingByName.Add(BoneName, Index);
 				}
 			}
+		}
 
-			auto UpdatingMappingFromSpecificTransferList = [] (
-				TArray<FBoneReference>& InTransferList,
-				const TWeakObjectPtr<UNodeMappingContainer>& InMappingContainer,
-				FBoneContainer& InRequiredBones,
-				const FReferenceSkeleton& InRefSkeleton,
-				const TArray<FBoneIndexType>& InRequiredBonesArray,
-				const UBaseControlRig* InControlRig,
-				TMap<FName, uint16>& OutMapping
-			) {
-				OutMapping.Reset();
-				
-				if (InMappingContainer.IsValid())
-				{
-					// get target to source mapping table - this is reversed mapping table
-					TMap<FName, FName> TargetToSourceMappingTable;
-					InMappingContainer->GetTargetToSourceMappingTable(TargetToSourceMappingTable);
+		auto UpdatingMappingFromSpecificTransferList = [] (
+			TArray<FBoneReference>& InTransferList,
+			const TWeakObjectPtr<UNodeMappingContainer>& InMappingContainer,
+			const FBoneContainer& InRequiredBones,
+			const FReferenceSkeleton& InRefSkeleton,
+			const TArray<FBoneIndexType>& InRequiredBonesArray,
+			const UBaseControlRig* InControlRig,
+			TMap<FName, uint16>& OutMapping
+		) {
+			OutMapping.Reset();
+			
+			if (InMappingContainer.IsValid())
+			{
+				// get target to source mapping table - this is reversed mapping table
+				TMap<FName, FName> TargetToSourceMappingTable;
+				InMappingContainer->GetTargetToSourceMappingTable(TargetToSourceMappingTable);
 
-					for(FBoneReference& InputBoneToTransfer : InTransferList)
+				for(FBoneReference& InputBoneToTransfer : InTransferList)
+				{
+					if(!InputBoneToTransfer.Initialize(InRequiredBones))
 					{
-						if(!InputBoneToTransfer.Initialize(InRequiredBones))
-						{
-							continue;
-						}
-						const FName TargetNodeName = InRefSkeleton.GetBoneName(InputBoneToTransfer.BoneIndex);
-						if (const FName* SourceName = TargetToSourceMappingTable.Find(TargetNodeName))
-						{
-							OutMapping.Add(*SourceName, InputBoneToTransfer.BoneIndex);
-						}
+						continue;
 					}
-				}
-				else
-				{
-					TArray<FName> NodeNames;
-					TArray<FNodeItem> NodeItems;
-					InControlRig->GetMappableNodeData(NodeNames, NodeItems);
-
-					for(FBoneReference& InputBoneToTransfer : InTransferList)
+					const FName TargetNodeName = InRefSkeleton.GetBoneName(InputBoneToTransfer.BoneIndex);
+					if (const FName* SourceName = TargetToSourceMappingTable.Find(TargetNodeName))
 					{
-						if(!InputBoneToTransfer.Initialize(InRequiredBones))
-						{
-							continue;
-						}
-						const FName& BoneName = InRefSkeleton.GetBoneName(InRequiredBonesArray[InputBoneToTransfer.BoneIndex]);
-						if (NodeNames.Contains(BoneName))
-						{
-							OutMapping.Add(BoneName, InputBoneToTransfer.BoneIndex);
-						}
+						OutMapping.Add(*SourceName, InputBoneToTransfer.BoneIndex);
 					}
 				}
-			};
-			
-			if(!InputBonesToTransfer.IsEmpty())
-			{
-				ControlRigBoneOutputMappingByName = ControlRigBoneInputMappingByName;
-
-				UpdatingMappingFromSpecificTransferList(
-					InputBonesToTransfer,
-					NodeMappingContainer,
-					RequiredBones,
-					RefSkeleton,
-					RequiredBonesArray,
-					ControlRig,
-					ControlRigBoneInputMappingByName);
 			}
-
-			if(!OutputBonesToTransfer.IsEmpty())
+			else
 			{
-				UpdatingMappingFromSpecificTransferList(
-					OutputBonesToTransfer,
-					NodeMappingContainer,
-					RequiredBones,
-					RefSkeleton,
-					RequiredBonesArray,
-					ControlRig,
-					ControlRigBoneOutputMappingByName);
-			}
+				TArray<FName> NodeNames;
+				TArray<FNodeItem> NodeItems;
+				InControlRig->GetMappableNodeData(NodeNames, NodeItems);
 
-			// check if we can switch the bones to an index based mapping.
-			// we can only do that if there is no node mapping container set.
-			if(!NodeMappingContainer.IsValid())
-			{
-				for(int32 InputOutput = 0; InputOutput < 2; InputOutput++)
+				for(FBoneReference& InputBoneToTransfer : InTransferList)
 				{
-					bool bIsMappingByIndex = true;
-					TMap<FName, uint16>& NameBasedMapping = InputOutput == 0 ? ControlRigBoneInputMappingByName : ControlRigBoneOutputMappingByName;
-					if(NameBasedMapping.IsEmpty())
+					if(!InputBoneToTransfer.Initialize(InRequiredBones))
 					{
 						continue;
 					}
-					
-					TArray<TPair<uint16, uint16>>& IndexBasedMapping = InputOutput == 0 ? ControlRigBoneInputMappingByIndex : ControlRigBoneOutputMappingByIndex;
-					
-					for (auto Iter = NameBasedMapping.CreateConstIterator(); Iter; ++Iter)
+					const FName& BoneName = InRefSkeleton.GetBoneName(InRequiredBonesArray[InputBoneToTransfer.BoneIndex]);
+					if (NodeNames.Contains(BoneName))
 					{
-						const uint16 SkeletonIndex = Iter.Value();
-						const int32 ControlRigIndex = Hierarchy->GetIndex(FRigElementKey(Iter.Key(), ERigElementType::Bone));
-						if(ControlRigIndex != INDEX_NONE)
-						{
-							IndexBasedMapping.Add(TPair<uint16, uint16>((uint16)ControlRigIndex, SkeletonIndex));
-						}
-						else
-						{
-							bIsMappingByIndex = false;
-						}
+						OutMapping.Add(BoneName, InputBoneToTransfer.BoneIndex);
 					}
+				}
+			}
+		};
+		
+		if(!InputBonesToTransfer.IsEmpty())
+		{
+			ControlRigBoneOutputMappingByName = ControlRigBoneInputMappingByName;
+
+			UpdatingMappingFromSpecificTransferList(
+				InputBonesToTransfer,
+				NodeMappingContainer,
+				InRequiredBones,
+				RefSkeleton,
+				RequiredBonesArray,
+				InControlRig,
+				ControlRigBoneInputMappingByName);
+		}
+
+		if(!OutputBonesToTransfer.IsEmpty())
+		{
+			UpdatingMappingFromSpecificTransferList(
+				OutputBonesToTransfer,
+				NodeMappingContainer,
+				InRequiredBones,
+				RefSkeleton,
+				RequiredBonesArray,
+				InControlRig,
+				ControlRigBoneOutputMappingByName);
+		}
 
-					if(bIsMappingByIndex)
+		// check if we can switch the bones to an index based mapping.
+		// we can only do that if there is no node mapping container set.
+		if(!NodeMappingContainer.IsValid())
+		{
+			for(int32 InputOutput = 0; InputOutput < 2; InputOutput++)
+			{
+				bool bIsMappingByIndex = true;
+				TMap<FName, uint16>& NameBasedMapping = InputOutput == 0 ? ControlRigBoneInputMappingByName : ControlRigBoneOutputMappingByName;
+				if(NameBasedMapping.IsEmpty())
+				{
+					continue;
+				}
+				
+				TArray<TPair<uint16, uint16>>& IndexBasedMapping = InputOutput == 0 ? ControlRigBoneInputMappingByIndex : ControlRigBoneOutputMappingByIndex;
+				
+				for (auto Iter = NameBasedMapping.CreateConstIterator(); Iter; ++Iter)
+				{
+					const uint16 SkeletonIndex = Iter.Value();
+					const int32 ControlRigIndex = Hierarchy->GetIndex(FRigElementKey(Iter.Key(), ERigElementType::Bone));
+					if(ControlRigIndex != INDEX_NONE)
 					{
-						NameBasedMapping.Reset();
+						IndexBasedMapping.Add(TPair<uint16, uint16>((uint16)ControlRigIndex, SkeletonIndex));
 					}
 					else
 					{
-						IndexBasedMapping.Reset();
+						bIsMappingByIndex = false;
 					}
 				}
-			}
-		}
 
-		if(bControlRigRequiresInitialization)
-		{
-			// re-init only if this is the first run
-			// and restore control values
-			ControlRig->RequestInit();
-			bControlRigRequiresInitialization = false;
+				if(bIsMappingByIndex)
+				{
+					NameBasedMapping.Reset();
+				}
+				else
+				{
+					IndexBasedMapping.Reset();
+				}
+			}
 		}
-		
-		LastBonesSerialNumberForCacheBones = BonesSerialNumber;
 	}
 }
 
diff --git a/Engine/Plugins/Animation/ControlRig/Source/ControlRig/Private/AnimNode_ControlRig_Library.cpp b/Engine/Plugins/Animation/ControlRig/Source/ControlRig/Private/AnimNode_ControlRig_Library.cpp
new file mode 100644
--- /dev/null
+++ b/Engine/Plugins/Animation/ControlRig/Source/ControlRig/Private/AnimNode_ControlRig_Library.cpp
@@ -0,0 +1,22 @@
+// Copyright Epic Games, Inc. All Rights Reserved.
+
+#include "AnimNode_ControlRig_Library.h"
+
+#include UE_INLINE_GENERATED_CPP_BY_NAME(AnimNode_ControlRig_Library)
+
+FControlRigReference UAnimNodeControlRigLibrary::ConvertToControlRig(const FAnimNodeReference& Node, EAnimNodeReferenceConversionResult& Result)
+{
+	return FAnimNodeReference::ConvertToType<FControlRigReference>(Node, Result);
+}
+
+FControlRigReference UAnimNodeControlRigLibrary::SetControlRigClass(const FControlRigReference& Node, TSubclassOf<UBaseControlRig> ControlRigClass)
+{
+	Node.CallAnimNodeFunction<FAnimNode_ControlRig>(
+	TEXT("SetSequence"),
+	[ControlRigClass](FAnimNode_ControlRig& InControlRigNode)
+	{
+		InControlRigNode.SetControlRigClass(ControlRigClass);
+	});
+
+	return Node;
+}
\ No newline at end of file
diff --git a/Engine/Plugins/Animation/ControlRig/Source/ControlRig/Public/AnimNode_ControlRig.h b/Engine/Plugins/Animation/ControlRig/Source/ControlRig/Public/AnimNode_ControlRig.h
--- a/Engine/Plugins/Animation/ControlRig/Source/ControlRig/Public/AnimNode_ControlRig.h
+++ b/Engine/Plugins/Animation/ControlRig/Source/ControlRig/Public/AnimNode_ControlRig.h
@@ -18,6 +18,8 @@ struct CONTROLRIG_API FAnimNode_ControlRig : public FAnimNode_ControlRigBase
 {
 	GENERATED_BODY()
 
+public:
+
 	FAnimNode_ControlRig();
 	~FAnimNode_ControlRig();
 
@@ -45,13 +47,23 @@ private:
 #endif
 private:
 
-	/** Cached ControlRig */
+	// The class to use for the rig. 
 	UPROPERTY(EditAnywhere, Category = ControlRig)
 	TSubclassOf<UBaseControlRig> ControlRigClass;
 
+	// The default class to use for the rig. This is needed
+	// only if the Control Rig Class is exposed as a pin.
+	UPROPERTY()
+	TSubclassOf<UBaseControlRig> DefaultControlRigClass;
+
+	/** Cached ControlRig */
 	UPROPERTY(transient)
 	TObjectPtr<UBaseControlRig> ControlRig;
 
+	/** Cached ControlRigs per class */
+	UPROPERTY(transient)
+	TMap<UClass*, TObjectPtr<UBaseControlRig>> ControlRigPerClass;
+
 	// alpha value handler
 	UPROPERTY(EditAnywhere, Category = Settings, meta = (PinShownByDefault))
 	float Alpha;
@@ -121,10 +133,13 @@ private:
 	FCurveMappings OutputCurveMappings;
 
 protected:
-	virtual UClass* GetTargetClass() const override { return *ControlRigClass; }
+	virtual UClass* GetTargetClass() const override;
 	virtual void UpdateInput(UBaseControlRig* InControlRig, const FPoseContext& InOutput) override;
 	virtual void UpdateOutput(UBaseControlRig* InControlRig, FPoseContext& InOutput) override;
 
+	void SetControlRigClass(TSubclassOf<UBaseControlRig> InControlRigClass);
+	bool UpdateControlRigIfNeeded(const UAnimInstance* InAnimInstance);
+
 	// Helper function to update the initial ref pose within the Control Rig if needed
 	void UpdateControlRigRefPoseIfNeeded(const FAnimInstanceProxy* InProxy, bool bIncludePoseInHash = false);
 	
@@ -136,6 +151,7 @@ public:
 	void PostSerialize(const FArchive& Ar);
 
 	friend class UAnimGraphNode_ControlRig;
+	friend class UAnimNodeControlRigLibrary;
 };
 
 template<>
diff --git a/Engine/Plugins/Animation/ControlRig/Source/ControlRig/Public/AnimNode_ControlRigBase.h b/Engine/Plugins/Animation/ControlRig/Source/ControlRig/Public/AnimNode_ControlRigBase.h
--- a/Engine/Plugins/Animation/ControlRig/Source/ControlRig/Public/AnimNode_ControlRigBase.h
+++ b/Engine/Plugins/Animation/ControlRig/Source/ControlRig/Public/AnimNode_ControlRigBase.h
@@ -77,6 +77,8 @@ struct CONTROLRIG_API FAnimNode_ControlRigBase : public FAnimNode_CustomProperty
 
 protected:
 
+	void UpdateInputOutputMappingIfRequired(UBaseControlRig* InControlRig, const FBoneContainer& RequiredBones); 
+
 	UPROPERTY(EditAnywhere, Category = Links)
 	FPoseLink Source;
 
diff --git a/Engine/Plugins/Animation/ControlRig/Source/ControlRig/Public/AnimNode_ControlRig_Library.h b/Engine/Plugins/Animation/ControlRig/Source/ControlRig/Public/AnimNode_ControlRig_Library.h
new file mode 100644
--- /dev/null
+++ b/Engine/Plugins/Animation/ControlRig/Source/ControlRig/Public/AnimNode_ControlRig_Library.h
@@ -0,0 +1,44 @@
+// Copyright Epic Games, Inc. All Rights Reserved.
+
+#pragma once
+
+#include "CoreMinimal.h"
+
+#include "Kismet/BlueprintFunctionLibrary.h"
+#include "Animation/AnimNodeReference.h"
+#include "AnimNode_ControlRig.h"
+#include "AnimNode_ControlRig_Library.generated.h"
+
+USTRUCT(BlueprintType)
+struct FControlRigReference : public FAnimNodeReference
+{
+	GENERATED_BODY()
+
+	typedef FAnimNode_ControlRig FInternalNodeType;
+};
+
+// Exposes operations to be performed on a control rig anim node
+UCLASS(Experimental, MinimalAPI)
+class UAnimNodeControlRigLibrary : public UBlueprintFunctionLibrary
+{
+	GENERATED_BODY()
+
+public:
+
+	/** Get a control rig context from an anim node context */
+	UFUNCTION(BlueprintCallable, Category = "Animation|ControlRig", meta = (BlueprintThreadSafe, ExpandEnumAsExecs = "Result"))
+	static CONTROLRIG_API FControlRigReference ConvertToControlRig(const FAnimNodeReference& Node, EAnimNodeReferenceConversionResult& Result);
+
+	/** Get a control rig context from an anim node context (pure) */
+	UFUNCTION(BlueprintPure, Category = "Animation|ControlRig", meta = (BlueprintThreadSafe, DisplayName = "Convert to Sequence Player"))
+	static void ConvertToControlRigPure(const FAnimNodeReference& Node, FControlRigReference& ControlRig, bool& Result)
+	{
+		EAnimNodeReferenceConversionResult ConversionResult;
+		ControlRig = ConvertToControlRig(Node, ConversionResult);
+		Result = (ConversionResult == EAnimNodeReferenceConversionResult::Succeeded);
+	}
+
+	/** Set the control rig class on the node */
+	UFUNCTION(BlueprintCallable, Category = "Animation|ControlRig", meta = (BlueprintThreadSafe))
+	static CONTROLRIG_API FControlRigReference SetControlRigClass(const FControlRigReference& Node, TSubclassOf<UBaseControlRig> ControlRigClass);
+};
