Control Rig: Fix broken sequencer control rig blueprint library from UBaseControlRig changes.

diff --git a/Engine/Plugins/Animation/ControlRig/Source/ControlRigEditor/Private/ControlRigSequencerEditorLibrary.cpp b/Engine/Plugins/Animation/ControlRig/Source/ControlRigEditor/Private/ControlRigSequencerEditorLibrary.cpp
--- a/Engine/Plugins/Animation/ControlRig/Source/ControlRigEditor/Private/ControlRigSequencerEditorLibrary.cpp
+++ b/Engine/Plugins/Animation/ControlRig/Source/ControlRigEditor/Private/ControlRigSequencerEditorLibrary.cpp
@@ -17,7 +17,7 @@
 #include "MovieSceneToolHelpers.h"
 #include "Rigs/FKControlRig.h"
 #include "Units/Execution/RigUnit_InverseExecution.h"
-#include "BaseControlRig.h"
+#include "ControlRig.h"
 #include "EditMode/ControlRigEditMode.h"
 #include "EditorModeManager.h"
 #include "Engine/Selection.h"
@@ -54,13 +54,21 @@
 
 #define LOCTEXT_NAMESPACE "ControlrigSequencerEditorLibrary"
 
-TArray<UBaseControlRig*> UControlRigSequencerEditorLibrary::GetVisibleControlRigs()
+TArray<UControlRig*> UControlRigSequencerEditorLibrary::GetVisibleControlRigs()
 {
-	TArray<UBaseControlRig*> ControlRigs;
+	TArray<UControlRig*> ControlRigs;
+	TArray<UBaseControlRig*> BaseControlRigs;
 	FControlRigEditMode* ControlRigEditMode = static_cast<FControlRigEditMode*>(GLevelEditorModeTools().GetActiveMode(FControlRigEditMode::ModeName));
 	if (ControlRigEditMode)
 	{
-		ControlRigs = ControlRigEditMode->GetControlRigsArray(true /*bIsVisible*/);
+		BaseControlRigs = ControlRigEditMode->GetControlRigsArray(true /*bIsVisible*/);
+	}
+	for (UBaseControlRig* BaseRig : BaseControlRigs)
+	{
+		if (UControlRig* ControlRig = Cast<UControlRig>(BaseRig))
+		{
+			ControlRigs.Add(ControlRig);
+		}
 	}
 	return ControlRigs;
 }
@@ -84,7 +92,8 @@ TArray<FControlRigSequencerBindingProxy> UControlRigSequencerEditorLibrary::GetC
 					{
 						FControlRigSequencerBindingProxy BindingProxy;
 						BindingProxy.Track = Track;
-						BindingProxy.ControlRig = Track->GetControlRig();
+						
+						BindingProxy.ControlRig = Cast<UControlRig>(Track->GetControlRig());
 						BindingProxy.Proxy.BindingID = Binding.GetObjectGuid();
 						BindingProxy.Proxy.Sequence = LevelSequence;
 						ControlRigBindingProxies.Add(BindingProxy);
@@ -199,10 +208,10 @@ static TSharedPtr<ISequencer> GetSequencerFromAsset()
 	return Sequencer;
 }
 
-static UMovieSceneControlRigParameterTrack* AddControlRig(ULevelSequence* LevelSequence,const UClass* InClass, UObject* BoundActor, FGuid ObjectBinding, UBaseControlRig* InExistingControlRig, bool bIsAdditiveControlRig)
+static UMovieSceneControlRigParameterTrack* AddControlRig(ULevelSequence* LevelSequence,const UClass* InClass, UObject* BoundActor, FGuid ObjectBinding, UControlRig* InExistingControlRig, bool bIsAdditiveControlRig)
 {
 	FSlateApplication::Get().DismissAllMenus();
-	if (!InClass || !InClass->IsChildOf(UBaseControlRig::StaticClass()) ||
+	if (!InClass || !InClass->IsChildOf(UControlRig::StaticClass()) ||
 		!LevelSequence || !LevelSequence->GetMovieScene())
 	{
 		return nullptr;
@@ -218,7 +227,7 @@ static UMovieSceneControlRigParameterTrack* AddControlRig(ULevelSequence* LevelS
 	LevelSequence->Modify();
 	OwnerMovieScene->Modify();
 	
-	if (bIsAdditiveControlRig && InClass != UFKControlRig::StaticClass() && !InClass->GetDefaultObject<UBaseControlRig>()->SupportsEvent(FRigUnit_InverseExecution::EventName))
+	if (bIsAdditiveControlRig && InClass != UFKControlRig::StaticClass() && !InClass->GetDefaultObject<UControlRig>()->SupportsEvent(FRigUnit_InverseExecution::EventName))
 	{
 		UE_LOG(LogControlRigEditor, Error, TEXT("Cannot add an additive control rig which does not contain a backwards solve event."));
 		return nullptr;
@@ -233,15 +242,15 @@ static UMovieSceneControlRigParameterTrack* AddControlRig(ULevelSequence* LevelS
 		ObjectName.RemoveFromEnd(TEXT("_C"));
 
 		bool bSequencerOwnsControlRig = false;
-		UBaseControlRig* ControlRig = InExistingControlRig;
+		UControlRig* ControlRig = InExistingControlRig;
 		if (ControlRig == nullptr)
 		{
-			ControlRig = NewObject<UBaseControlRig>(Track, InClass, FName(*ObjectName), RF_Transactional);
+			ControlRig = NewObject<UControlRig>(Track, InClass, FName(*ObjectName), RF_Transactional);
 			bSequencerOwnsControlRig = true;
 		}
 
 		ControlRig->Modify();
-		if (UFKControlRig* FKControlRig = Cast<UFKControlRig>(ControlRig))
+		if (UFKControlRig* FKControlRig = Cast<UFKControlRig>(Cast<UBaseControlRig>(ControlRig)))
 		{
 			if (bIsAdditiveControlRig)
 			{
@@ -254,7 +263,7 @@ static UMovieSceneControlRigParameterTrack* AddControlRig(ULevelSequence* LevelS
 		}
 		ControlRig->SetObjectBinding(MakeShared<FControlRigObjectBinding>());
 		ControlRig->GetObjectBinding()->BindToObject(BoundActor);
-		ControlRig->GetDataSourceRegistry()->RegisterDataSource(UBaseControlRig::OwnerComponent, ControlRig->GetObjectBinding()->GetBoundObject());
+		ControlRig->GetDataSourceRegistry()->RegisterDataSource(UControlRig::OwnerComponent, ControlRig->GetObjectBinding()->GetBoundObject());
 		// Do not re-initialize existing control rig
 		if (!InExistingControlRig)
 		{
@@ -338,7 +347,7 @@ UMovieSceneTrack* UControlRigSequencerEditorLibrary::FindOrCreateControlRigTrack
 				USkeletalMeshComponent* SkeletalMeshComponent = nullptr;
 				AcquireSkeletonAndSkelMeshCompFromObject(BoundObject, &Skeleton, &SkeletalMeshComponent);
 
-				UBaseControlRig* ControlRig = nullptr;
+				UControlRig* ControlRig = nullptr;
 				if (SkeletalMeshComponent && SkeletalMeshComponent->GetSkeletalMeshAsset() && SkeletalMeshComponent->GetSkeletalMeshAsset()->GetSkeleton())
 				{
 					UMovieSceneControlRigParameterTrack* Track = AddControlRig(LevelSequence, ControlRigClass, SkeletalMeshComponent, InBinding.BindingID, nullptr, bIsAdditiveControlRig);
@@ -387,7 +396,7 @@ TArray<UMovieSceneTrack*> UControlRigSequencerEditorLibrary::FindOrCreateControl
 			BoundActor->GetComponents<UControlRigComponent>(ControlRigComponents);
 			for (UControlRigComponent* ControlRigComponent : ControlRigComponents)
 			{
-				if (UBaseControlRig* CR = ControlRigComponent->GetControlRig())
+				if (UControlRig* CR = Cast<UControlRig>(ControlRigComponent->GetControlRig()))
 				{
 					UMovieSceneControlRigParameterTrack* GoodTrack = nullptr;
 					if (FMovieSceneBinding* Binding = MovieScene->FindBinding(InBinding.BindingID))
@@ -428,7 +437,7 @@ TArray<UMovieSceneTrack*> UControlRigSequencerEditorLibrary::FindOrCreateControl
 	return Tracks;
 }
 
-bool UControlRigSequencerEditorLibrary::TweenControlRig(ULevelSequence* LevelSequence, UBaseControlRig* ControlRig, float TweenValue)
+bool UControlRigSequencerEditorLibrary::TweenControlRig(ULevelSequence* LevelSequence, UControlRig* ControlRig, float TweenValue)
 {
 	TWeakPtr<ISequencer> WeakSequencer = GetSequencerFromAsset();
 	if (WeakSequencer.IsValid() && WeakSequencer.Pin()->GetFocusedMovieSceneSequence() == LevelSequence
@@ -1012,7 +1021,7 @@ TArray<FTransform> UControlRigSequencerEditorLibrary::GetSkeletalMeshComponentWo
 	return OutWorldTransforms;
 }
 
-FTransform UControlRigSequencerEditorLibrary::GetControlRigWorldTransform(ULevelSequence* LevelSequence,UBaseControlRig* ControlRig, FName ControlName, FFrameNumber Frame,
+FTransform UControlRigSequencerEditorLibrary::GetControlRigWorldTransform(ULevelSequence* LevelSequence,UControlRig* ControlRig, FName ControlName, FFrameNumber Frame,
 	ESequenceTimeUnit TimeUnit)
 {
 	TArray<FFrameNumber> Frames;
@@ -1025,7 +1034,7 @@ FTransform UControlRigSequencerEditorLibrary::GetControlRigWorldTransform(ULevel
 	return FTransform::Identity;
 }
 
-TArray<FTransform> UControlRigSequencerEditorLibrary::GetControlRigWorldTransforms(ULevelSequence* LevelSequence,UBaseControlRig* ControlRig, FName ControlName, const TArray<FFrameNumber>& InFrames,
+TArray<FTransform> UControlRigSequencerEditorLibrary::GetControlRigWorldTransforms(ULevelSequence* LevelSequence,UControlRig* ControlRig, FName ControlName, const TArray<FFrameNumber>& InFrames,
 	ESequenceTimeUnit TimeUnit)
 {
 	TWeakPtr<ISequencer> WeakSequencer = GetSequencerFromAsset();
@@ -1057,7 +1066,7 @@ TArray<FTransform> UControlRigSequencerEditorLibrary::GetControlRigWorldTransfor
 }
 
 
-static void LocalSetControlRigWorldTransforms(ULevelSequence* LevelSequence,UBaseControlRig* ControlRig, FName ControlName, EControlRigSetKey SetKey, const TArray<FFrameNumber>& InFrames, 
+static void LocalSetControlRigWorldTransforms(ULevelSequence* LevelSequence,UControlRig* ControlRig, FName ControlName, EControlRigSetKey SetKey, const TArray<FFrameNumber>& InFrames, 
 	const TArray<FTransform>& WorldTransforms, ESequenceTimeUnit TimeUnit)
 {
 	TWeakPtr<ISequencer> WeakSequencer = GetSequencerFromAsset();
@@ -1111,7 +1120,7 @@ static void LocalSetControlRigWorldTransforms(ULevelSequence* LevelSequence,UBas
 	}
 }
 
-void UControlRigSequencerEditorLibrary::SetControlRigWorldTransform(ULevelSequence* LevelSequence, UBaseControlRig* ControlRig, FName ControlName, FFrameNumber Frame, const FTransform& WorldTransform,
+void UControlRigSequencerEditorLibrary::SetControlRigWorldTransform(ULevelSequence* LevelSequence, UControlRig* ControlRig, FName ControlName, FFrameNumber Frame, const FTransform& WorldTransform,
 	ESequenceTimeUnit TimeUnit, bool bSetKey)
 {
 	EControlRigSetKey SetKey = bSetKey ? EControlRigSetKey::Always : EControlRigSetKey::DoNotCare;
@@ -1124,7 +1133,7 @@ void UControlRigSequencerEditorLibrary::SetControlRigWorldTransform(ULevelSequen
 
 }
 
-void UControlRigSequencerEditorLibrary::SetControlRigWorldTransforms(ULevelSequence* LevelSequence, UBaseControlRig* ControlRig, FName ControlName, const TArray<FFrameNumber>& Frames, const TArray<FTransform>& WorldTransforms,
+void UControlRigSequencerEditorLibrary::SetControlRigWorldTransforms(ULevelSequence* LevelSequence, UControlRig* ControlRig, FName ControlName, const TArray<FFrameNumber>& Frames, const TArray<FTransform>& WorldTransforms,
 	ESequenceTimeUnit TimeUnit)
 {
 	LocalSetControlRigWorldTransforms(LevelSequence, ControlRig, ControlName, EControlRigSetKey::Always, Frames, WorldTransforms,TimeUnit);
@@ -1227,7 +1236,7 @@ bool UControlRigSequencerEditorLibrary::BakeToControlRig(UWorld* World, ULevelSe
 				{
 					FString ObjectName = InClass->GetName();
 					ObjectName.RemoveFromEnd(TEXT("_C"));
-					UBaseControlRig* ControlRig = NewObject<UBaseControlRig>(Track, InClass, FName(*ObjectName), RF_Transactional);
+					UControlRig* ControlRig = NewObject<UControlRig>(Track, InClass, FName(*ObjectName), RF_Transactional);
 					FName OldEventString = FName(FString(TEXT("Inverse")));
 
 					if (InClass != UFKControlRig::StaticClass() && !(ControlRig->SupportsEvent(FRigUnit_InverseExecution::EventName) || ControlRig->SupportsEvent(OldEventString)))
@@ -1253,7 +1262,7 @@ bool UControlRigSequencerEditorLibrary::BakeToControlRig(UWorld* World, ULevelSe
 						else
 						{
 							/* mz todo we don't unbind  will test more
-							UBaseControlRig* OldControlRig = ControlRigEditMode->GetControlRig(false);
+							UControlRig* OldControlRig = ControlRigEditMode->GetControlRig(false);
 							if (OldControlRig)
 							{
 								WeakSequencer.Pin()->ObjectImplicitlyRemoved(OldControlRig);
@@ -1265,7 +1274,7 @@ bool UControlRigSequencerEditorLibrary::BakeToControlRig(UWorld* World, ULevelSe
 					ControlRig->Modify();
 					ControlRig->SetObjectBinding(MakeShared<FControlRigObjectBinding>());
 					ControlRig->GetObjectBinding()->BindToObject(SkeletalMeshComp);
-					ControlRig->GetDataSourceRegistry()->RegisterDataSource(UBaseControlRig::OwnerComponent, ControlRig->GetObjectBinding()->GetBoundObject());
+					ControlRig->GetDataSourceRegistry()->RegisterDataSource(UControlRig::OwnerComponent, ControlRig->GetObjectBinding()->GetBoundObject());
 					ControlRig->Initialize();
 					ControlRig->RequestInit();
 					ControlRig->SetBoneInitialTransformsFromSkeletalMeshComponent(SkeletalMeshComp, true);
@@ -1358,7 +1367,7 @@ bool UControlRigSequencerEditorLibrary::LoadAnimSequenceIntoControlRigSection(UM
 }
 
 static bool LocalGetControlRigControlValues(IMovieScenePlayer* Player, UMovieSceneSequence* MovieSceneSequence, FMovieSceneSequenceIDRef Template, FMovieSceneSequenceTransform& RootToLocalTransform,
-	UBaseControlRig* ControlRig, const FName& ControlName, ESequenceTimeUnit TimeUnit,
+	UControlRig* ControlRig, const FName& ControlName, ESequenceTimeUnit TimeUnit,
 	const TArray<FFrameNumber>& InFrames, TArray<FRigControlValue>& OutValues)
 {
 	if (Player == nullptr || MovieSceneSequence == nullptr || ControlRig == nullptr)
@@ -1394,7 +1403,7 @@ static bool LocalGetControlRigControlValues(IMovieScenePlayer* Player, UMovieSce
 	return true;
 }
 
-static bool GetControlRigValues(ISequencer* Sequencer, UBaseControlRig* ControlRig, const FName& ControlName, ESequenceTimeUnit TimeUnit,
+static bool GetControlRigValues(ISequencer* Sequencer, UControlRig* ControlRig, const FName& ControlName, ESequenceTimeUnit TimeUnit,
 	const TArray<FFrameNumber>& Frames,  TArray<FRigControlValue>& OutValues)
 {
 	if (Sequencer->GetFocusedMovieSceneSequence())
@@ -1409,7 +1418,7 @@ static bool GetControlRigValues(ISequencer* Sequencer, UBaseControlRig* ControlR
 	return false;
 }
 
-static bool GetControlRigValue(ISequencer* Sequencer, UBaseControlRig* ControlRig, const FName& ControlName, ESequenceTimeUnit TimeUnit,
+static bool GetControlRigValue(ISequencer* Sequencer, UControlRig* ControlRig, const FName& ControlName, ESequenceTimeUnit TimeUnit,
 	const FFrameNumber Frame, FRigControlValue& OutValue)
 {
 	if (Sequencer->GetFocusedMovieSceneSequence())
@@ -1430,7 +1439,7 @@ static bool GetControlRigValue(ISequencer* Sequencer, UBaseControlRig* ControlRi
 	return false;
 }
 
-static bool GetControlRigValues(UWorld* World, ULevelSequence* LevelSequence, UBaseControlRig* ControlRig, const FName& ControlName, ESequenceTimeUnit TimeUnit,
+static bool GetControlRigValues(UWorld* World, ULevelSequence* LevelSequence, UControlRig* ControlRig, const FName& ControlName, ESequenceTimeUnit TimeUnit,
 	const TArray<FFrameNumber>& Frames, TArray<FRigControlValue>& OutValues)
 {
 	if (LevelSequence)
@@ -1449,7 +1458,7 @@ static bool GetControlRigValues(UWorld* World, ULevelSequence* LevelSequence, UB
 }
 
 
-float UControlRigSequencerEditorLibrary::GetLocalControlRigFloat(ULevelSequence* LevelSequence,UBaseControlRig* ControlRig, FName ControlName, FFrameNumber Frame, ESequenceTimeUnit TimeUnit)
+float UControlRigSequencerEditorLibrary::GetLocalControlRigFloat(ULevelSequence* LevelSequence,UControlRig* ControlRig, FName ControlName, FFrameNumber Frame, ESequenceTimeUnit TimeUnit)
 {
 	float Value = 0.0f;
 	TWeakPtr<ISequencer> WeakSequencer = GetSequencerFromAsset();
@@ -1473,7 +1482,7 @@ float UControlRigSequencerEditorLibrary::GetLocalControlRigFloat(ULevelSequence*
 	return Value;
 }
 
-TArray<float> UControlRigSequencerEditorLibrary::GetLocalControlRigFloats(ULevelSequence* LevelSequence,UBaseControlRig* ControlRig, FName ControlName, const TArray<FFrameNumber>& Frames, ESequenceTimeUnit TimeUnit)
+TArray<float> UControlRigSequencerEditorLibrary::GetLocalControlRigFloats(ULevelSequence* LevelSequence,UControlRig* ControlRig, FName ControlName, const TArray<FFrameNumber>& Frames, ESequenceTimeUnit TimeUnit)
 {
 	TArray<float> Values;
 	TWeakPtr<ISequencer> WeakSequencer = GetSequencerFromAsset();
@@ -1503,7 +1512,7 @@ TArray<float> UControlRigSequencerEditorLibrary::GetLocalControlRigFloats(ULevel
 	return Values;
 }
 
-void UControlRigSequencerEditorLibrary::SetLocalControlRigFloat(ULevelSequence* LevelSequence, UBaseControlRig* ControlRig, FName ControlName, FFrameNumber Frame, float Value, ESequenceTimeUnit TimeUnit, bool bSetKey)
+void UControlRigSequencerEditorLibrary::SetLocalControlRigFloat(ULevelSequence* LevelSequence, UControlRig* ControlRig, FName ControlName, FFrameNumber Frame, float Value, ESequenceTimeUnit TimeUnit, bool bSetKey)
 {
 	if (LevelSequence == nullptr || ControlRig == nullptr)
 	{
@@ -1531,7 +1540,7 @@ void UControlRigSequencerEditorLibrary::SetLocalControlRigFloat(ULevelSequence*
 	}
 }
 
-void UControlRigSequencerEditorLibrary::SetLocalControlRigFloats(ULevelSequence* LevelSequence, UBaseControlRig* ControlRig, FName ControlName, const TArray<FFrameNumber>& Frames, const TArray<float> Values,
+void UControlRigSequencerEditorLibrary::SetLocalControlRigFloats(ULevelSequence* LevelSequence, UControlRig* ControlRig, FName ControlName, const TArray<FFrameNumber>& Frames, const TArray<float> Values,
 	ESequenceTimeUnit TimeUnit)
 {
 	if (LevelSequence == nullptr || ControlRig == nullptr || (Frames.Num() != Values.Num()))
@@ -1566,7 +1575,7 @@ void UControlRigSequencerEditorLibrary::SetLocalControlRigFloats(ULevelSequence*
 }
 
 
-bool UControlRigSequencerEditorLibrary::GetLocalControlRigBool(ULevelSequence* LevelSequence, UBaseControlRig* ControlRig, FName ControlName, FFrameNumber Frame, ESequenceTimeUnit TimeUnit)
+bool UControlRigSequencerEditorLibrary::GetLocalControlRigBool(ULevelSequence* LevelSequence, UControlRig* ControlRig, FName ControlName, FFrameNumber Frame, ESequenceTimeUnit TimeUnit)
 {
 	bool Value = true;
 	TWeakPtr<ISequencer> WeakSequencer = GetSequencerFromAsset();
@@ -1590,7 +1599,7 @@ bool UControlRigSequencerEditorLibrary::GetLocalControlRigBool(ULevelSequence* L
 	return Value;
 }
 
-TArray<bool> UControlRigSequencerEditorLibrary::GetLocalControlRigBools(ULevelSequence* LevelSequence, UBaseControlRig* ControlRig, FName ControlName, const TArray<FFrameNumber>& Frames, ESequenceTimeUnit TimeUnit)
+TArray<bool> UControlRigSequencerEditorLibrary::GetLocalControlRigBools(ULevelSequence* LevelSequence, UControlRig* ControlRig, FName ControlName, const TArray<FFrameNumber>& Frames, ESequenceTimeUnit TimeUnit)
 {
 	TArray<bool> Values;
 	TWeakPtr<ISequencer> WeakSequencer = GetSequencerFromAsset();
@@ -1619,7 +1628,7 @@ TArray<bool> UControlRigSequencerEditorLibrary::GetLocalControlRigBools(ULevelSe
 	return Values;
 }
 
-void UControlRigSequencerEditorLibrary::SetLocalControlRigBool(ULevelSequence* LevelSequence, UBaseControlRig* ControlRig, FName ControlName, FFrameNumber Frame, bool Value, ESequenceTimeUnit TimeUnit, bool bSetKey)
+void UControlRigSequencerEditorLibrary::SetLocalControlRigBool(ULevelSequence* LevelSequence, UControlRig* ControlRig, FName ControlName, FFrameNumber Frame, bool Value, ESequenceTimeUnit TimeUnit, bool bSetKey)
 {
 	if (LevelSequence == nullptr || ControlRig == nullptr)
 	{
@@ -1647,7 +1656,7 @@ void UControlRigSequencerEditorLibrary::SetLocalControlRigBool(ULevelSequence* L
 	}
 }
 
-void UControlRigSequencerEditorLibrary::SetLocalControlRigBools(ULevelSequence* LevelSequence, UBaseControlRig* ControlRig, FName ControlName, const TArray<FFrameNumber>& Frames,
+void UControlRigSequencerEditorLibrary::SetLocalControlRigBools(ULevelSequence* LevelSequence, UControlRig* ControlRig, FName ControlName, const TArray<FFrameNumber>& Frames,
 	const TArray<bool> Values, ESequenceTimeUnit TimeUnit)
 {
 	if (LevelSequence == nullptr || ControlRig == nullptr || (Frames.Num() != Values.Num()))
@@ -1681,7 +1690,7 @@ void UControlRigSequencerEditorLibrary::SetLocalControlRigBools(ULevelSequence*
 	}
 }
 
-int32 UControlRigSequencerEditorLibrary::GetLocalControlRigInt(ULevelSequence* LevelSequence, UBaseControlRig* ControlRig, FName ControlName, FFrameNumber Frame, ESequenceTimeUnit TimeUnit)
+int32 UControlRigSequencerEditorLibrary::GetLocalControlRigInt(ULevelSequence* LevelSequence, UControlRig* ControlRig, FName ControlName, FFrameNumber Frame, ESequenceTimeUnit TimeUnit)
 {
 	int32 Value = 0;
 	TWeakPtr<ISequencer> WeakSequencer = GetSequencerFromAsset();
@@ -1705,7 +1714,7 @@ int32 UControlRigSequencerEditorLibrary::GetLocalControlRigInt(ULevelSequence* L
 	return Value;
 }
 
-TArray<int32> UControlRigSequencerEditorLibrary::GetLocalControlRigInts(ULevelSequence* LevelSequence, UBaseControlRig* ControlRig, FName ControlName, const TArray<FFrameNumber>& Frames, ESequenceTimeUnit TimeUnit)
+TArray<int32> UControlRigSequencerEditorLibrary::GetLocalControlRigInts(ULevelSequence* LevelSequence, UControlRig* ControlRig, FName ControlName, const TArray<FFrameNumber>& Frames, ESequenceTimeUnit TimeUnit)
 {
 	TArray<int32> Values;
 	TWeakPtr<ISequencer> WeakSequencer = GetSequencerFromAsset();
@@ -1734,7 +1743,7 @@ TArray<int32> UControlRigSequencerEditorLibrary::GetLocalControlRigInts(ULevelSe
 	return Values;
 }
 
-void UControlRigSequencerEditorLibrary::SetLocalControlRigInt(ULevelSequence* LevelSequence, UBaseControlRig* ControlRig, FName ControlName, FFrameNumber Frame, int32 Value, ESequenceTimeUnit TimeUnit,bool bSetKey)
+void UControlRigSequencerEditorLibrary::SetLocalControlRigInt(ULevelSequence* LevelSequence, UControlRig* ControlRig, FName ControlName, FFrameNumber Frame, int32 Value, ESequenceTimeUnit TimeUnit,bool bSetKey)
 {
 	if (LevelSequence == nullptr || ControlRig == nullptr)
 	{
@@ -1762,7 +1771,7 @@ void UControlRigSequencerEditorLibrary::SetLocalControlRigInt(ULevelSequence* Le
 	}
 }
 
-void UControlRigSequencerEditorLibrary::SetLocalControlRigInts(ULevelSequence* LevelSequence, UBaseControlRig* ControlRig, FName ControlName, const TArray<FFrameNumber>& Frames,
+void UControlRigSequencerEditorLibrary::SetLocalControlRigInts(ULevelSequence* LevelSequence, UControlRig* ControlRig, FName ControlName, const TArray<FFrameNumber>& Frames,
 	const TArray<int32> Values , ESequenceTimeUnit TimeUnit)
 {
 	if (LevelSequence == nullptr || ControlRig == nullptr || (Frames.Num() != Values.Num()))
@@ -1797,7 +1806,7 @@ void UControlRigSequencerEditorLibrary::SetLocalControlRigInts(ULevelSequence* L
 }
 
 
-FVector2D UControlRigSequencerEditorLibrary::GetLocalControlRigVector2D(ULevelSequence* LevelSequence, UBaseControlRig* ControlRig, FName ControlName, FFrameNumber Frame, ESequenceTimeUnit TimeUnit)
+FVector2D UControlRigSequencerEditorLibrary::GetLocalControlRigVector2D(ULevelSequence* LevelSequence, UControlRig* ControlRig, FName ControlName, FFrameNumber Frame, ESequenceTimeUnit TimeUnit)
 {
 	FVector2D Value;
 	TWeakPtr<ISequencer> WeakSequencer = GetSequencerFromAsset();
@@ -1822,7 +1831,7 @@ FVector2D UControlRigSequencerEditorLibrary::GetLocalControlRigVector2D(ULevelSe
 	return Value;
 }
 
-TArray<FVector2D> UControlRigSequencerEditorLibrary::GetLocalControlRigVector2Ds(ULevelSequence* LevelSequence, UBaseControlRig* ControlRig, FName ControlName, const TArray<FFrameNumber>& Frames, ESequenceTimeUnit TimeUnit)
+TArray<FVector2D> UControlRigSequencerEditorLibrary::GetLocalControlRigVector2Ds(ULevelSequence* LevelSequence, UControlRig* ControlRig, FName ControlName, const TArray<FFrameNumber>& Frames, ESequenceTimeUnit TimeUnit)
 {
 	TArray<FVector2D> Values;
 	TWeakPtr<ISequencer> WeakSequencer = GetSequencerFromAsset();
@@ -1852,7 +1861,7 @@ TArray<FVector2D> UControlRigSequencerEditorLibrary::GetLocalControlRigVector2Ds
 
 }
 
-void UControlRigSequencerEditorLibrary::SetLocalControlRigVector2D(ULevelSequence* LevelSequence, UBaseControlRig* ControlRig, FName ControlName, FFrameNumber Frame, FVector2D Value, ESequenceTimeUnit TimeUnit,bool bSetKey)
+void UControlRigSequencerEditorLibrary::SetLocalControlRigVector2D(ULevelSequence* LevelSequence, UControlRig* ControlRig, FName ControlName, FFrameNumber Frame, FVector2D Value, ESequenceTimeUnit TimeUnit,bool bSetKey)
 {
 	if (LevelSequence == nullptr || ControlRig == nullptr)
 	{
@@ -1880,7 +1889,7 @@ void UControlRigSequencerEditorLibrary::SetLocalControlRigVector2D(ULevelSequenc
 	}
 }
 
-void UControlRigSequencerEditorLibrary::SetLocalControlRigVector2Ds(ULevelSequence* LevelSequence, UBaseControlRig* ControlRig, FName ControlName, const TArray<FFrameNumber>& Frames, 
+void UControlRigSequencerEditorLibrary::SetLocalControlRigVector2Ds(ULevelSequence* LevelSequence, UControlRig* ControlRig, FName ControlName, const TArray<FFrameNumber>& Frames, 
 	const TArray<FVector2D> Values, ESequenceTimeUnit TimeUnit)
 {
 	if (LevelSequence == nullptr || ControlRig == nullptr || (Frames.Num() != Values.Num()))
@@ -1915,7 +1924,7 @@ void UControlRigSequencerEditorLibrary::SetLocalControlRigVector2Ds(ULevelSequen
 }
 
 
-FVector UControlRigSequencerEditorLibrary::GetLocalControlRigPosition(ULevelSequence* LevelSequence, UBaseControlRig* ControlRig, FName ControlName, FFrameNumber Frame, ESequenceTimeUnit TimeUnit)
+FVector UControlRigSequencerEditorLibrary::GetLocalControlRigPosition(ULevelSequence* LevelSequence, UControlRig* ControlRig, FName ControlName, FFrameNumber Frame, ESequenceTimeUnit TimeUnit)
 {
 	FVector Value;
 	TWeakPtr<ISequencer> WeakSequencer = GetSequencerFromAsset();
@@ -1939,7 +1948,7 @@ FVector UControlRigSequencerEditorLibrary::GetLocalControlRigPosition(ULevelSequ
 	return Value;
 }
 
-TArray<FVector> UControlRigSequencerEditorLibrary::GetLocalControlRigPositions(ULevelSequence* LevelSequence, UBaseControlRig* ControlRig, FName ControlName, const TArray<FFrameNumber>& Frames, ESequenceTimeUnit TimeUnit)
+TArray<FVector> UControlRigSequencerEditorLibrary::GetLocalControlRigPositions(ULevelSequence* LevelSequence, UControlRig* ControlRig, FName ControlName, const TArray<FFrameNumber>& Frames, ESequenceTimeUnit TimeUnit)
 {
 	TArray<FVector> Values;
 	TWeakPtr<ISequencer> WeakSequencer = GetSequencerFromAsset();
@@ -1968,7 +1977,7 @@ TArray<FVector> UControlRigSequencerEditorLibrary::GetLocalControlRigPositions(U
 	return Values;
 }
 
-void UControlRigSequencerEditorLibrary::SetLocalControlRigPosition(ULevelSequence* LevelSequence, UBaseControlRig* ControlRig, FName ControlName, FFrameNumber Frame, FVector Value, ESequenceTimeUnit TimeUnit,bool bSetKey)
+void UControlRigSequencerEditorLibrary::SetLocalControlRigPosition(ULevelSequence* LevelSequence, UControlRig* ControlRig, FName ControlName, FFrameNumber Frame, FVector Value, ESequenceTimeUnit TimeUnit,bool bSetKey)
 {
 	if (LevelSequence == nullptr || ControlRig == nullptr)
 	{
@@ -1996,7 +2005,7 @@ void UControlRigSequencerEditorLibrary::SetLocalControlRigPosition(ULevelSequenc
 	}
 }
 
-void UControlRigSequencerEditorLibrary::SetLocalControlRigPositions(ULevelSequence* LevelSequence, UBaseControlRig* ControlRig, FName ControlName, const TArray<FFrameNumber>& Frames, 
+void UControlRigSequencerEditorLibrary::SetLocalControlRigPositions(ULevelSequence* LevelSequence, UControlRig* ControlRig, FName ControlName, const TArray<FFrameNumber>& Frames, 
 	const TArray<FVector> Values, ESequenceTimeUnit TimeUnit)
 {
 	if (LevelSequence == nullptr || ControlRig == nullptr || (Frames.Num() != Values.Num()))
@@ -2031,7 +2040,7 @@ void UControlRigSequencerEditorLibrary::SetLocalControlRigPositions(ULevelSequen
 }
 
 
-FRotator UControlRigSequencerEditorLibrary::GetLocalControlRigRotator(ULevelSequence* LevelSequence, UBaseControlRig* ControlRig, FName ControlName, FFrameNumber Frame, ESequenceTimeUnit TimeUnit)
+FRotator UControlRigSequencerEditorLibrary::GetLocalControlRigRotator(ULevelSequence* LevelSequence, UControlRig* ControlRig, FName ControlName, FFrameNumber Frame, ESequenceTimeUnit TimeUnit)
 {
 	FRotator Value;
 	TWeakPtr<ISequencer> WeakSequencer = GetSequencerFromAsset();
@@ -2055,7 +2064,7 @@ FRotator UControlRigSequencerEditorLibrary::GetLocalControlRigRotator(ULevelSequ
 	return Value;
 }
 
-TArray<FRotator> UControlRigSequencerEditorLibrary::GetLocalControlRigRotators(ULevelSequence* LevelSequence, UBaseControlRig* ControlRig, FName ControlName, const TArray<FFrameNumber>& Frames, ESequenceTimeUnit TimeUnit)
+TArray<FRotator> UControlRigSequencerEditorLibrary::GetLocalControlRigRotators(ULevelSequence* LevelSequence, UControlRig* ControlRig, FName ControlName, const TArray<FFrameNumber>& Frames, ESequenceTimeUnit TimeUnit)
 {
 	TArray<FRotator> Values;
 	TWeakPtr<ISequencer> WeakSequencer = GetSequencerFromAsset();
@@ -2084,7 +2093,7 @@ TArray<FRotator> UControlRigSequencerEditorLibrary::GetLocalControlRigRotators(U
 	return Values;
 }
 
-void UControlRigSequencerEditorLibrary::SetLocalControlRigRotator(ULevelSequence* LevelSequence, UBaseControlRig* ControlRig, FName ControlName, FFrameNumber Frame, FRotator Value, ESequenceTimeUnit TimeUnit, bool bSetKey)
+void UControlRigSequencerEditorLibrary::SetLocalControlRigRotator(ULevelSequence* LevelSequence, UControlRig* ControlRig, FName ControlName, FFrameNumber Frame, FRotator Value, ESequenceTimeUnit TimeUnit, bool bSetKey)
 {
 	if (LevelSequence == nullptr || ControlRig == nullptr)
 	{
@@ -2112,7 +2121,7 @@ void UControlRigSequencerEditorLibrary::SetLocalControlRigRotator(ULevelSequence
 	}
 }
 
-void UControlRigSequencerEditorLibrary::SetLocalControlRigRotators(ULevelSequence* LevelSequence, UBaseControlRig* ControlRig, FName ControlName, const TArray<FFrameNumber>& Frames,
+void UControlRigSequencerEditorLibrary::SetLocalControlRigRotators(ULevelSequence* LevelSequence, UControlRig* ControlRig, FName ControlName, const TArray<FFrameNumber>& Frames,
 	const TArray<FRotator> Values, ESequenceTimeUnit TimeUnit)
 {
 	if (LevelSequence == nullptr || ControlRig == nullptr || (Frames.Num() != Values.Num()))
@@ -2147,7 +2156,7 @@ void UControlRigSequencerEditorLibrary::SetLocalControlRigRotators(ULevelSequenc
 }
 
 
-FVector UControlRigSequencerEditorLibrary::GetLocalControlRigScale(ULevelSequence* LevelSequence, UBaseControlRig* ControlRig, FName ControlName, FFrameNumber Frame,  ESequenceTimeUnit TimeUnit)
+FVector UControlRigSequencerEditorLibrary::GetLocalControlRigScale(ULevelSequence* LevelSequence, UControlRig* ControlRig, FName ControlName, FFrameNumber Frame,  ESequenceTimeUnit TimeUnit)
 {
 	FVector Value;
 	TWeakPtr<ISequencer> WeakSequencer = GetSequencerFromAsset();
@@ -2171,7 +2180,7 @@ FVector UControlRigSequencerEditorLibrary::GetLocalControlRigScale(ULevelSequenc
 	return Value;
 }
 
-TArray<FVector>UControlRigSequencerEditorLibrary::GetLocalControlRigScales(ULevelSequence* LevelSequence, UBaseControlRig* ControlRig, FName ControlName, const TArray<FFrameNumber>& Frames, ESequenceTimeUnit TimeUnit)
+TArray<FVector>UControlRigSequencerEditorLibrary::GetLocalControlRigScales(ULevelSequence* LevelSequence, UControlRig* ControlRig, FName ControlName, const TArray<FFrameNumber>& Frames, ESequenceTimeUnit TimeUnit)
 {
 	TArray<FVector> Values;
 	TWeakPtr<ISequencer> WeakSequencer = GetSequencerFromAsset();
@@ -2200,7 +2209,7 @@ TArray<FVector>UControlRigSequencerEditorLibrary::GetLocalControlRigScales(ULeve
 	return Values;
 }
 
-void UControlRigSequencerEditorLibrary::SetLocalControlRigScale(ULevelSequence* LevelSequence, UBaseControlRig* ControlRig, FName ControlName, FFrameNumber Frame, FVector Value, ESequenceTimeUnit TimeUnit, bool bSetKey)
+void UControlRigSequencerEditorLibrary::SetLocalControlRigScale(ULevelSequence* LevelSequence, UControlRig* ControlRig, FName ControlName, FFrameNumber Frame, FVector Value, ESequenceTimeUnit TimeUnit, bool bSetKey)
 {
 	if (LevelSequence == nullptr || ControlRig == nullptr)
 	{
@@ -2228,7 +2237,7 @@ void UControlRigSequencerEditorLibrary::SetLocalControlRigScale(ULevelSequence*
 	}
 }
 
-void UControlRigSequencerEditorLibrary::SetLocalControlRigScales(ULevelSequence* LevelSequence, UBaseControlRig* ControlRig, FName ControlName, const TArray<FFrameNumber>& Frames, 
+void UControlRigSequencerEditorLibrary::SetLocalControlRigScales(ULevelSequence* LevelSequence, UControlRig* ControlRig, FName ControlName, const TArray<FFrameNumber>& Frames, 
 	const TArray<FVector> Values, ESequenceTimeUnit TimeUnit)
 {
 	if (LevelSequence == nullptr || ControlRig == nullptr || (Frames.Num() != Values.Num()))
@@ -2263,7 +2272,7 @@ void UControlRigSequencerEditorLibrary::SetLocalControlRigScales(ULevelSequence*
 }
 
 
-FEulerTransform UControlRigSequencerEditorLibrary::GetLocalControlRigEulerTransform(ULevelSequence* LevelSequence, UBaseControlRig* ControlRig, FName ControlName, FFrameNumber Frame, ESequenceTimeUnit TimeUnit)
+FEulerTransform UControlRigSequencerEditorLibrary::GetLocalControlRigEulerTransform(ULevelSequence* LevelSequence, UControlRig* ControlRig, FName ControlName, FFrameNumber Frame, ESequenceTimeUnit TimeUnit)
 {
 	FEulerTransform Value = FEulerTransform::Identity;
 	TWeakPtr<ISequencer> WeakSequencer = GetSequencerFromAsset();
@@ -2287,7 +2296,7 @@ FEulerTransform UControlRigSequencerEditorLibrary::GetLocalControlRigEulerTransf
 	return Value;
 }
 
-TArray<FEulerTransform> UControlRigSequencerEditorLibrary::GetLocalControlRigEulerTransforms(ULevelSequence* LevelSequence, UBaseControlRig* ControlRig, FName ControlName, const TArray<FFrameNumber>& Frames, ESequenceTimeUnit TimeUnit)
+TArray<FEulerTransform> UControlRigSequencerEditorLibrary::GetLocalControlRigEulerTransforms(ULevelSequence* LevelSequence, UControlRig* ControlRig, FName ControlName, const TArray<FFrameNumber>& Frames, ESequenceTimeUnit TimeUnit)
 {
 	TArray<FEulerTransform> Values;
 	TWeakPtr<ISequencer> WeakSequencer = GetSequencerFromAsset();
@@ -2316,7 +2325,7 @@ TArray<FEulerTransform> UControlRigSequencerEditorLibrary::GetLocalControlRigEul
 	return Values;
 }
 
-void UControlRigSequencerEditorLibrary::SetLocalControlRigEulerTransform(ULevelSequence* LevelSequence, UBaseControlRig* ControlRig, FName ControlName, FFrameNumber Frame, FEulerTransform Value,
+void UControlRigSequencerEditorLibrary::SetLocalControlRigEulerTransform(ULevelSequence* LevelSequence, UControlRig* ControlRig, FName ControlName, FFrameNumber Frame, FEulerTransform Value,
 	ESequenceTimeUnit TimeUnit, bool bSetKey)
 {
 	if (LevelSequence == nullptr || ControlRig == nullptr)
@@ -2350,7 +2359,7 @@ void UControlRigSequencerEditorLibrary::SetLocalControlRigEulerTransform(ULevelS
 	}
 }
 
-void UControlRigSequencerEditorLibrary::SetLocalControlRigEulerTransforms(ULevelSequence* LevelSequence, UBaseControlRig* ControlRig, FName ControlName, const TArray<FFrameNumber>& Frames,
+void UControlRigSequencerEditorLibrary::SetLocalControlRigEulerTransforms(ULevelSequence* LevelSequence, UControlRig* ControlRig, FName ControlName, const TArray<FFrameNumber>& Frames,
 	const TArray<FEulerTransform> Values, ESequenceTimeUnit TimeUnit)
 {
 	if (LevelSequence == nullptr || ControlRig == nullptr || (Frames.Num() != Values.Num()))
@@ -2391,7 +2400,7 @@ void UControlRigSequencerEditorLibrary::SetLocalControlRigEulerTransforms(ULevel
 }
 
 
-FTransformNoScale UControlRigSequencerEditorLibrary::GetLocalControlRigTransformNoScale(ULevelSequence* LevelSequence, UBaseControlRig* ControlRig, FName ControlName, FFrameNumber Frame, ESequenceTimeUnit TimeUnit)
+FTransformNoScale UControlRigSequencerEditorLibrary::GetLocalControlRigTransformNoScale(ULevelSequence* LevelSequence, UControlRig* ControlRig, FName ControlName, FFrameNumber Frame, ESequenceTimeUnit TimeUnit)
 {
 	FTransformNoScale Value = FTransformNoScale::Identity;
 	TWeakPtr<ISequencer> WeakSequencer = GetSequencerFromAsset();
@@ -2415,7 +2424,7 @@ FTransformNoScale UControlRigSequencerEditorLibrary::GetLocalControlRigTransform
 	return Value;
 }
 
-TArray<FTransformNoScale> UControlRigSequencerEditorLibrary::GetLocalControlRigTransformNoScales(ULevelSequence* LevelSequence, UBaseControlRig* ControlRig, FName ControlName, const TArray<FFrameNumber>& Frames, ESequenceTimeUnit TimeUnit)
+TArray<FTransformNoScale> UControlRigSequencerEditorLibrary::GetLocalControlRigTransformNoScales(ULevelSequence* LevelSequence, UControlRig* ControlRig, FName ControlName, const TArray<FFrameNumber>& Frames, ESequenceTimeUnit TimeUnit)
 {
 	TArray<FTransformNoScale> Values;
 	TWeakPtr<ISequencer> WeakSequencer = GetSequencerFromAsset();
@@ -2444,7 +2453,7 @@ TArray<FTransformNoScale> UControlRigSequencerEditorLibrary::GetLocalControlRigT
 	return Values;
 }
 
-void UControlRigSequencerEditorLibrary::SetLocalControlRigTransformNoScale(ULevelSequence* LevelSequence, UBaseControlRig* ControlRig, FName ControlName, FFrameNumber Frame, FTransformNoScale Value,
+void UControlRigSequencerEditorLibrary::SetLocalControlRigTransformNoScale(ULevelSequence* LevelSequence, UControlRig* ControlRig, FName ControlName, FFrameNumber Frame, FTransformNoScale Value,
 	ESequenceTimeUnit TimeUnit,bool bSetKey)
 {
 	if (LevelSequence == nullptr || ControlRig == nullptr)
@@ -2473,7 +2482,7 @@ void UControlRigSequencerEditorLibrary::SetLocalControlRigTransformNoScale(ULeve
 	}
 }
 
-void UControlRigSequencerEditorLibrary::SetLocalControlRigTransformNoScales(ULevelSequence* LevelSequence, UBaseControlRig* ControlRig, FName ControlName, const TArray<FFrameNumber>& Frames, 
+void UControlRigSequencerEditorLibrary::SetLocalControlRigTransformNoScales(ULevelSequence* LevelSequence, UControlRig* ControlRig, FName ControlName, const TArray<FFrameNumber>& Frames, 
 	const TArray<FTransformNoScale> Values, ESequenceTimeUnit TimeUnit)
 {
 	if (LevelSequence == nullptr || ControlRig == nullptr || (Frames.Num() != Values.Num()))
@@ -2508,7 +2517,7 @@ void UControlRigSequencerEditorLibrary::SetLocalControlRigTransformNoScales(ULev
 }
 
 
-FTransform UControlRigSequencerEditorLibrary::GetLocalControlRigTransform(ULevelSequence* LevelSequence, UBaseControlRig* ControlRig, FName ControlName, FFrameNumber Frame, ESequenceTimeUnit TimeUnit)
+FTransform UControlRigSequencerEditorLibrary::GetLocalControlRigTransform(ULevelSequence* LevelSequence, UControlRig* ControlRig, FName ControlName, FFrameNumber Frame, ESequenceTimeUnit TimeUnit)
 {
 	FTransform Value = FTransform::Identity;
 	TWeakPtr<ISequencer> WeakSequencer = GetSequencerFromAsset();
@@ -2532,7 +2541,7 @@ FTransform UControlRigSequencerEditorLibrary::GetLocalControlRigTransform(ULevel
 	return Value;
 }
 
-TArray<FTransform> UControlRigSequencerEditorLibrary::GetLocalControlRigTransforms(ULevelSequence* LevelSequence, UBaseControlRig* ControlRig, FName ControlName, const TArray<FFrameNumber>& Frames, ESequenceTimeUnit TimeUnit)
+TArray<FTransform> UControlRigSequencerEditorLibrary::GetLocalControlRigTransforms(ULevelSequence* LevelSequence, UControlRig* ControlRig, FName ControlName, const TArray<FFrameNumber>& Frames, ESequenceTimeUnit TimeUnit)
 {
 	TArray<FTransform> Values;
 	TWeakPtr<ISequencer> WeakSequencer = GetSequencerFromAsset();
@@ -2561,7 +2570,7 @@ TArray<FTransform> UControlRigSequencerEditorLibrary::GetLocalControlRigTransfor
 	return Values;
 }
 
-void UControlRigSequencerEditorLibrary::SetLocalControlRigTransform(ULevelSequence* LevelSequence, UBaseControlRig* ControlRig, FName ControlName, FFrameNumber Frame, FTransform Value,
+void UControlRigSequencerEditorLibrary::SetLocalControlRigTransform(ULevelSequence* LevelSequence, UControlRig* ControlRig, FName ControlName, FFrameNumber Frame, FTransform Value,
 	ESequenceTimeUnit TimeUnit,bool bSetKey)
 {
 	if (LevelSequence == nullptr || ControlRig == nullptr)
@@ -2590,7 +2599,7 @@ void UControlRigSequencerEditorLibrary::SetLocalControlRigTransform(ULevelSequen
 	}
 }
 
-void UControlRigSequencerEditorLibrary::SetLocalControlRigTransforms(ULevelSequence* LevelSequence, UBaseControlRig* ControlRig, FName ControlName, const TArray<FFrameNumber>& Frames, 
+void UControlRigSequencerEditorLibrary::SetLocalControlRigTransforms(ULevelSequence* LevelSequence, UControlRig* ControlRig, FName ControlName, const TArray<FFrameNumber>& Frames, 
 	const TArray<FTransform> Values, ESequenceTimeUnit TimeUnit)
 {
 	if (LevelSequence == nullptr || ControlRig == nullptr || (Frames.Num() != Values.Num()))
@@ -2728,7 +2737,7 @@ bool UControlRigSequencerEditorLibrary::CollapseControlRigAnimLayers(ULevelSeque
 	return bValid;
 }
 
-bool UControlRigSequencerEditorLibrary::SetControlRigSpace(ULevelSequence* LevelSequence, UBaseControlRig* ControlRig, FName ControlName, const FRigElementKey& InSpaceKey, FFrameNumber InTime, ESequenceTimeUnit TimeUnit)
+bool UControlRigSequencerEditorLibrary::SetControlRigSpace(ULevelSequence* LevelSequence, UControlRig* ControlRig, FName ControlName, const FRigElementKey& InSpaceKey, FFrameNumber InTime, ESequenceTimeUnit TimeUnit)
 {
 	TWeakPtr<ISequencer> WeakSequencer = GetSequencerFromAsset();
 	bool bValid = false;
@@ -2764,7 +2773,7 @@ bool UControlRigSequencerEditorLibrary::SetControlRigSpace(ULevelSequence* Level
 	return bValid;
 }
 
-bool UControlRigSequencerEditorLibrary::BakeControlRigSpace(ULevelSequence* InSequence, UBaseControlRig* InControlRig, const TArray<FName>& InControlNames, FRigSpacePickerBakeSettings InSettings, ESequenceTimeUnit TimeUnit)
+bool UControlRigSequencerEditorLibrary::BakeControlRigSpace(ULevelSequence* InSequence, UControlRig* InControlRig, const TArray<FName>& InControlNames, FRigSpacePickerBakeSettings InSettings, ESequenceTimeUnit TimeUnit)
 {
 	TWeakPtr<ISequencer> WeakSequencer = GetSequencerFromAsset();
 	bool bValid = false;
@@ -2799,7 +2808,7 @@ bool UControlRigSequencerEditorLibrary::BakeControlRigSpace(ULevelSequence* InSe
 	return bValid;
 }
 
-bool UControlRigSequencerEditorLibrary::DeleteControlRigSpace(ULevelSequence* LevelSequence, UBaseControlRig* ControlRig, FName ControlName, FFrameNumber InTime, ESequenceTimeUnit TimeUnit)
+bool UControlRigSequencerEditorLibrary::DeleteControlRigSpace(ULevelSequence* LevelSequence, UControlRig* ControlRig, FName ControlName, FFrameNumber InTime, ESequenceTimeUnit TimeUnit)
 {
 	TWeakPtr<ISequencer> WeakSequencer = GetSequencerFromAsset();
 	bool bValid = false;
@@ -2856,7 +2865,7 @@ bool UControlRigSequencerEditorLibrary::DeleteControlRigSpace(ULevelSequence* Le
 	return bValid;
 }
 
-bool UControlRigSequencerEditorLibrary::MoveControlRigSpace(ULevelSequence* LevelSequence, UBaseControlRig* ControlRig, FName ControlName, FFrameNumber InTime, FFrameNumber InNewTime, ESequenceTimeUnit TimeUnit)
+bool UControlRigSequencerEditorLibrary::MoveControlRigSpace(ULevelSequence* LevelSequence, UControlRig* ControlRig, FName ControlName, FFrameNumber InTime, FFrameNumber InNewTime, ESequenceTimeUnit TimeUnit)
 {
 	TWeakPtr<ISequencer> WeakSequencer = GetSequencerFromAsset();
 	bool bValid = false;
@@ -2915,7 +2924,7 @@ bool UControlRigSequencerEditorLibrary::MoveControlRigSpace(ULevelSequence* Leve
 	return bValid;
 }
 
-bool UControlRigSequencerEditorLibrary::RenameControlRigControlChannels(ULevelSequence* LevelSequence, UBaseControlRig* ControlRig, const TArray<FName>& InOldControlNames, const TArray<FName>& InNewControlNames)
+bool UControlRigSequencerEditorLibrary::RenameControlRigControlChannels(ULevelSequence* LevelSequence, UControlRig* ControlRig, const TArray<FName>& InOldControlNames, const TArray<FName>& InNewControlNames)
 {
 	if (LevelSequence == nullptr || ControlRig == nullptr)
 	{
@@ -2984,7 +2993,7 @@ bool UControlRigSequencerEditorLibrary::GetControlsMask(UMovieSceneSection* InSe
 		return false;
 	}
 	
-	UBaseControlRig* ControlRig = ParameterSection->GetControlRig();
+	UControlRig* ControlRig = Cast<UControlRig>(ParameterSection->GetControlRig());
 	if (!ControlRig)
 	{
 		FFrame::KismetExecutionMessage(TEXT("Section does not have a control rig"), ELogVerbosity::Error);
@@ -3016,7 +3025,7 @@ void UControlRigSequencerEditorLibrary::SetControlsMask(UMovieSceneSection* InSe
 		return;
 	}
 
-	UBaseControlRig* ControlRig = ParameterSection->GetControlRig();
+	UControlRig* ControlRig = Cast<UControlRig>(ParameterSection->GetControlRig());
 	if (!ControlRig)
 	{
 		FFrame::KismetExecutionMessage(TEXT("Section does not have a control rig"), ELogVerbosity::Error);
@@ -3064,29 +3073,29 @@ void UControlRigSequencerEditorLibrary::HideAllControls(UMovieSceneSection* InSe
 	ParameterSection->FillControlsMask(false);
 }
 
-bool UControlRigSequencerEditorLibrary::IsFKControlRig(UBaseControlRig* InControlRig)
+bool UControlRigSequencerEditorLibrary::IsFKControlRig(UControlRig* InControlRig)
 {
 	return (InControlRig && InControlRig->IsA<UFKControlRig>());
 }
 
-bool UControlRigSequencerEditorLibrary::IsAdditiveControlRig(UBaseControlRig* InControlRig)
+bool UControlRigSequencerEditorLibrary::IsAdditiveControlRig(UControlRig* InControlRig)
 {
 	return (InControlRig && InControlRig->IsAdditive());
 }
 
-EControlRigFKRigExecuteMode UControlRigSequencerEditorLibrary::GetFKControlRigApplyMode(UBaseControlRig* InControlRig)
+EControlRigFKRigExecuteMode UControlRigSequencerEditorLibrary::GetFKControlRigApplyMode(UControlRig* InControlRig)
 {
 	EControlRigFKRigExecuteMode ApplyMode = EControlRigFKRigExecuteMode::Direct;
-	if (UFKControlRig* FKRig = Cast<UFKControlRig>(InControlRig))
+	if (UFKControlRig* FKRig = Cast<UFKControlRig>(Cast<UBaseControlRig>(InControlRig)))
 	{
 		ApplyMode = FKRig->GetApplyMode();
 	}
 	return ApplyMode;
 }
 
-bool UControlRigSequencerEditorLibrary::SetControlRigApplyMode(UBaseControlRig* InControlRig, EControlRigFKRigExecuteMode InApplyMode)
+bool UControlRigSequencerEditorLibrary::SetControlRigApplyMode(UControlRig* InControlRig, EControlRigFKRigExecuteMode InApplyMode)
 {
-	if (UFKControlRig* FKRig = Cast<UFKControlRig>(InControlRig))
+	if (UFKControlRig* FKRig = Cast<UFKControlRig>(Cast<UBaseControlRig>(InControlRig)))
 	{
 		FKRig->SetApplyMode(InApplyMode);
 		return true;
diff --git a/Engine/Plugins/Animation/ControlRig/Source/ControlRigEditor/Private/EditMode/ControlRigEditMode.cpp b/Engine/Plugins/Animation/ControlRig/Source/ControlRigEditor/Private/EditMode/ControlRigEditMode.cpp
--- a/Engine/Plugins/Animation/ControlRig/Source/ControlRigEditor/Private/EditMode/ControlRigEditMode.cpp
+++ b/Engine/Plugins/Animation/ControlRig/Source/ControlRigEditor/Private/EditMode/ControlRigEditMode.cpp
@@ -8,7 +8,7 @@
 #include "Toolkits/ToolkitManager.h"
 #include "EditMode/SControlRigEditModeTools.h"
 #include "Algo/Transform.h"
-#include "BaseControlRig.h"
+#include "ControlRig.h"
 #include "HitProxies.h"
 #include "EditMode/ControlRigEditModeSettings.h"
 #include "ISequencer.h"
diff --git a/Engine/Plugins/Animation/ControlRig/Source/ControlRigEditor/Private/EditMode/SControlRigBaseListWidget.cpp b/Engine/Plugins/Animation/ControlRig/Source/ControlRigEditor/Private/EditMode/SControlRigBaseListWidget.cpp
--- a/Engine/Plugins/Animation/ControlRig/Source/ControlRigEditor/Private/EditMode/SControlRigBaseListWidget.cpp
+++ b/Engine/Plugins/Animation/ControlRig/Source/ControlRigEditor/Private/EditMode/SControlRigBaseListWidget.cpp
@@ -17,7 +17,7 @@
 
 #include "ScopedTransaction.h"
 
-#include "BaseControlRig.h"
+#include "ControlRig.h"
 #include "UnrealEdGlobals.h"
 #include "EditMode/ControlRigEditMode.h"
 #include "Tools/ControlRigPose.h"
diff --git a/Engine/Plugins/Animation/ControlRig/Source/ControlRigEditor/Public/ControlRigSequencerEditorLibrary.h b/Engine/Plugins/Animation/ControlRig/Source/ControlRigEditor/Public/ControlRigSequencerEditorLibrary.h
--- a/Engine/Plugins/Animation/ControlRig/Source/ControlRigEditor/Public/ControlRigSequencerEditorLibrary.h
+++ b/Engine/Plugins/Animation/ControlRig/Source/ControlRigEditor/Public/ControlRigSequencerEditorLibrary.h
@@ -24,7 +24,7 @@ class UTickableConstraint;
 class UTickableTransformConstraint;
 class UTransformableHandle;
 struct FBakingAnimationKeySettings;
-class UBaseControlRig;
+class UControlRig;
 
 USTRUCT(BlueprintType)
 struct FControlRigSequencerBindingProxy
@@ -36,7 +36,7 @@ struct FControlRigSequencerBindingProxy
 		, Track(nullptr)
 	{}
 
-	FControlRigSequencerBindingProxy(const FMovieSceneBindingProxy& InProxy, UBaseControlRig* InControlRig, UMovieSceneControlRigParameterTrack* InTrack)
+	FControlRigSequencerBindingProxy(const FMovieSceneBindingProxy& InProxy, UControlRig* InControlRig, UMovieSceneControlRigParameterTrack* InTrack)
 		: Proxy(InProxy)
 		, ControlRig(InControlRig)
 		, Track(InTrack)
@@ -46,7 +46,7 @@ struct FControlRigSequencerBindingProxy
 	FMovieSceneBindingProxy Proxy;
 
 	UPROPERTY(BlueprintReadOnly, Category = ControlRig)
-	TObjectPtr<UBaseControlRig> ControlRig;
+	TObjectPtr<UControlRig> ControlRig;
 
 	UPROPERTY(BlueprintReadOnly, Category = ControlRig)
 	TObjectPtr<UMovieSceneControlRigParameterTrack> Track;
@@ -80,7 +80,7 @@ public:
 	* @return returns list of visible Control Rigs
 	*/
 	UFUNCTION(BlueprintCallable, Category = "Editor Scripting | Sequencer Tools | Control Rig")
-	static TArray<UBaseControlRig*> GetVisibleControlRigs();
+	static TArray<UControlRig*> GetVisibleControlRigs();
 
 	/**
 	* Get all of the control rigs and their bindings in the level sequence
@@ -251,7 +251,7 @@ public:
 	* @param TweenValue The tween value to use, range from -1(blend to previous) to 1(blend to next)
 	*/
 	UFUNCTION(BlueprintCallable, Category = "Editor Scripting | Sequencer Tools | Control Rig")
-	static bool TweenControlRig(ULevelSequence* LevelSequence, UBaseControlRig* ControlRig, float TweenValue);
+	static bool TweenControlRig(ULevelSequence* LevelSequence, UControlRig* ControlRig, float TweenValue);
 	
 	/**
 	* Peform specified blend operation based upon selected keys in the curve editor or selected control rig controls
@@ -336,7 +336,7 @@ public:
 	* @return Returns World Transform
 	*/
 	UFUNCTION(BlueprintCallable, Category = "Editor Scripting | Sequencer Tools | Control Rig")
-	static FTransform GetControlRigWorldTransform(ULevelSequence* LevelSequence, UBaseControlRig* ControlRig, FName ControlName, FFrameNumber Frame,
+	static FTransform GetControlRigWorldTransform(ULevelSequence* LevelSequence, UControlRig* ControlRig, FName ControlName, FFrameNumber Frame,
 		ESequenceTimeUnit TimeUnit = ESequenceTimeUnit::DisplayRate);
 
 	/**
@@ -349,7 +349,7 @@ public:
 	* @return Returns World Transforms
 	*/
 	UFUNCTION(BlueprintCallable, Category = "Editor Scripting | Sequencer Tools | Control Rig")
-	static TArray<FTransform> GetControlRigWorldTransforms(ULevelSequence* LevelSequence, UBaseControlRig* ControlRig, FName ControlName, const TArray<FFrameNumber>& Frames,
+	static TArray<FTransform> GetControlRigWorldTransforms(ULevelSequence* LevelSequence, UControlRig* ControlRig, FName ControlName, const TArray<FFrameNumber>& Frames,
 		ESequenceTimeUnit TimeUnit = ESequenceTimeUnit::DisplayRate);
 
 	/**
@@ -363,7 +363,7 @@ public:
 	* @param bSetKey Whether or not to set a key.
 	*/
 	UFUNCTION(BlueprintCallable, Category = "Editor Scripting | Sequencer Tools | Control Rig")
-	static void SetControlRigWorldTransform(ULevelSequence* LevelSequence, UBaseControlRig* ControlRig, FName ControlName, FFrameNumber Frame, const FTransform& WorldTransform,
+	static void SetControlRigWorldTransform(ULevelSequence* LevelSequence, UControlRig* ControlRig, FName ControlName, FFrameNumber Frame, const FTransform& WorldTransform,
 		ESequenceTimeUnit TimeUnit = ESequenceTimeUnit::DisplayRate, bool bSetKey = true);
 
 	/**
@@ -376,7 +376,7 @@ public:
 	* @param TimeUnit Unit for frame values, either in display rate or tick resolution
 	*/
 	UFUNCTION(BlueprintCallable, Category = "Editor Scripting | Sequencer Tools | Control Rig")
-	static void SetControlRigWorldTransforms(ULevelSequence* LevelSequence, UBaseControlRig* ControlRig, FName ControlName, const TArray<FFrameNumber>& Frames, 
+	static void SetControlRigWorldTransforms(ULevelSequence* LevelSequence, UControlRig* ControlRig, FName ControlName, const TArray<FFrameNumber>& Frames, 
 		 const TArray<FTransform>& WorldTransforms, ESequenceTimeUnit TimeUnit = ESequenceTimeUnit::DisplayRate);
 	
 	/**
@@ -389,7 +389,7 @@ public:
 	* @return Returns Value at that time
 	*/
 	UFUNCTION(BlueprintCallable, Category = "Editor Scripting | Sequencer Tools | Control Rig")
-	static float GetLocalControlRigFloat(ULevelSequence* LevelSequence, UBaseControlRig* ControlRig, FName ControlName, FFrameNumber Frame,
+	static float GetLocalControlRigFloat(ULevelSequence* LevelSequence, UControlRig* ControlRig, FName ControlName, FFrameNumber Frame,
 		ESequenceTimeUnit TimeUnit = ESequenceTimeUnit::DisplayRate);
 
 	/**
@@ -402,7 +402,7 @@ public:
 	* @return Returns Values at those times
 	*/
 	UFUNCTION(BlueprintCallable, Category = "Editor Scripting | Sequencer Tools | Control Rig")
-	static TArray<float> GetLocalControlRigFloats(ULevelSequence* LevelSequence, UBaseControlRig* ControlRig, FName ControlName, const TArray<FFrameNumber>& Frames,
+	static TArray<float> GetLocalControlRigFloats(ULevelSequence* LevelSequence, UControlRig* ControlRig, FName ControlName, const TArray<FFrameNumber>& Frames,
 		ESequenceTimeUnit TimeUnit = ESequenceTimeUnit::DisplayRate);
 
 	/**
@@ -416,7 +416,7 @@ public:
 	* @param bSetKey If True set a key, if not just set the value
 	*/
 	UFUNCTION(BlueprintCallable, Category = "Editor Scripting | Sequencer Tools | Control Rig")
-	static void SetLocalControlRigFloat(ULevelSequence* LevelSequence, UBaseControlRig* ControlRig, FName ControlName, FFrameNumber Frame, float Value,
+	static void SetLocalControlRigFloat(ULevelSequence* LevelSequence, UControlRig* ControlRig, FName ControlName, FFrameNumber Frame, float Value,
 		ESequenceTimeUnit TimeUnit = ESequenceTimeUnit::DisplayRate,bool bSetKey = true);
 
 	/**
@@ -429,7 +429,7 @@ public:
 	* @param TimeUnit Unit for frame values, either in display rate or tick resolution
 	*/
 	UFUNCTION(BlueprintCallable, Category = "Editor Scripting | Sequencer Tools | Control Rig")
-	static void SetLocalControlRigFloats(ULevelSequence* LevelSequence, UBaseControlRig* ControlRig, FName ControlName, const TArray<FFrameNumber>& Frames, const TArray<float> Values,
+	static void SetLocalControlRigFloats(ULevelSequence* LevelSequence, UControlRig* ControlRig, FName ControlName, const TArray<FFrameNumber>& Frames, const TArray<float> Values,
 		ESequenceTimeUnit TimeUnit = ESequenceTimeUnit::DisplayRate);
 
 	/**
@@ -442,7 +442,7 @@ public:
 	* @return Returns Value at that time
 	*/
 	UFUNCTION(BlueprintCallable, Category = "Editor Scripting | Sequencer Tools | Control Rig")
-	static bool GetLocalControlRigBool(ULevelSequence* LevelSequence, UBaseControlRig* ControlRig, FName ControlName, FFrameNumber Frame,
+	static bool GetLocalControlRigBool(ULevelSequence* LevelSequence, UControlRig* ControlRig, FName ControlName, FFrameNumber Frame,
 		ESequenceTimeUnit TimeUnit = ESequenceTimeUnit::DisplayRate);
 
 	/**
@@ -455,7 +455,7 @@ public:
 	* @return Returns Values at those times
 	*/
 	UFUNCTION(BlueprintCallable, Category = "Editor Scripting | Sequencer Tools | Control Rig")
-	static TArray<bool> GetLocalControlRigBools(ULevelSequence* LevelSequence, UBaseControlRig* ControlRig, FName ControlName, const TArray<FFrameNumber>& Frames,
+	static TArray<bool> GetLocalControlRigBools(ULevelSequence* LevelSequence, UControlRig* ControlRig, FName ControlName, const TArray<FFrameNumber>& Frames,
 		ESequenceTimeUnit TimeUnit = ESequenceTimeUnit::DisplayRate);
 
 	/**
@@ -469,7 +469,7 @@ public:
 	* @param bSetKey If True set a key, if not just set the value
 	*/
 	UFUNCTION(BlueprintCallable, Category = "Editor Scripting | Sequencer Tools | Control Rig")
-	static void SetLocalControlRigBool(ULevelSequence* LevelSequence, UBaseControlRig* ControlRig, FName ControlName, FFrameNumber Frame, bool Value,
+	static void SetLocalControlRigBool(ULevelSequence* LevelSequence, UControlRig* ControlRig, FName ControlName, FFrameNumber Frame, bool Value,
 		ESequenceTimeUnit TimeUnit = ESequenceTimeUnit::DisplayRate, bool bSetKey = true);
 
 	/**
@@ -482,7 +482,7 @@ public:
 	* @param TimeUnit Unit for frame values, either in display rate or tick resolution
 */
 	UFUNCTION(BlueprintCallable, Category = "Editor Scripting | Sequencer Tools | Control Rig")
-	static void SetLocalControlRigBools(ULevelSequence* LevelSequence, UBaseControlRig* ControlRig, FName ControlName, const TArray<FFrameNumber>& Frames,
+	static void SetLocalControlRigBools(ULevelSequence* LevelSequence, UControlRig* ControlRig, FName ControlName, const TArray<FFrameNumber>& Frames,
 		const TArray<bool> Values, ESequenceTimeUnit TimeUnit = ESequenceTimeUnit::DisplayRate);
 
 	/**
@@ -495,7 +495,7 @@ public:
 	* @return Returns Value at that time
 	*/
 	UFUNCTION(BlueprintCallable, Category = "Editor Scripting | Sequencer Tools | Control Rig")
-	static int32 GetLocalControlRigInt(ULevelSequence* LevelSequence, UBaseControlRig* ControlRig, FName ControlName, FFrameNumber Frame, 
+	static int32 GetLocalControlRigInt(ULevelSequence* LevelSequence, UControlRig* ControlRig, FName ControlName, FFrameNumber Frame, 
 		ESequenceTimeUnit TimeUnit = ESequenceTimeUnit::DisplayRate);
 
 	/**
@@ -508,7 +508,7 @@ public:
 	* @return Returns Values at those times
 	*/
 	UFUNCTION(BlueprintCallable, Category = "Editor Scripting | Sequencer Tools | Control Rig")
-	static TArray<int32> GetLocalControlRigInts(ULevelSequence* LevelSequence, UBaseControlRig* ControlRig, FName ControlName, const TArray<FFrameNumber>& Frames,
+	static TArray<int32> GetLocalControlRigInts(ULevelSequence* LevelSequence, UControlRig* ControlRig, FName ControlName, const TArray<FFrameNumber>& Frames,
 		ESequenceTimeUnit TimeUnit = ESequenceTimeUnit::DisplayRate);
 
 	/**
@@ -522,7 +522,7 @@ public:
 	* @param bSetKey If True set a key, if not just set the value
 	*/
 	UFUNCTION(BlueprintCallable, Category = "Editor Scripting | Sequencer Tools | Control Rig")
-	static void SetLocalControlRigInt(ULevelSequence* LevelSequence, UBaseControlRig* ControlRig, FName ControlName, FFrameNumber Frame, int32 Value,
+	static void SetLocalControlRigInt(ULevelSequence* LevelSequence, UControlRig* ControlRig, FName ControlName, FFrameNumber Frame, int32 Value,
 		ESequenceTimeUnit TimeUnit = ESequenceTimeUnit::DisplayRate,bool bSetKey = true);
 
 
@@ -536,7 +536,7 @@ public:
 	* @param TimeUnit Unit for frame values, either in display rate or tick resolution
 	*/
 	UFUNCTION(BlueprintCallable, Category = "Editor Scripting | Sequencer Tools | Control Rig")
-	static void SetLocalControlRigInts(ULevelSequence* LevelSequence, UBaseControlRig* ControlRig, FName ControlName, const TArray<FFrameNumber>& Frames, const TArray<int32> Values,
+	static void SetLocalControlRigInts(ULevelSequence* LevelSequence, UControlRig* ControlRig, FName ControlName, const TArray<FFrameNumber>& Frames, const TArray<int32> Values,
 		ESequenceTimeUnit TimeUnit = ESequenceTimeUnit::DisplayRate);
 
 	/**
@@ -549,7 +549,7 @@ public:
 	* @return Returns Value at that time
 	*/
 	UFUNCTION(BlueprintCallable, Category = "Editor Scripting | Sequencer Tools | Control Rig")
-	static FVector2D GetLocalControlRigVector2D(ULevelSequence* LevelSequence, UBaseControlRig* ControlRig, FName ControlName, FFrameNumber Frame,
+	static FVector2D GetLocalControlRigVector2D(ULevelSequence* LevelSequence, UControlRig* ControlRig, FName ControlName, FFrameNumber Frame,
 		ESequenceTimeUnit TimeUnit = ESequenceTimeUnit::DisplayRate);
 
 	/**
@@ -562,7 +562,7 @@ public:
 	* @return Returns Values at those times
 	*/
 	UFUNCTION(BlueprintCallable, Category = "Editor Scripting | Sequencer Tools | Control Rig")
-	static TArray<FVector2D> GetLocalControlRigVector2Ds(ULevelSequence* LevelSequence, UBaseControlRig* ControlRig, FName ControlName, const TArray<FFrameNumber>& Frames,
+	static TArray<FVector2D> GetLocalControlRigVector2Ds(ULevelSequence* LevelSequence, UControlRig* ControlRig, FName ControlName, const TArray<FFrameNumber>& Frames,
 		ESequenceTimeUnit TimeUnit = ESequenceTimeUnit::DisplayRate);
 
 	/**
@@ -576,7 +576,7 @@ public:
 	* @param bSetKey If True set a key, if not just set the value
 	*/
 	UFUNCTION(BlueprintCallable, Category = "Editor Scripting | Sequencer Tools | Control Rig")
-	static void SetLocalControlRigVector2D(ULevelSequence* LevelSequence, UBaseControlRig* ControlRig, FName ControlName, FFrameNumber Frame, FVector2D Value,
+	static void SetLocalControlRigVector2D(ULevelSequence* LevelSequence, UControlRig* ControlRig, FName ControlName, FFrameNumber Frame, FVector2D Value,
 		ESequenceTimeUnit TimeUnit = ESequenceTimeUnit::DisplayRate,bool bSetKey = true);
 
 
@@ -590,7 +590,7 @@ public:
 	* @param TimeUnit Unit for frame values, either in display rate or tick resolution
 	*/
 	UFUNCTION(BlueprintCallable, Category = "Editor Scripting | Sequencer Tools | Control Rig")
-	static void SetLocalControlRigVector2Ds(ULevelSequence* LevelSequence, UBaseControlRig* ControlRig, FName ControlName, const TArray<FFrameNumber>& Frames, const TArray<FVector2D> Values,
+	static void SetLocalControlRigVector2Ds(ULevelSequence* LevelSequence, UControlRig* ControlRig, FName ControlName, const TArray<FFrameNumber>& Frames, const TArray<FVector2D> Values,
 		ESequenceTimeUnit TimeUnit = ESequenceTimeUnit::DisplayRate);
 
 	/**
@@ -603,7 +603,7 @@ public:
 	* @return Returns Value at that time
 	*/
 	UFUNCTION(BlueprintCallable, Category = "Editor Scripting | Sequencer Tools | Control Rig")
-	static FVector GetLocalControlRigPosition(ULevelSequence* LevelSequence, UBaseControlRig* ControlRig, FName ControlName, FFrameNumber Frame,
+	static FVector GetLocalControlRigPosition(ULevelSequence* LevelSequence, UControlRig* ControlRig, FName ControlName, FFrameNumber Frame,
 		ESequenceTimeUnit TimeUnit = ESequenceTimeUnit::DisplayRate);
 
 	/**
@@ -616,7 +616,7 @@ public:
 	* @return Returns Values at those times
 	*/
 	UFUNCTION(BlueprintCallable, Category = "Editor Scripting | Sequencer Tools | Control Rig")
-	static TArray<FVector> GetLocalControlRigPositions(ULevelSequence* LevelSequence, UBaseControlRig* ControlRig, FName ControlName, const TArray<FFrameNumber>& Frames,
+	static TArray<FVector> GetLocalControlRigPositions(ULevelSequence* LevelSequence, UControlRig* ControlRig, FName ControlName, const TArray<FFrameNumber>& Frames,
 		ESequenceTimeUnit TimeUnit = ESequenceTimeUnit::DisplayRate);
 
 	/**
@@ -630,7 +630,7 @@ public:
 	* @param bSetKey If True set a key, if not just set the value
 	*/
 	UFUNCTION(BlueprintCallable, Category = "Editor Scripting | Sequencer Tools | Control Rig")
-	static void SetLocalControlRigPosition(ULevelSequence* LevelSequence, UBaseControlRig* ControlRig, FName ControlName, FFrameNumber Frame, FVector Value,
+	static void SetLocalControlRigPosition(ULevelSequence* LevelSequence, UControlRig* ControlRig, FName ControlName, FFrameNumber Frame, FVector Value,
 		ESequenceTimeUnit TimeUnit = ESequenceTimeUnit::DisplayRate, bool bSetKey = true);
 
 	/**
@@ -643,7 +643,7 @@ public:
 	* @param TimeUnit Unit for frame values, either in display rate or tick resolution
 	*/
 	UFUNCTION(BlueprintCallable, Category = "Editor Scripting | Sequencer Tools | Control Rig")
-	static void SetLocalControlRigPositions(ULevelSequence* LevelSequence, UBaseControlRig* ControlRig, FName ControlName, const TArray<FFrameNumber>& Frames, const TArray<FVector> Values,
+	static void SetLocalControlRigPositions(ULevelSequence* LevelSequence, UControlRig* ControlRig, FName ControlName, const TArray<FFrameNumber>& Frames, const TArray<FVector> Values,
 		ESequenceTimeUnit TimeUnit = ESequenceTimeUnit::DisplayRate);
 
 	/**
@@ -656,7 +656,7 @@ public:
 	* @return Returns Value at that time
 	*/
 	UFUNCTION(BlueprintCallable, Category = "Editor Scripting | Sequencer Tools | Control Rig")
-	static FRotator GetLocalControlRigRotator(ULevelSequence* LevelSequence, UBaseControlRig* ControlRig, FName ControlName, FFrameNumber Frame,
+	static FRotator GetLocalControlRigRotator(ULevelSequence* LevelSequence, UControlRig* ControlRig, FName ControlName, FFrameNumber Frame,
 		ESequenceTimeUnit TimeUnit = ESequenceTimeUnit::DisplayRate);
 
 	/**
@@ -669,7 +669,7 @@ public:
 	* @return Returns Values at those times
 	*/
 	UFUNCTION(BlueprintCallable, Category = "Editor Scripting | Sequencer Tools | Control Rig")
-	static TArray<FRotator> GetLocalControlRigRotators(ULevelSequence* LevelSequence, UBaseControlRig* ControlRig, FName ControlName, const TArray<FFrameNumber>& Frames,
+	static TArray<FRotator> GetLocalControlRigRotators(ULevelSequence* LevelSequence, UControlRig* ControlRig, FName ControlName, const TArray<FFrameNumber>& Frames,
 		ESequenceTimeUnit TimeUnit = ESequenceTimeUnit::DisplayRate);
 
 	/**
@@ -683,7 +683,7 @@ public:
 	* @param bSetKey If True set a key, if not just set the value
 	*/
 	UFUNCTION(BlueprintCallable, Category = "Editor Scripting | Sequencer Tools | Control Rig")
-	static void SetLocalControlRigRotator(ULevelSequence* LevelSequence, UBaseControlRig* ControlRig, FName ControlName, FFrameNumber Frame, FRotator Value, 
+	static void SetLocalControlRigRotator(ULevelSequence* LevelSequence, UControlRig* ControlRig, FName ControlName, FFrameNumber Frame, FRotator Value, 
 		ESequenceTimeUnit TimeUnit = ESequenceTimeUnit::DisplayRate, bool bSetKey = true);
 
 	/**
@@ -696,7 +696,7 @@ public:
 	* @param TimeUnit Unit for frame values, either in display rate or tick resolution
 	*/
 	UFUNCTION(BlueprintCallable, Category = "Editor Scripting | Sequencer Tools | Control Rig")
-	static void SetLocalControlRigRotators(ULevelSequence* LevelSequence, UBaseControlRig* ControlRig, FName ControlName, const TArray<FFrameNumber>& Frames, const TArray<FRotator> Values,
+	static void SetLocalControlRigRotators(ULevelSequence* LevelSequence, UControlRig* ControlRig, FName ControlName, const TArray<FFrameNumber>& Frames, const TArray<FRotator> Values,
 		ESequenceTimeUnit TimeUnit = ESequenceTimeUnit::DisplayRate);
 
 	/**
@@ -709,7 +709,7 @@ public:
 	* @return Returns Value at that time
 	*/
 	UFUNCTION(BlueprintCallable, Category = "Editor Scripting | Sequencer Tools | Control Rig")
-	static FVector GetLocalControlRigScale(ULevelSequence* LevelSequence, UBaseControlRig* ControlRig, FName ControlName, FFrameNumber Frame,
+	static FVector GetLocalControlRigScale(ULevelSequence* LevelSequence, UControlRig* ControlRig, FName ControlName, FFrameNumber Frame,
 		ESequenceTimeUnit TimeUnit = ESequenceTimeUnit::DisplayRate);
 
 	/**
@@ -722,7 +722,7 @@ public:
 	* @return Returns Values at those times
 	*/
 	UFUNCTION(BlueprintCallable, Category = "Editor Scripting | Sequencer Tools | Control Rig")
-	static TArray<FVector> GetLocalControlRigScales(ULevelSequence* LevelSequence, UBaseControlRig* ControlRig, FName ControlName, const TArray<FFrameNumber>& Frames,
+	static TArray<FVector> GetLocalControlRigScales(ULevelSequence* LevelSequence, UControlRig* ControlRig, FName ControlName, const TArray<FFrameNumber>& Frames,
 		ESequenceTimeUnit TimeUnit = ESequenceTimeUnit::DisplayRate);
 
 	/**
@@ -736,7 +736,7 @@ public:
 	* @param bSetKey If True set a key, if not just set the value
 */
 	UFUNCTION(BlueprintCallable, Category = "Editor Scripting | Sequencer Tools | Control Rig")
-	static void SetLocalControlRigScale(ULevelSequence* LevelSequence, UBaseControlRig* ControlRig, FName ControlName, FFrameNumber Frame, FVector Value, 
+	static void SetLocalControlRigScale(ULevelSequence* LevelSequence, UControlRig* ControlRig, FName ControlName, FFrameNumber Frame, FVector Value, 
 		ESequenceTimeUnit TimeUnit = ESequenceTimeUnit::DisplayRate, bool bSetKey = true);
 
 	/**
@@ -749,7 +749,7 @@ public:
 	* @param TimeUnit Unit for frame values, either in display rate or tick resolution
 	*/
 	UFUNCTION(BlueprintCallable, Category = "Editor Scripting | Sequencer Tools | Control Rig")
-	static void SetLocalControlRigScales(ULevelSequence* LevelSequence, UBaseControlRig* ControlRig, FName ControlName, const TArray<FFrameNumber>& Frames, const TArray<FVector> Values,
+	static void SetLocalControlRigScales(ULevelSequence* LevelSequence, UControlRig* ControlRig, FName ControlName, const TArray<FFrameNumber>& Frames, const TArray<FVector> Values,
 		ESequenceTimeUnit TimeUnit = ESequenceTimeUnit::DisplayRate);
 
 
@@ -763,7 +763,7 @@ public:
 	* @return Returns Value at that time
 	*/
 	UFUNCTION(BlueprintCallable, Category = "Editor Scripting | Sequencer Tools | Control Rig")
-	static FEulerTransform GetLocalControlRigEulerTransform(ULevelSequence* LevelSequence, UBaseControlRig* ControlRig, FName ControlName, FFrameNumber Frame,
+	static FEulerTransform GetLocalControlRigEulerTransform(ULevelSequence* LevelSequence, UControlRig* ControlRig, FName ControlName, FFrameNumber Frame,
 		ESequenceTimeUnit TimeUnit = ESequenceTimeUnit::DisplayRate);
 
 	/**
@@ -776,7 +776,7 @@ public:
 	* @return Returns Values at those times
 	*/
 	UFUNCTION(BlueprintCallable, Category = "Editor Scripting | Sequencer Tools | Control Rig")
-	static TArray<FEulerTransform> GetLocalControlRigEulerTransforms(ULevelSequence* LevelSequence, UBaseControlRig* ControlRig, FName ControlName, const TArray<FFrameNumber>& Frames,
+	static TArray<FEulerTransform> GetLocalControlRigEulerTransforms(ULevelSequence* LevelSequence, UControlRig* ControlRig, FName ControlName, const TArray<FFrameNumber>& Frames,
 		ESequenceTimeUnit TimeUnit = ESequenceTimeUnit::DisplayRate);
 
 	/**
@@ -790,7 +790,7 @@ public:
 	* @param bSetKey If True set a key, if not just set the value
 	*/
 	UFUNCTION(BlueprintCallable, Category = "Editor Scripting | Sequencer Tools | Control Rig")
-	static void SetLocalControlRigEulerTransform(ULevelSequence* LevelSequence, UBaseControlRig* ControlRig, FName ControlName, FFrameNumber Frame, FEulerTransform Value,
+	static void SetLocalControlRigEulerTransform(ULevelSequence* LevelSequence, UControlRig* ControlRig, FName ControlName, FFrameNumber Frame, FEulerTransform Value,
 		ESequenceTimeUnit TimeUnit = ESequenceTimeUnit::DisplayRate, bool bSetKey = true);
 
 	/**
@@ -803,7 +803,7 @@ public:
 	* @param TimeUnit Unit for frame values, either in display rate or tick resolution
 	*/
 	UFUNCTION(BlueprintCallable, Category = "Editor Scripting | Sequencer Tools | Control Rig")
-	static void SetLocalControlRigEulerTransforms(ULevelSequence* LevelSequence, UBaseControlRig* ControlRig, FName ControlName, const TArray<FFrameNumber>& Frames, const TArray<FEulerTransform> Values,
+	static void SetLocalControlRigEulerTransforms(ULevelSequence* LevelSequence, UControlRig* ControlRig, FName ControlName, const TArray<FFrameNumber>& Frames, const TArray<FEulerTransform> Values,
 		ESequenceTimeUnit TimeUnit = ESequenceTimeUnit::DisplayRate);
 
 	/**
@@ -816,7 +816,7 @@ public:
 	* @return Returns Value at that time
 	*/
 	UFUNCTION(BlueprintCallable, Category = "Editor Scripting | Sequencer Tools | Control Rig")
-	static FTransformNoScale GetLocalControlRigTransformNoScale(ULevelSequence* LevelSequence, UBaseControlRig* ControlRig, FName ControlName, FFrameNumber Frame,
+	static FTransformNoScale GetLocalControlRigTransformNoScale(ULevelSequence* LevelSequence, UControlRig* ControlRig, FName ControlName, FFrameNumber Frame,
 		ESequenceTimeUnit TimeUnit = ESequenceTimeUnit::DisplayRate);
 
 	/**
@@ -829,7 +829,7 @@ public:
 	* @return Returns Values at those times
 	*/
 	UFUNCTION(BlueprintCallable, Category = "Editor Scripting | Sequencer Tools | Control Rig")
-	static TArray<FTransformNoScale> GetLocalControlRigTransformNoScales(ULevelSequence* LevelSequence, UBaseControlRig* ControlRig, FName ControlName, const TArray<FFrameNumber>& Frames,
+	static TArray<FTransformNoScale> GetLocalControlRigTransformNoScales(ULevelSequence* LevelSequence, UControlRig* ControlRig, FName ControlName, const TArray<FFrameNumber>& Frames,
 		ESequenceTimeUnit TimeUnit = ESequenceTimeUnit::DisplayRate);
 
 	/**
@@ -843,7 +843,7 @@ public:
 	* @param bSetKey If True set a key, if not just set the value
 	*/
 	UFUNCTION(BlueprintCallable, Category = "Editor Scripting | Sequencer Tools | Control Rig")
-	static void SetLocalControlRigTransformNoScale(ULevelSequence* LevelSequence, UBaseControlRig* ControlRig, FName ControlName, FFrameNumber Frame, FTransformNoScale Value, 
+	static void SetLocalControlRigTransformNoScale(ULevelSequence* LevelSequence, UControlRig* ControlRig, FName ControlName, FFrameNumber Frame, FTransformNoScale Value, 
 		ESequenceTimeUnit TimeUnit = ESequenceTimeUnit::DisplayRate, bool bSetKey = true);
 
 	/**
@@ -856,7 +856,7 @@ public:
 	* @param TimeUnit Unit for frame values, either in display rate or tick resolution
 	*/
 	UFUNCTION(BlueprintCallable, Category = "Editor Scripting | Sequencer Tools | Control Rig")
-	static void SetLocalControlRigTransformNoScales(ULevelSequence* LevelSequence, UBaseControlRig* ControlRig, FName ControlName, const TArray<FFrameNumber>& Frames, const TArray<FTransformNoScale> Values,
+	static void SetLocalControlRigTransformNoScales(ULevelSequence* LevelSequence, UControlRig* ControlRig, FName ControlName, const TArray<FFrameNumber>& Frames, const TArray<FTransformNoScale> Values,
 		ESequenceTimeUnit TimeUnit = ESequenceTimeUnit::DisplayRate);
 
 	/**
@@ -869,7 +869,7 @@ public:
 	* @return Returns Value at that time
 	*/
 	UFUNCTION(BlueprintCallable, Category = "Editor Scripting | Sequencer Tools | Control Rig")
-	static FTransform GetLocalControlRigTransform(ULevelSequence* LevelSequence, UBaseControlRig* ControlRig, FName ControlName, FFrameNumber Frame, 
+	static FTransform GetLocalControlRigTransform(ULevelSequence* LevelSequence, UControlRig* ControlRig, FName ControlName, FFrameNumber Frame, 
 		ESequenceTimeUnit TimeUnit = ESequenceTimeUnit::DisplayRate);
 
 	/**
@@ -882,7 +882,7 @@ public:
 	* @return Returns Values at those times
 	*/
 	UFUNCTION(BlueprintCallable, Category = "Editor Scripting | Sequencer Tools | Control Rig")
-	static TArray<FTransform> GetLocalControlRigTransforms(ULevelSequence* LevelSequence, UBaseControlRig* ControlRig, FName ControlName, const TArray<FFrameNumber>& Frames,
+	static TArray<FTransform> GetLocalControlRigTransforms(ULevelSequence* LevelSequence, UControlRig* ControlRig, FName ControlName, const TArray<FFrameNumber>& Frames,
 		ESequenceTimeUnit TimeUnit = ESequenceTimeUnit::DisplayRate);
 
 	/**
@@ -896,7 +896,7 @@ public:
 	* @param bSetKey If True set a key, if not just set the value
 	*/
 	UFUNCTION(BlueprintCallable, Category = "Editor Scripting | Sequencer Tools | Control Rig")
-	static void SetLocalControlRigTransform(ULevelSequence* LevelSequence, UBaseControlRig* ControlRig, FName ControlName, FFrameNumber Frame, FTransform Value, 
+	static void SetLocalControlRigTransform(ULevelSequence* LevelSequence, UControlRig* ControlRig, FName ControlName, FFrameNumber Frame, FTransform Value, 
 		ESequenceTimeUnit TimeUnit = ESequenceTimeUnit::DisplayRate, bool bSetKey = true);
 
 	/**
@@ -909,7 +909,7 @@ public:
 	* @param TimeUnit Unit for frame values, either in display rate or tick resolution
 	*/
 	UFUNCTION(BlueprintCallable, Category = "Editor Scripting | Sequencer Tools | Control Rig")
-	static void SetLocalControlRigTransforms(ULevelSequence* LevelSequence, UBaseControlRig* ControlRig, FName ControlName, const TArray<FFrameNumber>& Frames, const TArray<FTransform> Values,
+	static void SetLocalControlRigTransforms(ULevelSequence* LevelSequence, UControlRig* ControlRig, FName ControlName, const TArray<FFrameNumber>& Frames, const TArray<FTransform> Values,
 		ESequenceTimeUnit TimeUnit = ESequenceTimeUnit::DisplayRate);
 
 	/*
@@ -980,7 +980,7 @@ public:
 	 * @param TimeUnit Unit for the InTime, either in display rate or tick resolution
 	 */
 	UFUNCTION(BlueprintCallable, Category = "Editor Scripting | Sequencer Tools | Control Rig")
-	static bool SetControlRigSpace(ULevelSequence* InSequence, UBaseControlRig* InControlRig, FName InControlName, const FRigElementKey& InSpaceKey, FFrameNumber InTime,  ESequenceTimeUnit TimeUnit = ESequenceTimeUnit::DisplayRate);
+	static bool SetControlRigSpace(ULevelSequence* InSequence, UControlRig* InControlRig, FName InControlName, const FRigElementKey& InSpaceKey, FFrameNumber InTime,  ESequenceTimeUnit TimeUnit = ESequenceTimeUnit::DisplayRate);
 
 	/** Bake specified Control Rig Controls to a specified Space based upon the current settings
 	* @param InSequence Sequence to bake
@@ -990,7 +990,7 @@ public:
 	* @param TimeUnit Unit for the start and end times in the InSettings parameter.
 	*/
 	UFUNCTION(BlueprintCallable, Category = "Editor Scripting | Sequencer Tools | Control Rig")
-	static bool BakeControlRigSpace(ULevelSequence* InSequence, UBaseControlRig* InControlRig, const TArray<FName>& InControlNames, FRigSpacePickerBakeSettings InSettings, ESequenceTimeUnit TimeUnit = ESequenceTimeUnit::DisplayRate);
+	static bool BakeControlRigSpace(ULevelSequence* InSequence, UControlRig* InControlRig, const TArray<FName>& InControlNames, FRigSpacePickerBakeSettings InSettings, ESequenceTimeUnit TimeUnit = ESequenceTimeUnit::DisplayRate);
 	
 	/** Delete the Control Rig Space Key for the Control at the specified time. This will delete any attached Control Rig keys at this time and will perform any needed compensation to the new space.
 	*
@@ -1002,7 +1002,7 @@ public:
 	* @return Will return false if function fails,  for example if there is no key at this time it will fail.
 	*/
 	UFUNCTION(BlueprintCallable, Category = "Editor Scripting | Sequencer Tools | Control Rig")
-	static bool DeleteControlRigSpace(ULevelSequence* InSequence, UBaseControlRig* InControlRig, FName InControlName, FFrameNumber InTime, ESequenceTimeUnit TimeUnit = ESequenceTimeUnit::DisplayRate);
+	static bool DeleteControlRigSpace(ULevelSequence* InSequence, UControlRig* InControlRig, FName InControlName, FFrameNumber InTime, ESequenceTimeUnit TimeUnit = ESequenceTimeUnit::DisplayRate);
 
 	/** Move the Control Rig Space Key for the Control at the specified time to the new time. This will also move any Control Rig keys at this space switch boundary.
 	*
@@ -1015,7 +1015,7 @@ public:
 	* @return Will return false if function fails, for example if there is no key at this time it will fail, or if the new time is invalid it could fail also
 	*/
 	UFUNCTION(BlueprintCallable, Category = "Editor Scripting | Sequencer Tools | Control Rig")
-	static bool MoveControlRigSpace(ULevelSequence* InSequence, UBaseControlRig* InControlRig, FName InControlName, FFrameNumber InTime, FFrameNumber InNewTime, ESequenceTimeUnit TimeUnit = ESequenceTimeUnit::DisplayRate);
+	static bool MoveControlRigSpace(ULevelSequence* InSequence, UControlRig* InControlRig, FName InControlName, FFrameNumber InTime, FFrameNumber InNewTime, ESequenceTimeUnit TimeUnit = ESequenceTimeUnit::DisplayRate);
 
 	/** Rename the Control Rig Channels in Sequencer to the specified new control names, which should be present on the Control Rig
 	* @param InSequence Sequence to rename controls
@@ -1025,7 +1025,7 @@ public:
 	* @return Return true if the function succeeds, false if it doesn't which can happen if the name arrays don't match in size or any of the new Control Names aren't valid
 	*/
 	UFUNCTION(BlueprintCallable, Category = "Editor Scripting | Sequencer Tools | Control Rig")
-	static bool RenameControlRigControlChannels(ULevelSequence* InSequence, UBaseControlRig* InControlRig, const TArray<FName>& InOldControlNames, const TArray<FName>& InNewControlNames);
+	static bool RenameControlRigControlChannels(ULevelSequence* InSequence, UControlRig* InControlRig, const TArray<FName>& InOldControlNames, const TArray<FName>& InNewControlNames);
 
 	/** Get the controls mask for the given ControlName */
 	UFUNCTION(BlueprintPure, Category = "Editor Scripting | Sequencer Tools | Control Rig")
@@ -1047,20 +1047,20 @@ public:
 	@param InControlRig Rig to test to see if FK Control Rig
 	**/
 	UFUNCTION(BlueprintCallable, Category = "Editor Scripting | Sequencer Tools | Control Rig")
-	static bool IsFKControlRig(UBaseControlRig* InControlRig);
+	static bool IsFKControlRig(UControlRig* InControlRig);
 
 	/**	Whether or not the control rig is an Additive Control Rig.
 	@param InControlRig Rig to test to see if Additive Control Rig
 	**/
 	UFUNCTION(BlueprintCallable, Category = "Editor Scripting | Sequencer Tools | Control Rig")
-	static bool IsAdditiveControlRig(UBaseControlRig* InControlRig);
+	static bool IsAdditiveControlRig(UControlRig* InControlRig);
 
 	/**	Get FKControlRig Apply Mode.
 	@param InControlRig Rig to test
 	@return The EControlRigFKRigExecuteMode mode it is in, either Replace,Additive or Direct
 	**/
 	UFUNCTION(BlueprintCallable, Category = "Editor Scripting | Sequencer Tools | Control Rig")
-	static EControlRigFKRigExecuteMode  GetFKControlRigApplyMode(UBaseControlRig* InControlRig);
+	static EControlRigFKRigExecuteMode  GetFKControlRigApplyMode(UControlRig* InControlRig);
 
 	/**	Set the FK Control Rig to apply mode
 	@param InControlRig Rig to set 
@@ -1068,5 +1068,5 @@ public:
 	@return returns True if the mode was set, may not be set if the Control Rig doesn't support these modes currently only FKControlRig's do.
 	**/
 	UFUNCTION(BlueprintCallable, Category = "Editor Scripting | Sequencer Tools | Control Rig")
-	static bool SetControlRigApplyMode(UBaseControlRig* InControlRig, EControlRigFKRigExecuteMode InApplyMode);
+	static bool SetControlRigApplyMode(UControlRig* InControlRig, EControlRigFKRigExecuteMode InApplyMode);
 };
diff --git a/Engine/Plugins/Editor/SequencerAnimTools/Source/SequencerAnimTools/Private/SequencerAnimEditPivotTool.cpp b/Engine/Plugins/Editor/SequencerAnimTools/Source/SequencerAnimTools/Private/SequencerAnimEditPivotTool.cpp
--- a/Engine/Plugins/Editor/SequencerAnimTools/Source/SequencerAnimTools/Private/SequencerAnimEditPivotTool.cpp
+++ b/Engine/Plugins/Editor/SequencerAnimTools/Source/SequencerAnimTools/Private/SequencerAnimEditPivotTool.cpp
@@ -17,7 +17,7 @@
 #include "EditorModeManager.h"
 #include "Misc/QualifiedFrameTime.h"
 #include "Modules/ModuleManager.h"
-#include "BaseControlRig.h"
+#include "ControlRig.h"
 #include "ControlRigSequencerEditorLibrary.h"
 #include "LevelSequence.h"
 #include "LevelSequenceEditorBlueprintLibrary.h"
@@ -61,10 +61,11 @@ static void GetControlRigsAndSequencer(TArray<TWeakObjectPtr<UBaseControlRig>>&
 		SequencerPtr = LevelSequenceEditor ? LevelSequenceEditor->GetSequencer() : nullptr;
 		if (SequencerPtr.IsValid())
 		{
-			TArray<UBaseControlRig*> TempControlRigs;
+			TArray<UControlRig*> TempControlRigs;
 			TempControlRigs = UControlRigSequencerEditorLibrary::GetVisibleControlRigs();
-			for (UBaseControlRig* ControlRig : TempControlRigs)
+			for (UControlRig* ControlRig : TempControlRigs)
 			{
+				UBaseControlRig* Base = ControlRig;
 				ControlRigs.Add(ControlRig);
 			}
 		}
@@ -452,7 +453,8 @@ bool USequencerPivotTool::SetGizmoBasedOnSelection(bool bUseSaved)
 			}
 			if (bHasSavedPivot == false)
 			{
-				GizmoTransform = UControlRigSequencerEditorLibrary::GetControlRigWorldTransform(LevelSequence, ControlRig.Get(), Name, FrameTime.RoundToFrame(),
+				UControlRig* RealControlRig = Cast<UControlRig>(ControlRig.Get());
+				GizmoTransform = UControlRigSequencerEditorLibrary::GetControlRigWorldTransform(LevelSequence, RealControlRig, Name, FrameTime.RoundToFrame(),
 					ESequenceTimeUnit::TickResolution);
 				AverageLocation += GizmoTransform.GetLocation();
 				NumLocations++;
@@ -642,7 +644,8 @@ void USequencerPivotTool::GizmoTransformStarted(UTransformProxy* Proxy)
 								continue;
 							}
 						}
-						FTransform Transform = UControlRigSequencerEditorLibrary::GetControlRigWorldTransform(LevelSequence, ControlRig.Get(), Name, FrameNumber,
+						UControlRig* RealControlRig = Cast<UControlRig>(ControlRig.Get());
+						FTransform Transform = UControlRigSequencerEditorLibrary::GetControlRigWorldTransform(LevelSequence, RealControlRig, Name, FrameNumber,
 							ESequenceTimeUnit::TickResolution);
 
 						FControlRigSelectionDuringDrag ControlDrag;
@@ -770,6 +773,7 @@ void USequencerPivotTool::GizmoTransformChanged(UTransformProxy* Proxy, FTransfo
 		int32 Index = 0;
 		for (FControlRigSelectionDuringDrag& ControlDrag : ControlRigDrags)
 		{
+			UControlRig* RealControlRig = Cast<UControlRig>(ControlDrag.ControlRig);
 			if (bInPivotMode == false)
 			{
 				ControlDrag.CurrentTransform = Transform;
@@ -797,7 +801,7 @@ void USequencerPivotTool::GizmoTransformChanged(UTransformProxy* Proxy, FTransfo
 						OptQuat.EnforceShortestArcWith(ControlDrag.CurrentTransform.GetRotation());
 						ControlDrag.CurrentTransform.SetRotation(OptQuat);
 					}
-					UControlRigSequencerEditorLibrary::SetControlRigWorldTransform(ControlDrag.LevelSequence, ControlDrag.ControlRig, ControlDrag.ControlName,
+					UControlRigSequencerEditorLibrary::SetControlRigWorldTransform(ControlDrag.LevelSequence, RealControlRig, ControlDrag.ControlName,
 						ControlDrag.CurrentFrame, ControlDrag.CurrentTransform, ESequenceTimeUnit::TickResolution, bSetKey);
 					
 				}
@@ -805,13 +809,13 @@ void USequencerPivotTool::GizmoTransformChanged(UTransformProxy* Proxy, FTransfo
 				{
 					FVector DiffTranslation = StartDragTransform.GetRotation().RotateVector(Diff.GetTranslation());
 					ControlDrag.CurrentTransform.AddToTranslation(DiffTranslation);
-					UControlRigSequencerEditorLibrary::SetControlRigWorldTransform(ControlDrag.LevelSequence, ControlDrag.ControlRig, ControlDrag.ControlName,
+					UControlRigSequencerEditorLibrary::SetControlRigWorldTransform(ControlDrag.LevelSequence, RealControlRig, ControlDrag.ControlName,
 						ControlDrag.CurrentFrame, ControlDrag.CurrentTransform, ESequenceTimeUnit::TickResolution, bSetKey);
 				}
 			}
 			else //last one with shift we keep locked!
 			{
-				UControlRigSequencerEditorLibrary::SetControlRigWorldTransform(ControlDrag.LevelSequence, ControlDrag.ControlRig, ControlDrag.ControlName,
+				UControlRigSequencerEditorLibrary::SetControlRigWorldTransform(ControlDrag.LevelSequence, RealControlRig, ControlDrag.ControlName,
 					ControlDrag.CurrentFrame, ControlDrag.CurrentTransform, ESequenceTimeUnit::TickResolution, bSetKey);
 			}
 			
