Control Rig: Reduce FRigBaseElement memory cost by moving metadata storage from individual elements onto URigHierarchy. Reduces the size of each FRigBaseElement by 60%

#rnx

diff --git a/Engine/Plugins/Animation/ControlRig/Source/ControlRig/Private/BaseControlRig.cpp b/Engine/Plugins/Animation/ControlRig/Source/ControlRig/Private/BaseControlRig.cpp
--- a/Engine/Plugins/Animation/ControlRig/Source/ControlRig/Private/BaseControlRig.cpp
+++ b/Engine/Plugins/Animation/ControlRig/Source/ControlRig/Private/BaseControlRig.cpp
@@ -361,7 +361,7 @@ void UBaseControlRig::Evaluate_AnyThread()
 		
 		// Backwards solve
 		{
-			TGuardValue<bool> UpdatePreferredAngles(Hierarchy->bUpdatePreferedEulerAngleWhenSettingTransform, false);
+			TGuardValue<bool> UpdatePreferredAngles(Hierarchy->bUpdatePreferredEulerAngleWhenSettingTransform, false);
 			Execute(FRigUnit_InverseExecution::EventName);
 		}
 
@@ -421,7 +421,7 @@ void UBaseControlRig::Evaluate_AnyThread()
 		
 		// Forward solve
 		{
-			TGuardValue<bool> UpdatePreferredAngles(Hierarchy->bUpdatePreferedEulerAngleWhenSettingTransform, false);
+			TGuardValue<bool> UpdatePreferredAngles(Hierarchy->bUpdatePreferredEulerAngleWhenSettingTransform, false);
 			Execute(FRigUnit_BeginExecution::EventName);
 		}
 	}
diff --git a/Engine/Plugins/Animation/ControlRig/Source/ControlRig/Private/Rigs/RigHierarchy.cpp b/Engine/Plugins/Animation/ControlRig/Source/ControlRig/Private/Rigs/RigHierarchy.cpp
--- a/Engine/Plugins/Animation/ControlRig/Source/ControlRig/Private/Rigs/RigHierarchy.cpp
+++ b/Engine/Plugins/Animation/ControlRig/Source/ControlRig/Private/Rigs/RigHierarchy.cpp
@@ -119,7 +119,7 @@ URigHierarchy::URigHierarchy()
 #endif
 , bEnableCacheValidityCheck(bEnableValidityCheckbyDefault)
 , HierarchyForCacheValidation()
-, bUpdatePreferedEulerAngleWhenSettingTransform(true)
+, bUpdatePreferredEulerAngleWhenSettingTransform(true)
 , bAllowNameSpaceWhenSanitizingName(false)
 , ExecuteContext(nullptr)
 #if WITH_EDITOR
@@ -199,17 +199,19 @@ void URigHierarchy::Save(FArchive& Ar)
 		Ar << Key;
 
 		// allow the element to store more information
-		Element->Serialize(Ar, this, FRigBaseElement::StaticData);
+		Element->Serialize(Ar, FRigBaseElement::StaticData);
 	}
 
 	for(int32 ElementIndex = 0; ElementIndex < ElementCount; ElementIndex++)
 	{
 		FRigBaseElement* Element = Elements[ElementIndex];
-		Element->Serialize(Ar, this, FRigBaseElement::InterElementData);
+		Element->Serialize(Ar, FRigBaseElement::InterElementData);
 	}
 
 	Ar << PreviousNameMap;
 	Ar << PreviousParentMap;
+
+	Ar << ElementMetadata;
 }
 
 void URigHierarchy::Load(FArchive& Ar)
@@ -266,7 +268,7 @@ void URigHierarchy::Load(FArchive& Ar)
 		ElementsPerType[RigElementTypeToFlatIndex(Key.Type)].Add(Element);
 		IndexLookup.Add(Key, Element->Index);
 		
-		Element->Load(Ar, this, FRigBaseElement::StaticData);
+		Element->Load(Ar, FRigBaseElement::StaticData);
 	}
 
 	IncrementTopologyVersion();
@@ -274,7 +276,7 @@ void URigHierarchy::Load(FArchive& Ar)
 	for(int32 ElementIndex = 0; ElementIndex < ElementCount; ElementIndex++)
 	{
 		FRigBaseElement* Element = Elements[ElementIndex];
-		Element->Load(Ar, this, FRigBaseElement::InterElementData);
+		Element->Load(Ar, FRigBaseElement::InterElementData);
 	}
 
 	IncrementTopologyVersion();
@@ -318,6 +320,11 @@ void URigHierarchy::Load(FArchive& Ar)
 		PreviousNameMap.Reset();
 		PreviousParentMap.Reset();
 	}
+
+	if (Ar.CustomVer(FControlRigObjectVersion::GUID) >= FControlRigObjectVersion::RigHierarchyStoresElementMetadata)
+	{
+		Ar << ElementMetadata;
+	}
 }
 
 void URigHierarchy::PostLoad()
@@ -364,13 +371,22 @@ void URigHierarchy::Reset_Impl(bool bResetElements)
 		const UEnum* ElementTypeEnum = StaticEnum<ERigElementType>();
 		for(int32 TypeIndex=0; ;TypeIndex++)
 		{
-			if ((ERigElementType)ElementTypeEnum->GetValueByIndex(TypeIndex) == ERigElementType::All)
+			if (static_cast<ERigElementType>(ElementTypeEnum->GetValueByIndex(TypeIndex)) == ERigElementType::All)
 			{
 				break;
 			}
 			ElementsPerType.Add(TArray<FRigBaseElement*>());
 		}
 		IndexLookup.Reset();
+
+		for (TTuple<FRigElementKey, FMetaDataStorage>& Metadata: ElementMetadata)
+		{
+			for (TTuple<FName, FRigBaseMetadata*>& Item: Metadata.Value.MetadataMap)
+			{
+				FRigBaseMetadata::DestroyMetadata(&Item.Value);
+			} 
+		}
+		ElementMetadata.Reset();
 	}
 
 	ResetPoseHash = INDEX_NONE;
@@ -419,18 +435,16 @@ void URigHierarchy::CopyHierarchy(URigHierarchy* InHierarchy)
 		return;
 	}
 
-	bool bReallocateElements = false;
-
 	// check if we really need to do a deep copy all over again.
 	// for rigs which contain more elements (likely procedural elements)
 	// we'll assume we can just remove superfluous elements (from the end of the lists).
-	bReallocateElements = Elements.Num() < InHierarchy->Elements.Num();
+	bool bReallocateElements = Elements.Num() < InHierarchy->Elements.Num();
 	if(!bReallocateElements)
 	{
 		const UEnum* ElementTypeEnum = StaticEnum<ERigElementType>();
 		for(int32 ElementTypeIndex = 0; ; ElementTypeIndex++)
 		{
-			if ((ERigElementType)ElementTypeEnum->GetValueByIndex(ElementTypeIndex) == ERigElementType::All)
+			if (static_cast<ERigElementType>(ElementTypeEnum->GetValueByIndex(ElementTypeIndex)) == ERigElementType::All)
 			{
 				break;
 			}
@@ -485,7 +499,7 @@ void URigHierarchy::CopyHierarchy(URigHierarchy* InHierarchy)
 			{
 				FRigBaseElement* ElementMemory = MakeElement(ElementType, Count, &StructureSize);
 				verify(StructureSize == StructureSizePerType[ElementTypeIndex]);
-				NewElementsPerType.Add((uint8*)ElementMemory);
+				NewElementsPerType.Add(reinterpret_cast<uint8*>(ElementMemory));
 			}
 			else
 			{
@@ -500,7 +514,7 @@ void URigHierarchy::CopyHierarchy(URigHierarchy* InHierarchy)
 
 		for(int32 Index = 0; Index < InHierarchy->Num(); Index++)
 		{
-			FRigBaseElement* Source = InHierarchy->Get(Index);
+			const FRigBaseElement* Source = InHierarchy->Get(Index);
 			const FRigElementKey& Key = Source->Key;
 
 			const int32 ElementTypeIndex = RigElementTypeToFlatIndex(Key.Type);
@@ -509,16 +523,17 @@ void URigHierarchy::CopyHierarchy(URigHierarchy* InHierarchy)
 
 			const int32 StructureSize = StructureSizePerType[ElementTypeIndex];
 			check(NewElementsPerType[ElementTypeIndex] != nullptr);
-			FRigBaseElement* Target = (FRigBaseElement*)&NewElementsPerType[ElementTypeIndex][StructureSize * SubIndex];
-		
-			Target->Key = Key;
-			Target->NameString.Reset();
+			FRigBaseElement* Target = reinterpret_cast<FRigBaseElement*>(&NewElementsPerType[ElementTypeIndex][StructureSize * SubIndex]);
+
+			Target->InitializeFrom(Source);
+			
 			Target->SubIndex = SubIndex;
 			Target->Index = Elements.Add(Target);
-			Target->CreatedAtInstructionIndex = Source->CreatedAtInstructionIndex;
 
 			ElementsPerType[ElementTypeIndex].Add(Target);
 			IndexLookup.Add(Key, Target->Index);
+			
+			IncrementPoseVersion(Index);
 
 			check(Source->Index == Index);
 			check(Target->Index == Index);
@@ -550,26 +565,24 @@ void URigHierarchy::CopyHierarchy(URigHierarchy* InHierarchy)
 			FRigBaseElement* Target = Elements[Index];
 
 			check(Target->Key.Type == Source->Key.Type);
-            Target->Key = Source->Key;
-            Target->NameString.Reset();
-            Target->SubIndex = Source->SubIndex;
-            Target->Index = Source->Index;
-			Target->CreatedAtInstructionIndex = Source->CreatedAtInstructionIndex;
-			Target->CachedChildren.Reset();
-			Target->bSelected = false;
+			Target->InitializeFrom(Source);
 
 			Target->TopologyVersion = InHierarchy->GetTopologyVersion();
+			IncrementPoseVersion(Index);
 		}
 
 		IndexLookup = InHierarchy->IndexLookup;
 	}
 
+	// Copy all the element subclass data and all elements' metadata over.
 	for(int32 Index = 0; Index < InHierarchy->Num(); Index++)
 	{
-		FRigBaseElement* Source = InHierarchy->Get(Index);
+		const FRigBaseElement* Source = InHierarchy->Get(Index);
 		FRigBaseElement* Target = Elements[Index];
-		Target->CopyFrom(this, Source, InHierarchy);
-		IncrementPoseVersion(Index);
+
+		Target->CopyFrom(Source);
+
+		CopyAllMetadataFromElement(Target->GetKey(), Source);
 	}
 
 	PreviousNameMap.Append(InHierarchy->PreviousNameMap);
@@ -1010,29 +1023,23 @@ TArray<FRigElementKey> URigHierarchy::RestoreConnectorsFromInfos(TArray<FRigConn
 	return Keys;
 }
 
-TArray<FName> URigHierarchy::GetMetadataNames(FRigElementKey InItem)
+TArray<FName> URigHierarchy::GetMetadataNames(FRigElementKey InItem) const
 {
 	TArray<FName> Names;
-	if(const FRigBaseElement* Element = Find(InItem))
+	if(const FMetaDataStorage* Storage = ElementMetadata.Find(InItem))
 	{
-		for(int32 Index = 0; Index < Element->NumMetadata(); Index++)
-		{
-			Names.Add(Element->GetMetadata(Index)->GetName());
-		}
+		Storage->MetadataMap.GetKeys(Names);
 	}
 	return Names;
 }
 
-ERigMetadataType URigHierarchy::GetMetadataType(FRigElementKey InItem, FName InMetadataName)
+ERigMetadataType URigHierarchy::GetMetadataType(FRigElementKey InItem, FName InMetadataName) const
 {
-	if(const FRigBaseElement* Element = Find(InItem))
+	if(const FMetaDataStorage* Storage = ElementMetadata.Find(InItem))
 	{
-		for(int32 Index = 0; Index < Element->NumMetadata(); Index++)
+		if (const FRigBaseMetadata* const* MetadataPtrPtr = Storage->MetadataMap.Find(InMetadataName))
 		{
-			if(const FRigBaseMetadata* Metadata = Element->GetMetadata(InMetadataName))
-			{
-				return Metadata->GetType();
-			}
+			return (*MetadataPtrPtr)->GetType();
 		}
 	}
 	return ERigMetadataType::Invalid;
@@ -1040,20 +1047,12 @@ ERigMetadataType URigHierarchy::GetMetadataType(FRigElementKey InItem, FName InM
 
 bool URigHierarchy::RemoveMetadata(FRigElementKey InItem, FName InMetadataName)
 {
-	if(FRigBaseElement* Element = Find(InItem))
-	{
-		return Element->RemoveMetadata(InMetadataName);
-	}
-	return false;
+	return RemoveMetadataForElement(InItem, InMetadataName);
 }
 
 bool URigHierarchy::RemoveAllMetadata(FRigElementKey InItem)
 {
-	if(FRigBaseElement* Element = Find(InItem))
-	{
-		return Element->RemoveAllMetadata();
-	}
-	return false;
+	return RemoveAllMetadataForElement(InItem);
 }
 
 TArray<const FRigBaseElement*> URigHierarchy::GetSelectedElements(ERigElementType InTypeFilter) const
@@ -3299,7 +3298,7 @@ void URigHierarchy::SetTransform(FRigTransformElement* InTransformElement, const
 	{
 		ControlElement->Shape.MarkDirty(MakeGlobal(InTransformType));
 
-		if(bUpdatePreferedEulerAngleWhenSettingTransform && ERigTransformType::IsLocal(InTransformType))
+		if(bUpdatePreferredEulerAngleWhenSettingTransform && ERigTransformType::IsLocal(InTransformType))
 		{
 			const bool bInitial = ERigTransformType::IsInitial(InTransformType);
 			ControlElement->PreferredEulerAngles.SetRotator(InTransform.Rotator(), bInitial, true);
@@ -3964,7 +3963,7 @@ void URigHierarchy::SetControlValue(FRigControlElement* InControlElement, const
 				FRigControlValue Value = InValue;
 				InControlElement->Settings.ApplyLimits(Value);
 
-				TGuardValue<bool> DontSetPreferredEulerAngle(bUpdatePreferedEulerAngleWhenSettingTransform, false);
+				TGuardValue<bool> DontSetPreferredEulerAngle(bUpdatePreferredEulerAngleWhenSettingTransform, false);
 				SetTransform(
 					InControlElement,
 					Value.GetAsTransform(
@@ -3988,7 +3987,7 @@ void URigHierarchy::SetControlValue(FRigControlElement* InControlElement, const
 				FRigControlValue Value = InValue;
 				InControlElement->Settings.ApplyLimits(Value);
 
-				TGuardValue<bool> DontSetPreferredEulerAngle(bUpdatePreferedEulerAngleWhenSettingTransform, false);
+				TGuardValue<bool> DontSetPreferredEulerAngle(bUpdatePreferredEulerAngleWhenSettingTransform, false);
 				SetTransform(
 					InControlElement,
 					Value.GetAsTransform(
@@ -4815,9 +4814,6 @@ void URigHierarchy::DestroyElement(FRigBaseElement*& InElement)
 	LLM_SCOPE_BYNAME(TEXT("Animation/ControlRig"));
 	check(InElement != nullptr);
 
-	InElement->MetadataChangedDelegate.Unbind();
-	InElement->MetadataTagChangedDelegate.Unbind();
-
 	if(InElement->OwnedInstances == 0)
 	{
 		return;
@@ -5184,9 +5180,58 @@ void URigHierarchy::CleanupInvalidCaches()
 	EnsureCacheValidity();
 }
 
+void URigHierarchy::FMetaDataStorage::Reset()
+{
+	for (TTuple<FName, FRigBaseMetadata*>& Item: MetadataMap)
+	{
+		FRigBaseMetadata::DestroyMetadata(&Item.Value);
+	}
+	MetadataMap.Reset();
+}
+
+void URigHierarchy::FMetaDataStorage::Serialize(FArchive& Ar)
+{
+	static const UEnum* MetadataTypeEnum = StaticEnum<ERigMetadataType>();
+	
+	if (Ar.IsLoading())
+	{
+		Reset();
+		
+		int32 NumEntries = 0;
+		Ar << NumEntries;
+		MetadataMap.Reserve(NumEntries);
+		for (int32 Index = 0; Index < NumEntries; Index++)
+		{
+			FName ItemName, TypeName;
+			Ar << ItemName;
+			Ar << TypeName;
+
+			const ERigMetadataType Type = static_cast<ERigMetadataType>(MetadataTypeEnum->GetValueByName(TypeName));
+			FRigBaseMetadata* Metadata = FRigBaseMetadata::MakeMetadata(ItemName, Type);
+			Metadata->Serialize(Ar);
+
+			MetadataMap.Add(ItemName, Metadata);
+		}
+	}
+	else
+	{
+		int32 NumEntries = MetadataMap.Num();
+		Ar << NumEntries;
+		for (const TTuple<FName, FRigBaseMetadata*>& Item: MetadataMap)
+		{
+			FName ItemName = Item.Key;
+			FName TypeName = MetadataTypeEnum->GetNameByValue(static_cast<int64>(Item.Value->GetType()));
+			Ar << ItemName;
+			Ar << TypeName;
+			Item.Value->Serialize(Ar);
+		}
+	}
+}
+
 void URigHierarchy::OnMetadataChanged(const FRigElementKey& InKey, const FName& InName)
 {
-	IncrementMetadataVersion(InKey, InName);
+	MetadataVersion++;
+	
 	if(MetadataChangedDelegate.IsBound())
 	{
 		MetadataChangedDelegate.Broadcast(InKey, InName);
@@ -5195,13 +5240,173 @@ void URigHierarchy::OnMetadataChanged(const FRigElementKey& InKey, const FName&
 
 void URigHierarchy::OnMetadataTagChanged(const FRigElementKey& InKey, const FName& InTag, bool bAdded)
 {
-	IncrementMetadataTagVersion(InKey, InTag, bAdded);
+	MetadataTagVersion++;
+	
 	if(MetadataTagChangedDelegate.IsBound())
 	{
 		MetadataTagChangedDelegate.Broadcast(InKey, InTag, bAdded);
 	}
 }
 
+FRigBaseMetadata* URigHierarchy::GetMetadataForElement(const FRigElementKey& InKey, const FName& InName, ERigMetadataType InType, bool bInNotify)
+{
+	FMetaDataStorage* Storage = ElementMetadata.Find(InKey);
+	if (!Storage)
+	{
+		Storage = &ElementMetadata.Add(InKey);
+	}
+	
+	if (FRigBaseMetadata** MetadataPtrPtr = Storage->MetadataMap.Find(InName))
+	{
+		if ((*MetadataPtrPtr)->GetType() == InType)
+		{
+			return (*MetadataPtrPtr);
+		}
+		
+		// The type changed, replace the existing metadata with a new one of the correct type.
+		FRigBaseMetadata::DestroyMetadata(MetadataPtrPtr);
+		*MetadataPtrPtr = FRigBaseMetadata::MakeMetadata(InName, InType);
+		if (bInNotify)
+		{
+			OnMetadataChanged(InKey, InName);
+		}
+			
+		return *MetadataPtrPtr;
+	}
+	
+	// No metadata with that name existed on the element, create one from scratch.
+	FRigBaseMetadata* Metadata = FRigBaseMetadata::MakeMetadata(InName, InType);
+	Storage->MetadataMap.Add(InName, Metadata);
+	if (bInNotify)
+	{
+		OnMetadataChanged(InKey, InName);
+	}
+	return Metadata;
+}
+
+
+FRigBaseMetadata* URigHierarchy::FindMetadataForElement(const FRigElementKey& InKey, const FName& InName, ERigMetadataType InType)
+{
+	FMetaDataStorage* Storage = ElementMetadata.Find(InKey);
+	if (!Storage)
+	{
+		return nullptr;
+	}
+	
+	FRigBaseMetadata** MetadataPtrPtr = Storage->MetadataMap.Find(InName);
+	if (!MetadataPtrPtr)
+	{
+		return nullptr;
+	}
+
+	if (InType != ERigMetadataType::Invalid && (*MetadataPtrPtr)->GetType() != InType)
+	{
+		return nullptr;
+	}
+
+	return *MetadataPtrPtr;
+}
+
+const FRigBaseMetadata* URigHierarchy::FindMetadataForElement(const FRigElementKey& InKey, const FName& InName, ERigMetadataType InType) const
+{
+	return const_cast<URigHierarchy*>(this)->FindMetadataForElement(InKey, InName, InType);
+}
+
+
+bool URigHierarchy::RemoveMetadataForElement(const FRigElementKey& InKey, const FName& InName)
+{
+	FMetaDataStorage* Storage = ElementMetadata.Find(InKey);
+	if (!Storage)
+	{
+		return false;
+	}
+	
+	FRigBaseMetadata** MetadataPtrPtr = Storage->MetadataMap.Find(InName);
+	if (!MetadataPtrPtr)
+	{
+		return false;
+	}
+	
+	FRigBaseMetadata::DestroyMetadata(MetadataPtrPtr);
+	Storage->MetadataMap.Remove(InName);
+	
+	// If the storage is now empty, remove the element's storage, so we're not lugging it around
+	// unnecessarily. 
+	if (Storage->MetadataMap.IsEmpty())
+	{
+		ElementMetadata.Remove(InKey);
+	}
+	
+	OnMetadataChanged(InKey, InName);
+	return true;
+}
+
+
+bool URigHierarchy::RemoveAllMetadataForElement(const FRigElementKey& InKey)
+{
+	FMetaDataStorage* Storage = ElementMetadata.Find(InKey);
+	if (!Storage || Storage->MetadataMap.IsEmpty())
+	{
+		return false;
+	}
+
+	for (TTuple<FName, FRigBaseMetadata*>& Item: Storage->MetadataMap)
+	{
+		FRigBaseMetadata::DestroyMetadata(&Item.Value);
+	} 
+	
+	for (TTuple<FName, FRigBaseMetadata*>& Item: Storage->MetadataMap)
+	{
+		OnMetadataChanged(InKey, Item.Key);
+	}
+	
+	ElementMetadata.Remove(InKey);
+	
+	return true; 
+}
+
+
+void URigHierarchy::CopyAllMetadataFromElement(const FRigElementKey& InTargetKey, const FRigBaseElement* InSourceElement)
+{
+	if (!ensure(InSourceElement->Owner))
+	{
+		return;
+	}
+	
+	const FMetaDataStorage* SourceStorage = InSourceElement->Owner->ElementMetadata.Find(InSourceElement->GetKey());
+	if (!SourceStorage)
+	{
+		return;
+	}
+
+	FMetaDataStorage* TargetStorage = ElementMetadata.Find(InTargetKey);
+	if (!TargetStorage)
+	{
+		TargetStorage = &ElementMetadata.Add(InTargetKey);
+	}
+	
+	for (const TTuple<FName, FRigBaseMetadata*>& SourceItem: SourceStorage->MetadataMap)
+	{
+			
+		const FRigBaseMetadata* SourceMetadata = SourceItem.Value; 	
+		// Does the destination item already exist?
+		FRigBaseMetadata** TargetMetadataPtrPtr = TargetStorage->MetadataMap.Find(SourceItem.Key);
+		FRigBaseMetadata* TargetMetadata;
+		if (TargetMetadataPtrPtr)
+		{
+			FRigBaseMetadata::DestroyMetadata(TargetMetadataPtrPtr);
+			TargetMetadata = *TargetMetadataPtrPtr = FRigBaseMetadata::MakeMetadata(SourceItem.Key, SourceMetadata->GetType());
+		}
+		else
+		{
+			TargetMetadata = FRigBaseMetadata::MakeMetadata(SourceItem.Key, SourceMetadata->GetType());
+			TargetStorage->MetadataMap.Add(SourceItem.Key, TargetMetadata);
+		}
+		TargetMetadata->SetValueData(SourceMetadata->GetValueData(), SourceMetadata->GetValueSize());
+	}
+}
+
+
 void URigHierarchy::EnsureCacheValidityImpl()
 {
 	LLM_SCOPE_BYNAME(TEXT("Animation/ControlRig"));
@@ -6642,4 +6847,4 @@ TArray<FString> URigHierarchy::ConnectorSettingsToPythonCommands(const FRigConne
 FRigHierarchyRedirectorGuard::FRigHierarchyRedirectorGuard(UBaseControlRig* InControlRig)
 : Guard(InControlRig->GetHierarchy()->ElementKeyRedirector, &InControlRig->GetElementKeyRedirector())
 {
-}
\ No newline at end of file
+}
diff --git a/Engine/Plugins/Animation/ControlRig/Source/ControlRig/Private/Rigs/RigHierarchyController.cpp b/Engine/Plugins/Animation/ControlRig/Source/ControlRig/Private/Rigs/RigHierarchyController.cpp
--- a/Engine/Plugins/Animation/ControlRig/Source/ControlRig/Private/Rigs/RigHierarchyController.cpp
+++ b/Engine/Plugins/Animation/ControlRig/Source/ControlRig/Private/Rigs/RigHierarchyController.cpp
@@ -1875,7 +1875,7 @@ int32 URigHierarchyController::AddElement(FRigBaseElement* InElementToAdd, FRigB
 {
 	ensure(IsValid());
 
-	InElementToAdd->NameString.Reset();
+	InElementToAdd->CachedNameString.Reset();
 	InElementToAdd->SubIndex = Hierarchy->Num(InElementToAdd->Key.Type);
 	InElementToAdd->Index = Hierarchy->Elements.Add(InElementToAdd);
 	Hierarchy->ElementsPerType[URigHierarchy::RigElementTypeToFlatIndex(InElementToAdd->GetKey().Type)].Add(InElementToAdd);
@@ -2285,7 +2285,7 @@ bool URigHierarchyController::RenameElement(FRigBaseElement* InElement, const FN
    
 		TGuardValue<TMap<FRigElementKey, int32>> MapGuard(Hierarchy->IndexLookup, TemporaryMap);
 		InElement->Key.Name = GetSafeNewName(InName, InElement->GetType());
-		InElement->NameString.Reset();
+		InElement->CachedNameString.Reset();
 	}
 	
 	const FRigElementKey NewKey = InElement->GetKey();
@@ -2326,6 +2326,14 @@ bool URigHierarchyController::RenameElement(FRigBaseElement* InElement, const FN
 			}
 		}
 	}
+
+	// Rename metadata
+	URigHierarchy::FMetaDataStorage MetadataStorage;
+	if (Hierarchy->ElementMetadata.RemoveAndCopyValue(OldKey, MetadataStorage))
+	{
+		Hierarchy->ElementMetadata.Add(NewKey, MoveTemp(MetadataStorage));
+	}
+	
 	
 	Hierarchy->PreviousNameMap.FindOrAdd(NewKey) = OldKey;
 	Hierarchy->IncrementTopologyVersion();
diff --git a/Engine/Plugins/Animation/ControlRig/Source/ControlRig/Private/Rigs/RigHierarchyElements.cpp b/Engine/Plugins/Animation/ControlRig/Source/ControlRig/Private/Rigs/RigHierarchyElements.cpp
--- a/Engine/Plugins/Animation/ControlRig/Source/ControlRig/Private/Rigs/RigHierarchyElements.cpp
+++ b/Engine/Plugins/Animation/ControlRig/Source/ControlRig/Private/Rigs/RigHierarchyElements.cpp
@@ -15,37 +15,12 @@
 
 const FRigBaseElement::EElementIndex FRigBaseElement::ElementTypeIndex = BaseElement;
 
-FRigBaseElement::FRigBaseElement(const FRigBaseElement& InOther)
-{
-	*this = InOther;
-}
-
-FRigBaseElement& FRigBaseElement::operator=(const FRigBaseElement& InOther)
+FRigBaseElement::~FRigBaseElement()
 {
-	Key = InOther.Key;
-	NameString = InOther.NameString;
-	Index = InOther.Index;
-	SubIndex = InOther.SubIndex;
-	bSelected = InOther.bSelected;
-	CreatedAtInstructionIndex = InOther.CreatedAtInstructionIndex;
-	TopologyVersion = InOther.TopologyVersion;
-	CachedChildren.Reset();
-	OwnedInstances = 1;
-
-	RemoveAllMetadata();
-	for(const FRigBaseMetadata* InOtherMd : InOther.Metadata)
+	if (Owner)
 	{
-		FRigBaseMetadata* Md = SetupValidMetadata(InOtherMd->Name, InOtherMd->Type);
-		check(Md);
-		Md->SetValueData(InOtherMd->GetValueData(), InOtherMd->GetValueSize());
+		Owner->RemoveAllMetadataForElement(Key);
 	}
-
-	return *this;
-}
-
-FRigBaseElement::~FRigBaseElement()
-{
-	RemoveAllMetadata();
 }
 
 UScriptStruct* FRigBaseElement::GetElementStruct() const
@@ -92,49 +67,32 @@ UScriptStruct* FRigBaseElement::GetElementStruct() const
 	return FRigBaseElement::StaticStruct();
 }
 
-void FRigBaseElement::Serialize(FArchive& Ar, URigHierarchy* Hierarchy, ESerializationPhase SerializationPhase)
+void FRigBaseElement::Serialize(FArchive& Ar, ESerializationPhase SerializationPhase)
 {
 	Ar.UsingCustomVersion(FControlRigObjectVersion::GUID);
 
-	if (Ar.IsSaving() || Ar.IsObjectReferenceCollector() || Ar.IsCountingMemory())
+	if (Ar.IsLoading())
 	{
-		Save(Ar, Hierarchy, SerializationPhase);
-	}
-	else if (Ar.IsLoading())
-	{
-		Load(Ar, Hierarchy, SerializationPhase);
+		Load(Ar, SerializationPhase);
 	}
 	else
 	{
-		// remove due to FPIEFixupSerializer hitting this checkNoEntry();
+		Save(Ar, SerializationPhase);
 	}
 }
 
-void FRigBaseElement::Save(FArchive& Ar, URigHierarchy* Hierarchy, ESerializationPhase SerializationPhase)
+void FRigBaseElement::Save(FArchive& Ar, ESerializationPhase SerializationPhase)
 {
 	if(SerializationPhase == ESerializationPhase::StaticData)
 	{
 		Ar << Key;
-
-		static const UEnum* MetadataTypeEnum = StaticEnum<ERigMetadataType>();
-
-		int32 MetadataNum = Metadata.Num();
-		Ar << MetadataNum;
-
-		for(FRigBaseMetadata* Md : Metadata)
-		{
-			FName MetadataName = Md->GetName();
-			FName MetadataTypeName = MetadataTypeEnum->GetNameByValue((int64)Md->GetType());
-
-			Ar << MetadataName;
-			Ar << MetadataTypeName;
-			Md->Serialize(Ar, false);
-		}
 	}
 }
 
-void FRigBaseElement::Load(FArchive& Ar, URigHierarchy* Hierarchy, ESerializationPhase SerializationPhase)
+void FRigBaseElement::Load(FArchive& Ar, ESerializationPhase SerializationPhase)
 {
+	checkf(Owner != nullptr, TEXT("Loading should not happen on a rig element without an owner"));
+	
 	if(SerializationPhase == ESerializationPhase::StaticData)
 	{
 		FRigElementKey LoadedKey;
@@ -144,11 +102,10 @@ void FRigBaseElement::Load(FArchive& Ar, URigHierarchy* Hierarchy, ESerializatio
 		ensure(LoadedKey.Type == Key.Type);
 		Key = LoadedKey;
 
-		NameString.Reset();
+		CachedNameString.Reset();
 
-		RemoveAllMetadata();
-		
-		if (Ar.CustomVer(FControlRigObjectVersion::GUID) >= FControlRigObjectVersion::HierarchyElementMetadata)
+		if (Ar.CustomVer(FControlRigObjectVersion::GUID) >= FControlRigObjectVersion::HierarchyElementMetadata &&
+			Ar.CustomVer(FControlRigObjectVersion::GUID) < FControlRigObjectVersion::RigHierarchyStoresElementMetadata)
 		{
 			static const UEnum* MetadataTypeEnum = StaticEnum<ERigMetadataType>();
 
@@ -157,145 +114,110 @@ void FRigBaseElement::Load(FArchive& Ar, URigHierarchy* Hierarchy, ESerializatio
 
 			for(int32 MetadataIndex = 0; MetadataIndex < MetadataNum; MetadataIndex++)
 			{
-				FName MetadataName(NAME_None);
-				FName MetadataTypeName(NAME_None);
+				FName MetadataName;
+				FName MetadataTypeName;
 				Ar << MetadataName;
 				Ar << MetadataTypeName;
-				
-				const ERigMetadataType MetadataType = (ERigMetadataType)MetadataTypeEnum->GetValueByName(MetadataTypeName);
-				FRigBaseMetadata* Md = FRigBaseMetadata::MakeMetadata(this, MetadataName, MetadataType);
-				Md->Serialize(Ar, true);
-				Metadata.Add(Md);
-			}
 
-			for(int32 MetadataIndex = 0; MetadataIndex < Metadata.Num(); MetadataIndex++)
-			{
-				MetadataNameToIndex.Add(Metadata[MetadataIndex]->GetName(), MetadataIndex);
-			}
+				const ERigMetadataType MetadataType = static_cast<ERigMetadataType>(MetadataTypeEnum->GetValueByName(MetadataTypeName));
 
-			for(int32 MetadataIndex = 0; MetadataIndex < Metadata.Num(); MetadataIndex++)
-			{
-				NotifyMetadataChanged(Metadata[MetadataIndex]->GetName());
+				FRigBaseMetadata* Md = Owner->GetMetadataForElement(Key, MetadataName, MetadataType, false);
+				Md->Serialize(Ar);
 			}
 		}
 	}
 }
 
-bool FRigBaseElement::RemoveMetadata(const FName& InName)
+
+FRigBaseMetadata* FRigBaseElement::GetMetadata(const FName& InName, ERigMetadataType InType)
 {
-	if(const int32* MetadataIndexPtr = MetadataNameToIndex.Find(InName))
+	if (!Owner)
 	{
-		const int32 MetadataIndex = *MetadataIndexPtr;
-		FRigBaseMetadata::DestroyMetadata(&Metadata[MetadataIndex]);
-		MetadataNameToIndex.Remove(InName);
-		Metadata.RemoveAt(MetadataIndex);
-		for(TPair<FName, int32>& Pair : MetadataNameToIndex)
-		{
-			if(Pair.Value > MetadataIndex)
-			{
-				Pair.Value--;
-			}
-		}
-		NotifyMetadataChanged(InName);
-		return true;
+		return nullptr;
 	}
-	return false;
+	return Owner->FindMetadataForElement(Key, InName, InType);
 }
 
-bool FRigBaseElement::RemoveAllMetadata()
+
+const FRigBaseMetadata* FRigBaseElement::GetMetadata(const FName& InName, ERigMetadataType InType) const
 {
-	if(!Metadata.IsEmpty())
+	if (Owner == nullptr)
 	{
-		TArray<FName> MetaDataNames;
-		MetaDataNames.Reserve(Metadata.Num());
-		for(FRigBaseMetadata* Md : Metadata)
-		{
-			MetaDataNames.Add(Md->GetName());
-			FRigBaseMetadata::DestroyMetadata(&Md);
-		}
-		Metadata.Reset();
-		MetadataNameToIndex.Reset();
-		for(const FName& MetaDataName: MetaDataNames)
+		return nullptr;
+	}
+	return Owner->FindMetadataForElement(Key, InName, InType);
+}
+
+
+bool FRigBaseElement::SetMetadata(const FName& InName, ERigMetadataType InType, const void* InData, int32 InSize)
+{
+	if (Owner)
+	{
+		constexpr bool bNotify = true;
+		if (FRigBaseMetadata* Metadata = Owner->GetMetadataForElement(Key, InName, InType, bNotify))
 		{
-			NotifyMetadataChanged(MetaDataName);
+			Metadata->SetValueData(InData, InSize);
+			return true;
 		}
-		return true;
 	}
 	return false;
 }
 
-void FRigBaseElement::CopyFrom(URigHierarchy* InHierarchy, FRigBaseElement* InOther, URigHierarchy* InOtherHierarchy)
+FRigBaseMetadata* FRigBaseElement::SetupValidMetadata(const FName& InName, ERigMetadataType InType)
 {
-	// remember all previous names
-	TArray<FName> RemainingNames;
-	Algo::Transform(Metadata, RemainingNames, [](const FRigBaseMetadata* Md) -> FName
-	{
-		return Md->GetName();
-	});
-
-	// copy over all metadata. this also takes care of potential type changes
-	for(const FRigBaseMetadata* InOtherMd : InOther->Metadata)
+	if (Owner == nullptr)
 	{
-		FRigBaseMetadata* Md = SetupValidMetadata(InOtherMd->Name, InOtherMd->Type);
-		check(Md);
-		Md->SetValueData(InOtherMd->GetValueData(), InOtherMd->GetValueSize());
-		RemainingNames.Remove(InOtherMd->Name);
+		return nullptr;
 	}
+	constexpr bool bNotify = true;
+	return Owner->GetMetadataForElement(Key, InName, InType, bNotify);
+}
 
-	// remove all remaining metadata
-	for(const FName& NameToRemove : RemainingNames)
-	{
-		RemoveMetadata(NameToRemove);
-	}
 
-	// rebuild the name map
-	MetadataNameToIndex.Reset();
-	for(int32 MetadataIndex = 0; MetadataIndex < Metadata.Num(); MetadataIndex++)
+bool FRigBaseElement::RemoveMetadata(const FName& InName)
+{
+	if (Owner == nullptr)
 	{
-		MetadataNameToIndex.Add(Metadata[MetadataIndex]->GetName(), MetadataIndex);
+		return false;
 	}
+	return Owner->RemoveMetadataForElement(Key, InName);
 }
 
-FRigBaseMetadata* FRigBaseElement::SetupValidMetadata(const FName& InName, ERigMetadataType InType)
+bool FRigBaseElement::RemoveAllMetadata()
 {
-	if(const int32* MetadataIndexPtr = MetadataNameToIndex.Find(InName))
+	if (Owner == nullptr)
 	{
-		const int32 MetadataIndex = *MetadataIndexPtr;
-		if(Metadata[MetadataIndex]->GetType() == InType)
-		{
-			return Metadata[MetadataIndex];
-		}
-
-		FRigBaseMetadata::DestroyMetadata(&Metadata[MetadataIndex]);
-		Metadata[MetadataIndex] = FRigBaseMetadata::MakeMetadata(this, InName, InType);
-		NotifyMetadataChanged(InName);
-		return Metadata[MetadataIndex];
+		return false;
 	}
-
-	FRigBaseMetadata* Md = FRigBaseMetadata::MakeMetadata(this, InName, InType);
-	const int32 MetadataIndex = Metadata.Add(Md);
-	MetadataNameToIndex.Add(InName, MetadataIndex);
-	NotifyMetadataChanged(InName);
-	return Md;
+	return Owner->RemoveAllMetadataForElement(Key);
 }
 
-void FRigBaseElement::NotifyMetadataChanged(const FName& InName)
+void FRigBaseElement::NotifyMetadataTagChanged(const FName& InTag, bool bAdded)
 {
-	MetadataVersion++;
-	if(MetadataChangedDelegate.IsBound())
+	if (Owner)
 	{
-		MetadataChangedDelegate.Execute(GetKey(), InName);
+		Owner->OnMetadataTagChanged(Key, InTag, bAdded);
 	}
 }
 
-void FRigBaseElement::NotifyMetadataTagChanged(const FName& InTag, bool bAdded)
+
+void FRigBaseElement::InitializeFrom(const FRigBaseElement* InOther)
+{
+	Key = InOther->Key;
+	Index = InOther->Index;
+	SubIndex = InOther->SubIndex;
+	CreatedAtInstructionIndex = InOther->CreatedAtInstructionIndex;
+	bSelected = false;
+	TopologyVersion = 0;
+	CachedChildren.Reset();
+}
+
+
+void FRigBaseElement::CopyFrom(const FRigBaseElement* InOther)
 {
-	if(MetadataTagChangedDelegate.IsBound())
-	{
-		MetadataTagChangedDelegate.Execute(GetKey(), InTag, bAdded);
-	}
 }
 
+
 ////////////////////////////////////////////////////////////////////////////////
 // FRigComputedTransform
 ////////////////////////////////////////////////////////////////////////////////
@@ -489,9 +411,9 @@ FRigBaseElement* FRigElementHandle::Get()
 
 const FRigBaseElement::EElementIndex FRigTransformElement::ElementTypeIndex = TransformElement;
 
-void FRigTransformElement::Save(FArchive& Ar, URigHierarchy* Hierarchy, ESerializationPhase SerializationPhase)
+void FRigTransformElement::Save(FArchive& Ar, ESerializationPhase SerializationPhase)
 {
-	Super::Save(Ar, Hierarchy, SerializationPhase);
+	Super::Save(Ar, SerializationPhase);
 
 	if(SerializationPhase == ESerializationPhase::StaticData)
 	{
@@ -499,9 +421,9 @@ void FRigTransformElement::Save(FArchive& Ar, URigHierarchy* Hierarchy, ESeriali
 	}
 }
 
-void FRigTransformElement::Load(FArchive& Ar, URigHierarchy* Hierarchy, ESerializationPhase SerializationPhase)
+void FRigTransformElement::Load(FArchive& Ar, ESerializationPhase SerializationPhase)
 {
-	Super::Load(Ar, Hierarchy, SerializationPhase);
+	Super::Load(Ar, SerializationPhase);
 
 	if(SerializationPhase == ESerializationPhase::StaticData)
 	{
@@ -526,10 +448,9 @@ void FRigTransformElement::CopyPose(FRigBaseElement* InOther, bool bCurrent, boo
 	}
 }
 
-void FRigTransformElement::CopyFrom(URigHierarchy* InHierarchy, FRigBaseElement* InOther,
-	URigHierarchy* InOtherHierarchy)
+void FRigTransformElement::CopyFrom(const FRigBaseElement* InOther)
 {
-	Super::CopyFrom(InHierarchy, InOther, InOtherHierarchy);
+	Super::CopyFrom(InOther);
 	
 	const FRigTransformElement* SourceTransform = CastChecked<FRigTransformElement>(InOther);
 	Pose = SourceTransform->Pose;
@@ -540,7 +461,7 @@ void FRigTransformElement::CopyFrom(URigHierarchy* InHierarchy, FRigBaseElement*
 	for(int32 ElementToDirtyIndex = 0; ElementToDirtyIndex < SourceTransform->ElementsToDirty.Num(); ElementToDirtyIndex++)
 	{
 		const FElementToDirty& Source = SourceTransform->ElementsToDirty[ElementToDirtyIndex];
-		FRigTransformElement* TargetTransform = CastChecked<FRigTransformElement>(InHierarchy->Get(Source.Element->Index));
+		FRigTransformElement* TargetTransform = CastChecked<FRigTransformElement>(Owner->Get(Source.Element->Index));
 		const FElementToDirty Target(TargetTransform, Source.HierarchyDistance);
 		ElementsToDirty.Add(Target);
 		check(ElementsToDirty[ElementToDirtyIndex].Element->GetKey() == Source.Element->GetKey());
@@ -553,9 +474,9 @@ void FRigTransformElement::CopyFrom(URigHierarchy* InHierarchy, FRigBaseElement*
 
 const FRigBaseElement::EElementIndex FRigSingleParentElement::ElementTypeIndex = SingleParentElement;
 
-void FRigSingleParentElement::Save(FArchive& Ar, URigHierarchy* Hierarchy, ESerializationPhase SerializationPhase)
+void FRigSingleParentElement::Save(FArchive& Ar, ESerializationPhase SerializationPhase)
 {
-	Super::Save(Ar, Hierarchy, SerializationPhase);
+	Super::Save(Ar, SerializationPhase);
 
 	if(SerializationPhase == ESerializationPhase::InterElementData)
 	{
@@ -568,9 +489,9 @@ void FRigSingleParentElement::Save(FArchive& Ar, URigHierarchy* Hierarchy, ESeri
 	}
 }
 
-void FRigSingleParentElement::Load(FArchive& Ar, URigHierarchy* Hierarchy, ESerializationPhase SerializationPhase)
+void FRigSingleParentElement::Load(FArchive& Ar, ESerializationPhase SerializationPhase)
 {
-	Super::Load(Ar, Hierarchy, SerializationPhase);
+	Super::Load(Ar, SerializationPhase);
 
 	if(SerializationPhase == ESerializationPhase::InterElementData)
 	{
@@ -579,20 +500,19 @@ void FRigSingleParentElement::Load(FArchive& Ar, URigHierarchy* Hierarchy, ESeri
 
 		if(ParentKey.IsValid())
 		{
-			ParentElement = Hierarchy->FindChecked<FRigTransformElement>(ParentKey);
+			ParentElement = Owner->FindChecked<FRigTransformElement>(ParentKey);
 		}
 	}
 }
 
-void FRigSingleParentElement::CopyFrom(URigHierarchy* InHierarchy, FRigBaseElement* InOther,
-	URigHierarchy* InOtherHierarchy)
+void FRigSingleParentElement::CopyFrom(const FRigBaseElement* InOther)
 {
-	Super::CopyFrom(InHierarchy, InOther, InOtherHierarchy);
+	Super::CopyFrom(InOther);
 
 	const FRigSingleParentElement* Source = CastChecked<FRigSingleParentElement>(InOther); 
 	if(Source->ParentElement)
 	{
-		ParentElement = CastChecked<FRigTransformElement>(InHierarchy->Get(Source->ParentElement->Index));
+		ParentElement = CastChecked<FRigTransformElement>(Owner->Get(Source->ParentElement->GetIndex()));
 		check(ParentElement->GetKey() == Source->ParentElement->GetKey());
 	}
 	else
@@ -607,9 +527,9 @@ void FRigSingleParentElement::CopyFrom(URigHierarchy* InHierarchy, FRigBaseEleme
 
 const FRigBaseElement::EElementIndex FRigMultiParentElement::ElementTypeIndex = MultiParentElement;
 
-void FRigMultiParentElement::Save(FArchive& Ar, URigHierarchy* Hierarchy, ESerializationPhase SerializationPhase)
+void FRigMultiParentElement::Save(FArchive& Ar, ESerializationPhase SerializationPhase)
 {
-	Super::Save(Ar, Hierarchy, SerializationPhase);
+	Super::Save(Ar, SerializationPhase);
 
 	if(SerializationPhase == ESerializationPhase::StaticData)
 	{
@@ -633,9 +553,9 @@ void FRigMultiParentElement::Save(FArchive& Ar, URigHierarchy* Hierarchy, ESeria
 	}
 }
 
-void FRigMultiParentElement::Load(FArchive& Ar, URigHierarchy* Hierarchy, ESerializationPhase SerializationPhase)
+void FRigMultiParentElement::Load(FArchive& Ar, ESerializationPhase SerializationPhase)
 {
-	Super::Load(Ar, Hierarchy, SerializationPhase);
+	Super::Load(Ar, SerializationPhase);
 
 	if(SerializationPhase == ESerializationPhase::StaticData)
 	{
@@ -658,7 +578,7 @@ void FRigMultiParentElement::Load(FArchive& Ar, URigHierarchy* Hierarchy, ESeria
 			Ar << ParentKey;
 			ensure(ParentKey.IsValid());
 
-			ParentConstraints[ParentIndex].ParentElement = Hierarchy->FindChecked<FRigTransformElement>(ParentKey);
+			ParentConstraints[ParentIndex].ParentElement = Owner->FindChecked<FRigTransformElement>(ParentKey);
 			ParentConstraints[ParentIndex].bCacheIsDirty = true;
 
 			if (Ar.CustomVer(FControlRigObjectVersion::GUID) >= FControlRigObjectVersion::RigHierarchyMultiParentConstraints)
@@ -682,10 +602,9 @@ void FRigMultiParentElement::Load(FArchive& Ar, URigHierarchy* Hierarchy, ESeria
 	}
 }
 
-void FRigMultiParentElement::CopyFrom(URigHierarchy* InHierarchy, FRigBaseElement* InOther,
-                                      URigHierarchy* InOtherHierarchy)
+void FRigMultiParentElement::CopyFrom(const FRigBaseElement* InOther)
 {
-	Super::CopyFrom(InHierarchy, InOther, InOtherHierarchy);
+	Super::CopyFrom(InOther);
 	
 	const FRigMultiParentElement* Source = CastChecked<FRigMultiParentElement>(InOther);
 	ParentConstraints.Reset();
@@ -697,7 +616,7 @@ void FRigMultiParentElement::CopyFrom(URigHierarchy* InHierarchy, FRigBaseElemen
 	{
 		FRigElementParentConstraint ParentConstraint = Source->ParentConstraints[ParentIndex];
 		const FRigTransformElement* SourceParentElement = ParentConstraint.ParentElement;
-		ParentConstraint.ParentElement = CastChecked<FRigTransformElement>(InHierarchy->Get(SourceParentElement->Index));
+		ParentConstraint.ParentElement = CastChecked<FRigTransformElement>(Owner->Get(SourceParentElement->GetIndex()));
 		ParentConstraints.Add(ParentConstraint);
 		check(ParentConstraints[ParentIndex].ParentElement->GetKey() == SourceParentElement->GetKey());
 		IndexLookup.Add(ParentConstraint.ParentElement->GetKey(), ParentIndex);
@@ -760,9 +679,9 @@ static_assert(sizeof(FRigBoneElement) <= 736, "FRigBoneElement was optimized to
 
 const FRigBaseElement::EElementIndex FRigBoneElement::ElementTypeIndex = BoneElement;
 
-void FRigBoneElement::Save(FArchive& Ar, URigHierarchy* Hierarchy, ESerializationPhase SerializationPhase)
+void FRigBoneElement::Save(FArchive& Ar, ESerializationPhase SerializationPhase)
 {
-	Super::Save(Ar, Hierarchy, SerializationPhase);
+	Super::Save(Ar, SerializationPhase);
 
 	if(SerializationPhase == ESerializationPhase::StaticData)
 	{
@@ -772,9 +691,9 @@ void FRigBoneElement::Save(FArchive& Ar, URigHierarchy* Hierarchy, ESerializatio
 	}
 }
 
-void FRigBoneElement::Load(FArchive& Ar, URigHierarchy* Hierarchy, ESerializationPhase SerializationPhase)
+void FRigBoneElement::Load(FArchive& Ar, ESerializationPhase SerializationPhase)
 {
-	Super::Load(Ar, Hierarchy, SerializationPhase);
+	Super::Load(Ar, SerializationPhase);
 
 	if(SerializationPhase == ESerializationPhase::StaticData)
 	{
@@ -785,9 +704,9 @@ void FRigBoneElement::Load(FArchive& Ar, URigHierarchy* Hierarchy, ESerializatio
 	}
 }
 
-void FRigBoneElement::CopyFrom(URigHierarchy* InHierarchy, FRigBaseElement* InOther, URigHierarchy* InOtherHierarchy)
+void FRigBoneElement::CopyFrom(const FRigBaseElement* InOther)
 {
-	Super::CopyFrom(InHierarchy, InOther, InOtherHierarchy);
+	Super::CopyFrom(InOther);
 	
 	const FRigBoneElement* Source = CastChecked<FRigBoneElement>(InOther);
 	BoneType = Source->BoneType;
@@ -1133,7 +1052,7 @@ bool FRigControlSettings::operator==(const FRigControlSettings& InOther) const
 	{
 		return false;
 	}
-	if( ControlEnum != InOther. ControlEnum)
+	if(ControlEnum != InOther.ControlEnum)
 	{
 		return false;
 	}
@@ -1260,9 +1179,9 @@ void FRigControlSettings::SetupLimitArrayForType(bool bLimitTranslation, bool bL
 
 const FRigBaseElement::EElementIndex FRigControlElement::ElementTypeIndex = ControlElement;
 
-void FRigControlElement::Save(FArchive& Ar, URigHierarchy* Hierarchy, ESerializationPhase SerializationPhase)
+void FRigControlElement::Save(FArchive& Ar, ESerializationPhase SerializationPhase)
 {
-	Super::Save(Ar, Hierarchy, SerializationPhase);
+	Super::Save(Ar, SerializationPhase);
 
 	if(SerializationPhase == ESerializationPhase::StaticData)
 	{
@@ -1273,9 +1192,9 @@ void FRigControlElement::Save(FArchive& Ar, URigHierarchy* Hierarchy, ESerializa
 	}
 }
 
-void FRigControlElement::Load(FArchive& Ar, URigHierarchy* Hierarchy, ESerializationPhase SerializationPhase)
+void FRigControlElement::Load(FArchive& Ar, ESerializationPhase SerializationPhase)
 {
-	Super::Load(Ar, Hierarchy, SerializationPhase);
+	Super::Load(Ar, SerializationPhase);
 
 	if(SerializationPhase == ESerializationPhase::StaticData)
 	{
@@ -1295,9 +1214,9 @@ void FRigControlElement::Load(FArchive& Ar, URigHierarchy* Hierarchy, ESerializa
 	}
 }
 
-void FRigControlElement::CopyFrom(URigHierarchy* InHierarchy, FRigBaseElement* InOther, URigHierarchy* InOtherHierarchy)
+void FRigControlElement::CopyFrom(const FRigBaseElement* InOther)
 {
-	Super::CopyFrom(InHierarchy, InOther, InOtherHierarchy);
+	Super::CopyFrom(InOther);
 	
 	const FRigControlElement* Source = CastChecked<FRigControlElement>(InOther);
 	Settings = Source->Settings;
@@ -1333,9 +1252,9 @@ void FRigControlElement::CopyPose(FRigBaseElement* InOther, bool bCurrent, bool
 
 const FRigBaseElement::EElementIndex FRigCurveElement::ElementTypeIndex = CurveElement;
 
-void FRigCurveElement::Save(FArchive& Ar, URigHierarchy* Hierarchy, ESerializationPhase SerializationPhase)
+void FRigCurveElement::Save(FArchive& Ar, ESerializationPhase SerializationPhase)
 {
-	Super::Save(Ar, Hierarchy, SerializationPhase);
+	Super::Save(Ar, SerializationPhase);
 
 	if(SerializationPhase == ESerializationPhase::StaticData)
 	{
@@ -1344,9 +1263,9 @@ void FRigCurveElement::Save(FArchive& Ar, URigHierarchy* Hierarchy, ESerializati
 	}
 }
 
-void FRigCurveElement::Load(FArchive& Ar, URigHierarchy* Hierarchy, ESerializationPhase SerializationPhase)
+void FRigCurveElement::Load(FArchive& Ar, ESerializationPhase SerializationPhase)
 {
-	Super::Load(Ar, Hierarchy, SerializationPhase);
+	Super::Load(Ar, SerializationPhase);
 
 	if(SerializationPhase == ESerializationPhase::StaticData)
 	{
@@ -1373,9 +1292,9 @@ void FRigCurveElement::CopyPose(FRigBaseElement* InOther, bool bCurrent, bool bI
 	}
 }
 
-void FRigCurveElement::CopyFrom(URigHierarchy* InHierarchy, FRigBaseElement* InOther, URigHierarchy* InOtherHierarchy)
+void FRigCurveElement::CopyFrom(const FRigBaseElement* InOther)
 {
-	Super::CopyFrom(InHierarchy, InOther, InOtherHierarchy);
+	Super::CopyFrom(InOther);
 	
 	if(const FRigCurveElement* Other = CastChecked<FRigCurveElement>(InOther))
 	{
@@ -1409,9 +1328,9 @@ void FRigRigidBodySettings::Load(FArchive& Ar)
 
 const FRigBaseElement::EElementIndex FRigRigidBodyElement::ElementTypeIndex = RigidBodyElement;
 
-void FRigRigidBodyElement::Save(FArchive& Ar, URigHierarchy* Hierarchy, ESerializationPhase SerializationPhase)
+void FRigRigidBodyElement::Save(FArchive& Ar, ESerializationPhase SerializationPhase)
 {
-	Super::Save(Ar, Hierarchy, SerializationPhase);
+	Super::Save(Ar, SerializationPhase);
 
 	if(SerializationPhase == ESerializationPhase::StaticData)
 	{
@@ -1419,9 +1338,9 @@ void FRigRigidBodyElement::Save(FArchive& Ar, URigHierarchy* Hierarchy, ESeriali
 	}
 }
 
-void FRigRigidBodyElement::Load(FArchive& Ar, URigHierarchy* Hierarchy, ESerializationPhase SerializationPhase)
+void FRigRigidBodyElement::Load(FArchive& Ar, ESerializationPhase SerializationPhase)
 {
-	Super::Load(Ar, Hierarchy, SerializationPhase);
+	Super::Load(Ar, SerializationPhase);
 
 	if(SerializationPhase == ESerializationPhase::StaticData)
 	{
@@ -1429,10 +1348,9 @@ void FRigRigidBodyElement::Load(FArchive& Ar, URigHierarchy* Hierarchy, ESeriali
 	}
 }
 
-void FRigRigidBodyElement::CopyFrom(URigHierarchy* InHierarchy, FRigBaseElement* InOther,
-	URigHierarchy* InOtherHierarchy)
+void FRigRigidBodyElement::CopyFrom(const FRigBaseElement* InOther)
 {
-	Super::CopyFrom(InHierarchy, InOther, InOtherHierarchy);
+	Super::CopyFrom(InOther);
 	
 	const FRigRigidBodyElement* Source = CastChecked<FRigRigidBodyElement>(InOther);
 	Settings = Source->Settings;
@@ -1444,19 +1362,19 @@ void FRigRigidBodyElement::CopyFrom(URigHierarchy* InHierarchy, FRigBaseElement*
 
 const FRigBaseElement::EElementIndex FRigReferenceElement::ElementTypeIndex = ReferenceElement;
 
-void FRigReferenceElement::Save(FArchive& Ar, URigHierarchy* Hierarchy, ESerializationPhase SerializationPhase)
+void FRigReferenceElement::Save(FArchive& Ar, ESerializationPhase SerializationPhase)
 {
-	Super::Save(Ar, Hierarchy, SerializationPhase);
+	Super::Save(Ar, SerializationPhase);
 }
 
-void FRigReferenceElement::Load(FArchive& Ar, URigHierarchy* Hierarchy, ESerializationPhase SerializationPhase)
+void FRigReferenceElement::Load(FArchive& Ar, ESerializationPhase SerializationPhase)
 {
-	Super::Load(Ar, Hierarchy, SerializationPhase);
+	Super::Load(Ar, SerializationPhase);
 }
 
-void FRigReferenceElement::CopyFrom(URigHierarchy* InHierarchy, FRigBaseElement* InOther, URigHierarchy* InOtherHierarchy)
+void FRigReferenceElement::CopyFrom(const FRigBaseElement* InOther)
 {
-	Super::CopyFrom(InHierarchy, InOther, InOtherHierarchy);
+	Super::CopyFrom(InOther);
 	
 	const FRigReferenceElement* Source = CastChecked<FRigReferenceElement>(InOther);
 	GetWorldTransformDelegate = Source->GetWorldTransformDelegate;
@@ -1581,9 +1499,9 @@ uint32 GetTypeHash(const FRigConnectorSettings& Settings)
 
 const FRigBaseElement::EElementIndex FRigConnectorElement::ElementTypeIndex = ConnectorElement;
 
-void FRigConnectorElement::Save(FArchive& Ar, URigHierarchy* Hierarchy, ESerializationPhase SerializationPhase)
+void FRigConnectorElement::Save(FArchive& Ar, ESerializationPhase SerializationPhase)
 {
-	Super::Save(Ar, Hierarchy, SerializationPhase);
+	Super::Save(Ar, SerializationPhase);
 
 	if(SerializationPhase == ESerializationPhase::StaticData)
 	{
@@ -1591,9 +1509,9 @@ void FRigConnectorElement::Save(FArchive& Ar, URigHierarchy* Hierarchy, ESeriali
 	}
 }
 
-void FRigConnectorElement::Load(FArchive& Ar, URigHierarchy* Hierarchy, ESerializationPhase SerializationPhase)
+void FRigConnectorElement::Load(FArchive& Ar, ESerializationPhase SerializationPhase)
 {
-	Super::Load(Ar, Hierarchy, SerializationPhase);
+	Super::Load(Ar, SerializationPhase);
 
 	if(SerializationPhase == ESerializationPhase::StaticData)
 	{
@@ -1701,9 +1619,9 @@ FRigConnectorInfo FRigConnectorElement::GetConnectorInfo(const URigHierarchy* In
 	return Info;
 }
 
-void FRigConnectorElement::CopyFrom(URigHierarchy* InHierarchy, FRigBaseElement* InOther, URigHierarchy* InOtherHierarchy)
+void FRigConnectorElement::CopyFrom(const FRigBaseElement* InOther)
 {
-	Super::CopyFrom(InHierarchy, InOther, InOtherHierarchy);
+	Super::CopyFrom(InOther);
 	
 	const FRigConnectorElement* Source = CastChecked<FRigConnectorElement>(InOther);
 	Settings = Source->Settings;
@@ -1715,17 +1633,17 @@ void FRigConnectorElement::CopyFrom(URigHierarchy* InHierarchy, FRigBaseElement*
 
 const FRigBaseElement::EElementIndex FRigSocketElement::ElementTypeIndex = SocketElement;
 
-void FRigSocketElement::Save(FArchive& Ar, URigHierarchy* Hierarchy, ESerializationPhase SerializationPhase)
+void FRigSocketElement::Save(FArchive& Ar, ESerializationPhase SerializationPhase)
 {
-	Super::Save(Ar, Hierarchy, SerializationPhase);
+	Super::Save(Ar, SerializationPhase);
 }
 
-void FRigSocketElement::Load(FArchive& Ar, URigHierarchy* Hierarchy, ESerializationPhase SerializationPhase)
+void FRigSocketElement::Load(FArchive& Ar, ESerializationPhase SerializationPhase)
 {
-	Super::Load(Ar, Hierarchy, SerializationPhase);
+	Super::Load(Ar, SerializationPhase);
 }
 
-void FRigSocketElement::CopyFrom(URigHierarchy* InHierarchy, FRigBaseElement* InOther, URigHierarchy* InOtherHierarchy)
+void FRigSocketElement::CopyFrom(const FRigBaseElement* InOther)
 {
-	Super::CopyFrom(InHierarchy, InOther, InOtherHierarchy);
-}
\ No newline at end of file
+	Super::CopyFrom(InOther);
+}
diff --git a/Engine/Plugins/Animation/ControlRig/Source/ControlRig/Private/Rigs/RigHierarchyMetadata.cpp b/Engine/Plugins/Animation/ControlRig/Source/ControlRig/Private/Rigs/RigHierarchyMetadata.cpp
--- a/Engine/Plugins/Animation/ControlRig/Source/ControlRig/Private/Rigs/RigHierarchyMetadata.cpp
+++ b/Engine/Plugins/Animation/ControlRig/Source/ControlRig/Private/Rigs/RigHierarchyMetadata.cpp
@@ -1,8 +1,7 @@
 // Copyright Epic Games, Inc. All Rights Reserved.
 
 #include "Rigs/RigHierarchyMetadata.h"
-#include "Rigs/RigHierarchyElements.h"
-#include "Rigs/RigHierarchy.h"
+
 #include "ControlRigObjectVersion.h"
 
 #include UE_INLINE_GENERATED_CPP_BY_NAME(RigHierarchyMetadata)
@@ -108,17 +107,16 @@ UScriptStruct* FRigBaseMetadata::GetMetadataStruct(const ERigMetadataType& InTyp
 	return StaticStruct();
 }
 
-FRigBaseMetadata* FRigBaseMetadata::MakeMetadata(const FRigBaseElement* InElement, const FName& InName, ERigMetadataType InType)
+FRigBaseMetadata* FRigBaseMetadata::MakeMetadata(const FName& InName, ERigMetadataType InType)
 {
 	check(InType != ERigMetadataType::Invalid);
 	
-	UScriptStruct* Struct = GetMetadataStruct(InType);
+	const UScriptStruct* Struct = GetMetadataStruct(InType);
 	check(Struct);
 	
-	FRigBaseMetadata* Md = (FRigBaseMetadata*)FMemory::Malloc(Struct->GetStructureSize());
+	FRigBaseMetadata* Md = static_cast<FRigBaseMetadata*>(FMemory::Malloc(Struct->GetStructureSize()));
 	Struct->InitializeStruct(Md, 1);
 
-	Md->Element = InElement;
 	Md->Name = InName;
 	Md->Type = InType;
 	return Md;
@@ -137,22 +135,9 @@ void FRigBaseMetadata::DestroyMetadata(FRigBaseMetadata** Metadata)
 	Md = nullptr;
 }
 
-void FRigBaseMetadata::Serialize(FArchive& Ar, bool bIsLoading)
+void FRigBaseMetadata::Serialize(FArchive& Ar)
 {
 	Ar.UsingCustomVersion(FControlRigObjectVersion::GUID);
 	Ar << Name;
 	Ar << Type;
 }
-
-const FRigElementKey& FRigBaseMetadata::GetKey() const
-{
-	if(Element)
-	{
-		return Element->GetKey();
-	}
-
-	static const FRigElementKey EmptyKey;
-	return EmptyKey;
-}
-
-
diff --git a/Engine/Plugins/Animation/ControlRig/Source/ControlRig/Private/Units/Hierarchy/RigUnit_Metadata.cpp b/Engine/Plugins/Animation/ControlRig/Source/ControlRig/Private/Units/Hierarchy/RigUnit_Metadata.cpp
--- a/Engine/Plugins/Animation/ControlRig/Source/ControlRig/Private/Units/Hierarchy/RigUnit_Metadata.cpp
+++ b/Engine/Plugins/Animation/ControlRig/Source/ControlRig/Private/Units/Hierarchy/RigUnit_Metadata.cpp
@@ -190,7 +190,7 @@ FRigBaseMetadata* FRigDispatch_GetMetadata::FindMetadata(const FRigVMExtendedExe
 	const FControlRigExecuteContext& ExecuteContext = InContext.GetPublicData<FControlRigExecuteContext>();
 	if(Cache.UpdateCache(InKey, ExecuteContext.Hierarchy))
 	{
-		if(const FRigBaseElement* Element = ExecuteContext.Hierarchy->Get(Cache.GetIndex()))
+		if(FRigBaseElement* Element = ExecuteContext.Hierarchy->Get(Cache.GetIndex()))
 		{
 			// first try to find the metadata in the namespace
 			const FName Name = ExecuteContext.AdaptMetadataName(bUseNameSpace, InName);
diff --git a/Engine/Plugins/Animation/ControlRig/Source/ControlRig/Private/Units/Hierarchy/RigUnit_Metadata.h b/Engine/Plugins/Animation/ControlRig/Source/ControlRig/Private/Units/Hierarchy/RigUnit_Metadata.h
--- a/Engine/Plugins/Animation/ControlRig/Source/ControlRig/Private/Units/Hierarchy/RigUnit_Metadata.h
+++ b/Engine/Plugins/Animation/ControlRig/Source/ControlRig/Private/Units/Hierarchy/RigUnit_Metadata.h
@@ -102,7 +102,7 @@ protected:
 		bool& Found = *(bool*)Handles[Factory->FoundArgIndex].GetData();
 
 		// extract the metadata
-		if (MetadataType* Md = Cast<MetadataType>(FindMetadata(InContext, Item, Name, EnumValue, bUseNameSpace, Cache)))
+		if (const MetadataType* Md = Cast<MetadataType>(FindMetadata(InContext, Item, Name, EnumValue, bUseNameSpace, Cache)))
 		{
 			Value = Md->GetValue();
 			Found = true;
diff --git a/Engine/Plugins/Animation/ControlRig/Source/ControlRig/Public/ControlRigObjectVersion.h b/Engine/Plugins/Animation/ControlRig/Source/ControlRig/Public/ControlRigObjectVersion.h
--- a/Engine/Plugins/Animation/ControlRig/Source/ControlRig/Public/ControlRigObjectVersion.h
+++ b/Engine/Plugins/Animation/ControlRig/Source/ControlRig/Public/ControlRigObjectVersion.h
@@ -1,7 +1,6 @@
 // Copyright Epic Games, Inc. All Rights Reserved.
 #pragma once
 
-#include "CoreTypes.h"
 #include "Misc/Guid.h"
 
 // Custom serialization version for changes made in Dev-Anim stream
@@ -93,6 +92,9 @@ struct CONTROLRIG_API FControlRigObjectVersion
 		// Last bit required for Control supporting preferred rotation order
 		RigHierarchyControlPreferredRotationOrderFlag,
 
+		// Element metadata is now stored on URigHierarchy, rather than FRigBaseElement
+		RigHierarchyStoresElementMetadata,
+
 		// -----<new versions can be added above this line>-------------------------------------------------
 		VersionPlusOne,
 		LatestVersion = VersionPlusOne - 1
diff --git a/Engine/Plugins/Animation/ControlRig/Source/ControlRig/Public/Rigs/RigHierarchy.h b/Engine/Plugins/Animation/ControlRig/Source/ControlRig/Public/Rigs/RigHierarchy.h
--- a/Engine/Plugins/Animation/ControlRig/Source/ControlRig/Public/Rigs/RigHierarchy.h
+++ b/Engine/Plugins/Animation/ControlRig/Source/ControlRig/Public/Rigs/RigHierarchy.h
@@ -670,7 +670,7 @@ private:
 		{
 			return *Bone;
 		}
-		return FRigBoneElement();
+		return FRigBoneElement{};
 	}	
 	
 	/**
@@ -684,7 +684,7 @@ private:
 		{
 			return *Control;
 		}
-		return FRigControlElement();
+		return FRigControlElement{};
 	}	
 
 	/**
@@ -698,7 +698,7 @@ private:
 		{
 			return *Null;
 		}
-		return FRigNullElement();
+		return FRigNullElement{};
 	}
 	
 public:	
@@ -1034,7 +1034,7 @@ public:
 	 * @param InItem The element key to return the metadata keys for
 	 */
 	UFUNCTION(BlueprintPure, Category = URigHierarchy)
-	TArray<FName> GetMetadataNames(FRigElementKey InItem);
+	TArray<FName> GetMetadataNames(FRigElementKey InItem) const;
 
 	/**
 	 * Returns the type of metadata given its name the item it is stored under
@@ -1042,7 +1042,7 @@ public:
 	 * @param InMetadataName The name of the metadata to return the type for
 	 */
 	UFUNCTION(BlueprintPure, Category = URigHierarchy)
-	ERigMetadataType GetMetadataType(FRigElementKey InItem, FName InMetadataName);
+	ERigMetadataType GetMetadataType(FRigElementKey InItem, FName InMetadataName) const;
 
 	/**
 	 * Removes the metadata under a given element 
@@ -3429,28 +3429,12 @@ public:
 	/**
 	 * Returns the metadata version of this hierarchy
 	 */
-	uint16 GetMetadataVersion() const { return MetadataVersion; }
-
-	/**
-	 * Increments the metadata version
-	 */
-	void IncrementMetadataVersion(const FRigElementKey& InKey, const FName& InName)
-	{
-		MetadataVersion += 1 + (int32)HashCombine(GetTypeHash(InKey), GetTypeHash(InName));
-	}
+	uint32 GetMetadataVersion() const { return MetadataVersion; }
 
 	/**
      * Returns the metadata tag version of this hierarchy
 	 */
-	uint16 GetMetadataTagVersion() const { return MetadataTagVersion; }
-
-	/**
-	 * Increments the metadataTag version
-	 */
-	void IncrementMetadataTagVersion(const FRigElementKey& InKey, const FName& InTag, bool bAdded)
-	{
-		MetadataTagVersion += 1 + (int32)HashCombine(GetTypeHash(InKey), GetTypeHash(InTag));
-	}
+	uint32 GetMetadataTagVersion() const { return MetadataTagVersion; }
 
 	/**
 	 * Returns the current / initial pose of the hierarchy
@@ -4059,12 +4043,10 @@ private:
 	template<typename ElementType = FRigBaseElement>
 	ElementType* NewElement(int32 Num = 1)
 	{
-		ElementType* NewElements = (ElementType*)FMemory::Malloc(sizeof(ElementType) * Num);
+		ElementType* NewElements = static_cast<ElementType*>(FMemory::Malloc(sizeof(ElementType) * Num));
 		for(int32 Index=0;Index<Num;Index++)
 		{
-			new(&NewElements[Index]) ElementType();
-			NewElements[Index].MetadataChangedDelegate.BindStatic(&URigHierarchy::OnMetadataChanged_Static, this);
-			NewElements[Index].MetadataTagChangedDelegate.BindStatic(&URigHierarchy::OnMetadataTagChanged_Static, this);
+			new(&NewElements[Index]) ElementType(this);
 		}
 		NewElements[0].OwnedInstances = Num;
 		return NewElements;
@@ -4113,7 +4095,7 @@ private:
 	 * created or removed (not when the metadata values changes)
 	 */
 	UPROPERTY(transient)
-	uint16 MetadataVersion;
+	uint32 MetadataVersion;
 
 	/**
 	 * The metadata version of the hierarchy changes when metadata is being
@@ -4135,9 +4117,28 @@ private:
 	// Storage for the elements
 	mutable TArray<TArray<FRigBaseElement*>> ElementsPerType;
 
+	//
+	struct FMetaDataStorage
+	{
+		TMap<FName, FRigBaseMetadata*> MetadataMap;
+
+		void Reset();
+		void Serialize(FArchive& Ar);
+
+		friend FArchive& operator<<(FArchive& Ar, FMetaDataStorage& Storage)
+		{
+			Storage.Serialize(Ar);
+			return Ar;
+		}
+	};
+
 	// Managed lookup from Key to Index
 	TMap<FRigElementKey, int32> IndexLookup;
 
+	TMap<FRigElementKey, FString> UserDefinedElementName;
+
+	TMap<FRigElementKey, FMetaDataStorage> ElementMetadata;
+
 	// Static empty element array used for ref returns
 	static const FRigBaseElementChildrenArray EmptyElementArray;
 
@@ -4513,9 +4514,9 @@ private:
 	{
 		if(InElement)
 		{
-			if(FRigBaseMetadata* Metadata = InElement->GetMetadata(InMetadataName, InType))
+			if(const FRigBaseMetadata* Metadata = FindMetadataForElement(InElement->GetKey(), InMetadataName, InType))
 			{
-				return *(const T*)Metadata->GetValueData();
+				return *static_cast<const T*>(Metadata->GetValueData());
 			}
 		}
 		return DefaultValue;
@@ -4545,7 +4546,11 @@ private:
 	{
 		if(InElement)
 		{
-			return InElement->SetMetaData(InMetadataName, InType, &InValue, sizeof(T));
+			constexpr bool bNotify = true;
+			if (FRigBaseMetadata* Metadata = GetMetadataForElement(InElement->GetKey(), InMetadataName, InType, bNotify))
+			{
+				return Metadata->SetValueData(&InValue, sizeof(T));
+			}
 		}
 		return false;
 	}
@@ -4565,21 +4570,27 @@ private:
 	void OnMetadataChanged(const FRigElementKey& InKey, const FName& InName);
 	void OnMetadataTagChanged(const FRigElementKey& InKey, const FName& InTag, bool bAdded);
 
-protected:
-
-	static void OnMetadataChanged_Static(const FRigElementKey& InKey, const FName& InName, URigHierarchy* InHierarchy)
-	{
-		check(InHierarchy);
-		check(IsValid(InHierarchy));
-		InHierarchy->OnMetadataChanged(InKey, InName);
-	}
-	static void OnMetadataTagChanged_Static(const FRigElementKey& InKey, const FName& InTag, bool bAdded, URigHierarchy* InHierarchy)
-	{
-		check(InHierarchy);
-		check(IsValid(InHierarchy));
-		InHierarchy->OnMetadataTagChanged(InKey, InTag, bAdded);
-	}
+	/** Returns a metadata ptr to the given element's metadata. If the meta data, with the same name, doesn't exist already a new entry
+	    is created for that element. If the name matches but the type differs, the existing metadata is destroyed and a new one with the
+	    matching type is created instead.
+	    */
+	FRigBaseMetadata* GetMetadataForElement(const FRigElementKey& InKey, const FName& InName, ERigMetadataType InType, bool bInNotify);
+	
+	/** Attempts to find element's metadata of the given name and type. If either the element doesnt exist, the name doesn't exist or the
+	    type doesn't match, then \c nullptr is returned.
+	    */ 
+	FRigBaseMetadata* FindMetadataForElement(const FRigElementKey& InKey, const FName& InName, ERigMetadataType InType);
+	const FRigBaseMetadata* FindMetadataForElement(const FRigElementKey& InKey, const FName& InName, ERigMetadataType InType) const;
 	
+	/** Removes the named meta data for the given element, regardless of type. If the element doesn't exist, or it doesn't have any
+	    metadata of the given name, this function does nothing and returns \c false.
+		*/
+	bool RemoveMetadataForElement(const FRigElementKey& InKey, const FName& InName);
+	bool RemoveAllMetadataForElement(const FRigElementKey& InKey);
+	
+	void CopyAllMetadataFromElement(const FRigElementKey& InTargetKey, const FRigBaseElement* InSourceElement);
+	
+protected:
 	bool bEnableCacheValidityCheck;
 
 	static bool bEnableValidityCheckbyDefault;
@@ -4589,7 +4600,7 @@ protected:
 
 	mutable TMap<FRigElementKey, FRigElementKey> DefaultParentPerElement;
 
-	bool bUpdatePreferedEulerAngleWhenSettingTransform;
+	bool bUpdatePreferredEulerAngleWhenSettingTransform;
 	mutable bool bAllowNameSpaceWhenSanitizingName;
 
 private:
@@ -4659,6 +4670,7 @@ private:
 	friend class UBaseControlRig;
 	friend class UControlRig;
 	friend class FControlRigEditor;
+	friend struct FRigBaseElement;
 	friend struct FRigHierarchyValidityBracket;
 	friend struct FRigHierarchyGlobalValidityBracket;
 	friend struct FControlRigVisualGraphUtils;
diff --git a/Engine/Plugins/Animation/ControlRig/Source/ControlRig/Public/Rigs/RigHierarchyDefines.h b/Engine/Plugins/Animation/ControlRig/Source/ControlRig/Public/Rigs/RigHierarchyDefines.h
--- a/Engine/Plugins/Animation/ControlRig/Source/ControlRig/Public/Rigs/RigHierarchyDefines.h
+++ b/Engine/Plugins/Animation/ControlRig/Source/ControlRig/Public/Rigs/RigHierarchyDefines.h
@@ -1452,12 +1452,12 @@ public:
 	UPROPERTY(BlueprintReadWrite, EditAnywhere, Category = "Hierarchy", meta = (CustomWidget = "ElementName"))
 	FName Name;
 
-	FRigElementKey()
-		: Type(ERigElementType::None)
+	FRigElementKey(const ERigElementType InType = ERigElementType::None)
+		: Type(InType)
 		, Name(NAME_None)
 	{}
 
-	FRigElementKey(const FName& InName, ERigElementType InType)
+	FRigElementKey(const FName& InName, const ERigElementType InType)
 		: Type(InType)
 		, Name(InName)
 	{}
diff --git a/Engine/Plugins/Animation/ControlRig/Source/ControlRig/Public/Rigs/RigHierarchyElements.h b/Engine/Plugins/Animation/ControlRig/Source/ControlRig/Public/Rigs/RigHierarchyElements.h
--- a/Engine/Plugins/Animation/ControlRig/Source/ControlRig/Public/Rigs/RigHierarchyElements.h
+++ b/Engine/Plugins/Animation/ControlRig/Source/ControlRig/Public/Rigs/RigHierarchyElements.h
@@ -286,7 +286,7 @@ USTRUCT(BlueprintType)
 struct CONTROLRIG_API FRigLocalAndGlobalTransform
 {
 	GENERATED_BODY()
-
+	
 	FRigLocalAndGlobalTransform()
     : Local()
     , Global()
@@ -576,25 +576,30 @@ public:
 
 	static const EElementIndex ElementTypeIndex;
 
-	FRigBaseElement()
-    : Key()
-    , Index(INDEX_NONE)
-	, SubIndex(INDEX_NONE)
-	, NameString()
-	, CreatedAtInstructionIndex(INDEX_NONE)
-	, OwnedInstances(0)
-	, TopologyVersion(0)
-	, MetadataVersion(0)
-	, bSelected(false)
-	{}
-
-	FRigBaseElement(const FRigBaseElement& InOther);
-	FRigBaseElement& operator= (const FRigBaseElement& InOther);
-
+	FRigBaseElement() = default;
 	virtual ~FRigBaseElement();
 
-	virtual int32 GetElementTypeIndex() const { return (int32)ElementTypeIndex; }
-	int32 GetElementTypeCount() const { return (int32)EElementIndex::Max; }
+	FRigBaseElement(const FRigBaseElement& InOther)
+	{
+		*this = InOther;
+	}
+	
+	FRigBaseElement& operator=(const FRigBaseElement& InOther)
+	{
+		// We purposefully do not copy any non-UPROPERTY entries, including Owner. This is so that when the copied
+		// element is deleted, the metadata is not deleted with it. These copies are purely intended for interfacing
+		// with BP and details view wrappers. 
+		// These copies are solely intended for UBaseControlRig::OnControlSelected_BP
+		Key = InOther.Key;
+		Index = InOther.Index;
+		SubIndex = InOther.SubIndex;
+		CreatedAtInstructionIndex = InOther.CreatedAtInstructionIndex;
+		bSelected = InOther.bSelected;
+		return *this;
+	}
+	
+	virtual int32 GetElementTypeIndex() const { return ElementTypeIndex; }
+	static int32 GetElementTypeCount() { return EElementIndex::Max; }
 
 	enum ESerializationPhase
 	{
@@ -603,41 +608,61 @@ public:
 	};
 
 protected:
+	// Only derived types should be able to construct this one.
+	explicit FRigBaseElement(URigHierarchy* InOwner, ERigElementType InElementType)
+		: Owner(InOwner)
+		, Key(InElementType)
+	{
+	}
 
+	static bool IsClassOf(const FRigBaseElement* InElement)
+	{
+		return true;
+	}
+
+	// ReSharper disable once CppUE4ProbableMemoryIssuesWithUObject
+	URigHierarchy* Owner = nullptr;
+	
 	UPROPERTY(BlueprintReadOnly, VisibleAnywhere, Category = RigElement, meta = (AllowPrivateAccess = "true"))
 	FRigElementKey Key;
 
 	UPROPERTY(BlueprintReadOnly, VisibleAnywhere, Category = RigElement, meta = (AllowPrivateAccess = "true"))
-	int32 Index;
+	int32 Index = INDEX_NONE;
 
 	UPROPERTY(BlueprintReadOnly, VisibleAnywhere, Category = RigElement, meta = (AllowPrivateAccess = "true"))
-	int32 SubIndex;
+	int32 SubIndex = INDEX_NONE;
 
-	TArray<FRigBaseMetadata*> Metadata;
-	TMap<FName,int32> MetadataNameToIndex;
+	UPROPERTY(BlueprintReadOnly, Transient, Category = RigElement, meta = (AllowPrivateAccess = "true"))
+	int32 CreatedAtInstructionIndex = INDEX_NONE;
 
-	static bool IsClassOf(const FRigBaseElement* InElement)
-	{
-		return true;
-	}
+	UPROPERTY(BlueprintReadOnly, Transient, Category = RigElement, meta = (AllowPrivateAccess = "true"))
+	bool bSelected = false;
 
-	mutable TOptional<FString> NameString;
+	// used for constructing / destructing the memory. typically == 1
+	// Set by URigHierarchy::NewElement.
+	int32 OwnedInstances = 0;
+
+	mutable uint16 TopologyVersion = 0;
+	
+	mutable FRigBaseElementChildrenArray CachedChildren;
+	
+	mutable FString CachedNameString;
 
 public:
 
 	UScriptStruct* GetElementStruct() const;
-	void Serialize(FArchive& Ar, URigHierarchy* Hierarchy, ESerializationPhase SerializationPhase);
-	virtual void Save(FArchive& Ar, URigHierarchy* Hierarchy, ESerializationPhase SerializationPhase);
-	virtual void Load(FArchive& Ar, URigHierarchy* Hierarchy, ESerializationPhase SerializationPhase);
+	void Serialize(FArchive& Ar, ESerializationPhase SerializationPhase);
+	virtual void Save(FArchive& Ar, ESerializationPhase SerializationPhase);
+	virtual void Load(FArchive& Ar, ESerializationPhase SerializationPhase);
 
 	const FName& GetFName() const { return Key.Name; }
 	const FString& GetName() const
 	{
-		if(!NameString.IsSet())
+		if(CachedNameString.IsEmpty() && !Key.Name.IsNone())
 		{
-			NameString = Key.Name.ToString();
+			CachedNameString = Key.Name.ToString();
 		}
-		return NameString.GetValue();
+		return CachedNameString;
 	}
 	virtual const FName& GetDisplayName() const { return GetFName(); }
 	ERigElementType GetType() const { return Key.Type; }
@@ -647,41 +672,19 @@ public:
 	bool IsSelected() const { return bSelected; }
 	int32 GetCreatedAtInstructionIndex() const { return CreatedAtInstructionIndex; }
 	bool IsProcedural() const { return CreatedAtInstructionIndex != INDEX_NONE; }
-	int32 GetMetadataVersion() const { return MetadataVersion; }
+	const URigHierarchy* GetOwner() const { return Owner; }
+	URigHierarchy* GetOwner() { return Owner; }
 
-	int32 NumMetadata() const { return Metadata.Num(); }
-	FRigBaseMetadata* GetMetadata(int32 InIndex) const { return Metadata[InIndex]; }
-	FRigBaseMetadata* GetMetadata(const FName& InName) const
-	{
-		if(const int32* MetadataIndex = MetadataNameToIndex.Find(InName))
-		{
-			return GetMetadata(*MetadataIndex);
-		}
-		return nullptr;
-	}
-	FRigBaseMetadata* GetMetadata(const FName& InName, ERigMetadataType InType) const
-	{
-		if(const int32* MetadataIndex = MetadataNameToIndex.Find(InName))
-		{
-			FRigBaseMetadata* Md = GetMetadata(*MetadataIndex);
-			if(Md->GetType() == InType)
-			{
-				return Md;
-			}
-		}
-		return nullptr;
-	}
-	bool SetMetaData(const FName& InName, ERigMetadataType InType, const void* InData, int32 InSize)
-	{
-		if(FRigBaseMetadata* Md = SetupValidMetadata(InName, InType))
-		{
-			return Md->SetValueData(InData, InSize);
-		}
-		return false;
-	}
+	// Metadata
+	FRigBaseMetadata* GetMetadata(const FName& InName, ERigMetadataType InType = ERigMetadataType::Invalid);
+	const FRigBaseMetadata* GetMetadata(const FName& InName, ERigMetadataType InType) const;
+	bool SetMetadata(const FName& InName, ERigMetadataType InType, const void* InData, int32 InSize);
+	FRigBaseMetadata* SetupValidMetadata(const FName& InName, ERigMetadataType InType);
 	bool RemoveMetadata(const FName& InName);
 	bool RemoveAllMetadata();
-
+	
+	void NotifyMetadataTagChanged(const FName& InTag, bool bAdded);
+	
 	template<typename T>
 	bool IsA() const { return T::IsClassOf(this); }
 
@@ -735,66 +738,53 @@ public:
 	virtual void CopyPose(FRigBaseElement* InOther, bool bCurrent, bool bInitial, bool bWeights) {}
 
 protected:
+	// Used to initialize this base element during URigHierarchy::CopyHierarchy. Once all elements are
+	// initialized, the sub-class data copying is done using CopyFrom.
+	void InitializeFrom(const FRigBaseElement* InOther);
 
 	// helper function to be called as part of URigHierarchy::CopyHierarchy
-	virtual void CopyFrom(URigHierarchy* InHierarchy, FRigBaseElement* InOther, URigHierarchy* InOtherHierarchy);
-	
-	// sets up the metadata and ensures the right type
-	FRigBaseMetadata* SetupValidMetadata(const FName& InName, ERigMetadataType InType);
-
-	void NotifyMetadataChanged(const FName& InName);
-	void NotifyMetadataTagChanged(const FName& InTag, bool bAdded);
-
-	mutable FRigBaseElementChildrenArray CachedChildren;
-
-	FRigElementMetadataChangedDelegate MetadataChangedDelegate;
-	FRigElementMetadataTagChangedDelegate MetadataTagChangedDelegate;
-
-	UPROPERTY(BlueprintReadOnly, Transient, Category = RigElement, meta = (AllowPrivateAccess = "true"))
-	int32 CreatedAtInstructionIndex;
-
-	// used for constructing / destructing the memory. typically == 1
-	int32 OwnedInstances;
-
-	mutable uint16 TopologyVersion;
-	mutable uint16 MetadataVersion;
-
-	UPROPERTY(BlueprintReadOnly, Transient, Category = RigElement, meta = (AllowPrivateAccess = "true"))
-	bool bSelected;
+	virtual void CopyFrom(const FRigBaseElement* InOther);
 
+	
+	friend class FControlRigEditor;
 	friend class URigHierarchy;
 	friend class URigHierarchyController;
-	friend struct FRigDispatch_SetMetadata;
-	friend struct FRigUnit_SetMetadataTag;
-	friend struct FRigUnit_SetMetadataTagArray;
-	friend struct FRigUnit_RemoveMetadataTag;
-	friend class FControlRigEditor;
 };
 
 USTRUCT(BlueprintType)
 struct CONTROLRIG_API FRigTransformElement : public FRigBaseElement
 {
-public:
-	
 	GENERATED_BODY()
 	DECLARE_RIG_ELEMENT_METHODS(FRigTransformElement)
 
 	static const EElementIndex ElementTypeIndex;
 
-	FRigTransformElement()
-		: FRigBaseElement()
-	{}
-
-	virtual ~FRigTransformElement(){}
+	FRigTransformElement() = default;
+	FRigTransformElement(const FRigTransformElement& InOther)
+	{
+		*this = InOther;
+	}
+	FRigTransformElement& operator=(const FRigTransformElement& InOther)
+	{
+		Super::operator=(InOther);
+		Pose = InOther.Pose;
+		return *this;
+	}
+	virtual ~FRigTransformElement() override {}
 
-	virtual void Save(FArchive& A, URigHierarchy* Hierarchy, ESerializationPhase SerializationPhase) override;
-	virtual void Load(FArchive& Ar, URigHierarchy* Hierarchy, ESerializationPhase SerializationPhase) override;
+	virtual void Save(FArchive& A, ESerializationPhase SerializationPhase) override;
+	virtual void Load(FArchive& Ar, ESerializationPhase SerializationPhase) override;
 
+	virtual void CopyPose(FRigBaseElement* InOther, bool bCurrent, bool bInitial, bool bWeights) override;
+	
 	UPROPERTY(BlueprintReadOnly, EditAnywhere, Category = RigElement, meta = (DisplayAfter = "Index"))
 	FRigCurrentAndInitialTransform Pose;
 	
 protected:
-
+	FRigTransformElement(URigHierarchy* InOwner, const ERigElementType InType) :
+		FRigBaseElement(InOwner, InType)
+	{}
+	
 	struct FElementToDirty
 	{
 		FElementToDirty()
@@ -835,46 +825,35 @@ protected:
 			InElement->GetType() == ERigElementType::Socket;
 	}
 
-public:
-	
-	virtual void CopyPose(FRigBaseElement* InOther, bool bCurrent, bool bInitial, bool bWeights) override;
-	
-protected:
-
-	virtual void CopyFrom(URigHierarchy* InHierarchy, FRigBaseElement* InOther, URigHierarchy* InOtherHierarchy) override;
+	virtual void CopyFrom(const FRigBaseElement* InOther) override;
 
-	friend struct FRigBaseElement;
-	friend struct FRigSingleParentElement;
-	friend struct FRigMultiParentElement;
 	friend class URigHierarchy;
 	friend class URigHierarchyController;
+	friend struct FRigBaseElement;
 };
 
 USTRUCT(BlueprintType)
 struct CONTROLRIG_API FRigSingleParentElement : public FRigTransformElement
 {
-public:
-	
 	GENERATED_BODY()
 	DECLARE_RIG_ELEMENT_METHODS(FRigSingleParentElement)
 
 	static const EElementIndex ElementTypeIndex;
 
-	FRigSingleParentElement()
-	: FRigTransformElement()
-	, ParentElement(nullptr)
-	{}
-
-	virtual ~FRigSingleParentElement(){}
+	FRigSingleParentElement() = default;
+	virtual ~FRigSingleParentElement() override {}
 
-	virtual void Save(FArchive& A, URigHierarchy* Hierarchy, ESerializationPhase SerializationPhase) override;
-	virtual void Load(FArchive& Ar, URigHierarchy* Hierarchy, ESerializationPhase SerializationPhase) override;
+	virtual void Save(FArchive& A, ESerializationPhase SerializationPhase) override;
+	virtual void Load(FArchive& Ar, ESerializationPhase SerializationPhase) override;
 
-	FRigTransformElement* ParentElement;
+	FRigTransformElement* ParentElement = nullptr;
 
 protected:
+	explicit FRigSingleParentElement(URigHierarchy* InOwner, ERigElementType InType)
+		: FRigTransformElement(InOwner, InType)
+	{}
 
-	virtual void CopyFrom(URigHierarchy* InHierarchy, FRigBaseElement* InOther, URigHierarchy* InOtherHierarchy) override;
+	virtual void CopyFrom(const FRigBaseElement* InOther) override;
 
 	static bool IsClassOf(const FRigBaseElement* InElement)
 	{
@@ -883,7 +862,8 @@ protected:
 			InElement->GetType() == ERigElementType::Reference ||
 			InElement->GetType() == ERigElementType::Socket;
 	}
-
+	
+	friend class URigHierarchy;
 	friend struct FRigBaseElement;
 };
 
@@ -1004,41 +984,37 @@ typedef TArray<FRigElementParentConstraint, TInlineAllocator<1>> FRigElementPare
 USTRUCT(BlueprintType)
 struct CONTROLRIG_API FRigMultiParentElement : public FRigTransformElement
 {
-public:
-	
 	GENERATED_BODY()
 	DECLARE_RIG_ELEMENT_METHODS(FRigMultiParentElement)
 
 	static const EElementIndex ElementTypeIndex;
 
-    FRigMultiParentElement()
-    : FRigTransformElement()
-	{}
+	FRigMultiParentElement() = default;	
+	virtual ~FRigMultiParentElement() override {}
 
-	virtual ~FRigMultiParentElement(){}
+	virtual void Save(FArchive& A, ESerializationPhase SerializationPhase) override;
+	virtual void Load(FArchive& Ar, ESerializationPhase SerializationPhase) override;
 
-	virtual void Save(FArchive& A, URigHierarchy* Hierarchy, ESerializationPhase SerializationPhase) override;
-	virtual void Load(FArchive& Ar, URigHierarchy* Hierarchy, ESerializationPhase SerializationPhase) override;
+	virtual void CopyPose(FRigBaseElement* InOther, bool bCurrent, bool bInitial, bool bWeights) override;
 	
 	FRigElementParentConstraintArray ParentConstraints;
 	TMap<FRigElementKey, int32> IndexLookup;
 
 protected:
+	explicit FRigMultiParentElement(URigHierarchy* InOwner, const ERigElementType InType)
+		: FRigTransformElement(InOwner, InType)
+	{}
 
-	virtual void CopyFrom(URigHierarchy* InHierarchy, FRigBaseElement* InOther, URigHierarchy* InOtherHierarchy) override;
-
+	virtual void CopyFrom(const FRigBaseElement* InOther) override;
+	
+private:
 	static bool IsClassOf(const FRigBaseElement* InElement)
 	{
 		return InElement->GetType() == ERigElementType::Null ||
 			InElement->GetType() == ERigElementType::Control;
 	}
 
-public:
-	
-	virtual void CopyPose(FRigBaseElement* InOther, bool bCurrent, bool bInitial, bool bWeights) override;
-
-protected:
-	
+	friend class URigHierarchy;
 	friend struct FRigBaseElement;
 };
 
@@ -1053,35 +1029,45 @@ public:
 
 	static const EElementIndex ElementTypeIndex;
 
+	UPROPERTY(BlueprintReadOnly, VisibleAnywhere, Category = RigElement)
+	ERigBoneType BoneType = ERigBoneType::User;
+
 	FRigBoneElement()
-		: FRigSingleParentElement()
+		: FRigBoneElement(nullptr)
+	{}
+	FRigBoneElement(const FRigBoneElement& InOther)
 	{
-		Key.Type = ERigElementType::Bone;
-		BoneType = ERigBoneType::User;
+		*this = InOther;
+	}
+	FRigBoneElement& operator=(const FRigBoneElement& InOther)
+	{
+		BoneType = InOther.BoneType;
+		return *this;
 	}
 	
-	virtual ~FRigBoneElement(){}
+	virtual ~FRigBoneElement() override {}
 
-	virtual void Save(FArchive& A, URigHierarchy* Hierarchy, ESerializationPhase SerializationPhase) override;
-	virtual void Load(FArchive& Ar, URigHierarchy* Hierarchy, ESerializationPhase SerializationPhase) override;
-
-	UPROPERTY(BlueprintReadOnly, VisibleAnywhere, Category = RigElement)
-	ERigBoneType BoneType;
+	virtual void Save(FArchive& A, ESerializationPhase SerializationPhase) override;
+	virtual void Load(FArchive& Ar, ESerializationPhase SerializationPhase) override;
 
-protected:
+private:
+	explicit FRigBoneElement(URigHierarchy* InOwner)
+		: FRigSingleParentElement(InOwner, ERigElementType::Bone)
+	{}
 
-	virtual void CopyFrom(URigHierarchy* InHierarchy, FRigBaseElement* InOther, URigHierarchy* InOtherHierarchy) override;
+	virtual void CopyFrom(const FRigBaseElement* InOther) override;
 
 	static bool IsClassOf(const FRigBaseElement* InElement)
 	{
 		return InElement->GetType() == ERigElementType::Bone;
 	}
 
+	friend class URigHierarchy;
 	friend struct FRigBaseElement;
 };
 
 USTRUCT(BlueprintType)
-struct CONTROLRIG_API FRigNullElement : public FRigMultiParentElement
+struct CONTROLRIG_API FRigNullElement final : public FRigMultiParentElement
 {
 public:
 	
@@ -1090,21 +1076,23 @@ public:
 
 	static const EElementIndex ElementTypeIndex;
 
-	FRigNullElement()
-    : FRigMultiParentElement()
-	{
-		Key.Type = ERigElementType::Null; 
-	}
+	FRigNullElement() 
+		: FRigNullElement(nullptr)
+	{}
 
-	virtual ~FRigNullElement(){}
+	virtual ~FRigNullElement() override {}
+
+private:
+	explicit FRigNullElement(URigHierarchy* InOwner)
+		: FRigMultiParentElement(InOwner, ERigElementType::Null)
+	{}
 
-protected:
-	
 	static bool IsClassOf(const FRigBaseElement* InElement)
 	{
 		return InElement->GetType() == ERigElementType::Null;
 	}
 
+	friend class URigHierarchy;
 	friend struct FRigBaseElement;
 };
 
@@ -1248,7 +1236,7 @@ struct CONTROLRIG_API FRigControlSettings
 	EEulerRotationOrder PreferredRotationOrder;
 
 	/**
-	* Whether to use a specfied rotation order or just use the default FRotator order and conversion functions
+	* Whether to use a specified rotation order or just use the default FRotator order and conversion functions
 	*/
 	UPROPERTY(BlueprintReadOnly, EditAnywhere, Category = Animation)
 	bool bUsePreferredRotationOrder;
@@ -1397,7 +1385,7 @@ struct CONTROLRIG_API FRigControlSettings
 };
 
 USTRUCT(BlueprintType)
-struct CONTROLRIG_API FRigControlElement : public FRigMultiParentElement
+struct CONTROLRIG_API FRigControlElement final : public FRigMultiParentElement
 {
 	public:
 	
@@ -1406,13 +1394,39 @@ struct CONTROLRIG_API FRigControlElement : public FRigMultiParentElement
 
 	static const EElementIndex ElementTypeIndex;
 
+	UPROPERTY(BlueprintReadOnly, EditAnywhere, Category = Control)
+	FRigControlSettings Settings;
+
+	UPROPERTY(BlueprintReadOnly, EditAnywhere, Category = RigElement)
+	FRigCurrentAndInitialTransform Offset;
+
+	UPROPERTY(BlueprintReadOnly, EditAnywhere, Category = RigElement)
+	FRigCurrentAndInitialTransform Shape;
+
+	UPROPERTY(BlueprintReadOnly, EditAnywhere, Category = RigElement)
+	FRigPreferredEulerAngles PreferredEulerAngles;
+	
+
 	FRigControlElement()
-		: FRigMultiParentElement()
+		: FRigControlElement(nullptr)
+	{ }
+	
+	FRigControlElement(const FRigControlElement& InOther)
+	{
+		*this = InOther;
+	}
+	
+	FRigControlElement& operator=(const FRigControlElement& InOther)
 	{
-		Key.Type = ERigElementType::Control; 
+		Super::operator=(InOther);
+		Settings = InOther.Settings;
+		Offset = InOther.Offset;
+		Shape = InOther.Shape;
+		PreferredEulerAngles = InOther.PreferredEulerAngles;
+		return *this;
 	}
 
-	virtual ~FRigControlElement(){}
+	virtual ~FRigControlElement() override {}
 	
 	virtual const FName& GetDisplayName() const override
 	{
@@ -1427,90 +1441,67 @@ struct CONTROLRIG_API FRigControlElement : public FRigMultiParentElement
 
 	bool CanDriveControls() const { return Settings.AnimationType == ERigControlAnimationType::ProxyControl || Settings.AnimationType == ERigControlAnimationType::AnimationControl; }
 
-	virtual void Save(FArchive& A, URigHierarchy* Hierarchy, ESerializationPhase SerializationPhase) override;
-	virtual void Load(FArchive& Ar, URigHierarchy* Hierarchy, ESerializationPhase SerializationPhase) override;
+	virtual void Save(FArchive& A, ESerializationPhase SerializationPhase) override;
+	virtual void Load(FArchive& Ar, ESerializationPhase SerializationPhase) override;
 
+	virtual void CopyPose(FRigBaseElement* InOther, bool bCurrent, bool bInitial, bool bWeights) override;
+	
 private:
+	explicit FRigControlElement(URigHierarchy* InOwner)
+		: FRigMultiParentElement(InOwner, ERigElementType::Control)
+	{ }
 
-	virtual void CopyFrom(URigHierarchy* InHierarchy, FRigBaseElement* InOther, URigHierarchy* InOtherHierarchy) override;
-
-public:
-
-	UPROPERTY(BlueprintReadOnly, EditAnywhere, Category = Control)
-	FRigControlSettings Settings;
-
-	UPROPERTY(BlueprintReadOnly, EditAnywhere, Category = RigElement)
-	FRigCurrentAndInitialTransform Offset;
-
-	UPROPERTY(BlueprintReadOnly, EditAnywhere, Category = RigElement)
-	FRigCurrentAndInitialTransform Shape;
-
-	UPROPERTY(BlueprintReadOnly, EditAnywhere, Category = RigElement)
-	FRigPreferredEulerAngles PreferredEulerAngles;
-	
-protected:
+	virtual void CopyFrom(const FRigBaseElement* InOther) override;
 
 	static bool IsClassOf(const FRigBaseElement* InElement)
 	{
 		return InElement->GetType() == ERigElementType::Control;
 	}
 
-public:
-	
-	virtual void CopyPose(FRigBaseElement* InOther, bool bCurrent, bool bInitial, bool bWeights) override;
-
-protected:
-
-	friend struct FRigBaseElement;
 	friend class URigHierarchy;
-	friend class URigHierarchyController;
+	friend struct FRigBaseElement;
 };
 
 USTRUCT(BlueprintType)
-struct CONTROLRIG_API FRigCurveElement : public FRigBaseElement
+struct CONTROLRIG_API FRigCurveElement final : public FRigBaseElement
 {
-public:
-	
 	GENERATED_BODY()
 	DECLARE_RIG_ELEMENT_METHODS(FRigCurveElement)
 
 	static const EElementIndex ElementTypeIndex;
 
-	FRigCurveElement()
-		: FRigBaseElement()
-		, bIsValueSet(true)
-		, Value(0.f)
-	{
-		Key.Type = ERigElementType::Curve;
-	}
+	// Set to true if the value was actually set. Used to carry back and forth blend curve
+	// value validity state.
+	bool bIsValueSet = true;
+	
+	float Value = 0.0f;
 
+	
+	FRigCurveElement()
+		: FRigCurveElement(nullptr)
+	{}
+	
 	virtual ~FRigCurveElement() override {}
 
-	virtual void Save(FArchive& A, URigHierarchy* Hierarchy, ESerializationPhase SerializationPhase) override;
-	virtual void Load(FArchive& Ar, URigHierarchy* Hierarchy, ESerializationPhase SerializationPhase) override;
+	virtual void Save(FArchive& A, ESerializationPhase SerializationPhase) override;
+	virtual void Load(FArchive& Ar, ESerializationPhase SerializationPhase) override;
 
+	virtual void CopyPose(FRigBaseElement* InOther, bool bCurrent, bool bInitial, bool bWeights) override;
+	
 private:
-
-	virtual void CopyFrom(URigHierarchy* InHierarchy, FRigBaseElement* InOther, URigHierarchy* InOtherHierarchy) override;
-
-public:
-	// Set to true if the value was actually set. Used to carry back and forth blend curve
-	// value validity state.
-	bool bIsValueSet;
+	FRigCurveElement(URigHierarchy* InOwner)
+		: FRigBaseElement(InOwner, ERigElementType::Curve)
+	{}
 	
-	float Value;
+	virtual void CopyFrom(const FRigBaseElement* InOther) override;
 
 	static bool IsClassOf(const FRigBaseElement* InElement)
 	{
 		return InElement->GetType() == ERigElementType::Curve;
 	}
 
-public:
-	
-	virtual void CopyPose(FRigBaseElement* InOther, bool bCurrent, bool bInitial, bool bWeights) override;
-
-protected:
-	
+	friend class URigHierarchy;
+	friend class URigHierarchyController;
 	friend struct FRigBaseElement;
 };
 
@@ -1539,38 +1530,48 @@ public:
 
 	static const EElementIndex ElementTypeIndex;
 
-    FRigRigidBodyElement()
-        : FRigSingleParentElement()
+	UPROPERTY(BlueprintReadOnly, EditAnywhere, Category = Control, meta=(ShowOnlyInnerProperties))
+	FRigRigidBodySettings Settings;
+	
+	FRigRigidBodyElement()
+        : FRigRigidBodyElement(nullptr)
+	{ }
+	
+	FRigRigidBodyElement(const FRigRigidBodyElement& InOther)
 	{
-		Key.Type = ERigElementType::RigidBody;
+		*this = InOther;
 	}
 	
-	virtual ~FRigRigidBodyElement(){}
+	FRigRigidBodyElement& operator=(const FRigRigidBodyElement& InOther)
+	{
+		Super::operator=(InOther);
+		Settings = InOther.Settings;
+		return *this;
+	}
+	
+	virtual ~FRigRigidBodyElement() override {}
 
-	virtual void Save(FArchive& A, URigHierarchy* Hierarchy, ESerializationPhase SerializationPhase) override;
-	virtual void Load(FArchive& Ar, URigHierarchy* Hierarchy, ESerializationPhase SerializationPhase) override;
+	virtual void Save(FArchive& A, ESerializationPhase SerializationPhase) override;
+	virtual void Load(FArchive& Ar, ESerializationPhase SerializationPhase) override;
 
 private:
-
-	virtual void CopyFrom(URigHierarchy* InHierarchy, FRigBaseElement* InOther, URigHierarchy* InOtherHierarchy) override;
-
-public:
-
-	UPROPERTY(BlueprintReadOnly, EditAnywhere, Category = Control, meta=(ShowOnlyInnerProperties))
-	FRigRigidBodySettings Settings;
-
-protected:
+	explicit FRigRigidBodyElement(URigHierarchy* InOwner)
+		: FRigSingleParentElement(InOwner, ERigElementType::RigidBody)
+	{ }
 	
+	virtual void CopyFrom(const FRigBaseElement* InOther) override;
+
 	static bool IsClassOf(const FRigBaseElement* InElement)
 	{
 		return InElement->GetType() == ERigElementType::RigidBody;
 	}
 
+	friend class URigHierarchy;
 	friend struct FRigBaseElement;
 };
 
 USTRUCT(BlueprintType)
-struct CONTROLRIG_API FRigReferenceElement : public FRigSingleParentElement
+struct CONTROLRIG_API FRigReferenceElement final : public FRigSingleParentElement
 {
 public:
 	
@@ -1580,33 +1581,35 @@ public:
 	static const EElementIndex ElementTypeIndex;
 
     FRigReferenceElement()
-        : FRigSingleParentElement()
-	{
-		Key.Type = ERigElementType::Reference;
-	}
+        : FRigReferenceElement(nullptr)
+	{ }
 	
-	virtual ~FRigReferenceElement(){}
+	virtual ~FRigReferenceElement() override {}
 
-	virtual void Save(FArchive& A, URigHierarchy* Hierarchy, ESerializationPhase SerializationPhase) override;
-	virtual void Load(FArchive& Ar, URigHierarchy* Hierarchy, ESerializationPhase SerializationPhase) override;
+	virtual void Save(FArchive& A, ESerializationPhase SerializationPhase) override;
+	virtual void Load(FArchive& Ar, ESerializationPhase SerializationPhase) override;
 
 	FTransform GetReferenceWorldTransform(const FRigVMExecuteContext* InContext, bool bInitial) const;
 
 	virtual void CopyPose(FRigBaseElement* InOther, bool bCurrent, bool bInitial, bool bWeights) override;
 
-protected:
+private:
+	explicit FRigReferenceElement(URigHierarchy* InOwner)
+		: FRigSingleParentElement(InOwner, ERigElementType::Reference)
+	{ }
 
 	FRigReferenceGetWorldTransformDelegate GetWorldTransformDelegate;
 
-	virtual void CopyFrom(URigHierarchy* InHierarchy, FRigBaseElement* InOther, URigHierarchy* InOtherHierarchy) override;
+	virtual void CopyFrom(const FRigBaseElement* InOther) override;
 
 	static bool IsClassOf(const FRigBaseElement* InElement)
 	{
 		return InElement->GetType() == ERigElementType::Reference;
 	}
 
-	friend struct FRigBaseElement;
+	friend class URigHierarchy;
 	friend class URigHierarchyController;
+	friend struct FRigBaseElement;
 };
 
 USTRUCT(BlueprintType)
@@ -1687,12 +1690,6 @@ struct CONTROLRIG_API FRigConnectorInfo
 {
 	GENERATED_BODY()
 	
-	FRigConnectorInfo()
-		: Name(NAME_None)
-		, ResolvedTarget()
-		, Settings()
-	{}
-	
 	UPROPERTY(BlueprintReadWrite, EditAnywhere, Category = Connector)
 	FName Name;
 
@@ -1704,92 +1701,90 @@ struct CONTROLRIG_API FRigConnectorInfo
 };
 
 USTRUCT(BlueprintType)
-struct CONTROLRIG_API FRigConnectorElement : public FRigBaseElement
+struct CONTROLRIG_API FRigConnectorElement final : public FRigBaseElement
 {
-public:
-	
 	GENERATED_BODY()
 	DECLARE_RIG_ELEMENT_METHODS(FRigConnectorElement)
 
 	static const EElementIndex ElementTypeIndex;
 
+	UPROPERTY(BlueprintReadOnly, EditAnywhere, Category = Control)
+	FRigConnectorSettings Settings;
+	
 	FRigConnectorElement()
-		: FRigBaseElement()
+		: FRigConnectorElement(nullptr)
+	{}
+	FRigConnectorElement(const FRigConnectorElement& InOther)
+	{
+		*this = InOther;
+	}
+	FRigConnectorElement& operator=(const FRigConnectorElement& InOther)
 	{
-		Key.Type = ERigElementType::Connector; 
+		Super::operator=(InOther);
+		Settings = InOther.Settings;
+		return *this;
 	}
 
-	virtual ~FRigConnectorElement(){}
+	virtual ~FRigConnectorElement() override {}
 	
-	virtual void Save(FArchive& A, URigHierarchy* Hierarchy, ESerializationPhase SerializationPhase) override;
-	virtual void Load(FArchive& Ar, URigHierarchy* Hierarchy, ESerializationPhase SerializationPhase) override;
+	virtual void Save(FArchive& A, ESerializationPhase SerializationPhase) override;
+	virtual void Load(FArchive& Ar, ESerializationPhase SerializationPhase) override;
 
 	bool CanConnect(const FRigConnectionInfo* InConnectionInfo, FString* OutFailureReason) const;
 
 	FRigConnectorInfo GetConnectorInfo(const URigHierarchy* InHierarchy) const;
-	
-private:
-
-	virtual void CopyFrom(URigHierarchy* InHierarchy, FRigBaseElement* InOther, URigHierarchy* InOtherHierarchy) override;
-
-public:
 
-	UPROPERTY(BlueprintReadOnly, EditAnywhere, Category = Control)
-	FRigConnectorSettings Settings;
-
-protected:
+private:
+	explicit FRigConnectorElement(URigHierarchy* InOwner)
+		: FRigBaseElement(InOwner, ERigElementType::Connector)
+	{ }
+	
+	virtual void CopyFrom(const FRigBaseElement* InOther) override;
 
 	static bool IsClassOf(const FRigBaseElement* InElement)
 	{
 		return InElement->GetType() == ERigElementType::Connector;
 	}
 
-protected:
-
-	friend struct FRigBaseElement;
 	friend class URigHierarchy;
-	friend class URigHierarchyController;
+	friend struct FRigBaseElement;
 };
 
+
 USTRUCT(BlueprintType)
-struct CONTROLRIG_API FRigSocketElement : public FRigSingleParentElement
+struct CONTROLRIG_API FRigSocketElement final : public FRigSingleParentElement
 {
-public:
-	
 	GENERATED_BODY()
 	DECLARE_RIG_ELEMENT_METHODS(FRigSocketElement)
 
 	static const EElementIndex ElementTypeIndex;
 
 	FRigSocketElement()
-		: FRigSingleParentElement()
-	{
-		Key.Type = ERigElementType::Socket; 
-	}
-
-	virtual ~FRigSocketElement(){}
+		: FRigSocketElement(nullptr)
+	{}
+			
+	virtual ~FRigSocketElement() override {}
 	
-	virtual void Save(FArchive& A, URigHierarchy* Hierarchy, ESerializationPhase SerializationPhase) override;
-	virtual void Load(FArchive& Ar, URigHierarchy* Hierarchy, ESerializationPhase SerializationPhase) override;
+	virtual void Save(FArchive& A, ESerializationPhase SerializationPhase) override;
+	virtual void Load(FArchive& Ar, ESerializationPhase SerializationPhase) override;
 
 private:
-
-	virtual void CopyFrom(URigHierarchy* InHierarchy, FRigBaseElement* InOther, URigHierarchy* InOtherHierarchy) override;
-
-protected:
+	explicit FRigSocketElement(URigHierarchy* InOwner)
+		: FRigSingleParentElement(InOwner, ERigElementType::Socket)
+	{ }
+	
+	virtual void CopyFrom(const FRigBaseElement* InOther) override;
 
 	static bool IsClassOf(const FRigBaseElement* InElement)
 	{
 		return InElement->GetType() == ERigElementType::Socket;
 	}
 
-protected:
-
-	friend struct FRigBaseElement;
 	friend class URigHierarchy;
-	friend class URigHierarchyController;
+	friend struct FRigBaseElement;
 };
 
+
 USTRUCT()
 struct CONTROLRIG_API FRigHierarchyCopyPasteContentPerElement
 {
diff --git a/Engine/Plugins/Animation/ControlRig/Source/ControlRig/Public/Rigs/RigHierarchyMetadata.h b/Engine/Plugins/Animation/ControlRig/Source/ControlRig/Public/Rigs/RigHierarchyMetadata.h
--- a/Engine/Plugins/Animation/ControlRig/Source/ControlRig/Public/Rigs/RigHierarchyMetadata.h
+++ b/Engine/Plugins/Animation/ControlRig/Source/ControlRig/Public/Rigs/RigHierarchyMetadata.h
@@ -36,28 +36,15 @@ struct CONTROLRIG_API FRigBaseMetadata
 {
 	GENERATED_BODY()
 
-public:
-
-	FRigBaseMetadata()
-    	: Element(nullptr)
-		, Name(NAME_None)
-		, Type(ERigMetadataType::Invalid)
-		, ValueProperty(nullptr)
-	{}
-
-	virtual ~FRigBaseMetadata(){}
-
 protected:
 
-	const FRigBaseElement* Element;
-
 	UPROPERTY()
 	FName Name;
 
 	UPROPERTY()
-	ERigMetadataType Type;
+	ERigMetadataType Type = ERigMetadataType::Invalid;
 
-	mutable const FProperty* ValueProperty;
+	mutable const FProperty* ValueProperty = nullptr;
 
 	static bool IsClassOf(const FRigBaseMetadata* InMetadata)
 	{
@@ -65,15 +52,14 @@ protected:
 	}
 
 public:
+	virtual ~FRigBaseMetadata() = default;
 
 	UScriptStruct* GetMetadataStruct() const;
-	virtual void Serialize(FArchive& Ar, bool bIsLoading);
+	virtual void Serialize(FArchive& Ar);
 
-	bool IsValid() const { return (GetType() != ERigMetadataType::Invalid) && GetKey().IsValid(); }
-	const FRigBaseElement* GetElement() const { return Element; }
-	const FRigElementKey& GetKey() const;
+	bool IsValid() const { return GetType() != ERigMetadataType::Invalid; }
 	const FName& GetName() const { return Name; }
-	const ERigMetadataType& GetType() const { return Type; }
+	ERigMetadataType GetType() const { return Type; }
 	bool IsArray() const
 	{
 		return GetValueProperty()->IsA<FArrayProperty>();
@@ -146,7 +132,7 @@ public:
 protected:
 
 	static UScriptStruct* GetMetadataStruct(const ERigMetadataType& InType);
-	static FRigBaseMetadata* MakeMetadata(const FRigBaseElement* InElement, const FName& InName, ERigMetadataType InType);
+	static FRigBaseMetadata* MakeMetadata(const FName& InName, ERigMetadataType InType);
 	static void DestroyMetadata(FRigBaseMetadata** Metadata);
 
 	const FProperty* GetValueProperty() const
@@ -188,14 +174,9 @@ public:
 	GENERATED_BODY()
 	DECLARE_RIG_METADATA_METHODS(FRigBoolMetadata)
 
-	FRigBoolMetadata()
-		: FRigBaseMetadata()
-		, Value(false)
-	{}
-	virtual ~FRigBoolMetadata() override {}
-	virtual void Serialize(FArchive& Ar, bool bIsLoading) override
+	virtual void Serialize(FArchive& Ar) override
 	{
-		Super::Serialize(Ar, bIsLoading);
+		Super::Serialize(Ar);
 		Ar << Value;
 	}
 	const bool& GetValue() const { return Value; }
@@ -210,24 +191,18 @@ public:
 protected:
 
 	UPROPERTY()
-	bool Value;
+	bool Value = false;
 };
 
 USTRUCT()
 struct CONTROLRIG_API FRigBoolArrayMetadata : public FRigBaseMetadata
 {
-public:
-	
 	GENERATED_BODY()
 	DECLARE_RIG_METADATA_METHODS(FRigBoolArrayMetadata)
 
-	FRigBoolArrayMetadata()
-		: FRigBaseMetadata()
-	{}
-	virtual ~FRigBoolArrayMetadata() override {}
-	virtual void Serialize(FArchive& Ar, bool bIsLoading) override
+	virtual void Serialize(FArchive& Ar) override
 	{
-		Super::Serialize(Ar, bIsLoading);
+		Super::Serialize(Ar);
 		Ar << Value;
 	}
 	const TArray<bool>& GetValue() const { return Value; }
@@ -253,14 +228,9 @@ public:
 	GENERATED_BODY()
 	DECLARE_RIG_METADATA_METHODS(FRigFloatMetadata)
 
-	FRigFloatMetadata()
-		: FRigBaseMetadata()
-		, Value(0.f)
-	{}
-	virtual ~FRigFloatMetadata() override {}
-	virtual void Serialize(FArchive& Ar, bool bIsLoading) override
+	virtual void Serialize(FArchive& Ar) override
 	{
-		Super::Serialize(Ar, bIsLoading);
+		Super::Serialize(Ar);
 		Ar << Value;
 	}
 	const float& GetValue() const { return Value; } 
@@ -275,7 +245,7 @@ public:
 protected:
 
 	UPROPERTY()
-	float Value;
+	float Value = 0.f;
 };
 
 USTRUCT()
@@ -286,13 +256,9 @@ public:
 	GENERATED_BODY()
 	DECLARE_RIG_METADATA_METHODS(FRigFloatArrayMetadata)
 
-	FRigFloatArrayMetadata()
-		: FRigBaseMetadata()
-	{}
-	virtual ~FRigFloatArrayMetadata() override {}
-	virtual void Serialize(FArchive& Ar, bool bIsLoading) override
+	virtual void Serialize(FArchive& Ar) override
 	{
-		Super::Serialize(Ar, bIsLoading);
+		Super::Serialize(Ar);
 		Ar << Value;
 	}
 	const TArray<float>& GetValue() const { return Value; }
@@ -318,14 +284,9 @@ public:
 	GENERATED_BODY()
 	DECLARE_RIG_METADATA_METHODS(FRigInt32Metadata)
 
-	FRigInt32Metadata()
-		: FRigBaseMetadata()
-		, Value(0)
-	{}
-	virtual ~FRigInt32Metadata() override {}
-	virtual void Serialize(FArchive& Ar, bool bIsLoading) override
+	virtual void Serialize(FArchive& Ar) override
 	{
-		Super::Serialize(Ar, bIsLoading);
+		Super::Serialize(Ar);
 		Ar << Value;
 	}
 	const int32& GetValue() const { return Value; } 
@@ -340,7 +301,7 @@ public:
 protected:
 
 	UPROPERTY()
-	int32 Value;
+	int32 Value = 0;
 };
 
 USTRUCT()
@@ -351,13 +312,9 @@ public:
 	GENERATED_BODY()
 	DECLARE_RIG_METADATA_METHODS(FRigInt32ArrayMetadata)
 
-	FRigInt32ArrayMetadata()
-		: FRigBaseMetadata()
-	{}
-	virtual ~FRigInt32ArrayMetadata() override {}
-	virtual void Serialize(FArchive& Ar, bool bIsLoading) override
+	virtual void Serialize(FArchive& Ar) override
 	{
-		Super::Serialize(Ar, bIsLoading);
+		Super::Serialize(Ar);
 		Ar << Value;
 	}
 	const TArray<int32>& GetValue() const { return Value; }
@@ -383,14 +340,9 @@ public:
 	GENERATED_BODY()
 	DECLARE_RIG_METADATA_METHODS(FRigNameMetadata)
 
-	FRigNameMetadata()
-		: FRigBaseMetadata()
-		, Value(NAME_None)
-	{}
-	virtual ~FRigNameMetadata() override {}
-	virtual void Serialize(FArchive& Ar, bool bIsLoading) override
+	virtual void Serialize(FArchive& Ar) override
 	{
-		Super::Serialize(Ar, bIsLoading);
+		Super::Serialize(Ar);
 		Ar << Value;
 	}
 	const FName& GetValue() const { return Value; } 
@@ -416,13 +368,9 @@ public:
 	GENERATED_BODY()
 	DECLARE_RIG_METADATA_METHODS(FRigNameArrayMetadata)
 
-	FRigNameArrayMetadata()
-		: FRigBaseMetadata()
-	{}
-	virtual ~FRigNameArrayMetadata() override {}
-	virtual void Serialize(FArchive& Ar, bool bIsLoading) override
+	virtual void Serialize(FArchive& Ar) override
 	{
-		Super::Serialize(Ar, bIsLoading);
+		Super::Serialize(Ar);
 		Ar << Value;
 	}
 	const TArray<FName>& GetValue() const { return Value; }
@@ -448,14 +396,9 @@ public:
 	GENERATED_BODY()
 	DECLARE_RIG_METADATA_METHODS(FRigVectorMetadata)
 
-	FRigVectorMetadata()
-		: FRigBaseMetadata()
-		, Value(FVector::ZeroVector)
-	{}
-	virtual ~FRigVectorMetadata() override {}
-	virtual void Serialize(FArchive& Ar, bool bIsLoading) override
+	virtual void Serialize(FArchive& Ar) override
 	{
-		Super::Serialize(Ar, bIsLoading);
+		Super::Serialize(Ar);
 		Ar << Value;
 	}
 	const FVector& GetValue() const { return Value; } 
@@ -470,7 +413,7 @@ public:
 protected:
 
 	UPROPERTY()
-	FVector Value;
+	FVector Value = FVector::ZeroVector;
 };
 
 USTRUCT()
@@ -481,13 +424,9 @@ public:
 	GENERATED_BODY()
 	DECLARE_RIG_METADATA_METHODS(FRigVectorArrayMetadata)
 
-	FRigVectorArrayMetadata()
-		: FRigBaseMetadata()
-	{}
-	virtual ~FRigVectorArrayMetadata() override {}
-	virtual void Serialize(FArchive& Ar, bool bIsLoading) override
+	virtual void Serialize(FArchive& Ar) override
 	{
-		Super::Serialize(Ar, bIsLoading);
+		Super::Serialize(Ar);
 		Ar << Value;
 	}
 	const TArray<FVector>& GetValue() const { return Value; }
@@ -513,14 +452,9 @@ public:
 	GENERATED_BODY()
 	DECLARE_RIG_METADATA_METHODS(FRigRotatorMetadata)
 
-	FRigRotatorMetadata()
-		: FRigBaseMetadata()
-		, Value(FRotator::ZeroRotator)
-	{}
-	virtual ~FRigRotatorMetadata() override {}
-	virtual void Serialize(FArchive& Ar, bool bIsLoading) override
+	virtual void Serialize(FArchive& Ar) override
 	{
-		Super::Serialize(Ar, bIsLoading);
+		Super::Serialize(Ar);
 		Ar << Value;
 	}
 	const FRotator& GetValue() const { return Value; } 
@@ -535,7 +469,7 @@ public:
 protected:
 
 	UPROPERTY()
-	FRotator Value;
+	FRotator Value = FRotator::ZeroRotator;
 };
 
 USTRUCT()
@@ -546,13 +480,9 @@ public:
 	GENERATED_BODY()
 	DECLARE_RIG_METADATA_METHODS(FRigRotatorArrayMetadata)
 
-	FRigRotatorArrayMetadata()
-		: FRigBaseMetadata()
-	{}
-	virtual ~FRigRotatorArrayMetadata() override {}
-	virtual void Serialize(FArchive& Ar, bool bIsLoading) override
+	virtual void Serialize(FArchive& Ar) override
 	{
-		Super::Serialize(Ar, bIsLoading);
+		Super::Serialize(Ar);
 		Ar << Value;
 	}
 	const TArray<FRotator>& GetValue() const { return Value; }
@@ -578,14 +508,9 @@ public:
 	GENERATED_BODY()
 	DECLARE_RIG_METADATA_METHODS(FRigQuatMetadata)
 
-	FRigQuatMetadata()
-		: FRigBaseMetadata()
-		, Value(FQuat::Identity)
-	{}
-	virtual ~FRigQuatMetadata() override {}
-	virtual void Serialize(FArchive& Ar, bool bIsLoading) override
+	virtual void Serialize(FArchive& Ar) override
 	{
-		Super::Serialize(Ar, bIsLoading);
+		Super::Serialize(Ar);
 		Ar << Value;
 	}
 	const FQuat& GetValue() const { return Value; } 
@@ -600,7 +525,7 @@ public:
 protected:
 
 	UPROPERTY()
-	FQuat Value;
+	FQuat Value = FQuat::Identity;
 };
 
 USTRUCT()
@@ -611,13 +536,9 @@ public:
 	GENERATED_BODY()
 	DECLARE_RIG_METADATA_METHODS(FRigQuatArrayMetadata)
 
-	FRigQuatArrayMetadata()
-		: FRigBaseMetadata()
-	{}
-	virtual ~FRigQuatArrayMetadata() override {}
-	virtual void Serialize(FArchive& Ar, bool bIsLoading) override
+	virtual void Serialize(FArchive& Ar) override
 	{
-		Super::Serialize(Ar, bIsLoading);
+		Super::Serialize(Ar);
 		Ar << Value;
 	}
 	const TArray<FQuat>& GetValue() const { return Value; }
@@ -643,14 +564,9 @@ public:
 	GENERATED_BODY()
 	DECLARE_RIG_METADATA_METHODS(FRigTransformMetadata)
 
-	FRigTransformMetadata()
-		: FRigBaseMetadata()
-		, Value(FTransform::Identity)
-	{}
-	virtual ~FRigTransformMetadata() override {}
-	virtual void Serialize(FArchive& Ar, bool bIsLoading) override
+	virtual void Serialize(FArchive& Ar) override
 	{
-		Super::Serialize(Ar, bIsLoading);
+		Super::Serialize(Ar);
 		Ar << Value;
 	}
 	const FTransform& GetValue() const { return Value; } 
@@ -665,7 +581,7 @@ public:
 protected:
 
 	UPROPERTY()
-	FTransform Value;
+	FTransform Value = FTransform::Identity;
 };
 
 USTRUCT()
@@ -676,13 +592,9 @@ public:
 	GENERATED_BODY()
 	DECLARE_RIG_METADATA_METHODS(FRigTransformArrayMetadata)
 
-	FRigTransformArrayMetadata()
-		: FRigBaseMetadata()
-	{}
-	virtual ~FRigTransformArrayMetadata() override {}
-	virtual void Serialize(FArchive& Ar, bool bIsLoading) override
+	virtual void Serialize(FArchive& Ar) override
 	{
-		Super::Serialize(Ar, bIsLoading);
+		Super::Serialize(Ar);
 		Ar << Value;
 	}
 	const TArray<FTransform>& GetValue() const { return Value; }
@@ -708,14 +620,9 @@ public:
 	GENERATED_BODY()
 	DECLARE_RIG_METADATA_METHODS(FRigLinearColorMetadata)
 
-	FRigLinearColorMetadata()
-		: FRigBaseMetadata()
-		, Value(FLinearColor::White)
-	{}
-	virtual ~FRigLinearColorMetadata() override {}
-	virtual void Serialize(FArchive& Ar, bool bIsLoading) override
+	virtual void Serialize(FArchive& Ar) override
 	{
-		Super::Serialize(Ar, bIsLoading);
+		Super::Serialize(Ar);
 		Ar << Value;
 	}
 	const FLinearColor& GetValue() const { return Value; } 
@@ -730,7 +637,7 @@ public:
 protected:
 
 	UPROPERTY()
-	FLinearColor Value;
+	FLinearColor Value = FLinearColor::White;
 };
 
 USTRUCT()
@@ -741,13 +648,9 @@ public:
 	GENERATED_BODY()
 	DECLARE_RIG_METADATA_METHODS(FRigLinearColorArrayMetadata)
 
-	FRigLinearColorArrayMetadata()
-		: FRigBaseMetadata()
-	{}
-	virtual ~FRigLinearColorArrayMetadata() override {}
-	virtual void Serialize(FArchive& Ar, bool bIsLoading) override
+	virtual void Serialize(FArchive& Ar) override
 	{
-		Super::Serialize(Ar, bIsLoading);
+		Super::Serialize(Ar);
 		Ar << Value;
 	}
 	const TArray<FLinearColor>& GetValue() const { return Value; }
@@ -760,7 +663,6 @@ public:
 	}
 
 protected:
-
 	UPROPERTY()
 	TArray<FLinearColor> Value;
 };
@@ -773,13 +675,9 @@ public:
 	GENERATED_BODY()
 	DECLARE_RIG_METADATA_METHODS(FRigElementKeyMetadata)
 
-	FRigElementKeyMetadata()
-		: FRigBaseMetadata()
-	{}
-	virtual ~FRigElementKeyMetadata() override {}
-	virtual void Serialize(FArchive& Ar, bool bIsLoading) override
+	virtual void Serialize(FArchive& Ar) override
 	{
-		Super::Serialize(Ar, bIsLoading);
+		Super::Serialize(Ar);
 		Ar << Value;
 	}
 	const FRigElementKey& GetValue() const { return Value; } 
@@ -805,13 +703,9 @@ public:
 	GENERATED_BODY()
 	DECLARE_RIG_METADATA_METHODS(FRigElementKeyArrayMetadata)
 
-	FRigElementKeyArrayMetadata()
-		: FRigBaseMetadata()
-	{}
-	virtual ~FRigElementKeyArrayMetadata() override {}
-	virtual void Serialize(FArchive& Ar, bool bIsLoading) override
+	virtual void Serialize(FArchive& Ar) override
 	{
-		Super::Serialize(Ar, bIsLoading);
+		Super::Serialize(Ar);
 		Ar << Value;
 	}
 	const TArray<FRigElementKey>& GetValue() const { return Value; }
diff --git a/Engine/Plugins/Animation/ControlRig/Source/ControlRigEditor/Private/ControlRigElementDetails.cpp b/Engine/Plugins/Animation/ControlRig/Source/ControlRigEditor/Private/ControlRigElementDetails.cpp
--- a/Engine/Plugins/Animation/ControlRig/Source/ControlRigEditor/Private/ControlRigElementDetails.cpp
+++ b/Engine/Plugins/Animation/ControlRig/Source/ControlRigEditor/Private/ControlRigElementDetails.cpp
@@ -1221,10 +1221,12 @@ void FRigBaseElementDetails::CustomizeMetadata(IDetailLayoutBuilder& DetailBuild
 		return;
 	}
 
+	URigHierarchy* Hierarchy = nullptr;
 	if (!MetadataHandle.IsValid())
 	{
 		const FPerElementInfo& Info = PerElementInfos[0];
-		URigHierarchy* Hierarchy = Info.IsValid() ? Info.GetHierarchy() : nullptr;
+		
+		Hierarchy = Info.IsValid() ? Info.GetHierarchy() : nullptr;
 		if (!Hierarchy)
 		{
 			return;
@@ -1241,17 +1243,19 @@ void FRigBaseElementDetails::CustomizeMetadata(IDetailLayoutBuilder& DetailBuild
 		});
 	}
 	
-	const FRigBaseElement* Element = PerElementInfos[0].Element.Get();
-	if(Element->NumMetadata() == 0)
+	FRigBaseElement* Element = PerElementInfos[0].Element.Get();
+	TArray<FName> MetadataNames = Element->GetOwner()->GetMetadataNames(Element->GetKey());
+	
+	if(MetadataNames.IsEmpty())
 	{
 		return;
 	}
 
 	IDetailCategoryBuilder& MetadataCategory = DetailBuilder.EditCategory(TEXT("Metadata"), LOCTEXT("Metadata", "Metadata"));
-	for(int32 Index=0;Index<Element->NumMetadata();Index++)
+	for(FName MetadataName: MetadataNames)
 	{
-		FRigBaseMetadata* Metadata = Element->GetMetadata(Index);
-		TSharedPtr<FStructOnScope> StructOnScope = MakeShareable(new FStructOnScope(Metadata->GetMetadataStruct(), (uint8*)Metadata));
+		FRigBaseMetadata* Metadata = Element->GetMetadata(MetadataName);
+		TSharedPtr<FStructOnScope> StructOnScope = MakeShareable(new FStructOnScope(Metadata->GetMetadataStruct(), reinterpret_cast<uint8*>(Metadata)));
 
 		FAddPropertyParams Params;
 		Params.CreateCategoryNodes(false);
diff --git a/Engine/Plugins/Animation/ControlRig/Source/ControlRigEditor/Private/Graph/ControlRigGraphPanelPinFactory.cpp b/Engine/Plugins/Animation/ControlRig/Source/ControlRigEditor/Private/Graph/ControlRigGraphPanelPinFactory.cpp
--- a/Engine/Plugins/Animation/ControlRig/Source/ControlRigEditor/Private/Graph/ControlRigGraphPanelPinFactory.cpp
+++ b/Engine/Plugins/Animation/ControlRig/Source/ControlRigEditor/Private/Graph/ControlRigGraphPanelPinFactory.cpp
@@ -357,9 +357,9 @@ TSharedPtr<SGraphPin> FControlRigGraphPanelPinFactory::CreatePin_Internal(UEdGra
 										for(int32 ElementIndex=0; ElementIndex < ControlRig->GetHierarchy()->Num(); ElementIndex++)
 										{
 											const FRigBaseElement* OtherElement = ControlRig->GetHierarchy()->Get(ElementIndex);
-											for(int32 MetadataIndex = 0; MetadataIndex < OtherElement->NumMetadata(); MetadataIndex++)
+											for(FName MetadataName: ControlRig->GetHierarchy()->GetMetadataNames(OtherElement->GetKey()))
 											{
-												Names.AddUnique(OtherElement->GetMetadata(MetadataIndex)->GetName());
+												Names.AddUnique(MetadataName);
 											}
 										}
 
