Control Rig: Add ModularControlRig class

#virtualized

diff --git a/Engine/Plugins/Animation/ControlRig/Source/ControlRig/Private/BaseControlRig.cpp b/Engine/Plugins/Animation/ControlRig/Source/ControlRig/Private/BaseControlRig.cpp
--- a/Engine/Plugins/Animation/ControlRig/Source/ControlRig/Private/BaseControlRig.cpp
+++ b/Engine/Plugins/Animation/ControlRig/Source/ControlRig/Private/BaseControlRig.cpp
@@ -333,7 +333,7 @@ void UBaseControlRig::Evaluate_AnyThread()
 	{
 		if (bRequiresInitExecution)
 		{
-			EvaluateVMs_AnyThread();
+			Super::Evaluate_AnyThread();
 		}
 		
 		// we can have other systems trying to poke into running instances of Control Rigs
@@ -427,7 +427,7 @@ void UBaseControlRig::Evaluate_AnyThread()
 	}
 	else
 	{
-		EvaluateVMs_AnyThread();
+		Super::Evaluate_AnyThread();
 	}
 }
 
@@ -493,38 +493,41 @@ void UBaseControlRig::InitializeFromCDO()
 	// copy CDO property you need to here
 	if (!HasAnyFlags(RF_ClassDefaultObject))
 	{
-		// similar to FControlRigBlueprintCompilerContext::CopyTermDefaultsToDefaultObject,
-		// where CDO is initialized from BP there,
-		// we initialize all other instances of Control Rig from the CDO here
-		UBaseControlRig* CDO = GetClass()->GetDefaultObject<UBaseControlRig>();
-		URigHierarchy* Hierarchy = GetHierarchy();
-
-		// copy hierarchy
+		if (!IsRigModuleInstance())
 		{
-			FRigHierarchyValidityBracket ValidityBracketA(Hierarchy);
-			FRigHierarchyValidityBracket ValidityBracketB(CDO->GetHierarchy());
+			// similar to FControlRigBlueprintCompilerContext::CopyTermDefaultsToDefaultObject,
+			// where CDO is initialized from BP there,
+			// we initialize all other instances of Control Rig from the CDO here
+			UBaseControlRig* CDO = GetClass()->GetDefaultObject<UBaseControlRig>();
+			URigHierarchy* Hierarchy = GetHierarchy();
+
+			// copy hierarchy
+			{
+				FRigHierarchyValidityBracket ValidityBracketA(Hierarchy);
+				FRigHierarchyValidityBracket ValidityBracketB(CDO->GetHierarchy());
 			
-			TGuardValue<bool> Guard(Hierarchy->GetSuspendNotificationsFlag(), true);
-			Hierarchy->CopyHierarchy(CDO->GetHierarchy());
-			Hierarchy->ResetPoseToInitial(ERigElementType::All);
-		}
+				TGuardValue<bool> Guard(Hierarchy->GetSuspendNotificationsFlag(), true);
+				Hierarchy->CopyHierarchy(CDO->GetHierarchy());
+				Hierarchy->ResetPoseToInitial(ERigElementType::All);
+			}
 
 #if WITH_EDITOR
-		// current hierarchy should always mirror CDO's hierarchy whenever a change of interest happens
-		CDO->GetHierarchy()->RegisterListeningHierarchy(Hierarchy);
+			// current hierarchy should always mirror CDO's hierarchy whenever a change of interest happens
+			CDO->GetHierarchy()->RegisterListeningHierarchy(Hierarchy);
 #endif
 
-		// notify clients that the hierarchy has changed
-		Hierarchy->Notify(ERigHierarchyNotification::HierarchyReset, nullptr);
+			// notify clients that the hierarchy has changed
+			Hierarchy->Notify(ERigHierarchyNotification::HierarchyReset, nullptr);
 
-		// copy hierarchy settings
-		HierarchySettings = CDO->HierarchySettings;
-		ElementKeyRedirector = FRigElementKeyRedirector(CDO->ElementKeyRedirector, Hierarchy); 
+			// copy hierarchy settings
+			HierarchySettings = CDO->HierarchySettings;
+			ElementKeyRedirector = FRigElementKeyRedirector(CDO->ElementKeyRedirector, Hierarchy); 
 		
-		// increment the procedural limit based on the number of elements in the CDO
-		if(const URigHierarchy* CDOHierarchy = CDO->GetHierarchy())
-		{
-			HierarchySettings.ProceduralElementLimit += CDOHierarchy->Num();
+			// increment the procedural limit based on the number of elements in the CDO
+			if(const URigHierarchy* CDOHierarchy = CDO->GetHierarchy())
+			{
+				HierarchySettings.ProceduralElementLimit += CDOHierarchy->Num();
+			}
 		}
 
 		ExternalVariableDataAssetLinks.Reset();
@@ -569,28 +572,29 @@ bool UBaseControlRig::AllConnectorsAreResolved(FString* OutFailureReason, FRigEl
 		// todo: introduce a cache here based on ElementKeyDirector hash and
 		// topology hash of the hierarchy
 		
-		const TArray<FRigConnectorElement*> Connectors = Hierarchy->GetConnectors(false);
+		const TArray<FRigModuleConnector> Connectors = GetRigModuleSettings().ExposedConnectors;
 
 		// collect the connection map
 		TMap<FRigElementKey, FRigElementKey> ConnectionMap;
-		for(const FRigConnectorElement* Connector : Connectors)
+		for(const FRigModuleConnector& Connector : Connectors)
 		{
-			if(const FCachedRigElement* Cache = ElementKeyRedirector.Find(Connector->GetKey()))
+			const FRigElementKey ConnectorKey(*Connector.Name, ERigElementType::Connector);
+			if(const FCachedRigElement* Cache = ElementKeyRedirector.Find(ConnectorKey))
 			{
 				if(const_cast<FCachedRigElement*>(Cache)->UpdateCache(Hierarchy))
 				{
-					ConnectionMap.Add(Connector->GetKey(), Cache->GetKey());
+					ConnectionMap.Add(ConnectorKey, Cache->GetKey());
 				}
 				else
 				{
 					if(OutFailureReason)
 					{
-						static constexpr TCHAR Format[] = TEXT("Connector '%s' has invalid target '%s'.");
-						*OutFailureReason = FString::Printf(Format, *Connector->GetName(), *Cache->GetKey().ToString());
+						static constexpr TCHAR Format[] = TEXT("Connector '%s%s' has invalid target '%s'.");
+						*OutFailureReason = FString::Printf(Format, *GetRigModuleNameSpace(), *Connector.Name, *Cache->GetKey().ToString());
 					}
 					if(OutConnector)
 					{
-						*OutConnector = Connector->GetKey();
+						*OutConnector = ConnectorKey;
 					}
 					return false;
 				}
@@ -600,11 +604,11 @@ bool UBaseControlRig::AllConnectorsAreResolved(FString* OutFailureReason, FRigEl
 				if(OutFailureReason)
 				{
 					static constexpr TCHAR Format[] = TEXT("Connector '%s' is not resolved.");
-					*OutFailureReason = FString::Printf(Format, *Connector->GetName());
+					*OutFailureReason = FString::Printf(Format, *Connector.Name);
 				}
 				if(OutConnector)
 				{
-					*OutConnector = Connector->GetKey();
+					*OutConnector = ConnectorKey;
 				}
 				return false;
 			}
@@ -1570,7 +1574,7 @@ const FString& UBaseControlRig::GetRigModuleNameSpace() const
 		{
 			const FString& ParentNameSpace = ParentRig->GetRigModuleNameSpace();
 			static constexpr TCHAR JoinFormat[] = TEXT("%s%s:");
-			RigModuleNameSpace = FString::Printf(JoinFormat, *ParentNameSpace, *GetName());
+			RigModuleNameSpace = FString::Printf(JoinFormat, *ParentNameSpace, *GetFName().ToString());
 			return RigModuleNameSpace;
 		}
 
@@ -1603,6 +1607,11 @@ FRigElementKeyRedirector& UBaseControlRig::GetElementKeyRedirector()
 	return EmptyRedirector;
 }
 
+void UBaseControlRig::SetElementKeyRedirector(const FRigElementKeyRedirector InElementRedirector)
+{
+	ElementKeyRedirector = InElementRedirector;
+}
+
 TArray<FRigControlElement*> UBaseControlRig::AvailableControls() const
 {
 	if(DynamicHierarchy)
@@ -3072,6 +3081,18 @@ void UBaseControlRig::PostInitInstance(URigVMHost* InCDO)
 	RequestInit();
 }
 
+void UBaseControlRig::SetDynamicHierarchy(TObjectPtr<URigHierarchy> InHierarchy)
+{
+	// Delete any existing hierarchy
+	if (DynamicHierarchy->GetOuter() == this)
+	{
+		DynamicHierarchy->OnUndoRedo().RemoveAll(this);
+		DynamicHierarchy->Rename(nullptr, GetTransientPackage(), REN_ForceNoResetLoaders | REN_DoNotDirty | REN_DontCreateRedirectors | REN_NonTransactional);
+		DynamicHierarchy->MarkAsGarbage();
+	}
+	DynamicHierarchy = InHierarchy;
+}
+
 UTransformableControlHandle* UBaseControlRig::CreateTransformableControlHandle(
 	const FName& InControlName) const
 {
diff --git a/Engine/Plugins/Animation/ControlRig/Source/ControlRig/Private/ControlRig.cpp b/Engine/Plugins/Animation/ControlRig/Source/ControlRig/Private/ControlRig.cpp
--- a/Engine/Plugins/Animation/ControlRig/Source/ControlRig/Private/ControlRig.cpp
+++ b/Engine/Plugins/Animation/ControlRig/Source/ControlRig/Private/ControlRig.cpp
@@ -168,11 +168,20 @@ bool UControlRig::Execute_Internal(const FName& InEventName)
 		FRigHierarchyExecuteContextBracket HierarchyContextGuard(Hierarchy, &Context);
 
 		// setup the module information
-		FControlRigExecuteContext& PublicContext = Context.GetPublicDataSafe<FControlRigExecuteContext>();
-		FControlRigExecuteContextRigModuleGuard RigModuleGuard(PublicContext, this);
-		FRigHierarchyRedirectorGuard ElementRedirectorGuard(this);
+		bool bSuccess;
+		if (!IsRigModuleInstance())
+		{
+			FControlRigExecuteContext& PublicContext = Context.GetPublicDataSafe<FControlRigExecuteContext>();
+			FControlRigExecuteContextRigModuleGuard RigModuleGuard(PublicContext, this);
+			FRigHierarchyRedirectorGuard ElementRedirectorGuard(this);
+
+			bSuccess = VM->ExecuteVM(Context, InEventName) != ERigVMExecuteResult::Failed;
+		}
+		else
+		{
+			bSuccess = VM->ExecuteVM(Context, InEventName) != ERigVMExecuteResult::Failed;
+		}
 
-		const bool bSuccess = VM->ExecuteVM(Context, InEventName) != ERigVMExecuteResult::Failed;
 
 #if UE_RIGVM_PROFILE_EXECUTE_UNITS_NUM
 		const uint64 EndCycles = FPlatformTime::Cycles64();
@@ -191,11 +200,5 @@ bool UControlRig::Execute_Internal(const FName& InEventName)
 	return false;
 }
 
-#if WITH_EDITOR
-void UControlRig::SetFirstEntryEventInEventQueue(FRigVMExtendedExecuteContext& Context, const FName& InFirstEventName)
-{
-	VM->SetFirstEntryEventInEventQueue(Context, NAME_None);
-}
-#endif
 
 #undef LOCTEXT_NAMESPACE
\ No newline at end of file
diff --git a/Engine/Plugins/Animation/ControlRig/Source/ControlRig/Private/ModularRig.cpp b/Engine/Plugins/Animation/ControlRig/Source/ControlRig/Private/ModularRig.cpp
new file mode 100644
--- /dev/null
+++ b/Engine/Plugins/Animation/ControlRig/Source/ControlRig/Private/ModularRig.cpp
@@ -0,0 +1,246 @@
+// Copyright Epic Games, Inc. All Rights Reserved.
+
+#include "ModularRig.h"
+#include "Units/Execution/RigUnit_BeginExecution.h"
+#include "Units/Execution/RigUnit_InverseExecution.h"
+#include "Units/Execution/RigUnit_PrepareForExecution.h"
+#include "Units/Execution/RigUnit_InteractionExecution.h"
+#include "ControlRigObjectBinding.h"
+#include "Rigs/RigHierarchyController.h"
+#include "ControlRigComponent.h"
+
+#define LOCTEXT_NAMESPACE "ModularRig"
+
+const FString UModularRig::NamespaceSeparator = TEXT(":");
+
+FRigModuleInstance::~FRigModuleInstance()
+{
+	CachedChildren.Reset();
+	if (Rig && Rig.IsValid())
+	{
+		Rig->Rename(nullptr, GetTransientPackage(), REN_ForceNoResetLoaders | REN_DoNotDirty | REN_DontCreateRedirectors | REN_NonTransactional);
+		Rig->MarkAsGarbage();
+		Rig = nullptr;
+	}
+}
+
+UModularRig::UModularRig(const FObjectInitializer& ObjectInitializer)
+	: Super(ObjectInitializer)
+{
+#if WITH_EDITOR
+	FCoreUObjectDelegates::OnObjectsReplaced.AddUObject(this, &UModularRig::OnObjectsReplaced);
+#endif
+}
+
+void UModularRig::BeginDestroy()
+{
+	Super::BeginDestroy();
+	
+#if WITH_EDITOR
+	FCoreUObjectDelegates::OnObjectsReplaced.RemoveAll(this);
+#endif
+}
+
+void UModularRig::InitializeVMs(bool bRequestInit)
+{
+	URigVMHost::Initialize(bRequestInit);
+	ForEachModule([bRequestInit](FRigModuleInstance* Module) -> bool
+	{
+		if (Module->Rig.IsValid())
+		{
+			Module->Rig->InitializeVMs(bRequestInit);
+		}
+		return true;
+	});
+}
+
+bool UModularRig::InitializeVMs(const FName& InEventName)
+{
+	URigVMHost::InitializeVM(InEventName);
+	ForEachModule([InEventName](FRigModuleInstance* Module) -> bool
+	{
+		if (Module->Rig.IsValid())
+		{
+			Module->Rig->InitializeVMs(InEventName);
+		}
+		return true;
+	});
+	return true;
+}
+
+bool UModularRig::Execute_Internal(const FName& InEventName)
+{
+	if (VM)
+	{
+		FRigVMExtendedExecuteContext& Context = GetRigVMExtendedExecuteContext();
+		URigHierarchy* Hierarchy = GetHierarchy();
+#if WITH_EDITOR
+
+		bool bRecordTransformsAtRuntime = true;
+		if(const UObject* Outer = GetOuter())
+		{
+			if(Outer->IsA<UControlRigComponent>())
+			{
+				bRecordTransformsAtRuntime = false;
+			}
+		}
+		TGuardValue<bool> RecordTransformsPerInstructionGuard(Hierarchy->bRecordTransformsAtRuntime, bRecordTransformsAtRuntime);
+		
+		if(Hierarchy->bRecordTransformsAtRuntime)
+		{
+			Hierarchy->ReadTransformsAtRuntime.Reset();
+			Hierarchy->WrittenTransformsAtRuntime.Reset();
+		}
+		
+#endif
+		
+		ForEachModule([&InEventName, this, Hierarchy](FRigModuleInstance* Module) -> bool
+		{
+			if (Module->Rig.IsValid())
+			{
+				UBaseControlRig* Rig = Module->Rig.Get();
+				
+				// Make sure the hierarchy has the correct element redirector from this module rig
+				FRigHierarchyRedirectorGuard ElementRedirectorGuard(Rig);
+
+				// Make sure the hierarchy has the correct execute context with the rig module namespace
+				FRigHierarchyExecuteContextBracket ExecuteContextBracket(Hierarchy, &Rig->GetRigVMExtendedExecuteContext());
+				
+				Rig->Execute_Internal(InEventName);
+			}
+			return true;
+		});
+
+		return true;
+	}
+	return false;
+}
+
+void UModularRig::OnObjectsReplaced(const TMap<UObject*, UObject*>& OldToNewInstanceMap)
+{
+	for (FRigModuleInstance& ModuleInstance : Modules)
+	{
+		if (UObject*const * NewObject = OldToNewInstanceMap.Find(ModuleInstance.Rig.Get()))
+		{
+			ModuleInstance.Rig = Cast<UBaseControlRig>(*NewObject);
+		}
+	}
+	InitializeVMs(true);
+}
+
+void UModularRig::ResetModules()
+{
+	RootModules.Reset();
+	Modules.Reset();
+}
+
+bool UModularRig::AddModuleInstance(const FName& InModuleName, TSubclassOf<UBaseControlRig> InModuleClass, FString InParentPath, const TMap<FRigElementKey, FRigElementKey>& InConnectionMap)
+{
+	FRigModuleInstance* ParentModule = FindModule(InParentPath);
+	return AddModuleInstance(InModuleName, InModuleClass, ParentModule, InConnectionMap) != nullptr;
+}
+
+FRigModuleInstance* UModularRig::AddModuleInstance(const FName& InModuleName, TSubclassOf<UBaseControlRig> InModuleClass, FRigModuleInstance* InParent, const TMap<FRigElementKey, FRigElementKey>& InConnectionMap) 
+{
+	// Make sure there are no name clashes
+	if (InParent)
+	{
+		for (FRigModuleInstance* Child : InParent->CachedChildren)
+		{
+			if (Child->Name == InModuleName)
+			{
+				return nullptr;
+			}
+		}
+	}
+	else
+	{
+		for (FRigModuleInstance* RootModule : RootModules)
+		{
+			if (RootModule->Name == InModuleName)
+			{
+				return nullptr;
+			}
+		}
+	}
+
+	// For now, lets only allow rig modules
+	if (!InModuleClass->GetDefaultObject<UBaseControlRig>()->IsRigModule())
+	{
+		return nullptr;
+	}
+
+	FString Name = (InParent) ? InParent->Name.ToString() + NamespaceSeparator + InModuleName.ToString() : InModuleName.ToString();
+	FRigModuleInstance& NewModule = Modules.Add_GetRef(FRigModuleInstance());
+	NewModule.Rig = NewObject<UBaseControlRig>(this, InModuleClass, *Name, RF_Transient | RF_Transactional);
+	NewModule.Name = InModuleName;
+
+	if (InParent)
+	{
+		InParent->CachedChildren.Add(&NewModule);
+		NewModule.ParentPath = InParent->ParentPath.IsEmpty() ? InParent->Name.ToString() : InParent->ParentPath + NamespaceSeparator + InParent->Name.ToString();
+	}
+	else
+	{
+		RootModules.Add(&NewModule);
+	}
+
+	// Configure module
+	{
+		URigHierarchy* Hierarchy = GetHierarchy();
+		FRigVMExtendedExecuteContext& ModuleContext = NewModule.Rig->GetRigVMExtendedExecuteContext();
+		FControlRigExecuteContext& ModulePublicContext = ModuleContext.GetPublicDataSafe<FControlRigExecuteContext>();
+		NewModule.Rig->bCopyHierarchyBeforeConstruction = false;
+		NewModule.Rig->SetDynamicHierarchy(Hierarchy);
+		ModulePublicContext.Hierarchy = Hierarchy;
+		ModulePublicContext.RigModuleNameSpace = NewModule.Rig->GetRigModuleNameSpace();
+		ModulePublicContext.RigModuleNameSpaceHash = GetTypeHash(ModulePublicContext.RigModuleNameSpace);
+		NewModule.Rig->SetElementKeyRedirector(FRigElementKeyRedirector(InConnectionMap, Hierarchy));
+	}
+	
+	return &NewModule;
+}
+
+FRigModuleInstance* UModularRig::FindModule(const FString& InPath)
+{
+	TArray<FRigModuleInstance*>* Children = &RootModules;
+	FString Left = InPath, Right;
+	while (Left.Split(NamespaceSeparator, &Left, &Right))
+	{
+		FRigModuleInstance** Cur = Children->FindByPredicate([Left](FRigModuleInstance* Module)
+		{
+			return Module->Name == Left;
+		});
+		if (!Cur)
+		{
+			return nullptr;
+		}
+		Children = &(*Cur)->CachedChildren;
+		Left = Right;
+	}
+
+	FRigModuleInstance** Cur = Children->FindByPredicate([Left](FRigModuleInstance* Module)
+		{
+			return Module->Name == Left;
+		});
+	if (!Cur)
+	{
+		return nullptr;
+	}
+	return *Cur;
+}
+
+void UModularRig::ForEachModule(TFunctionRef<bool(FRigModuleInstance*)> PerModuleFunction)
+{
+	TArray<FRigModuleInstance*> ModuleInstances = RootModules;
+	for (int32 ModuleIndex = 0; ModuleIndex < ModuleInstances.Num(); ++ModuleIndex)
+	{
+		if (!PerModuleFunction(ModuleInstances[ModuleIndex]))
+		{
+			break;
+		}
+		ModuleInstances.Append(ModuleInstances[ModuleIndex]->CachedChildren);
+	}
+}
+
+#undef LOCTEXT_NAMESPACE
diff --git a/Engine/Plugins/Animation/ControlRig/Source/ControlRig/Private/Rigs/RigHierarchy.cpp b/Engine/Plugins/Animation/ControlRig/Source/ControlRig/Private/Rigs/RigHierarchy.cpp
--- a/Engine/Plugins/Animation/ControlRig/Source/ControlRig/Private/Rigs/RigHierarchy.cpp
+++ b/Engine/Plugins/Animation/ControlRig/Source/ControlRig/Private/Rigs/RigHierarchy.cpp
@@ -2386,6 +2386,9 @@ void URigHierarchy::Traverse(TFunction<void(FRigBaseElement*, bool& /* continue
 
 bool URigHierarchy::CanConnect(const FRigConnectionInfo* InConnectionInfo, FString* OutFailureReason, FRigElementKey* OutConnector) const
 {
+	// todo: Check if the rules are satisfied
+	return true;
+	
 	check(InConnectionInfo);
 	check(InConnectionInfo->IsValid());
 
@@ -2502,15 +2505,17 @@ bool URigHierarchy::CanConnect(const FRigConnectionInfo* InConnectionInfo, FStri
 
 const FRigElementKey& URigHierarchy::GetResolvedTarget(const FRigElementKey& InConnectorKey) const
 {
-	if(ElementKeyRedirector)
+	if (InConnectorKey.Type == ERigElementType::Connector)
 	{
-		if(const FCachedRigElement* Target = ElementKeyRedirector->Find(InConnectorKey))
+		if(ElementKeyRedirector)
 		{
-			return Target->GetKey();
+			if(const FCachedRigElement* Target = ElementKeyRedirector->Find(InConnectorKey))
+			{
+				return Target->GetKey();
+			}
 		}
 	}
-	static const FRigElementKey InValidKey;
-	return InValidKey;
+	return InConnectorKey;
 }
 
 bool URigHierarchy::Undo()
diff --git a/Engine/Plugins/Animation/ControlRig/Source/ControlRig/Private/Rigs/RigHierarchyController.cpp b/Engine/Plugins/Animation/ControlRig/Source/ControlRig/Private/Rigs/RigHierarchyController.cpp
--- a/Engine/Plugins/Animation/ControlRig/Source/ControlRig/Private/Rigs/RigHierarchyController.cpp
+++ b/Engine/Plugins/Animation/ControlRig/Source/ControlRig/Private/Rigs/RigHierarchyController.cpp
@@ -4,6 +4,7 @@
 #include "BaseControlRig.h"
 #include "AnimationCoreLibrary.h"
 #include "UObject/Package.h"
+#include "ModularRig.h"
 
 #if WITH_EDITOR
 #include "Framework/Notifications/NotificationManager.h"
@@ -1996,6 +1997,7 @@ bool URigHierarchyController::RemoveElement(FRigBaseElement* InElement)
 
 	// make sure this element is part of this hierarchy
 	ensure(Hierarchy->FindChecked(InElement->Key) == InElement);
+	ensure(InElement->OwnedInstances == 1);
 
 	// deselect if needed
 	if(InElement->IsSelected())
@@ -2096,7 +2098,10 @@ bool URigHierarchyController::RemoveElement(FRigBaseElement* InElement)
 		Notify(ERigHierarchyNotification::HierarchyReset, nullptr);
 	}
 
-	Hierarchy->DestroyElement(InElement);
+	if (InElement->OwnedInstances == 1)
+	{
+		Hierarchy->DestroyElement(InElement);
+	}
 
 	Hierarchy->EnsureCacheValidity();
 
diff --git a/Engine/Plugins/Animation/ControlRig/Source/ControlRig/Private/Units/Execution/RigUnit_Collection.cpp b/Engine/Plugins/Animation/ControlRig/Source/ControlRig/Private/Units/Execution/RigUnit_Collection.cpp
--- a/Engine/Plugins/Animation/ControlRig/Source/ControlRig/Private/Units/Execution/RigUnit_Collection.cpp
+++ b/Engine/Plugins/Animation/ControlRig/Source/ControlRig/Private/Units/Execution/RigUnit_Collection.cpp
@@ -115,7 +115,7 @@ FRigUnit_CollectionChildrenArray_Execute()
 	}
 	
 	uint32 Hash = GetTypeHash(StaticStruct()) + ExecuteContext.Hierarchy->GetTopologyVersion() * 17;
-	Hash = HashCombine(Hash, GetTypeHash(Parent));
+	Hash = HashCombine(Hash, GetTypeHash(ExecuteContext.Hierarchy->GetResolvedTarget(Parent)));
 	Hash = HashCombine(Hash, bRecursive ? 2 : 0);
 	Hash = HashCombine(Hash, bIncludeParent ? 1 : 0);
 	Hash = HashCombine(Hash, (int32)TypeToSearch * 8);
diff --git a/Engine/Plugins/Animation/ControlRig/Source/ControlRig/Private/Units/RigUnitContext.cpp b/Engine/Plugins/Animation/ControlRig/Source/ControlRig/Private/Units/RigUnitContext.cpp
--- a/Engine/Plugins/Animation/ControlRig/Source/ControlRig/Private/Units/RigUnitContext.cpp
+++ b/Engine/Plugins/Animation/ControlRig/Source/ControlRig/Private/Units/RigUnitContext.cpp
@@ -80,7 +80,16 @@ FControlRigExecuteContextRigModuleGuard::FControlRigExecuteContextRigModuleGuard
 	Context.RigModuleNameSpace = InControlRig->GetRigModuleNameSpace();
 	Context.RigModuleNameSpaceHash = GetTypeHash(Context.RigModuleNameSpace);
 }
-	
+
+FControlRigExecuteContextRigModuleGuard::FControlRigExecuteContextRigModuleGuard(FControlRigExecuteContext& InContext, const FString& InNewModuleNameSpace)
+	: Context(InContext)
+	, PreviousRigModuleNameSpace(InContext.RigModuleNameSpace)
+	, PreviousRigModuleNameSpaceHash(InContext.RigModuleNameSpaceHash)
+{
+	Context.RigModuleNameSpace = InNewModuleNameSpace;
+	Context.RigModuleNameSpaceHash = GetTypeHash(Context.RigModuleNameSpace);
+}
+
 FControlRigExecuteContextRigModuleGuard::~FControlRigExecuteContextRigModuleGuard()
 {
 	Context.RigModuleNameSpace = PreviousRigModuleNameSpace;
diff --git a/Engine/Plugins/Animation/ControlRig/Source/ControlRig/Public/BaseControlRig.h b/Engine/Plugins/Animation/ControlRig/Source/ControlRig/Public/BaseControlRig.h
--- a/Engine/Plugins/Animation/ControlRig/Source/ControlRig/Public/BaseControlRig.h
+++ b/Engine/Plugins/Animation/ControlRig/Source/ControlRig/Public/BaseControlRig.h
@@ -102,7 +102,10 @@ public:
 	const FString& GetRigModuleNameSpace() const;
 
 	// Returns the redirector from key to key for this rig
-	FRigElementKeyRedirector& GetElementKeyRedirector();
+	virtual FRigElementKeyRedirector& GetElementKeyRedirector();
+
+	// Returns the redirector from key to key for this rig
+	virtual void SetElementKeyRedirector(const FRigElementKeyRedirector InElementRedirector);
 
 	/** Creates a transformable control handle for the specified control to be used by the constraints system. Should use the UObject from 
 	ConstraintsScriptingLibrary::GetManager(UWorld* InWorld)*/
@@ -130,9 +133,6 @@ public:
 	/** Evaluates the ControlRig */
 	virtual void Evaluate_AnyThread() override;
 
-	/** Calls Evaluate_AnyThread on all VMs in this Rig */
-	virtual void EvaluateVMs_AnyThread() PURE_VIRTUAL(UBaseControlRig::EvaluateVMs_AnyThread, return Super::Evaluate_AnyThread(); );
-
 	/** Ticks animation of the skeletal mesh component bound to this control rig */
 	bool EvaluateSkeletalMeshComponent(double InDeltaTime);
 
@@ -366,6 +366,8 @@ public:
 
 	virtual USceneComponent* GetOwningSceneComponent() override;
 
+	void SetDynamicHierarchy(TObjectPtr<URigHierarchy> InHierarchy);
+
 protected:
 
 	virtual void PostInitInstance(URigVMHost* InCDO) override;
@@ -596,7 +598,7 @@ protected:
 
 
 #if WITH_EDITOR
-	virtual void SetFirstEntryEventInEventQueue(FRigVMExtendedExecuteContext& Context, const FName& InFirstEventName) PURE_VIRTUAL(UBaseControlRig::SetFirstEntryEventInEventQueue, return; );
+	virtual void SetFirstEntryEventInEventQueue(FRigVMExtendedExecuteContext& Context, const FName& InFirstEventName) { VM->SetFirstEntryEventInEventQueue(Context, NAME_None); }
 
 	static void OnHierarchyTransformUndoRedoWeak(URigHierarchy* InHierarchy, const FRigElementKey& InKey, ERigTransformType::Type InTransformType, const FTransform& InTransform, bool bIsUndo, TWeakObjectPtr<UBaseControlRig> WeakThis)
 	{
@@ -796,6 +798,7 @@ private:
 	friend class AControlRigShapeActor;
 	friend class FRigTransformElementDetails;
 	friend class FControlRigEditorModule;
+	friend class UModularRig;
 };
 
 class CONTROLRIG_API FControlRigBracketScope
diff --git a/Engine/Plugins/Animation/ControlRig/Source/ControlRig/Public/ControlRig.h b/Engine/Plugins/Animation/ControlRig/Source/ControlRig/Public/ControlRig.h
--- a/Engine/Plugins/Animation/ControlRig/Source/ControlRig/Public/ControlRig.h
+++ b/Engine/Plugins/Animation/ControlRig/Source/ControlRig/Public/ControlRig.h
@@ -19,10 +19,6 @@ public:
 	virtual void InitializeVMsFromCDO() override { URigVMHost::InitializeFromCDO(); }
 	virtual void RequestInitVMs() override { URigVMHost::RequestInit(); }
 	virtual bool Execute_Internal(const FName& InEventName) override;
-	virtual void EvaluateVMs_AnyThread() override { URigVMHost::Evaluate_AnyThread(); }
-#if WITH_EDITOR
-	virtual void SetFirstEntryEventInEventQueue(FRigVMExtendedExecuteContext& Context, const FName& InFirstEventName) override;
-#endif
 	// END BaseControlRig
 
 	UE_DEPRECATED(5.4, "InteractionRig is no longer used") UFUNCTION(BlueprintGetter, meta = (DeprecatedFunction, DeprecationMessage = "InteractionRig is no longer used"))
diff --git a/Engine/Plugins/Animation/ControlRig/Source/ControlRig/Public/ModularRig.h b/Engine/Plugins/Animation/ControlRig/Source/ControlRig/Public/ModularRig.h
new file mode 100644
--- /dev/null
+++ b/Engine/Plugins/Animation/ControlRig/Source/ControlRig/Public/ModularRig.h
@@ -0,0 +1,72 @@
+// Copyright Epic Games, Inc. All Rights Reserved.
+
+#pragma once
+#include "BaseControlRig.h"
+
+#include "ModularRig.generated.h"
+
+USTRUCT(BlueprintType)
+struct CONTROLRIG_API FRigModuleInstance 
+{
+	GENERATED_USTRUCT_BODY()
+
+	FRigModuleInstance()
+	: Name(NAME_None)
+	, Rig(nullptr)
+	, ParentPath(FString())
+	{
+	}
+
+	~FRigModuleInstance();
+
+	UPROPERTY()
+	FName Name;
+
+	UPROPERTY()
+	TSoftObjectPtr<UBaseControlRig> Rig;
+
+	UPROPERTY()
+	FString ParentPath;
+
+	TArray<FRigModuleInstance*> CachedChildren;
+};
+
+/** Runs logic for mapping input data to transforms (the "Rig") */
+UCLASS(Blueprintable, Abstract, editinlinenew)
+class CONTROLRIG_API UModularRig : public UBaseControlRig
+{
+	GENERATED_UCLASS_BODY()
+
+	UPROPERTY()
+	TArray<FRigModuleInstance> Modules;
+	TArray<FRigModuleInstance*> RootModules;
+
+public:
+
+	// BEGIN BaseControlRig
+	virtual void InitializeVMs(bool bRequestInit = true) override;
+	virtual bool InitializeVMs(const FName& InEventName) override;
+	virtual void InitializeVMsFromCDO() override { URigVMHost::InitializeFromCDO(); }
+	virtual void RequestInitVMs() override { URigVMHost::RequestInit(); }
+	virtual bool Execute_Internal(const FName& InEventName) override;
+	virtual FRigElementKeyRedirector& GetElementKeyRedirector() override { return ElementKeyRedirector; }
+	// END BaseControlRig
+
+	// BEGIN UObject
+	virtual void BeginDestroy() override;
+	// END UObject
+	
+	void OnObjectsReplaced(const TMap<UObject*, UObject*>& OldToNewInstanceMap);
+
+	void ResetModules();
+
+	/** Adds a module to the rig*/
+	bool AddModuleInstance(const FName& InModuleName, TSubclassOf<UBaseControlRig> InModuleClass, FString InParentPath, const TMap<FRigElementKey, FRigElementKey>& InConnectionMap);
+	FRigModuleInstance* AddModuleInstance(const FName& InModuleName, TSubclassOf<UBaseControlRig> InModuleClass, FRigModuleInstance* InParent, const TMap<FRigElementKey, FRigElementKey>& InConnectionMap);
+
+	FRigModuleInstance* FindModule(const FString& InPath);
+
+	void ForEachModule(TFunctionRef<bool(FRigModuleInstance*)> PerModuleFunction);
+
+	static const FString NamespaceSeparator;
+};
diff --git a/Engine/Plugins/Animation/ControlRig/Source/ControlRig/Public/Rigs/AdditiveControlRig.h b/Engine/Plugins/Animation/ControlRig/Source/ControlRig/Public/Rigs/AdditiveControlRig.h
--- a/Engine/Plugins/Animation/ControlRig/Source/ControlRig/Public/Rigs/AdditiveControlRig.h
+++ b/Engine/Plugins/Animation/ControlRig/Source/ControlRig/Public/Rigs/AdditiveControlRig.h
@@ -23,7 +23,6 @@ public:
 	virtual void InitializeVMsFromCDO() override { URigVMHost::InitializeFromCDO(); }
 	virtual void RequestInitVMs() override { URigVMHost::RequestInit(); }
 	virtual bool Execute_Internal(const FName& InEventName) override;
-	virtual void EvaluateVMs_AnyThread() override { URigVMHost::Evaluate_AnyThread(); }
 #if WITH_EDITOR
 	virtual void SetFirstEntryEventInEventQueue(FRigVMExtendedExecuteContext& Context, const FName& InFirstEventName) override;
 #endif
diff --git a/Engine/Plugins/Animation/ControlRig/Source/ControlRig/Public/Rigs/FKControlRig.h b/Engine/Plugins/Animation/ControlRig/Source/ControlRig/Public/Rigs/FKControlRig.h
--- a/Engine/Plugins/Animation/ControlRig/Source/ControlRig/Public/Rigs/FKControlRig.h
+++ b/Engine/Plugins/Animation/ControlRig/Source/ControlRig/Public/Rigs/FKControlRig.h
@@ -49,7 +49,6 @@ public:
 	virtual void InitializeVMsFromCDO() override { URigVMHost::InitializeFromCDO(); }
 	virtual void RequestInitVMs() override { URigVMHost::RequestInit(); }
 	virtual bool Execute_Internal(const FName& InEventName) override;
-	virtual void EvaluateVMs_AnyThread() override { URigVMHost::Evaluate_AnyThread(); }
 #if WITH_EDITOR
 	virtual void SetFirstEntryEventInEventQueue(FRigVMExtendedExecuteContext& Context, const FName& InFirstEventName) override;
 #endif
diff --git a/Engine/Plugins/Animation/ControlRig/Source/ControlRig/Public/Rigs/RigHierarchy.h b/Engine/Plugins/Animation/ControlRig/Source/ControlRig/Public/Rigs/RigHierarchy.h
--- a/Engine/Plugins/Animation/ControlRig/Source/ControlRig/Public/Rigs/RigHierarchy.h
+++ b/Engine/Plugins/Animation/ControlRig/Source/ControlRig/Public/Rigs/RigHierarchy.h
@@ -4669,6 +4669,7 @@ private:
 	friend class URigHierarchyController;
 	friend class UBaseControlRig;
 	friend class UControlRig;
+	friend class UModularRig;
 	friend class FControlRigEditor;
 	friend struct FRigBaseElement;
 	friend struct FRigHierarchyValidityBracket;
@@ -4744,6 +4745,8 @@ private:
 
 	friend class UBaseControlRig;
 	friend class UControlRig;
+	friend class UModularRig;
+	friend class UModularRigController;
 };
 
 struct CONTROLRIG_API FRigHierarchyValidityBracket
diff --git a/Engine/Plugins/Animation/ControlRig/Source/ControlRig/Public/Units/RigUnitContext.h b/Engine/Plugins/Animation/ControlRig/Source/ControlRig/Public/Units/RigUnitContext.h
--- a/Engine/Plugins/Animation/ControlRig/Source/ControlRig/Public/Units/RigUnitContext.h
+++ b/Engine/Plugins/Animation/ControlRig/Source/ControlRig/Public/Units/RigUnitContext.h
@@ -191,12 +191,14 @@ private:
 	uint32 RigModuleNameSpaceHash;
 
 	friend class FControlRigExecuteContextRigModuleGuard;
+	friend class UModularRig;
 };
 
 class CONTROLRIG_API FControlRigExecuteContextRigModuleGuard
 {
 public:
 	FControlRigExecuteContextRigModuleGuard(FControlRigExecuteContext& InContext, const UBaseControlRig* InControlRig);
+	FControlRigExecuteContextRigModuleGuard(FControlRigExecuteContext& InContext, const FString& InNewModuleNameSpace);
 	~FControlRigExecuteContextRigModuleGuard();
 
 private:
diff --git a/Engine/Plugins/Animation/ControlRig/Source/ControlRigDeveloper/Private/ControlRigBlueprint.cpp b/Engine/Plugins/Animation/ControlRig/Source/ControlRigDeveloper/Private/ControlRigBlueprint.cpp
--- a/Engine/Plugins/Animation/ControlRig/Source/ControlRigDeveloper/Private/ControlRigBlueprint.cpp
+++ b/Engine/Plugins/Animation/ControlRig/Source/ControlRigDeveloper/Private/ControlRigBlueprint.cpp
@@ -14,6 +14,7 @@
 #include "UObject/UObjectGlobals.h"
 #include "ControlRigObjectVersion.h"
 #include "BlueprintCompilationManager.h"
+#include "ModularRig.h"
 #include "RigVMCompiler/RigVMCompiler.h"
 #include "RigVMCore/RigVMRegistry.h"
 #include "Units/Execution/RigUnit_BeginExecution.h"
@@ -69,6 +70,11 @@ UControlRigBlueprint::UControlRigBlueprint(const FObjectInitializer& ObjectIniti
 	{
 		CommonInitialization(ObjectInitializer);
 	}
+
+	ModularRigModel.SetOuterClientHost(this);
+	UModularRigController* ModularController = ModularRigModel.GetController();
+	ModularController->OnModified().AddUObject(this, &UControlRigBlueprint::HandleRigModulesModified);
+	
 }
 
 UControlRigBlueprint::UControlRigBlueprint()
@@ -186,6 +192,16 @@ bool UControlRigBlueprint::CanTurnIntoControlRigModule(bool InAutoConvertHierarc
 		return false;
 	}
 
+	if (GetRigVMHostClass()->IsChildOf<UModularRig>())
+	{
+		if(OutErrorMessage)
+		{
+			static const FString Message = TEXT("This asset is a Modular Rig.");
+			*OutErrorMessage = Message;
+		}
+		return false;
+	}
+
 	if(Hierarchy == nullptr)
 	{
 		if(OutErrorMessage)
@@ -819,6 +835,7 @@ void UControlRigBlueprint::SetPreviewMesh(USkeletalMesh* PreviewMesh, bool bMark
 void UControlRigBlueprint::Serialize(FArchive& Ar)
 {
 	RigVMClient.SetOuterClientHost(this, GET_MEMBER_NAME_CHECKED(UControlRigBlueprint, RigVMClient));
+	ModularRigModel.SetOuterClientHost(this);
 	
 	Super::Serialize(Ar);
 
@@ -866,6 +883,8 @@ void UControlRigBlueprint::Serialize(FArchive& Ar)
 			TGuardValue<bool> DisableClientNotifs(RigVMClient.bSuspendNotifications, true);
 			RigVMClient.SetFromDeprecatedData(Model_DEPRECATED, FunctionLibrary_DEPRECATED);
 		}
+
+		ModularRigModel.UpdateCachedChildren();
 	}
 }
 
@@ -924,12 +943,20 @@ bool UControlRigBlueprint::ResolveConnector(const FRigElementKey& DraggedKey, co
 			return false;
 		}
 		ExistingTargetKey = TargetKey;
+
+		// Add connection to the model
+		if (UModularRigController* Controller = GetModularRigController())
+		{
+			Controller->ConnectModuleToElement(DraggedKey, TargetKey);
+		}
 	}
 	else
 	{
 		ConnectionMap.Remove(DraggedKey);
 	}
 
+	RecompileModularRig();
+
 	PropagateHierarchyFromBPToInstances();
 
 	if(UBaseControlRig* ControlRig = Cast<UBaseControlRig>(GetObjectBeingDebugged()))
@@ -1096,6 +1123,7 @@ void UControlRigBlueprint::HandlePackageDone()
 	PropagateHierarchyFromBPToInstances();
 
 	Super::HandlePackageDone();
+	RecompileModularRig();
 }
 
 void UControlRigBlueprint::HandleConfigureRigVMController(const FRigVMClient* InClient, URigVMController* InControllerToConfigure)
@@ -1144,6 +1172,16 @@ UClass* UControlRigBlueprint::GetRigVMEditorSettingsClass() const
 	return UControlRigEditorSettings::StaticClass();
 }
 
+void UControlRigBlueprint::GetPreloadDependencies(TArray<UObject*>& OutDeps)
+{
+	Super::GetPreloadDependencies(OutDeps);
+
+	for (FRigModuleReference& Module : ModularRigModel.Modules)
+	{
+		OutDeps.Add(Module.Class.Get());
+	}
+}
+
 #if WITH_EDITOR
 const FName& UControlRigBlueprint::GetPanelPinFactoryName() const
 {
@@ -1250,6 +1288,14 @@ void UControlRigBlueprint::PostDuplicate(bool bDuplicateForPIE)
 		Controller->OnModified().RemoveAll(this);
 		Controller->OnModified().AddUObject(this, &UControlRigBlueprint::HandleHierarchyModified);
 	}
+
+	if (UModularRigController* ModularController = ModularRigModel.GetController())
+	{
+		ModularController->OnModified().RemoveAll(this);
+		ModularController->OnModified().AddUObject(this, &UControlRigBlueprint::HandleRigModulesModified);
+	}
+
+	ModularRigModel.UpdateCachedChildren();
 }
 
 TArray<UControlRigBlueprint*> UControlRigBlueprint::GetCurrentlyOpenRigBlueprints()
@@ -1445,6 +1491,70 @@ void UControlRigBlueprint::ClearTransientControls()
 	}
 }
 
+UModularRigController* UControlRigBlueprint::GetModularRigController() 
+{
+	if (!GetControlRigClass()->IsChildOf(UModularRig::StaticClass()))
+	{
+		return nullptr;
+	}
+
+	return ModularRigModel.GetController();
+}
+
+static void AddModulesRecursively(UModularRig* Rig, const FRigModuleReference* InModule, FRigModuleInstance* InParent)
+{
+	// Make sure the inner is loaded and valid
+	{
+		if (!InModule->Class.IsValid())
+		{
+			InModule->Class.LoadSynchronous();
+		}
+		if (!InModule->Class.IsValid())
+		{
+			return;
+		}
+	}
+	
+	FRigModuleInstance* NewModule = Rig->AddModuleInstance(InModule->Name, InModule->Class.Get(), InParent, InModule->Connections);
+	
+	for (const FRigModuleReference* ChildModule : InModule->CachedChildren)
+	{
+		AddModulesRecursively(Rig, ChildModule, NewModule);
+	}
+}
+
+static void BuildModularRig(UModularRig* Rig, const FModularRigModel& Model)
+{
+	Rig->ResetModules();
+	for (const FRigModuleReference* RootModule : Model.RootModules)
+	{
+		AddModulesRecursively(Rig, RootModule, nullptr);
+	}
+	Rig->InitializeVMs(true);
+}
+
+void UControlRigBlueprint::RecompileModularRig()
+{
+	if (UClass* MyControlRigClass = GeneratedClass)
+	{
+		if (UModularRig* DefaultObject = Cast<UModularRig>(MyControlRigClass->GetDefaultObject(false)))
+		{
+			BuildModularRig(DefaultObject, ModularRigModel);
+
+			// Initialize Archetype Instances
+			TArray<UObject*> ArchetypeInstances;
+			DefaultObject->GetArchetypeInstances(ArchetypeInstances);
+			for (UObject* Instance : ArchetypeInstances)
+			{
+				if (UModularRig* ModularRigInstance = Cast<UModularRig>(Instance))
+				{
+					BuildModularRig(ModularRigInstance, ModularRigModel);
+				}
+			}
+		}
+	}
+}
+
 #endif
 
 void UControlRigBlueprint::SetupDefaultObjectDuringCompilation(URigVMHost* InCDO)
@@ -1892,7 +2002,7 @@ void UControlRigBlueprint::UpdateElementKeyRedirector(UBaseControlRig* InControl
 {
 	InControlRig->HierarchySettings = HierarchySettings;
 	InControlRig->RigModuleSettings = RigModuleSettings;
-	InControlRig->ElementKeyRedirector = FRigElementKeyRedirector(ConnectionMap, Hierarchy);
+	InControlRig->ElementKeyRedirector = FRigElementKeyRedirector(ConnectionMap, InControlRig->GetHierarchy());
 }
 
 void UControlRigBlueprint::PropagatePoseFromInstanceToBP(UBaseControlRig* InControlRig) const
@@ -1952,6 +2062,7 @@ void UControlRigBlueprint::PropagateHierarchyFromBPToInstances() const
 					InstanceRig->PostInitInstanceIfRequired();
 					InstanceRig->GetHierarchy()->CopyHierarchy(Hierarchy);
 					InstanceRig->HierarchySettings = HierarchySettings;
+					UpdateElementKeyRedirector(InstanceRig);
 					InstanceRig->Initialize(true);
 				}
 			}
@@ -2126,6 +2237,11 @@ void UControlRigBlueprint::HandleHierarchyModified(ERigHierarchyNotification InN
 #endif
 }
 
+void UControlRigBlueprint::HandleRigModulesModified(EModularRigNotification InNotification, const FRigModuleReference* InModule)
+{
+	RecompileModularRig();
+}
+
 UControlRigBlueprint::FControlValueScope::FControlValueScope(UControlRigBlueprint* InBlueprint)
 : Blueprint(InBlueprint)
 {
diff --git a/Engine/Plugins/Animation/ControlRig/Source/ControlRigDeveloper/Private/ModularRigController.cpp b/Engine/Plugins/Animation/ControlRig/Source/ControlRigDeveloper/Private/ModularRigController.cpp
new file mode 100644
--- /dev/null
+++ b/Engine/Plugins/Animation/ControlRig/Source/ControlRigDeveloper/Private/ModularRigController.cpp
@@ -0,0 +1,207 @@
+// Copyright Epic Games, Inc. All Rights Reserved.
+
+#include "ModularRigController.h"
+
+#include "ControlRig.h"
+#include "ModularRig.h"
+#include "ControlRigBlueprint.h"
+#include "ModularRigModel.h"
+#include "Rigs/RigHierarchyController.h"
+
+#include UE_INLINE_GENERATED_CPP_BY_NAME(ModularRigController)
+
+UModularRigController::UModularRigController(const FObjectInitializer& ObjectInitializer)
+	: Super(ObjectInitializer)
+{
+}
+
+bool UModularRigController::AddModule(const FName& InModuleName, TSubclassOf<UBaseControlRig> InClass, const FString& InParentModulePath)
+{
+	if (!InClass->GetDefaultObject<UBaseControlRig>()->IsRigModule())
+	{
+		return false;
+	}
+
+	FRigModuleReference* NewModule = nullptr;
+	if (InParentModulePath.IsEmpty())
+	{
+		for (FRigModuleReference* Module : Model->RootModules)
+		{
+			if (Module->Name.ToString() == InModuleName)
+			{
+				return false;
+			}
+		}
+
+		Model->Modules.Add(FRigModuleReference(InModuleName, InClass, FString()));
+		NewModule = &Model->Modules.Last();
+		Model->RootModules.Add(NewModule);
+	}
+	else if (FRigModuleReference* ParentModule = FindModule(InParentModulePath))
+	{
+		for (FRigModuleReference* Module : ParentModule->CachedChildren)
+		{
+			if (Module->Name.ToString() == InModuleName)
+			{
+				return false;
+			}
+		}
+
+		Model->Modules.Add(FRigModuleReference(InModuleName, InClass, ParentModule->GetNamespace()));
+		NewModule = &Model->Modules.Last();
+		ParentModule->CachedChildren.Add(&Model->Modules.Last());
+	}
+
+	if (!NewModule)
+	{
+		return false;
+	}
+
+	UControlRig* DefaultModule = InClass->GetDefaultObject<UControlRig>();
+	if (!DefaultModule)
+	{
+		return false;
+	}
+
+	const FString NewNamespace = NewModule->GetNamespace();
+
+	// Add Connectors/Sockets
+	UControlRigBlueprint* Blueprint = Cast<UControlRigBlueprint>(GetOuter());
+	if (URigHierarchyController* Controller = Blueprint->GetHierarchyController())
+	{
+		if (UBaseControlRig* CDO = Blueprint->GetControlRigClass()->GetDefaultObject<UBaseControlRig>())
+		{
+			FRigVMExtendedExecuteContext& Context = CDO->GetRigVMExtendedExecuteContext();
+			FRigHierarchyExecuteContextBracket HierarchyContextGuard(Controller->GetHierarchy(), &Context);
+
+			// setup the module information
+			FControlRigExecuteContext& PublicContext = Context.GetPublicDataSafe<FControlRigExecuteContext>();
+			FControlRigExecuteContextRigModuleGuard RigModuleGuard(PublicContext, NewNamespace);
+
+			
+			const TArray<FRigModuleConnector>& Connectors = DefaultModule->GetRigModuleSettings().ExposedConnectors;
+			for (const FRigModuleConnector& Connector : Connectors)
+			{
+				Controller->AddConnector(*Connector.Name, Connector.Settings);
+			}
+
+			// todo: copy the sockets
+			// TArray<FRigSocketElement*> Sockets = DefaultModule->GetHierarchy()->GetElementsOfType<FRigSocketElement>();
+			// for (FRigSocketElement* Socket : Sockets)
+			// {
+			// 	Controller->CopySocket(NewModule->GetPath(), Socket);
+			// }
+		}
+	}
+
+	Notify(EModularRigNotification::ModuleAdded, NewModule);
+	
+	return true;
+}
+
+FRigModuleReference* UModularRigController::FindModule(const FString& InPath)
+{
+	FString Left = InPath, Right;
+	TArray<FString> Tokens;
+	while (Left.Split(UModularRig::NamespaceSeparator, &Left, &Right))
+	{
+		Tokens.Add(Left);
+		Left = Right;
+	}
+	Tokens.Add(Left);
+
+	FRigModuleReference* CurModule = nullptr;
+	for (FRigModuleReference* Module : Model->RootModules)
+	{
+		if (Module->Name.ToString() == Tokens[0])
+		{
+			CurModule = Module;
+			break;
+		}
+	}
+	if (!CurModule)
+	{
+		return nullptr;
+	}
+
+	for (int32 Index=1; Index<Tokens.Num(); ++Index)
+	{
+		bool bFound = false;
+		for (FRigModuleReference* Child : CurModule->CachedChildren)
+		{
+			if (Child->Name.ToString() == Tokens[Index])
+			{
+				CurModule = Child;
+				bFound = true;
+				break;
+			}
+		}
+		if (!bFound)
+		{
+			return nullptr;
+		}
+	}
+
+	return CurModule;
+}
+
+bool UModularRigController::ConnectModuleToElement(const FRigElementKey& InConnectorKey, const FRigElementKey& InTargetKey)
+{
+	FString ConnectorNameSpace, ConnectorName;
+	if (!InConnectorKey.Name.ToString().Split(UModularRig::NamespaceSeparator, &ConnectorNameSpace, &ConnectorName, ESearchCase::CaseSensitive, ESearchDir::FromEnd))
+	{
+		return false;
+	}
+	
+	FRigModuleReference* Module = FindModule(ConnectorNameSpace);
+	if (!Module)
+	{
+		return false;
+	}
+
+	UControlRigBlueprint* Blueprint = Cast<UControlRigBlueprint>(GetOuter());
+	FRigConnectorElement* Connector = Cast<FRigConnectorElement>(Blueprint->Hierarchy->Find(InConnectorKey));
+	if (!Connector)
+	{
+		return false;
+	}
+
+	if (!Blueprint->Hierarchy->Contains(InTargetKey))
+	{
+		return false;
+	}
+
+	const FRigElementKey ConnectorKey(*ConnectorName, ERigElementType::Connector);
+	FRigElementKey& TargetKey = Module->Connections.FindOrAdd(ConnectorKey);
+	TargetKey = InTargetKey;
+
+	Notify(EModularRigNotification::ConnectionChanged, Module);
+	
+	return true;
+}
+
+bool UModularRigController::RemoveModule(const FString& InModulesPath)
+{
+	// todo: UE-199050
+	return false;
+}
+
+bool UModularRigController::RenameModule(const FString& InModulesPath, const FName& InNewName)
+{
+	// todo: UE-199050
+	return false;
+}
+
+bool UModularRigController::ReparentModule(const FString& InModulesPath, const FString& InNewParentModulePath)
+{
+	// todo: UE-199050
+	return false;
+}
+
+void UModularRigController::Notify(const EModularRigNotification& InNotification, const FRigModuleReference* InElement)
+{
+	if(!bSuspendNotifications)
+	{
+		ModifiedEvent.Broadcast(InNotification, InElement);
+	}
+}
diff --git a/Engine/Plugins/Animation/ControlRig/Source/ControlRigDeveloper/Private/ModularRigModel.cpp b/Engine/Plugins/Animation/ControlRig/Source/ControlRigDeveloper/Private/ModularRigModel.cpp
new file mode 100644
--- /dev/null
+++ b/Engine/Plugins/Animation/ControlRig/Source/ControlRigDeveloper/Private/ModularRigModel.cpp
@@ -0,0 +1,93 @@
+// Copyright Epic Games, Inc. All Rights Reserved.
+
+#include "ModularRigModel.h"
+#include "ModularRig.h"
+
+#include UE_INLINE_GENERATED_CPP_BY_NAME(ModularRigModel)
+
+FString FRigModuleReference::GetNamespace() const
+{
+	if (ParentNamespace.IsEmpty())
+	{
+		return Name.ToString() + UModularRig::NamespaceSeparator;
+	}
+	return ParentNamespace + UModularRig::NamespaceSeparator + Name.ToString() + UModularRig::NamespaceSeparator;
+}
+
+UModularRigController* FModularRigModel::GetController(bool bCreateIfNeeded)
+{
+	if (Controller)
+	{
+		return Controller;
+	}
+
+	if (bCreateIfNeeded)
+	{
+		const FName SafeControllerName = *FString::Printf(TEXT("%s_ModularRig_Controller"), *GetOuter()->GetPathName());
+		Controller = NewObject<UModularRigController>(GetOuter(), UModularRigController::StaticClass(), SafeControllerName);
+		Controller->SetModel(this);
+	}
+	return Controller;
+}
+
+void FModularRigModel::SetOuterClientHost(UObject* InOuterClientHost)
+{
+	OuterClientHost = InOuterClientHost;
+}
+
+void FModularRigModel::UpdateCachedChildren()
+{
+	TMap<FString, FRigModuleReference*> PathToModule;
+	for (FRigModuleReference& Module : Modules)
+	{
+		Module.CachedChildren.Reset();
+		PathToModule.Add(Module.GetNamespace(), &Module);
+	}
+	for (FRigModuleReference& Module : Modules)
+	{
+		if (Module.ParentNamespace.IsEmpty())
+		{
+			RootModules.Add(&Module);
+		}
+		else
+		{
+			if (FRigModuleReference** ParentModule = PathToModule.Find(Module.ParentNamespace))
+			{
+				(*ParentModule)->CachedChildren.Add(&Module);
+			}
+		}
+	}
+}
+
+FRigModuleReference* FModularRigModel::FindModule(const FString InNameSpace)
+{
+	TArray<FRigModuleReference*>* Children = &RootModules;
+
+	FString Left = InNameSpace, Right;
+	while (Left.Split(UModularRig::NamespaceSeparator, &Left, &Right))
+	{
+		FRigModuleReference** Child = Children->FindByPredicate([Left](FRigModuleReference* Module)
+		{
+			return Module->Name.ToString() == Left;
+		});
+		
+		if (!Child)
+		{
+			return nullptr;
+		}
+		
+		Children = &(*Child)->CachedChildren;
+		Left = Right;
+	}
+
+	FRigModuleReference** Child = Children->FindByPredicate([Left](FRigModuleReference* Module)
+		{
+			return Module->Name.ToString() == Left;
+		});
+
+	if (!Child)
+	{
+		return nullptr;
+	}
+	return *Child;
+}
diff --git a/Engine/Plugins/Animation/ControlRig/Source/ControlRigDeveloper/Public/ControlRigBlueprint.h b/Engine/Plugins/Animation/ControlRig/Source/ControlRigDeveloper/Public/ControlRigBlueprint.h
--- a/Engine/Plugins/Animation/ControlRig/Source/ControlRigDeveloper/Public/ControlRigBlueprint.h
+++ b/Engine/Plugins/Animation/ControlRig/Source/ControlRigDeveloper/Public/ControlRigBlueprint.h
@@ -18,6 +18,7 @@
 #include "ControlRigValidationPass.h"
 #include "RigVMBlueprint.h"
 #include "Rigs/RigModuleDefines.h"
+#include "ModularRigModel.h"
 
 #if WITH_EDITOR
 #include "Kismet2/CompilerResultsLog.h"
@@ -48,6 +49,7 @@ public:
 	virtual UClass* GetRigVMEdGraphSchemaClass() const override;
 	virtual TArray<FString> GeneratePythonCommands(const FString InNewBlueprintName) override;
 	virtual UClass* GetRigVMEditorSettingsClass() const override;
+	virtual void GetPreloadDependencies(TArray<UObject*>& OutDeps) override;
 #if WITH_EDITOR
 	virtual const FName& GetPanelPinFactoryName() const override;
 	static const FName ControlRigPanelNodeFactoryName;
@@ -205,6 +207,15 @@ public:
 	UFUNCTION(BlueprintCallable, Category = "Hierarchy")
 	URigHierarchyController* GetHierarchyController() { return Hierarchy->GetController(true); }
 
+	UPROPERTY(BlueprintReadOnly, Category = "Modules")
+	FModularRigModel ModularRigModel;
+
+	UFUNCTION(BlueprintCallable, Category = "Modules")
+	UModularRigController* GetModularRigController();
+
+	UFUNCTION(BlueprintCallable, Category = "Control Rig Blueprint")
+	void RecompileModularRig();
+
 private:
 
 	/** Whether or not this rig has an Inversion Event */
@@ -271,6 +282,8 @@ private:
 
 	void HandleHierarchyModified(ERigHierarchyNotification InNotification, URigHierarchy* InHierarchy, const FRigBaseElement* InElement);
 
+	void HandleRigModulesModified(EModularRigNotification InNotification, const FRigModuleReference* InModule);
+
 #if WITH_EDITOR
 	virtual void HandlePackageDone() override;
 	virtual void HandleConfigureRigVMController(const FRigVMClient* InClient, URigVMController* InControllerToConfigure) override;
diff --git a/Engine/Plugins/Animation/ControlRig/Source/ControlRigDeveloper/Public/ModularRigController.h b/Engine/Plugins/Animation/ControlRig/Source/ControlRigDeveloper/Public/ModularRigController.h
new file mode 100644
--- /dev/null
+++ b/Engine/Plugins/Animation/ControlRig/Source/ControlRigDeveloper/Public/ModularRigController.h
@@ -0,0 +1,63 @@
+// Copyright Epic Games, Inc. All Rights Reserved.
+
+#pragma once
+
+#include "BaseControlRig.h"
+#include "ModularRigController.generated.h"
+
+struct FRigModuleReference;
+struct FModularRigModel;
+
+DECLARE_MULTICAST_DELEGATE_TwoParams(FModularRigModifiedEvent, EModularRigNotification /* type */, const FRigModuleReference* /* element */);
+
+
+UENUM()
+enum class EModularRigNotification : uint8
+{
+	ModuleAdded,
+
+	ConnectionChanged,
+
+	/** MAX - invalid */
+	Max UMETA(Hidden),
+};
+
+UCLASS(BlueprintType)
+class CONTROLRIGDEVELOPER_API UModularRigController : public UObject
+{
+	GENERATED_UCLASS_BODY()
+
+	UModularRigController()
+		: Model(nullptr)
+		, bSuspendNotifications(false)
+	{
+	}
+
+	FModularRigModel* Model;
+
+	FModularRigModifiedEvent ModifiedEvent;
+
+	bool bSuspendNotifications;
+
+
+	UFUNCTION(BlueprintCallable, Category = "Control Rig | Modules")
+	bool AddModule(const FName& InModuleName, TSubclassOf<UBaseControlRig> InClass, const FString& InParentModulePath);
+
+	UFUNCTION(BlueprintCallable, Category = "Control Rig | Modules")
+	bool ConnectModuleToElement(const FRigElementKey& InConnectorKey, const FRigElementKey& InTargetKey);
+
+	UFUNCTION(BlueprintCallable, Category = "ControlRig | Modules")
+	bool RemoveModule(const FString& InModulesPath);
+
+	UFUNCTION(BlueprintCallable, Category = "ControlRig | Modules")
+	bool RenameModule(const FString& InModulesPath, const FName& InNewName);
+
+	UFUNCTION(BlueprintCallable, Category = "ControlRig | Modules")
+	bool ReparentModule(const FString& InModulesPath, const FString& InNewParentModulePath);
+
+
+	void SetModel(FModularRigModel* InModel) { Model = InModel; }
+	FRigModuleReference* FindModule(const FString& InPath);
+	FModularRigModifiedEvent& OnModified() { return ModifiedEvent; }
+	void Notify(const EModularRigNotification& InNotification, const FRigModuleReference* InElement);
+};
\ No newline at end of file
diff --git a/Engine/Plugins/Animation/ControlRig/Source/ControlRigDeveloper/Public/ModularRigModel.h b/Engine/Plugins/Animation/ControlRig/Source/ControlRigDeveloper/Public/ModularRigModel.h
new file mode 100644
--- /dev/null
+++ b/Engine/Plugins/Animation/ControlRig/Source/ControlRigDeveloper/Public/ModularRigModel.h
@@ -0,0 +1,75 @@
+// Copyright Epic Games, Inc. All Rights Reserved.
+
+#pragma once
+
+#include "BaseControlRig.h"
+#include "ModularRigController.h"
+#include "ModularRigModel.generated.h"
+
+USTRUCT(BlueprintType)
+struct CONTROLRIGDEVELOPER_API FRigModuleReference
+{
+	GENERATED_BODY()
+
+	FRigModuleReference()
+		: Name(NAME_None)
+		, ParentNamespace(FString())
+		, Class(nullptr)
+	{}
+	
+	FRigModuleReference(const FName& InName, TSubclassOf<UBaseControlRig> InClass, const FString& InParentPath)
+		: Name(InName)
+		, ParentNamespace(InParentPath)
+		, Class(InClass)
+	{}
+
+	UPROPERTY()
+	FName Name;
+
+	UPROPERTY()
+	FString ParentNamespace;
+
+	UPROPERTY()
+	TSoftClassPtr<UBaseControlRig> Class;
+
+	UPROPERTY()
+	TMap<FRigElementKey, FRigElementKey> Connections; // Connectors to Connection element
+
+	TArray<FRigModuleReference*> CachedChildren;
+
+	FString GetNamespace() const;
+	
+	friend class UModularRigController;
+};
+
+// A management struct containing all modules in the rig
+USTRUCT(BlueprintType)
+struct CONTROLRIGDEVELOPER_API FModularRigModel
+{
+public:
+
+	GENERATED_BODY()
+
+	UPROPERTY()
+	TArray<FRigModuleReference> Modules;
+	TArray<FRigModuleReference*> RootModules;
+
+	UPROPERTY(transient)
+	TObjectPtr<UModularRigController> Controller;
+
+	UModularRigController* GetController(bool bCreateIfNeeded = true);
+
+	UObject* GetOuter() const { return OuterClientHost.IsValid() ? OuterClientHost.Get() : nullptr; }
+
+	void SetOuterClientHost(UObject* InOuterClientHost);
+
+	void UpdateCachedChildren();
+
+	FRigModuleReference* FindModule(const FString InNameSpace);
+	
+private:
+	TWeakObjectPtr<UObject> OuterClientHost;
+
+	friend class UModularRigController;
+	friend struct FRigModuleReference;
+};
\ No newline at end of file
diff --git a/Engine/Plugins/Animation/ControlRig/Source/ControlRigEditor/Private/ControlRigBlueprintActions.cpp b/Engine/Plugins/Animation/ControlRig/Source/ControlRigEditor/Private/ControlRigBlueprintActions.cpp
--- a/Engine/Plugins/Animation/ControlRig/Source/ControlRigEditor/Private/ControlRigBlueprintActions.cpp
+++ b/Engine/Plugins/Animation/ControlRig/Source/ControlRigEditor/Private/ControlRigBlueprintActions.cpp
@@ -404,7 +404,7 @@ void FControlRigBlueprintActions::OnSpawnedSkeletalMeshActorChanged(UObject* InO
 
 					ObjectName.RemoveFromEnd(TEXT("_C"));
 
-					// This is either a UControlRig or a UModularControlRig
+					// This is either a UControlRig or a UModularRig
 					ControlRig = NewObject<UBaseControlRig>(Track, ControlRigClass, FName(*ObjectName), RF_Transactional);
 					ControlRig->SetObjectBinding(MakeShared<FControlRigObjectBinding>());
 					ControlRig->GetObjectBinding()->BindToObject(MeshActor->GetSkeletalMeshComponent());
diff --git a/Engine/Plugins/Animation/ControlRig/Source/ControlRigEditor/Private/ControlRigBlueprintFactory.cpp b/Engine/Plugins/Animation/ControlRig/Source/ControlRigEditor/Private/ControlRigBlueprintFactory.cpp
--- a/Engine/Plugins/Animation/ControlRig/Source/ControlRigEditor/Private/ControlRigBlueprintFactory.cpp
+++ b/Engine/Plugins/Animation/ControlRig/Source/ControlRigEditor/Private/ControlRigBlueprintFactory.cpp
@@ -30,6 +30,9 @@
 
 #define LOCTEXT_NAMESPACE "ControlRigBlueprintFactory"
 
+static TAutoConsoleVariable<bool> CVarControlRigEnableModularRigs(TEXT("ControlRig.Test.EnableModularRigs"), false, TEXT("When true we will allow the creation of modular rigs."));
+
+
 /** Dialog to configure creation properties */
 class SControlRigBlueprintCreateDialog : public SCompoundWidget
 {
@@ -276,8 +279,12 @@ UControlRigBlueprintFactory::UControlRigBlueprintFactory()
 
 bool UControlRigBlueprintFactory::ConfigureProperties()
 {
-	TSharedRef<SControlRigBlueprintCreateDialog> Dialog = SNew(SControlRigBlueprintCreateDialog);
-	return Dialog->ConfigureProperties(this);
+	if (CVarControlRigEnableModularRigs->GetBool())
+	{
+		TSharedRef<SControlRigBlueprintCreateDialog> Dialog = SNew(SControlRigBlueprintCreateDialog);
+		return Dialog->ConfigureProperties(this);
+	}
+	return true;
 };
 
 UObject* UControlRigBlueprintFactory::FactoryCreateNew(UClass* Class, UObject* InParent, FName Name, EObjectFlags Flags, UObject* Context, FFeedbackContext* Warn, FName CallingContext)
diff --git a/Engine/Plugins/Animation/ControlRig/Source/ControlRigEditor/Private/ControlRigElementDetails.cpp b/Engine/Plugins/Animation/ControlRig/Source/ControlRigEditor/Private/ControlRigElementDetails.cpp
--- a/Engine/Plugins/Animation/ControlRig/Source/ControlRigEditor/Private/ControlRigElementDetails.cpp
+++ b/Engine/Plugins/Animation/ControlRig/Source/ControlRigEditor/Private/ControlRigElementDetails.cpp
@@ -5113,7 +5113,7 @@ void FRigConnectorElementDetails::CustomizeDetails(IDetailLayoutBuilder& DetailB
 	FRigTransformElementDetails::CustomizeDetails(DetailBuilder);
 	CustomizeTransform(DetailBuilder);
 	CustomizeSettings(DetailBuilder);
-	CustomizeMetadata(DetailBuilder);
+	//CustomizeMetadata(DetailBuilder);
 }
 
 void FRigConnectorElementDetails::CustomizeSettings(IDetailLayoutBuilder& DetailBuilder)
diff --git a/Engine/Plugins/Animation/ControlRig/Source/ControlRigEditor/Private/Editor/ControlRigEditor.cpp b/Engine/Plugins/Animation/ControlRig/Source/ControlRigEditor/Private/Editor/ControlRigEditor.cpp
--- a/Engine/Plugins/Animation/ControlRig/Source/ControlRigEditor/Private/Editor/ControlRigEditor.cpp
+++ b/Engine/Plugins/Animation/ControlRig/Source/ControlRigEditor/Private/Editor/ControlRigEditor.cpp
@@ -28,6 +28,8 @@
 #include "IPersonaPreviewScene.h"
 #include "Animation/AnimData/BoneMaskFilter.h"
 #include "BaseControlRig.h"
+#include "ControlRig.h"
+#include "ModularRig.h"
 #include "Editor/ControlRigSkeletalMeshComponent.h"
 #include "ControlRigObjectBinding.h"
 #include "RigVMBlueprintUtils.h"
@@ -1101,6 +1103,8 @@ void FControlRigEditor::Compile()
 			FRigVMEditor::Compile();
 		}
 
+		ControlRigBlueprint->RecompileModularRig();
+
 		// ensure the skeletal mesh is still bound
 		UControlRigSkeletalMeshComponent* SkelMeshComponent = Cast<UControlRigSkeletalMeshComponent>(GetPersonaToolkit()->GetPreviewScene()->GetPreviewMeshComponent());
 		if (SkelMeshComponent)
@@ -3020,19 +3024,22 @@ void FControlRigEditor::OnHierarchyModified(ERigHierarchyNotification InNotif, U
 	{
 		case ERigHierarchyNotification::ElementAdded:
 		{
-			if(InElement->GetType() == ERigElementType::Connector)
+			if (GetControlRig()->IsA<UControlRig>())
 			{
-				if(InHierarchy->GetConnectors().Num() == 1)
+				if(InElement->GetType() == ERigElementType::Connector)
 				{
-					FNotificationInfo Info(LOCTEXT("FirstConnectorEncountered", "Looks like you have added the first connector. This rig will now be configured as a module, settings can be found in the class settings Hierarchy -> Module Settings."));
-					Info.bFireAndForget = true;
-					Info.FadeOutDuration = 5.0f;
-					Info.ExpireDuration = 5.0f;
+					if(InHierarchy->GetConnectors().Num() == 1)
+					{
+						FNotificationInfo Info(LOCTEXT("FirstConnectorEncountered", "Looks like you have added the first connector. This rig will now be configured as a module, settings can be found in the class settings Hierarchy -> Module Settings."));
+						Info.bFireAndForget = true;
+						Info.FadeOutDuration = 5.0f;
+						Info.ExpireDuration = 5.0f;
 
-					TSharedPtr<SNotificationItem> NotificationPtr = FSlateNotificationManager::Get().AddNotification(Info);
-					NotificationPtr->SetCompletionState(SNotificationItem::CS_Success);
+						TSharedPtr<SNotificationItem> NotificationPtr = FSlateNotificationManager::Get().AddNotification(Info);
+						NotificationPtr->SetCompletionState(SNotificationItem::CS_Success);
 
-					RigBlueprint->TurnIntoControlRigModule();
+						RigBlueprint->TurnIntoControlRigModule();
+					}
 				}
 			}
 			// no break - fall through
diff --git a/Engine/Plugins/Runtime/RigVM/Source/RigVMDeveloper/Public/RigVMBlueprint.h b/Engine/Plugins/Runtime/RigVM/Source/RigVMDeveloper/Public/RigVMBlueprint.h
--- a/Engine/Plugins/Runtime/RigVM/Source/RigVMDeveloper/Public/RigVMBlueprint.h
+++ b/Engine/Plugins/Runtime/RigVM/Source/RigVMDeveloper/Public/RigVMBlueprint.h
@@ -311,7 +311,7 @@ public:
 	virtual void PostEditChangeChainProperty(struct FPropertyChangedChainEvent& PropertyChangedEvent) override;
 	virtual void PostRename(UObject* OldOuter, const FName OldName) override;
 	/** Called during cooking. Must return all objects that will be Preload()ed when this is serialized at load time. */
-	void GetPreloadDependencies(TArray<UObject*>& OutDeps) override;
+	virtual void GetPreloadDependencies(TArray<UObject*>& OutDeps) override;
 
 	// IRigVMClientHost interface
 	virtual FRigVMClient* GetRigVMClient() override;
