Control Rig: Expose GetDebuggedControlRig in ControlRigBlueprint

diff --git a/Engine/Plugins/Animation/ControlRig/Source/ControlRigDeveloper/Public/ControlRigBlueprint.h b/Engine/Plugins/Animation/ControlRig/Source/ControlRigDeveloper/Public/ControlRigBlueprint.h
--- a/Engine/Plugins/Animation/ControlRig/Source/ControlRigDeveloper/Public/ControlRigBlueprint.h
+++ b/Engine/Plugins/Animation/ControlRig/Source/ControlRigDeveloper/Public/ControlRigBlueprint.h
@@ -84,6 +84,9 @@ public:
 	UFUNCTION(BlueprintCallable, Category = "Control Rig Blueprint")
 	UBaseControlRig* CreateControlRig() { return Cast<UBaseControlRig>(CreateRigVMHost()); }
 
+	UFUNCTION(BlueprintCallable, Category = "Control Rig Blueprint")
+	UBaseControlRig* GetDebuggedControlRig() { return Cast<UBaseControlRig>(GetDebuggedRigVMHost()); } 
+
 	/** IInterface_PreviewMeshProvider interface */
 	UFUNCTION(BlueprintCallable, Category = "Control Rig Blueprint")
 	virtual void SetPreviewMesh(USkeletalMesh* PreviewMesh, bool bMarkAsDirty = true) override;
diff --git a/Engine/Plugins/Runtime/RigVM/Source/RigVMDeveloper/Public/RigVMBlueprint.h b/Engine/Plugins/Runtime/RigVM/Source/RigVMDeveloper/Public/RigVMBlueprint.h
--- a/Engine/Plugins/Runtime/RigVM/Source/RigVMDeveloper/Public/RigVMBlueprint.h
+++ b/Engine/Plugins/Runtime/RigVM/Source/RigVMDeveloper/Public/RigVMBlueprint.h
@@ -520,6 +520,9 @@ public:
 	UFUNCTION(BlueprintCallable, Category = "VM")
 	URigVMHost* CreateRigVMHost();
 
+	UFUNCTION(BlueprintCallable, Category = "VM")
+	URigVMHost* GetDebuggedRigVMHost() { return Cast<URigVMHost>(GetObjectBeingDebugged()); }
+
 	UFUNCTION(BlueprintCallable, Category = "VM")
 	virtual TArray<UStruct*> GetAvailableRigVMStructs() const;
 
