Sequencer: Snapper: Don't include controls without shapes as snap objects

diff --git a/Engine/Plugins/Animation/ControlRig/Source/ControlRigEditor/Private/EditMode/SControlRigSnapper.cpp b/Engine/Plugins/Animation/ControlRig/Source/ControlRigEditor/Private/EditMode/SControlRigSnapper.cpp
--- a/Engine/Plugins/Animation/ControlRig/Source/ControlRigEditor/Private/EditMode/SControlRigSnapper.cpp
+++ b/Engine/Plugins/Animation/ControlRig/Source/ControlRigEditor/Private/EditMode/SControlRigSnapper.cpp
@@ -507,26 +507,45 @@ FControlRigSnapperSelection SControlRigSnapper::GetSelection(bool bGetAll)
 	{
 		if (ControlRig)
 		{
-			TArray<FName> SelectedControls = ControlRig->CurrentControlSelection();
-			if (SelectedControls.Num() > 0)
+			if (const URigHierarchy* Hierarchy = ControlRig->GetHierarchy())
 			{
-				FControlRigForWorldTransforms ControlRigAndSelection;
-				ControlRigAndSelection.ControlRig = ControlRig;
-				ControlRigAndSelection.ControlNames = SelectedControls;
-				if (bGetAll == false)
+				TArray<FName> SelectedControls = ControlRig->CurrentControlSelection();
+				if (SelectedControls.Num() > 0)
 				{
-					//make sure to only use first control
-					if (SelectedControls.Num() > 1)
+					FControlRigForWorldTransforms ControlRigAndSelection;
+					ControlRigAndSelection.ControlRig = ControlRig;
+					if (bGetAll == false)
 					{
-						ControlRigAndSelection.ControlNames.SetNum(0);
-						ControlRigAndSelection.ControlNames.Add(SelectedControls[0]);
+						//make sure to only use first control that is snappable(that has a shape)
+						for (const FName& ControlName : SelectedControls)
+						{
+							if (const FRigControlElement* ControlElement = Hierarchy->Find<FRigControlElement>(FRigElementKey(ControlName, ERigElementType::Control)))
+							{
+								if (ControlElement->Settings.SupportsShape())
+								{
+									ControlRigAndSelection.ControlNames.Add(ControlName);
+									break;
+								}
+							}
+						}
+						Selection.ControlRigs.Add(ControlRigAndSelection);
+						return Selection;
+					}
+					else
+					{
+						for (const FName& ControlName : SelectedControls)
+						{
+							if (const FRigControlElement* ControlElement = Hierarchy->Find<FRigControlElement>(FRigElementKey(ControlName, ERigElementType::Control)))
+							{
+								if (ControlElement->Settings.SupportsShape())
+								{
+									ControlRigAndSelection.ControlNames.Add(ControlName);
+								}
+							}
+						}
+						Selection.ControlRigs.Add(ControlRigAndSelection);
+
 					}
-					Selection.ControlRigs.Add(ControlRigAndSelection);
-					return Selection;
-				}
-				else
-				{
-					Selection.ControlRigs.Add(ControlRigAndSelection);
 				}
 			}
 		}
