Control Rig: Support for limits on spawn control node

#rnx

diff --git a/Engine/Plugins/Animation/ControlRig/Source/ControlRig/Private/Units/Execution/RigUnit_DynamicHierarchy.cpp b/Engine/Plugins/Animation/ControlRig/Source/ControlRig/Private/Units/Execution/RigUnit_DynamicHierarchy.cpp
--- a/Engine/Plugins/Animation/ControlRig/Source/ControlRig/Private/Units/Execution/RigUnit_DynamicHierarchy.cpp
+++ b/Engine/Plugins/Animation/ControlRig/Source/ControlRig/Private/Units/Execution/RigUnit_DynamicHierarchy.cpp
@@ -807,6 +807,39 @@ FRigUnit_HierarchyAddControlRotator_Execute()
 	}
 }
 
+void FRigUnit_HierarchyAddControlTransform_LimitSettings::ConfigureFrom(const FRigControlElement* InControlElement, const FRigControlSettings& InSettings)
+{
+	LimitTranslationX = InSettings.LimitEnabled[0];
+	LimitTranslationY = InSettings.LimitEnabled[1];
+	LimitTranslationZ = InSettings.LimitEnabled[2];
+	LimitPitch = InSettings.LimitEnabled[3];
+	LimitYaw = InSettings.LimitEnabled[4];
+	LimitRoll = InSettings.LimitEnabled[5];
+	LimitScaleX = InSettings.LimitEnabled[6];
+	LimitScaleY = InSettings.LimitEnabled[7];
+	LimitScaleZ = InSettings.LimitEnabled[8];
+	MinValue = InSettings.MinimumValue.Get<FRigControlValue::FEulerTransform_Float>().ToTransform();
+	MaxValue = InSettings.MaximumValue.Get<FRigControlValue::FEulerTransform_Float>().ToTransform();
+	bDrawLimits = InSettings.bDrawLimits;
+}
+
+void FRigUnit_HierarchyAddControlTransform_LimitSettings::Configure(FRigControlSettings& OutSettings) const
+{
+	OutSettings.SetupLimitArrayForType(false, false, false);
+	OutSettings.LimitEnabled[0] = LimitTranslationX;
+	OutSettings.LimitEnabled[1] = LimitTranslationY;
+	OutSettings.LimitEnabled[2] = LimitTranslationZ;
+	OutSettings.LimitEnabled[3] = LimitPitch;
+	OutSettings.LimitEnabled[4] = LimitYaw;
+	OutSettings.LimitEnabled[5] = LimitRoll;
+	OutSettings.LimitEnabled[6] = LimitScaleX;
+	OutSettings.LimitEnabled[7] = LimitScaleY;
+	OutSettings.LimitEnabled[8] = LimitScaleZ;
+	OutSettings.MinimumValue = FRigControlValue::Make<FRigControlValue::FEulerTransform_Float>(MinValue);
+	OutSettings.MaximumValue = FRigControlValue::Make<FRigControlValue::FEulerTransform_Float>(MaxValue);
+	OutSettings.bDrawLimits = bDrawLimits;
+}
+
 void FRigUnit_HierarchyAddControlTransform_Settings::ConfigureFrom(const FRigControlElement* InControlElement, const FRigControlSettings& InSettings)
 {
 	FRigUnit_HierarchyAddControl_Settings::ConfigureFrom(InControlElement, InSettings);
@@ -816,6 +849,7 @@ void FRigUnit_HierarchyAddControlTransform_Settings::ConfigureFrom(const FRigCon
 	PreferredRotationOrder = InSettings.PreferredRotationOrder;
 
 	Proxy.ConfigureFrom(InControlElement, InSettings);
+	Limits.ConfigureFrom(InControlElement, InSettings);
 	Shape.ConfigureFrom(InControlElement, InSettings);
 }
 
@@ -829,6 +863,7 @@ void FRigUnit_HierarchyAddControlTransform_Settings::Configure(FRigControlSettin
 	OutSettings.PreferredRotationOrder = PreferredRotationOrder;
 
 	Proxy.Configure(OutSettings);
+	Limits.Configure(OutSettings);
 	Shape.Configure(OutSettings);
 }
 
diff --git a/Engine/Plugins/Animation/ControlRig/Source/ControlRig/Private/Units/Execution/RigUnit_DynamicHierarchy.h b/Engine/Plugins/Animation/ControlRig/Source/ControlRig/Private/Units/Execution/RigUnit_DynamicHierarchy.h
--- a/Engine/Plugins/Animation/ControlRig/Source/ControlRig/Private/Units/Execution/RigUnit_DynamicHierarchy.h
+++ b/Engine/Plugins/Animation/ControlRig/Source/ControlRig/Private/Units/Execution/RigUnit_DynamicHierarchy.h
@@ -954,7 +954,7 @@ struct CONTROLRIG_API FRigUnit_HierarchyAddControlRotator_Settings : public FRig
 	ERigVMTransformSpace InitialSpace;
 
 	UPROPERTY(EditAnywhere, BlueprintReadWrite, Category = "Settings")
-	FRigUnit_HierarchyAddControlRotator_LimitSettings Limits;;
+	FRigUnit_HierarchyAddControlRotator_LimitSettings Limits;
 
 	UPROPERTY(EditAnywhere, BlueprintReadWrite, Category = "Settings")
 	FRigUnit_HierarchyAddControl_ShapeSettings Shape;
@@ -998,6 +998,67 @@ struct CONTROLRIG_API FRigUnit_HierarchyAddControlRotator : public FRigUnit_Hier
 	virtual void Execute() override;
 };
 
+USTRUCT(BlueprintType)
+struct CONTROLRIG_API FRigUnit_HierarchyAddControlTransform_LimitSettings
+{
+	GENERATED_BODY();
+	
+	FRigUnit_HierarchyAddControlTransform_LimitSettings()
+		: LimitTranslationX(false, false)
+		, LimitTranslationY(false, false)
+		, LimitTranslationZ(false, false)
+		, LimitPitch(false, false)
+		, LimitYaw(false, false)
+		, LimitRoll(false, false)
+		, LimitScaleX(false, false)
+		, LimitScaleY(false, false)
+		, LimitScaleZ(false, false)
+		, MinValue(FEulerTransform(FVector(-100.f, -100.f, -100.f), FRotator(-180.f, -180.f, -180.f), FVector(0.f, 0.f, 0.f)))
+		, MaxValue(FEulerTransform(FVector(100.f, 100.f, 100.f), FRotator(180.f, 180.f, 180.f), FVector(10.f, 10.f, 10.f)))
+		, bDrawLimits(true)
+	{}
+
+	void ConfigureFrom(const FRigControlElement* InControlElement, const FRigControlSettings& InSettings);
+	void Configure(FRigControlSettings& OutSettings) const;
+
+	UPROPERTY(EditAnywhere, BlueprintReadWrite, Category = "Settings")
+	FRigControlLimitEnabled LimitTranslationX;
+
+	UPROPERTY(EditAnywhere, BlueprintReadWrite, Category = "Settings")
+	FRigControlLimitEnabled LimitTranslationY;
+
+	UPROPERTY(EditAnywhere, BlueprintReadWrite, Category = "Settings")
+	FRigControlLimitEnabled LimitTranslationZ;
+
+	UPROPERTY(EditAnywhere, BlueprintReadWrite, Category = "Settings")
+	FRigControlLimitEnabled LimitPitch;
+
+	UPROPERTY(EditAnywhere, BlueprintReadWrite, Category = "Settings")
+	FRigControlLimitEnabled LimitYaw;
+
+	UPROPERTY(EditAnywhere, BlueprintReadWrite, Category = "Settings")
+	FRigControlLimitEnabled LimitRoll;
+
+	UPROPERTY(EditAnywhere, BlueprintReadWrite, Category = "Settings")
+	FRigControlLimitEnabled LimitScaleX;
+
+	UPROPERTY(EditAnywhere, BlueprintReadWrite, Category = "Settings")
+	FRigControlLimitEnabled LimitScaleY;
+
+	UPROPERTY(EditAnywhere, BlueprintReadWrite, Category = "Settings")
+	FRigControlLimitEnabled LimitScaleZ;
+	
+	UPROPERTY(EditAnywhere, BlueprintReadWrite, Category = "Settings")
+	FEulerTransform MinValue;
+
+	UPROPERTY(EditAnywhere, BlueprintReadWrite, Category = "Settings")
+	FEulerTransform MaxValue;
+
+	UPROPERTY(EditAnywhere, BlueprintReadWrite, Category = "Settings")
+	bool bDrawLimits;
+};
+
+
 USTRUCT(BlueprintType)
 struct CONTROLRIG_API FRigUnit_HierarchyAddControlTransform_Settings : public FRigUnit_HierarchyAddControl_Settings
 {
@@ -1026,6 +1087,9 @@ struct CONTROLRIG_API FRigUnit_HierarchyAddControlTransform_Settings : public FR
 	UPROPERTY(EditAnywhere, BlueprintReadWrite, Category = "Settings")
 	EEulerRotationOrder PreferredRotationOrder;
 
+	UPROPERTY(EditAnywhere, BlueprintReadWrite, Category = "Settings")
+	FRigUnit_HierarchyAddControlTransform_LimitSettings Limits;
+
 	UPROPERTY(EditAnywhere, BlueprintReadWrite, Category = "Settings")
 	FRigUnit_HierarchyAddControl_ShapeSettings Shape;
 
