Control Rig: Save the type of control rig in the asset registry of the control rig blueprint
#rnx

diff --git a/Engine/Plugins/Animation/ControlRig/Source/ControlRigDeveloper/Private/ControlRigBlueprint.cpp b/Engine/Plugins/Animation/ControlRig/Source/ControlRigDeveloper/Private/ControlRigBlueprint.cpp
--- a/Engine/Plugins/Animation/ControlRig/Source/ControlRigDeveloper/Private/ControlRigBlueprint.cpp
+++ b/Engine/Plugins/Animation/ControlRig/Source/ControlRigDeveloper/Private/ControlRigBlueprint.cpp
@@ -910,6 +910,22 @@ void UControlRigBlueprint::PreSave(FObjectPreSaveContext ObjectSaveContext)
 	});
 
 	UpdateExposedModuleConnectors();
+
+	if (IsControlRigModule())
+	{
+		ControlRigType = EControlRigType::RigModule;
+		ItemTypeDisplayName = TEXT("Rig Module");
+	}
+	else if (GetControlRigClass()->IsChildOf(UModularRig::StaticClass()))
+	{
+		ControlRigType = EControlRigType::ModularRig;
+		ItemTypeDisplayName = TEXT("Modular Rig");
+	}
+	else
+	{
+		ControlRigType = EControlRigType::IndependentRig;
+		ItemTypeDisplayName = TEXT("Control Rig");
+	}
 }
 
 void UControlRigBlueprint::UpdateExposedModuleConnectors() const
diff --git a/Engine/Plugins/Animation/ControlRig/Source/ControlRigDeveloper/Public/ControlRigBlueprint.h b/Engine/Plugins/Animation/ControlRig/Source/ControlRigDeveloper/Public/ControlRigBlueprint.h
--- a/Engine/Plugins/Animation/ControlRig/Source/ControlRigDeveloper/Public/ControlRigBlueprint.h
+++ b/Engine/Plugins/Animation/ControlRig/Source/ControlRigDeveloper/Public/ControlRigBlueprint.h
@@ -31,6 +31,14 @@ class USkeletalMesh;
 class UControlRigGraph;
 struct FEndLoadPackageContext;
 
+UENUM(BlueprintType)
+enum class EControlRigType : uint8
+{
+	IndependentRig = 0,
+	RigModule = 1,
+	ModularRig =2,
+};
+
 
 UCLASS(BlueprintType, meta=(IgnoreClassThumbnail))
 class CONTROLRIGDEVELOPER_API UControlRigBlueprint : public URigVMBlueprint, public IInterface_PreviewMeshProvider
@@ -219,6 +227,12 @@ public:
 	UFUNCTION(BlueprintCallable, Category = "Control Rig Blueprint")
 	void RecompileModularRig();
 
+	UPROPERTY(AssetRegistrySearchable)
+	EControlRigType ControlRigType;
+
+	UPROPERTY(AssetRegistrySearchable)
+	FName ItemTypeDisplayName = TEXT("Control Rig");
+
 private:
 
 	/** Whether or not this rig has an Inversion Event */
