Control Rig: Reduce the memory used for storing children of elements. The ptr-based, non-recursive GetChildren function is now faster by completely avoiding allocations.

#rnx

diff --git a/Engine/Plugins/Animation/ControlRig/Source/ControlRig/Private/Rigs/RigHierarchy.cpp b/Engine/Plugins/Animation/ControlRig/Source/ControlRig/Private/Rigs/RigHierarchy.cpp
--- a/Engine/Plugins/Animation/ControlRig/Source/ControlRig/Private/Rigs/RigHierarchy.cpp
+++ b/Engine/Plugins/Animation/ControlRig/Source/ControlRig/Private/Rigs/RigHierarchy.cpp
@@ -9,6 +9,7 @@
 #include "Math/ControlRigMathLibrary.h"
 #include "UObject/AnimObjectVersion.h"
 #include "ControlRigObjectVersion.h"
+#include "Algo/Count.h"
 #include "Serialization/MemoryReader.h"
 #include "Serialization/MemoryWriter.h"
 #include "HAL/LowLevelMemTracker.h"
@@ -88,8 +89,6 @@ FAutoConsoleCommandWithWorldAndArgs FCmdControlRigHierarchyTraceFrames
 // URigHierarchy
 ////////////////////////////////////////////////////////////////////////////////
 
-const FRigBaseElementChildrenArray URigHierarchy::EmptyElementArray;
-
 #if URIGHIERARCHY_ENSURE_CACHE_VALIDITY
 bool URigHierarchy::bEnableValidityCheckbyDefault = true;
 #else
@@ -296,8 +295,6 @@ void URigHierarchy::Load(FArchive& Ar)
 		}
 	}
 
-	UpdateAllCachedChildren(true);
-
 	if(Ar.IsTransacting())
 	{
 		for(const FRigElementKey& SelectedKey : SelectedKeys)
@@ -397,6 +394,11 @@ void URigHierarchy::Reset_Impl(bool bResetElements)
 	PoseVersionPerElement.Reset();
 	ElementDependencyCache.Reset();
 
+	ChildElementOffsetAndCountCache.Reset();
+	ChildElementCache.Reset();
+	ChildElementCacheTopologyVersion = std::numeric_limits<uint32>::max();
+
+
 	if(!IsGarbageCollecting())
 	{
 		Notify(ERigHierarchyNotification::HierarchyReset, nullptr);
@@ -567,7 +569,6 @@ void URigHierarchy::CopyHierarchy(URigHierarchy* InHierarchy)
 			check(Target->Key.Type == Source->Key.Type);
 			Target->InitializeFrom(Source);
 
-			Target->TopologyVersion = InHierarchy->GetTopologyVersion();
 			IncrementPoseVersion(Index);
 		}
 
@@ -586,10 +587,12 @@ void URigHierarchy::CopyHierarchy(URigHierarchy* InHierarchy)
 	}
 
 	PreviousNameMap.Append(InHierarchy->PreviousNameMap);
-	TopologyVersion = InHierarchy->GetTopologyVersion();
+
+	// Increment the topology version to invalidate our cached children.
+	IncrementTopologyVersion();
+	
 	MetadataVersion = InHierarchy->GetMetadataVersion();
 
-	UpdateAllCachedChildren(true);
 	EnsureCacheValidity();
 }
 
@@ -1510,78 +1513,115 @@ FString URigHierarchy::GetControlPinDefaultValue(FRigControlElement* InControlEl
 TArray<FRigElementKey> URigHierarchy::GetChildren(FRigElementKey InKey, bool bRecursive) const
 {
 	LLM_SCOPE_BYNAME(TEXT("Animation/ControlRig"));
-	FRigBaseElementChildrenArray LocalChildren;
-	const FRigBaseElementChildrenArray* ChildrenPtr = nullptr;
-	if(bRecursive)
+
+	auto ConvertElementsToKeys = [this](TConstArrayView<FRigBaseElement*> InElements) -> TArray<FRigElementKey>
 	{
-		LocalChildren = GetChildren(Find(InKey), true);
-		ChildrenPtr = & LocalChildren;
-	}
-	else
+		TArray<FRigElementKey> ElementKeys;
+		ElementKeys.Reserve(InElements.Num());
+		for (const FRigBaseElement* Element: InElements)
+		{
+			ElementKeys.Add(Element->Key);
+		}
+		return ElementKeys;
+	};
+
+	if (bRecursive)
 	{
-		ChildrenPtr = &GetChildren(Find(InKey));
+		return ConvertElementsToKeys(GetChildren(Find(InKey), true));
 	}
-
-	const FRigBaseElementChildrenArray& Children = *ChildrenPtr;
-
-	TArray<FRigElementKey> Keys;
-	for(const FRigBaseElement* Child : Children)
+	else
 	{
-		Keys.Add(Child->Key);
+		return ConvertElementsToKeys(GetChildren(Find(InKey)));
 	}
-	return Keys;
 }
 
 TArray<int32> URigHierarchy::GetChildren(int32 InIndex, bool bRecursive) const
 {
 	LLM_SCOPE_BYNAME(TEXT("Animation/ControlRig"));
-	FRigBaseElementChildrenArray LocalChildren;
-	const FRigBaseElementChildrenArray* ChildrenPtr = nullptr;
-	if(bRecursive)
+
+	if (!ensure(Elements.IsValidIndex(InIndex)))
 	{
-		LocalChildren = GetChildren(Get(InIndex), true);
-		ChildrenPtr = & LocalChildren;
+		return {};
 	}
-	else
+	
+	TArray<int32> ChildIndexes;
+
+	auto AddIndexesFromElements = [this, &ChildIndexes](TConstArrayView<FRigBaseElement*> InElements) -> void
 	{
-		ChildrenPtr = &GetChildren(Get(InIndex));
+		ChildIndexes.Reserve(ChildIndexes.Num() + InElements.Num());
+		for (const FRigBaseElement* Element: InElements)
+		{
+			ChildIndexes.Add(Element->Index);
+		}
+	};
+	
+	AddIndexesFromElements(GetChildren(Elements[InIndex]));
+
+	if (bRecursive)
+	{
+		// Go along the children array and add all children. Once we stop adding children, the traversal index
+		// will reach the end and we're done.
+		for(int32 TraversalIndex = 0; TraversalIndex != ChildIndexes.Num(); TraversalIndex++)
+		{
+			AddIndexesFromElements(GetChildren(Elements[ChildIndexes[TraversalIndex]]));
+		}
 	}
 
-	const FRigBaseElementChildrenArray& Children = *ChildrenPtr;
+	return ChildIndexes;
+}
 
-	TArray<int32> Indices;
-	for(const FRigBaseElement* Child : Children)
+TConstArrayView<FRigBaseElement*> URigHierarchy::GetChildren(const FRigBaseElement* InElement) const
+{
+	LLM_SCOPE_BYNAME(TEXT("Animation/ControlRig"));
+	if(InElement)
 	{
-		Indices.Add(Child->Index);
+		EnsureCachedChildrenAreCurrent();
+
+		if (InElement->ChildCacheIndex != INDEX_NONE)
+		{
+			const FChildElementOffsetAndCount& OffsetAndCount = ChildElementOffsetAndCountCache[InElement->ChildCacheIndex];
+			return TConstArrayView<FRigBaseElement*>(&ChildElementCache[OffsetAndCount.Offset], OffsetAndCount.Count);
+		}
 	}
-	return Indices;
+	return {};
 }
 
-const FRigBaseElementChildrenArray& URigHierarchy::GetChildren(const FRigBaseElement* InElement) const
+TArrayView<FRigBaseElement*> URigHierarchy::GetChildren(const FRigBaseElement* InElement)
 {
 	LLM_SCOPE_BYNAME(TEXT("Animation/ControlRig"));
 	if(InElement)
 	{
-		UpdateCachedChildren(InElement);
-		return InElement->CachedChildren;
+		EnsureCachedChildrenAreCurrent();
+
+		if (InElement->ChildCacheIndex != INDEX_NONE)
+		{
+			const FChildElementOffsetAndCount& OffsetAndCount = ChildElementOffsetAndCountCache[InElement->ChildCacheIndex];
+			return TArrayView<FRigBaseElement*>(&ChildElementCache[OffsetAndCount.Offset], OffsetAndCount.Count);
+		}
 	}
-	return EmptyElementArray;
+	return {};
 }
 
+
+
 FRigBaseElementChildrenArray URigHierarchy::GetChildren(const FRigBaseElement* InElement, bool bRecursive) const
 {
 	LLM_SCOPE_BYNAME(TEXT("Animation/ControlRig"));
-	// call the non-recursive variation
-	FRigBaseElementChildrenArray Children = GetChildren(InElement);
-	
-	if(bRecursive)
+
+	FRigBaseElementChildrenArray Children;
+
+	Children.Append(GetChildren(InElement));
+
+	if (bRecursive)
 	{
-		for(int32 ChildIndex = 0; ChildIndex < Children.Num(); ChildIndex++)
+		// Go along the children array and add all children. Once we stop adding children, the traversal index
+		// will reach the end and we're done.
+		for(int32 TraversalIndex = 0; TraversalIndex != Children.Num(); TraversalIndex++)
 		{
-			Children.Append(GetChildren(Children[ChildIndex], true));
+			Children.Append(GetChildren(Children[TraversalIndex]));
 		}
 	}
-
+	
 	return Children;
 }
 
@@ -2341,8 +2381,7 @@ void URigHierarchy::Traverse(FRigBaseElement* InElement, bool bTowardsChildren,
 	{
 		if(bTowardsChildren)
 		{
-			const FRigBaseElementChildrenArray& Children = GetChildren(InElement);
-			for (FRigBaseElement* Child : Children)
+			for (FRigBaseElement* Child : GetChildren(InElement))
 			{
 				Traverse(Child, true, PerElementFunction);
 			}
@@ -3246,8 +3285,7 @@ void URigHierarchy::SetTransform(FRigTransformElement* InTransformElement, const
 					// only fire these notes if we are inside the asset editor
 					if(World->WorldType == EWorldType::EditorPreview)
 					{
-						const FRigBaseElementChildrenArray& Children = GetChildren(InTransformElement);
-						for (FRigBaseElement* Child : Children)
+						for (FRigBaseElement* Child : GetChildren(InTransformElement))
 						{
 							const bool bChildFound = WrittenTransformsAtRuntime.ContainsByPredicate([Child](const TInstructionSliceElement& Entry) -> bool
 							{
@@ -4449,11 +4487,11 @@ int32 URigHierarchy::GetLocalIndex(const FRigBaseElement* InElement) const
 	
 	if(const FRigBaseElement* ParentElement = GetFirstParent(InElement))
 	{
-		const FRigBaseElementChildrenArray& Children = GetChildren(ParentElement);
-		return Children.Find((FRigBaseElement*)InElement);
+		TConstArrayView<FRigBaseElement*> Children = GetChildren(ParentElement);
+		return Children.Find(const_cast<FRigBaseElement*>(InElement));
 	}
 
-	return GetRootElements().Find((FRigBaseElement*)InElement);
+	return GetRootElements().Find(const_cast<FRigBaseElement*>(InElement));
 }
 
 bool URigHierarchy::IsTracingChanges() const
@@ -4607,99 +4645,124 @@ bool URigHierarchy::IsSelected(const FRigBaseElement* InElement) const
 	return bIsSelected;
 }
 
-void URigHierarchy::ResetCachedChildren()
-{
-	LLM_SCOPE_BYNAME(TEXT("Animation/ControlRig"));
-	for (int32 ElementIndex = 0; ElementIndex < Elements.Num(); ElementIndex++)
-	{
-		FRigBaseElement* Element = Elements[ElementIndex];
-		Element->CachedChildren.Reset();
-	}
-}
 
-bool URigHierarchy::UpdateCachedChildren(const FRigBaseElement* InElement, bool bForce) const
+void URigHierarchy::EnsureCachedChildrenAreCurrent() const
 {
 	LLM_SCOPE_BYNAME(TEXT("Animation/ControlRig"));
-	check(InElement);
 
-	if(InElement->TopologyVersion == TopologyVersion && !bForce)
+	if(ChildElementCacheTopologyVersion != TopologyVersion)
 	{
-		return false;
+		const_cast<URigHierarchy*>(this)->UpdateCachedChildren();
 	}
+}
+	
+void URigHierarchy::UpdateCachedChildren()
+{
+	FScopeLock Lock(&ElementsLock);
+	
+	// First we tally up how many children each element has, then we allocate for the total
+	// count and do the same loop again.
+	TArray<int32> ChildrenCount;
+	ChildrenCount.SetNumZeroed(Elements.Num());
 
-	// we collect the transforms into a separate array to loop faster,
-	// since some hierarchies have a very large number of curves / non transform elements.
-	TArray<FRigTransformElement*> TransformElements;
-	TransformElements.Reserve(Elements.Num());
+	// Bit array that denotes elements that have parents. We'll use this to quickly iterate
+	// for the second pass.
+	TBitArray<> ElementHasParent(false, Elements.Num());
 
 	for (int32 ElementIndex = 0; ElementIndex < Elements.Num(); ElementIndex++)
 	{
-		FRigBaseElement* Element = Elements[ElementIndex];
-		if(bForce || (Element->TopologyVersion != TopologyVersion))
+		const FRigBaseElement* Element = Elements[ElementIndex];
+		
+		if(const FRigSingleParentElement* SingleParentElement = Cast<FRigSingleParentElement>(Element))
 		{
-			if(FRigTransformElement* TransformElement = Cast<FRigTransformElement>(Element))
+			if(const FRigTransformElement* ParentElement = SingleParentElement->ParentElement)
 			{
-				TransformElements.Add(TransformElement);
-				TransformElement->CachedChildren.Reset();
+				ChildrenCount[ParentElement->Index]++;
+				ElementHasParent[ElementIndex] = true;
 			}
-			else
+		}
+		else if(const FRigMultiParentElement* MultiParentElement = Cast<FRigMultiParentElement>(Element))
+		{
+			for(const FRigElementParentConstraint& ParentConstraint : MultiParentElement->ParentConstraints)
 			{
-				Element->TopologyVersion = TopologyVersion;
+				if(const FRigTransformElement* ParentElement = ParentConstraint.ParentElement)
+				{
+					ChildrenCount[ParentElement->Index]++;
+					ElementHasParent[ElementIndex] = true;
+				}
 			}
 		}
 	}
 
-	// since we'll have to loop over all children anyway - it makes sense to update all of them
-	// at the same time.
-	for (FRigTransformElement* Element : TransformElements)
+	// Tally up how many elements have children.
+	const int32 NumElementsWithChildren = Algo::CountIf(ChildrenCount, [](int32 InCount) { return InCount > 0; });
+
+	ChildElementOffsetAndCountCache.Reset(NumElementsWithChildren);
+
+	// Tally up how many children there are in total and set the index on each of the elements as we go.
+	int32 TotalChildren = 0;
+	for (int32 ElementIndex = 0; ElementIndex < Elements.Num(); ElementIndex++)
 	{
-		if(FRigSingleParentElement* SingleParentElement = Cast<FRigSingleParentElement>(Element))
+		if (ChildrenCount[ElementIndex])
 		{
-			if(FRigTransformElement* ParentElement = SingleParentElement->ParentElement)
+			Elements[ElementIndex]->ChildCacheIndex = ChildElementOffsetAndCountCache.Num();
+			ChildElementOffsetAndCountCache.Add({TotalChildren, ChildrenCount[ElementIndex]});
+			TotalChildren += ChildrenCount[ElementIndex];
+		}
+		else
+		{
+			// This element has no children, mark it as having no entry in the child table.
+			Elements[ElementIndex]->ChildCacheIndex = INDEX_NONE;
+		}
+	}
+
+	// Now run through all elements that are known to have parents and start filling up the children array.
+	ChildElementCache.Reset();
+	ChildElementCache.SetNumZeroed(TotalChildren);
+
+	// Recycle this array to indicate where we are with each set of children, as a local offset into each
+	// element's children sub-array.
+	ChildrenCount.Reset();
+	ChildrenCount.SetNumZeroed(Elements.Num());	
+
+	auto SetChildElement = [this, &ChildrenCount](
+		const FRigTransformElement* InParentElement,
+		FRigBaseElement* InChildElement)
+	{
+		const int32 ParentElementIndex = InParentElement->Index;
+		const int32 CacheOffset = ChildElementOffsetAndCountCache[InParentElement->ChildCacheIndex].Offset;
+
+		ChildElementCache[CacheOffset + ChildrenCount[ParentElementIndex]] = InChildElement;
+		ChildrenCount[ParentElementIndex]++;
+	};
+	
+	for (TConstSetBitIterator<> ElementIt(ElementHasParent); ElementIt; ++ElementIt)
+	{
+		FRigBaseElement* Element = Elements[ElementIt.GetIndex()];
+		
+		if(const FRigSingleParentElement* SingleParentElement = Cast<FRigSingleParentElement>(Element))
+		{
+			if(const FRigTransformElement* ParentElement = SingleParentElement->ParentElement)
 			{
-				if(bForce || (ParentElement->TopologyVersion != TopologyVersion))
-				{
-					ParentElement->CachedChildren.Add(SingleParentElement);
-				}
+				SetChildElement(ParentElement, Element);
 			}
 		}
-		else if(FRigMultiParentElement* MultiParentElement = Cast<FRigMultiParentElement>(Element))
+		else if(const FRigMultiParentElement* MultiParentElement = Cast<FRigMultiParentElement>(Element))
 		{
 			for(const FRigElementParentConstraint& ParentConstraint : MultiParentElement->ParentConstraints)
 			{
-				if(FRigTransformElement* ParentElement = ParentConstraint.ParentElement)
+				if(const FRigTransformElement* ParentElement = ParentConstraint.ParentElement)
 				{
-					if(bForce || (ParentElement->TopologyVersion != TopologyVersion))
-					{
-						ParentElement->CachedChildren.Add(MultiParentElement);
-					}
+					SetChildElement(ParentElement, Element);
 				}
 			}
 		}
 	}
 
-	for (FRigTransformElement* Element : TransformElements)
-	{
-		Element->TopologyVersion = TopologyVersion;
-	}
-
-	return true;
+	// Mark the cache up-to-date.
+	ChildElementCacheTopologyVersion = TopologyVersion;
 }
 
-void URigHierarchy::UpdateAllCachedChildren(bool bForce) const
-{
-	LLM_SCOPE_BYNAME(TEXT("Animation/ControlRig"));
-	
-	for (const FRigBaseElement* Element : Elements)
-	{
-		// as soon as one element was updated,
-		// as a side effect all other elements are updated.
-		if(UpdateCachedChildren(Element, bForce))
-		{
-			return;
-		}
-	}
-}
 	
 FRigElementKey URigHierarchy::PreprocessParentElementKeyForSpaceSwitching(const FRigElementKey& InChildKey, const FRigElementKey& InParentKey)
 {
diff --git a/Engine/Plugins/Animation/ControlRig/Source/ControlRig/Private/Rigs/RigHierarchyController.cpp b/Engine/Plugins/Animation/ControlRig/Source/ControlRig/Private/Rigs/RigHierarchyController.cpp
--- a/Engine/Plugins/Animation/ControlRig/Source/ControlRig/Private/Rigs/RigHierarchyController.cpp
+++ b/Engine/Plugins/Animation/ControlRig/Source/ControlRig/Private/Rigs/RigHierarchyController.cpp
@@ -2362,11 +2362,9 @@ bool URigHierarchyController::ReorderElement(FRigBaseElement* InElement, int32 I
 	InIndex = FMath::Max<int32>(InIndex, 0);
 
 	TArray<FRigBaseElement*> LocalElements;
-	const FRigBaseElement* ParentElement = Hierarchy->GetFirstParent(InElement);
-	if(ParentElement)
+	if(const FRigBaseElement* ParentElement = Hierarchy->GetFirstParent(InElement))
 	{
-		const FRigBaseElementChildrenArray& Children = Hierarchy->GetChildren(ParentElement);
-		LocalElements.Append(Children);
+		LocalElements.Append(Hierarchy->GetChildren(ParentElement));
 	}
 	else
 	{
diff --git a/Engine/Plugins/Animation/ControlRig/Source/ControlRig/Private/Rigs/RigHierarchyElements.cpp b/Engine/Plugins/Animation/ControlRig/Source/ControlRig/Private/Rigs/RigHierarchyElements.cpp
--- a/Engine/Plugins/Animation/ControlRig/Source/ControlRig/Private/Rigs/RigHierarchyElements.cpp
+++ b/Engine/Plugins/Animation/ControlRig/Source/ControlRig/Private/Rigs/RigHierarchyElements.cpp
@@ -102,6 +102,7 @@ void FRigBaseElement::Load(FArchive& Ar, ESerializationPhase SerializationPhase)
 		ensure(LoadedKey.Type == Key.Type);
 		Key = LoadedKey;
 
+		ChildCacheIndex = INDEX_NONE;
 		CachedNameString.Reset();
 
 		if (Ar.CustomVer(FControlRigObjectVersion::GUID) >= FControlRigObjectVersion::HierarchyElementMetadata &&
@@ -208,8 +209,6 @@ void FRigBaseElement::InitializeFrom(const FRigBaseElement* InOther)
 	SubIndex = InOther->SubIndex;
 	CreatedAtInstructionIndex = InOther->CreatedAtInstructionIndex;
 	bSelected = false;
-	TopologyVersion = 0;
-	CachedChildren.Reset();
 }
 
 
diff --git a/Engine/Plugins/Animation/ControlRig/Source/ControlRig/Private/Units/Hierarchy/RigUnit_ControlChannel.cpp b/Engine/Plugins/Animation/ControlRig/Source/ControlRig/Private/Units/Hierarchy/RigUnit_ControlChannel.cpp
--- a/Engine/Plugins/Animation/ControlRig/Source/ControlRig/Private/Units/Hierarchy/RigUnit_ControlChannel.cpp
+++ b/Engine/Plugins/Animation/ControlRig/Source/ControlRig/Private/Units/Hierarchy/RigUnit_ControlChannel.cpp
@@ -22,8 +22,7 @@ bool FRigUnit_GetAnimationChannelBase::UpdateCache(const URigHierarchy* InHierar
 	{
 		if(const FRigControlElement* ControlElement = InHierarchy->Find<FRigControlElement>(FRigElementKey(Control, ERigElementType::Control)))
 		{
-			FRigBaseElementChildrenArray Children = InHierarchy->GetChildren(ControlElement);
-			for(const FRigBaseElement* Child : Children)
+			for(const FRigBaseElement* Child : InHierarchy->GetChildren(ControlElement))
 			{
 				if(const FRigControlElement* ChildControl = Cast<FRigControlElement>(Child))
 				{
diff --git a/Engine/Plugins/Animation/ControlRig/Source/ControlRig/Public/Rigs/RigHierarchy.h b/Engine/Plugins/Animation/ControlRig/Source/ControlRig/Public/Rigs/RigHierarchy.h
--- a/Engine/Plugins/Animation/ControlRig/Source/ControlRig/Public/Rigs/RigHierarchy.h
+++ b/Engine/Plugins/Animation/ControlRig/Source/ControlRig/Public/Rigs/RigHierarchy.h
@@ -2930,7 +2930,8 @@ public:
 	 * @param InElement The element to retrieve the children for
 	 * @return Returns the child elements
 	 */
-	const FRigBaseElementChildrenArray& GetChildren(const FRigBaseElement* InElement) const;
+	TConstArrayView<FRigBaseElement*> GetChildren(const FRigBaseElement* InElement) const;
+	TArrayView<FRigBaseElement*> GetChildren(const FRigBaseElement* InElement);
 
 	/**
 	 * Returns the child elements of a given element
@@ -3404,7 +3405,7 @@ public:
 	/**
 	 * Returns the topology version of this hierarchy
 	 */
-	uint16 GetTopologyVersion() const { return TopologyVersion; }
+	uint32 GetTopologyVersion() const { return TopologyVersion; }
 
 	/**
 	 * Returns the hash of this hierarchy used for cached element keys
@@ -4004,24 +4005,14 @@ private:
     bool IsSelected(const FRigBaseElement* InElement) const;
 
 	/**
-	 * Removes the transient cached children table for all elements.
-	 */
-	void ResetCachedChildren();
-
-	/**
-	 * Updates the transient cached children table for a given element if needed (or if bForce == true).
-	 * @param InElement The element to update the children table for
-	 * @param bForce If set to true the table will always be updated
+	 * Updates the transient cached children table if the topology version is out of date with the one
+	 * stored with the cached table.
 	 * @return Returns true if a change was performed
 	 */
-	bool UpdateCachedChildren(const FRigBaseElement* InElement, bool bForce = false) const;
-
-	/**
-	* Updates the transient cached children table for all elements if needed.
-	 * @param bForce If set to true the table will always be updated
-	*/
-	void UpdateAllCachedChildren(bool bForce = false) const;
+	void EnsureCachedChildrenAreCurrent() const;
 
+	void UpdateCachedChildren();
+	
 	/**
 	 * Corrects a parent element key for space switching
 	 */
@@ -4088,7 +4079,7 @@ private:
 	 * added, removed, re-parented or renamed.
 	 */
 	UPROPERTY(transient)
-	uint16 TopologyVersion;
+	uint32 TopologyVersion;
 
 	/**
 	 * The metadata version of the hierarchy changes when metadata is being
@@ -4137,10 +4128,29 @@ private:
 
 	TMap<FRigElementKey, FString> UserDefinedElementName;
 
+	// Element metadata storage. Storage is defined here rather than on the elements
+	// to reduce memory consumption. Only elements created by MakeElement point to
+	// the element storage. Copied elements via the copy constructor or copy operator
+	// do not have URigHierarchy as an owner and therefore do not carry metadata with them.
 	TMap<FRigElementKey, FMetaDataStorage> ElementMetadata;
 
-	// Static empty element array used for ref returns
-	static const FRigBaseElementChildrenArray EmptyElementArray;
+	// A quick-lookup cache for elements' children. Each element that has a ChildCacheIndex
+	// not equal to INDEX_NONE is an index into the offset and count cache below, which in
+	// turn contains an offset into the ChildElementCache, which stores consecutive runs of
+	// children for that element. This makes it quick to get a TArrayView on the list of
+	// children.
+	struct FChildElementOffsetAndCount
+	{
+		int32 Offset;
+		int32 Count;
+	};
+	
+	TArray<FChildElementOffsetAndCount> ChildElementOffsetAndCountCache;
+	TArray<FRigBaseElement*> ChildElementCache;
+
+	// The topology version at which the child element cache was constructed. If it differs
+	// from the stored TopologyVersion, then the cache is rebuilt. 
+	uint32 ChildElementCacheTopologyVersion = std::numeric_limits<uint32>::max();
 
 	///////////////////////////////////////////////
 	/// Undo redo related
diff --git a/Engine/Plugins/Animation/ControlRig/Source/ControlRig/Public/Rigs/RigHierarchyElements.h b/Engine/Plugins/Animation/ControlRig/Source/ControlRig/Public/Rigs/RigHierarchyElements.h
--- a/Engine/Plugins/Animation/ControlRig/Source/ControlRig/Public/Rigs/RigHierarchyElements.h
+++ b/Engine/Plugins/Animation/ControlRig/Source/ControlRig/Public/Rigs/RigHierarchyElements.h
@@ -642,9 +642,8 @@ protected:
 	// Set by URigHierarchy::NewElement.
 	int32 OwnedInstances = 0;
 
-	mutable uint16 TopologyVersion = 0;
-	
-	mutable FRigBaseElementChildrenArray CachedChildren;
+	// Index into the child cache offset and count table in URigHierarchy. Set by URigHierarchy::UpdateCachedChildren
+	int32 ChildCacheIndex = INDEX_NONE;
 	
 	mutable FString CachedNameString;
 
diff --git a/Engine/Plugins/Animation/ControlRig/Source/ControlRigEditor/Private/ControlRigElementDetails.cpp b/Engine/Plugins/Animation/ControlRig/Source/ControlRigEditor/Private/ControlRigElementDetails.cpp
--- a/Engine/Plugins/Animation/ControlRig/Source/ControlRigEditor/Private/ControlRigElementDetails.cpp
+++ b/Engine/Plugins/Animation/ControlRig/Source/ControlRigEditor/Private/ControlRigElementDetails.cpp
@@ -3270,8 +3270,7 @@ void FRigControlElementDetails::CustomizeAnimationChannels(IDetailLayoutBuilder&
 	Category.HeaderContent(HeaderContentWidget);
 
 	bool bHasAnimationChannels = false;
-	const FRigBaseElementChildrenArray ChildElements = Hierarchy->GetChildren(ControlElement);
-	for(const FRigBaseElement* ChildElement : ChildElements)
+	for(const FRigBaseElement* ChildElement : Hierarchy->GetChildren(ControlElement))
 	{
 		if(const FRigControlElement* ChildControlElement = Cast<FRigControlElement>(ChildElement))
 		{
diff --git a/Engine/Plugins/Animation/ControlRig/Source/ControlRigEditor/Private/Graph/ControlRigGraphPanelPinFactory.cpp b/Engine/Plugins/Animation/ControlRig/Source/ControlRigEditor/Private/Graph/ControlRigGraphPanelPinFactory.cpp
--- a/Engine/Plugins/Animation/ControlRig/Source/ControlRigEditor/Private/Graph/ControlRigGraphPanelPinFactory.cpp
+++ b/Engine/Plugins/Animation/ControlRig/Source/ControlRigEditor/Private/Graph/ControlRigGraphPanelPinFactory.cpp
@@ -288,8 +288,7 @@ TSharedPtr<SGraphPin> FControlRigGraphPanelPinFactory::CreatePin_Internal(UEdGra
 
 									if(const FRigControlElement* ControlElement = Hierarchy->Find<FRigControlElement>(ControlKey))
 									{
-										FRigBaseElementChildrenArray Children = Hierarchy->GetChildren(ControlElement);
-										for(const FRigBaseElement* Child : Children)
+										for(const FRigBaseElement* Child : Hierarchy->GetChildren(ControlElement))
 										{
 											if(const FRigControlElement* ChildControl = Cast<FRigControlElement>(Child))
 											{
