Control Rig: Rename BaseControlRig to ControlRig
#rnx

diff --git a/Engine/Plugins/Animation/AnimationData/Source/AnimationData/Private/AnimSequencerController.cpp b/Engine/Plugins/Animation/AnimationData/Source/AnimationData/Private/AnimSequencerController.cpp
--- a/Engine/Plugins/Animation/AnimationData/Source/AnimationData/Private/AnimSequencerController.cpp
+++ b/Engine/Plugins/Animation/AnimationData/Source/AnimationData/Private/AnimSequencerController.cpp
@@ -1305,7 +1305,7 @@ bool UAnimSequencerController::SetCurveAttributes(const FAnimationCurveIdentifie
 
 		if (UMovieSceneControlRigParameterSection* Section = Model->GetFKControlRigSection())
 		{
-			UBaseControlRig* ControlRig = Section->GetControlRig();
+			UControlRig* ControlRig = Section->GetControlRig();
 
 			if (UFKControlRig* FKRig = Cast<UFKControlRig>(ControlRig))
 			{
@@ -1543,7 +1543,7 @@ bool UAnimSequencerController::SetBoneTrackKeys(FName BoneName, const TArray<FVe
 		{
 			if (UMovieSceneControlRigParameterSection* Section = Model->GetFKControlRigSection())
 			{
-				if (UBaseControlRig* ControlRig = Model->GetControlRig())
+				if (UControlRig* ControlRig = Model->GetControlRig())
 				{
 					if (URigHierarchy* Hierarchy = ControlRig->GetHierarchy())
 					{
@@ -2150,7 +2150,7 @@ void UAnimSequencerController::RemoveUnusedControlsAndCurves() const
 {
 	if (UMovieSceneControlRigParameterSection* Section = Model->GetFKControlRigSection())
 	{
-		if (UBaseControlRig* ControlRig = Model->GetControlRig())
+		if (UControlRig* ControlRig = Model->GetControlRig())
 		{
 			if (URigHierarchy* Hierarchy = ControlRig->GetHierarchy())
 			{
@@ -2609,7 +2609,7 @@ bool UAnimSequencerController::SetBoneCurveKeys(const FName& BoneName, const TAr
 	check(UAnimationSequencerDataModel::UseDirectFKControlRigMode == 1);
 	if (UMovieSceneControlRigParameterSection* Section = Model->GetFKControlRigSection())
 	{
-		UBaseControlRig* ControlRig = Section->GetControlRig();
+		UControlRig* ControlRig = Section->GetControlRig();
 		if (UFKControlRig* FKRig = Cast<UFKControlRig>(ControlRig))
 		{
 			const URigHierarchy* Hierarchy = FKRig->GetHierarchy();
@@ -2801,7 +2801,7 @@ bool UAnimSequencerController::UpdateBoneCurveKeys(const FName& BoneName, const
 {
 	if (UMovieSceneControlRigParameterSection* Section = Model->GetFKControlRigSection())
 	{		
-		UBaseControlRig* ControlRig = Section->GetControlRig();
+		UControlRig* ControlRig = Section->GetControlRig();
 		if (UFKControlRig* FKRig = Cast<UFKControlRig>(ControlRig))
 		{
 			if(URigHierarchy* Hierarchy = FKRig->GetHierarchy())
@@ -2896,7 +2896,7 @@ bool UAnimSequencerController::RemoveBoneCurveKey(const FName& BoneName, float T
     	const FAnimationCurveIdentifier CurveId = UAnimationCurveIdentifierExtensions::GetCurveIdentifier(Skeleton, BoneName, ERawCurveTrackTypes::RCT_Transform);
 		if (const UMovieSceneControlRigParameterSection* Section = Model->GetFKControlRigSection())
 		{
-			UBaseControlRig* ControlRig = Section->GetControlRig();
+			UControlRig* ControlRig = Section->GetControlRig();
 	
 			if (UFKControlRig* FKRig = Cast<UFKControlRig>(ControlRig))
 			{
@@ -2970,7 +2970,7 @@ bool UAnimSequencerController::AddCurveControl(const FName& CurveName) const
 {
 	if (UMovieSceneControlRigParameterSection* Section = Model->GetFKControlRigSection())
 	{
-		UBaseControlRig* ControlRig = Section->GetControlRig();
+		UControlRig* ControlRig = Section->GetControlRig();
 		if (UFKControlRig* FKRig = Cast<UFKControlRig>(ControlRig))
 		{
 			if(URigHierarchy* Hierarchy = FKRig->GetHierarchy())
@@ -3053,7 +3053,7 @@ bool UAnimSequencerController::RenameCurveControl(const FName& CurveName, const
 	// Need to rename curve and control element and curve				
 	if (UMovieSceneControlRigParameterSection* Section = Model->GetFKControlRigSection())
 	{
-		if (UBaseControlRig* ControlRig = Section->GetControlRig())
+		if (UControlRig* ControlRig = Section->GetControlRig())
 		{
 			if (URigHierarchy* RigHierarchy = ControlRig->GetHierarchy())
 			{
@@ -3140,7 +3140,7 @@ bool UAnimSequencerController::RemoveCurveControl(const FName& CurveName) const
 {
 	if (UMovieSceneControlRigParameterSection* Section = Model->GetFKControlRigSection())
 	{
-		UBaseControlRig* ControlRig = Section->GetControlRig();
+		UControlRig* ControlRig = Section->GetControlRig();
 
 		if (UFKControlRig* FKRig = Cast<UFKControlRig>(ControlRig))
 		{
@@ -3204,7 +3204,7 @@ bool UAnimSequencerController::SetCurveControlKeys(const FName& CurveName, const
 {
 	if (UMovieSceneControlRigParameterSection* Section = Model->GetFKControlRigSection())
 	{
-		UBaseControlRig* ControlRig = Section->GetControlRig();
+		UControlRig* ControlRig = Section->GetControlRig();
 
 		if (UFKControlRig* FKRig = Cast<UFKControlRig>(ControlRig))
 		{
@@ -3292,7 +3292,7 @@ bool UAnimSequencerController::SetCurveControlKey(const FName& CurveName, const
 {
 	if (UMovieSceneControlRigParameterSection* Section = Model->GetFKControlRigSection())
 	{
-		UBaseControlRig* ControlRig = Section->GetControlRig();
+		UControlRig* ControlRig = Section->GetControlRig();
 
 		if (UFKControlRig* FKRig = Cast<UFKControlRig>(ControlRig))
 		{
@@ -3381,7 +3381,7 @@ bool UAnimSequencerController::RemoveCurveControlKey(const FName& CurveName, flo
 {
 	if (UMovieSceneControlRigParameterSection* Section = Model->GetFKControlRigSection())
 	{
-		UBaseControlRig* ControlRig = Section->GetControlRig();
+		UControlRig* ControlRig = Section->GetControlRig();
 
 		if (UFKControlRig* FKRig = Cast<UFKControlRig>(ControlRig))
 		{
diff --git a/Engine/Plugins/Animation/AnimationData/Source/AnimationData/Private/AnimSequencerDataModel.cpp b/Engine/Plugins/Animation/AnimationData/Source/AnimationData/Private/AnimSequencerDataModel.cpp
--- a/Engine/Plugins/Animation/AnimationData/Source/AnimationData/Private/AnimSequencerDataModel.cpp
+++ b/Engine/Plugins/Animation/AnimationData/Source/AnimationData/Private/AnimSequencerDataModel.cpp
@@ -144,7 +144,7 @@ void UAnimationSequencerDataModel::InitializeFKControlRig(UFKControlRig* FKContr
 	}
 }
 
-UBaseControlRig* UAnimationSequencerDataModel::GetControlRig() const
+UControlRig* UAnimationSequencerDataModel::GetControlRig() const
 {
 	if(const UMovieSceneControlRigParameterTrack* Track = GetControlRigTrack())
 	{
@@ -196,7 +196,7 @@ void UAnimationSequencerDataModel::PostLoad()
 void UAnimationSequencerDataModel::DeclareConstructClasses(TArray<FTopLevelAssetPath>& OutConstructClasses, const UClass* SpecificSubclass)
 {
 	Super::DeclareConstructClasses(OutConstructClasses, SpecificSubclass);
-	OutConstructClasses.Add(FTopLevelAssetPath(UBaseControlRig::StaticClass()));
+	OutConstructClasses.Add(FTopLevelAssetPath(UControlRig::StaticClass()));
 }
 #endif
 
@@ -751,7 +751,7 @@ void UAnimationSequencerDataModel::Evaluate(FAnimationPoseData& InOutPoseData, c
 		EvaluateTrack(Track, EvaluationContext);
 
 		// Generate/populate the output animation pose data
-		UBaseControlRig* ControlRig = Track->GetControlRig();
+		UControlRig* ControlRig = Track->GetControlRig();
 		GeneratePoseData(ControlRig, InOutPoseData, EvaluationContext);
 	}
 }
@@ -848,7 +848,7 @@ UMovieSceneControlRigParameterSection* UAnimationSequencerDataModel::GetFKContro
 			{
 				if (UMovieSceneControlRigParameterSection* Section = Cast<UMovieSceneControlRigParameterSection>(TrackSection))
 				{
-					if (const UBaseControlRig* ControlRig = Section->GetControlRig())
+					if (const UControlRig* ControlRig = Section->GetControlRig())
 					{
 						if (ControlRig->IsA<UFKControlRig>())
 						{
@@ -874,7 +874,7 @@ void UAnimationSequencerDataModel::GenerateLegacyCurveData()
 		{
 			if (const UMovieSceneControlRigParameterSection* Section = Cast<UMovieSceneControlRigParameterSection>(TrackSection))
 			{
-				if (const UBaseControlRig* ControlRig = Section->GetControlRig())
+				if (const UControlRig* ControlRig = Section->GetControlRig())
 				{
 						if(URigHierarchy* Hierarchy = ControlRig->GetHierarchy())
 						{
@@ -974,7 +974,7 @@ void UAnimationSequencerDataModel::ValidateControlRigData() const
 	const UMovieSceneControlRigParameterSection* Section = GetFKControlRigSection();
 	checkf(Section, TEXT("Unable to find Control Rig Section"));
 
-	UBaseControlRig* ControlRig = Section->GetControlRig();
+	UControlRig* ControlRig = Section->GetControlRig();
 	checkf(ControlRig, TEXT("Unable to find Control Rig instance for Section"));
 
 	checkf(ControlRig->IsA<UFKControlRig>(), TEXT("Invalid class for Control Rig expected UFKControlRig"));
@@ -1018,7 +1018,7 @@ void UAnimationSequencerDataModel::ValidateLegacyAgainstControlRigData() const
 {
 	UMovieSceneControlRigParameterSection* Section = GetFKControlRigSection();
 
-	UBaseControlRig* ControlRig = Section->GetControlRig();
+	UControlRig* ControlRig = Section->GetControlRig();
 	URigHierarchy* Hierarchy = ControlRig->GetHierarchy();
 
 	// Validate bone tracks against controls
@@ -1083,7 +1083,7 @@ void UAnimationSequencerDataModel::IterateTransformControlCurve(const FName& Bon
 	ValidateControlRigData();
 
 	const UMovieSceneControlRigParameterSection* Section = GetFKControlRigSection();
-	UBaseControlRig* ControlRig = Section->GetControlRig();
+	UControlRig* ControlRig = Section->GetControlRig();
 	const URigHierarchy* Hierarchy = ControlRig->GetHierarchy();
 	
 	const FRigElementKey BoneControlKey(UFKControlRig::GetControlName(BoneName, ERigElementType::Bone), ERigElementType::Control);
@@ -1164,7 +1164,7 @@ UObject* UAnimationSequencerDataModel::GetParentObject(UObject* MovieSceneBlends
 	return GetOuter();
 }
 
-void UAnimationSequencerDataModel::GeneratePoseData(UBaseControlRig* ControlRig, FAnimationPoseData& InOutPoseData, const UE::Anim::DataModel::FEvaluationContext& EvaluationContext) const
+void UAnimationSequencerDataModel::GeneratePoseData(UControlRig* ControlRig, FAnimationPoseData& InOutPoseData, const UE::Anim::DataModel::FEvaluationContext& EvaluationContext) const
 {
 	QUICK_SCOPE_CYCLE_COUNTER(STAT_GeneratePoseData);
 	
@@ -1315,7 +1315,7 @@ void UAnimationSequencerDataModel::EvaluateTrack(UMovieSceneControlRigParameterT
 		bWasDoNotKey = FKRigSection->GetDoNotKey();
 		FKRigSection->SetDoNotKey(true);
 
-		UBaseControlRig* ControlRig = FKRigSection->GetControlRig();
+		UControlRig* ControlRig = FKRigSection->GetControlRig();
 		check(ControlRig);
 
 		// Reset to ref-pose
@@ -1562,7 +1562,7 @@ void UAnimationSequencerDataModel::IterateBoneKeys(const FName& BoneName, TFunct
 	ValidateControlRigData();
 
 	const UMovieSceneControlRigParameterSection* Section = GetFKControlRigSection();
-	UBaseControlRig* ControlRig = Section->GetControlRig();
+	UControlRig* ControlRig = Section->GetControlRig();
 	const URigHierarchy* Hierarchy = ControlRig->GetHierarchy();
 	
 	const FRigElementKey BoneControlKey(UFKControlRig::GetControlName(BoneName, ERigElementType::Bone), ERigElementType::Control);
diff --git a/Engine/Plugins/Animation/AnimationData/Source/AnimationData/Private/AnimSequencerDataModel.h b/Engine/Plugins/Animation/AnimationData/Source/AnimationData/Private/AnimSequencerDataModel.h
--- a/Engine/Plugins/Animation/AnimationData/Source/AnimationData/Private/AnimSequencerDataModel.h
+++ b/Engine/Plugins/Animation/AnimationData/Source/AnimationData/Private/AnimSequencerDataModel.h
@@ -9,7 +9,7 @@
 
 #include "AnimSequencerDataModel.generated.h"
 
-class UBaseControlRig;
+class UControlRig;
 class UMovieSceneControlRigParameterTrack;
 class UMovieSceneControlRigParameterSection;
 class UFKControlRig;
@@ -146,14 +146,14 @@ protected:
 	void ValidateControlRigData() const;
 	void ValidateLegacyAgainstControlRigData() const;
 	
-	void GeneratePoseData(UBaseControlRig* ControlRig, FAnimationPoseData& InOutPoseData, const UE::Anim::DataModel::FEvaluationContext& EvaluationContext) const;
+	void GeneratePoseData(UControlRig* ControlRig, FAnimationPoseData& InOutPoseData, const UE::Anim::DataModel::FEvaluationContext& EvaluationContext) const;
 	void EvaluateTrack(UMovieSceneControlRigParameterTrack* CR_Track, const UE::Anim::DataModel::FEvaluationContext& EvaluationContext) const;
 
 	UMovieSceneControlRigParameterTrack* GetControlRigTrack() const;
 	UMovieSceneControlRigParameterSection* GetFKControlRigSection() const;
 	USkeleton* GetSkeleton() const;
 	void InitializeFKControlRig(UFKControlRig* FKControlRig, USkeleton* Skeleton) const;
-	UBaseControlRig* GetControlRig() const;
+	UControlRig* GetControlRig() const;
 	
 	void IterateTransformControlCurve(const FName& BoneName, TFunction<void(const FTransform&, const FFrameNumber&)> IterationFunction, const TArray<FFrameNumber>* InFrameNumbers = nullptr) const;
 	void GenerateTransformKeysForControl(const FName& BoneName, TArray<FTransform>& InOutTransforms, TArray<FFrameNumber>& InOutFrameNumbers) const;
diff --git a/Engine/Plugins/Animation/ControlRig/Source/ControlRig/Private/AnimNode_ControlRig.cpp b/Engine/Plugins/Animation/ControlRig/Source/ControlRig/Private/AnimNode_ControlRig.cpp
--- a/Engine/Plugins/Animation/ControlRig/Source/ControlRig/Private/AnimNode_ControlRig.cpp
+++ b/Engine/Plugins/Animation/ControlRig/Source/ControlRig/Private/AnimNode_ControlRig.cpp
@@ -198,7 +198,7 @@ void FAnimNode_ControlRig::CacheBones_AnyThread(const FAnimationCacheBonesContex
 		};
 
 		URigHierarchy* Hierarchy = nullptr;
-		if(UBaseControlRig* CurrentControlRig = GetControlRig())
+		if(UControlRig* CurrentControlRig = GetControlRig())
 		{
 			Hierarchy = CurrentControlRig->GetHierarchy();
 		}
@@ -231,7 +231,7 @@ UClass* FAnimNode_ControlRig::GetTargetClass() const
 	return DefaultControlRigClass;
 }
 
-void FAnimNode_ControlRig::UpdateInput(UBaseControlRig* InControlRig, const FPoseContext& InOutput)
+void FAnimNode_ControlRig::UpdateInput(UControlRig* InControlRig, const FPoseContext& InOutput)
 {
 	DECLARE_SCOPE_HIERARCHICAL_COUNTER_FUNC()
 
@@ -256,7 +256,7 @@ void FAnimNode_ControlRig::UpdateInput(UBaseControlRig* InControlRig, const FPos
 	}
 }
 
-void FAnimNode_ControlRig::UpdateOutput(UBaseControlRig* InControlRig, FPoseContext& InOutput)
+void FAnimNode_ControlRig::UpdateOutput(UControlRig* InControlRig, FPoseContext& InOutput)
 {
 	DECLARE_SCOPE_HIERARCHICAL_COUNTER_FUNC()
 
@@ -282,7 +282,7 @@ void FAnimNode_ControlRig::UpdateOutput(UBaseControlRig* InControlRig, FPoseCont
 	}
 }
 
-void FAnimNode_ControlRig::SetControlRigClass(TSubclassOf<UBaseControlRig> InControlRigClass)
+void FAnimNode_ControlRig::SetControlRigClass(TSubclassOf<UControlRig> InControlRigClass)
 {
 	if(DefaultControlRigClass == nullptr)
 	{
@@ -302,7 +302,7 @@ bool FAnimNode_ControlRig::UpdateControlRigIfNeeded(const UAnimInstance* InAnimI
 		{
 			if(ControlRig->GetClass() != ExpectedClass)
 			{
-				UBaseControlRig* NewControlRig = nullptr;
+				UControlRig* NewControlRig = nullptr;
 
 				auto ReportErrorAndSwitchToDefaultRig = [this, InAnimInstance, ExpectedClass](const FString& InMessage) -> bool
 				{
@@ -325,7 +325,7 @@ bool FAnimNode_ControlRig::UpdateControlRigIfNeeded(const UAnimInstance* InAnimI
 					(ExpectedClass != DefaultControlRigClass))
 				{
 					// check if we already created this before
-					if(const TObjectPtr<UBaseControlRig>* ExistingControlRig = ControlRigPerClass.Find(ExpectedClass))
+					if(const TObjectPtr<UControlRig>* ExistingControlRig = ControlRigPerClass.Find(ExpectedClass))
 					{
 						NewControlRig = *ExistingControlRig;
 
@@ -374,7 +374,7 @@ bool FAnimNode_ControlRig::UpdateControlRigIfNeeded(const UAnimInstance* InAnimI
 						}
 						
 						// create a new control rig using the new class
-						NewControlRig = NewObject<UBaseControlRig>(InAnimInstance->GetOwningComponent(), ExpectedClass);
+						NewControlRig = NewObject<UControlRig>(InAnimInstance->GetOwningComponent(), ExpectedClass);
 						NewControlRig->Initialize(true);
 						NewControlRig->RequestInit();
 
@@ -461,7 +461,7 @@ bool FAnimNode_ControlRig::UpdateControlRigIfNeeded(const UAnimInstance* InAnimI
 
 		if(ControlRig == nullptr)
 		{
-			ControlRig = NewObject<UBaseControlRig>(InAnimInstance->GetOwningComponent(), ExpectedClass);
+			ControlRig = NewObject<UControlRig>(InAnimInstance->GetOwningComponent(), ExpectedClass);
 			ControlRig->Initialize(true);
 			ControlRig->RequestInit();
 		}
@@ -534,7 +534,7 @@ void FAnimNode_ControlRig::SetIOMapping(bool bInput, const FName& SourceProperty
 	UClass* TargetClass = GetTargetClass();
 	if (TargetClass)
 	{
-		UBaseControlRig* CDO = TargetClass->GetDefaultObject<UBaseControlRig>();
+		UControlRig* CDO = TargetClass->GetDefaultObject<UControlRig>();
 		if (CDO)
 		{
 			TMap<FName, FName>& MappingData = (bInput) ? InputMapping : OutputMapping;
@@ -592,7 +592,7 @@ void FAnimNode_ControlRig::PropagateInputProperties(const UObject* InSourceInsta
 {
 	if (TargetInstance)
 	{
-		UBaseControlRig* TargetControlRig = Cast<UBaseControlRig>((UObject*)TargetInstance);
+		UControlRig* TargetControlRig = Cast<UControlRig>((UObject*)TargetInstance);
 		if(TargetControlRig == nullptr)
 		{
 			return;
diff --git a/Engine/Plugins/Animation/ControlRig/Source/ControlRig/Private/AnimNode_ControlRigBase.cpp b/Engine/Plugins/Animation/ControlRig/Source/ControlRig/Private/AnimNode_ControlRigBase.cpp
--- a/Engine/Plugins/Animation/ControlRig/Source/ControlRig/Private/AnimNode_ControlRigBase.cpp
+++ b/Engine/Plugins/Animation/ControlRig/Source/ControlRig/Private/AnimNode_ControlRigBase.cpp
@@ -52,7 +52,7 @@ void FAnimNode_ControlRigBase::OnInitializeAnimInstance(const FAnimInstanceProxy
 	WeakAnimInstanceObject = TWeakObjectPtr<const UAnimInstance>(InAnimInstance);
 
 	USkeletalMeshComponent* Component = InAnimInstance->GetOwningComponent();
-	UBaseControlRig* ControlRig = GetControlRig();
+	UControlRig* ControlRig = GetControlRig();
 	if (Component && Component->GetSkeletalMeshAsset() && ControlRig)
 	{
 #if WITH_EDITORONLY_DATA
@@ -66,7 +66,7 @@ void FAnimNode_ControlRigBase::OnInitializeAnimInstance(const FAnimInstanceProxy
 #endif
 
 		// register skeletalmesh component for now
-		ControlRig->GetDataSourceRegistry()->RegisterDataSource(UBaseControlRig::OwnerComponent, InAnimInstance->GetOwningComponent());
+		ControlRig->GetDataSourceRegistry()->RegisterDataSource(UControlRig::OwnerComponent, InAnimInstance->GetOwningComponent());
 		UpdateGetAssetUserDataDelegate(ControlRig);
 	}
 }
@@ -93,7 +93,7 @@ void FAnimNode_ControlRigBase::Update_AnyThread(const FAnimationUpdateContext& C
 
 	if (bExecute)
 	{
-		if (UBaseControlRig* ControlRig = GetControlRig())
+		if (UControlRig* ControlRig = GetControlRig())
 		{
 			// @TODO: fix this to be thread-safe
 			// Pre-update doesn't work for custom anim instances
@@ -110,7 +110,7 @@ bool FAnimNode_ControlRigBase::CanExecute()
 		return false;
 	}
 
-	if (UBaseControlRig* ControlRig = GetControlRig())
+	if (UControlRig* ControlRig = GetControlRig())
 	{
 		return ControlRig->CanExecute(); 
 	}
@@ -118,7 +118,7 @@ bool FAnimNode_ControlRigBase::CanExecute()
 	return false;
 }
 
-void FAnimNode_ControlRigBase::UpdateInput(UBaseControlRig* ControlRig, const FPoseContext& InOutput)
+void FAnimNode_ControlRigBase::UpdateInput(UControlRig* ControlRig, const FPoseContext& InOutput)
 {
 	SCOPE_CYCLE_COUNTER(STAT_ControlRig_UpdateInput);
 
@@ -161,7 +161,7 @@ void FAnimNode_ControlRigBase::UpdateInput(UBaseControlRig* ControlRig, const FP
 			{
 #if WITH_EDITOR
 				// make sure transient controls don't get reset
-				UBaseControlRig::FTransientControlPoseScope PoseScope(ControlRig);
+				UControlRig::FTransientControlPoseScope PoseScope(ControlRig);
 #endif 
 				Hierarchy->ResetPoseToInitial(ERigElementType::Bone);
 			}
@@ -259,7 +259,7 @@ void FAnimNode_ControlRigBase::UpdateInput(UBaseControlRig* ControlRig, const FP
 #endif
 }
 
-void FAnimNode_ControlRigBase::UpdateOutput(UBaseControlRig* ControlRig, FPoseContext& InOutput)
+void FAnimNode_ControlRigBase::UpdateOutput(UControlRig* ControlRig, FPoseContext& InOutput)
 {
 	SCOPE_CYCLE_COUNTER(STAT_ControlRig_UpdateOutput);
 
@@ -440,11 +440,11 @@ void FAnimNode_ControlRigBase::ExecuteControlRig(FPoseContext& InOutput)
 {
 	SCOPE_CYCLE_COUNTER(STAT_ControlRig_Evaluate);
 
-	if (UBaseControlRig* ControlRig = GetControlRig())
+	if (UControlRig* ControlRig = GetControlRig())
 	{
 		// temporarily give control rig access to the stack allocated attribute container
 		// control rig may have rig units that can add/get attributes to/from this container
-		UBaseControlRig::FAnimAttributeContainerPtrScope AttributeScope(ControlRig, InOutput.CustomAttributes);
+		UControlRig::FAnimAttributeContainerPtrScope AttributeScope(ControlRig, InOutput.CustomAttributes);
 		
 		// first update input to the system
 		UpdateInput(ControlRig, InOutput);
@@ -519,7 +519,7 @@ void FAnimNode_ControlRigBase::ExecuteControlRig(FPoseContext& InOutput)
 
 struct FControlRigControlScope
 {
-	FControlRigControlScope(UBaseControlRig* InControlRig)
+	FControlRigControlScope(UControlRig* InControlRig)
 		: ControlRig(InControlRig)
 	{
 		if (ControlRig.IsValid())
@@ -546,7 +546,7 @@ struct FControlRigControlScope
 	}
 
 	TMap<FRigElementKey, FRigControlValue> ControlValues;
-	TWeakObjectPtr<UBaseControlRig> ControlRig;
+	TWeakObjectPtr<UControlRig> ControlRig;
 };
 
 void FAnimNode_ControlRigBase::CacheBones_AnyThread(const FAnimationCacheBonesContext& Context)
@@ -556,7 +556,7 @@ void FAnimNode_ControlRigBase::CacheBones_AnyThread(const FAnimationCacheBonesCo
 	FAnimNode_CustomProperty::CacheBones_AnyThread(Context);
 	Source.CacheBones(Context);
 
-	if (UBaseControlRig* ControlRig = GetControlRig())
+	if (UControlRig* ControlRig = GetControlRig())
 	{
 		// fill up node names
 		FBoneContainer& RequiredBones = Context.AnimInstanceProxy->GetRequiredBones();
@@ -587,7 +587,7 @@ void FAnimNode_ControlRigBase::CacheBones_AnyThread(const FAnimationCacheBonesCo
 	}
 }
 
-void FAnimNode_ControlRigBase::UpdateInputOutputMappingIfRequired(UBaseControlRig* InControlRig, const FBoneContainer& InRequiredBones)
+void FAnimNode_ControlRigBase::UpdateInputOutputMappingIfRequired(UControlRig* InControlRig, const FBoneContainer& InRequiredBones)
 {
 	const URigHierarchy* Hierarchy = InControlRig->GetHierarchy();
 	if(Hierarchy == nullptr)
@@ -651,7 +651,7 @@ void FAnimNode_ControlRigBase::UpdateInputOutputMappingIfRequired(UBaseControlRi
 			const FBoneContainer& InRequiredBones,
 			const FReferenceSkeleton& InRefSkeleton,
 			const TArray<FBoneIndexType>& InRequiredBonesArray,
-			const UBaseControlRig* InControlRig,
+			const UControlRig* InControlRig,
 			TMap<FName, uint16>& OutMapping
 		) {
 			OutMapping.Reset();
@@ -766,7 +766,7 @@ void FAnimNode_ControlRigBase::UpdateInputOutputMappingIfRequired(UBaseControlRi
 
 UClass* FAnimNode_ControlRigBase::GetTargetClass() const
 {
-	if (UBaseControlRig* ControlRig = GetControlRig())
+	if (UControlRig* ControlRig = GetControlRig())
 	{
 		return ControlRig->GetClass();
 	}
@@ -774,7 +774,7 @@ UClass* FAnimNode_ControlRigBase::GetTargetClass() const
 	return nullptr;
 }
 
-void FAnimNode_ControlRigBase::QueueControlRigDrawInstructions(UBaseControlRig* ControlRig, FAnimInstanceProxy* Proxy) const
+void FAnimNode_ControlRigBase::QueueControlRigDrawInstructions(UControlRig* ControlRig, FAnimInstanceProxy* Proxy) const
 {
 	ensure(ControlRig);
 	ensure(Proxy);
@@ -830,7 +830,7 @@ void FAnimNode_ControlRigBase::QueueControlRigDrawInstructions(UBaseControlRig*
 	}
 }
 
-void FAnimNode_ControlRigBase::UpdateGetAssetUserDataDelegate(UBaseControlRig* InControlRig) const
+void FAnimNode_ControlRigBase::UpdateGetAssetUserDataDelegate(UControlRig* InControlRig) const
 {
 	if(GetAssetUserData().IsEmpty() || !WeakAnimInstanceObject.IsValid())
 	{
@@ -842,7 +842,7 @@ void FAnimNode_ControlRigBase::UpdateGetAssetUserDataDelegate(UBaseControlRig* I
 	// since the delegate may go stale quickly. to guard against destroyed anim nodes
 	// we'll rely on the anim instance to provide an indication if the memory is still valid. 
 	TWeakObjectPtr<const UAnimInstance> LocalWeakAnimInstance = WeakAnimInstanceObject;
-	InControlRig->GetExternalAssetUserDataDelegate = UBaseControlRig::FGetExternalAssetUserData::CreateLambda([InControlRig, LocalWeakAnimInstance, this]
+	InControlRig->GetExternalAssetUserDataDelegate = UControlRig::FGetExternalAssetUserData::CreateLambda([InControlRig, LocalWeakAnimInstance, this]
 	{
 		if(LocalWeakAnimInstance.IsValid())
 		{
diff --git a/Engine/Plugins/Animation/ControlRig/Source/ControlRig/Private/AnimNode_ControlRig_ExternalSource.cpp b/Engine/Plugins/Animation/ControlRig/Source/ControlRig/Private/AnimNode_ControlRig_ExternalSource.cpp
--- a/Engine/Plugins/Animation/ControlRig/Source/ControlRig/Private/AnimNode_ControlRig_ExternalSource.cpp
+++ b/Engine/Plugins/Animation/ControlRig/Source/ControlRig/Private/AnimNode_ControlRig_ExternalSource.cpp
@@ -10,20 +10,20 @@ FAnimNode_ControlRig_ExternalSource::FAnimNode_ControlRig_ExternalSource()
 	bControlRigRequiresInitialization = false;
 }
 
-void FAnimNode_ControlRig_ExternalSource::SetControlRig(UBaseControlRig* InControlRig)
+void FAnimNode_ControlRig_ExternalSource::SetControlRig(UControlRig* InControlRig)
 {
 	ControlRig = InControlRig;
 	// requires initializing animation system
 }
 
-UBaseControlRig* FAnimNode_ControlRig_ExternalSource::GetControlRig() const
+UControlRig* FAnimNode_ControlRig_ExternalSource::GetControlRig() const
 {
 	return (ControlRig.IsValid()? ControlRig.Get() : nullptr);
 }
 
-TSubclassOf<UBaseControlRig> FAnimNode_ControlRig_ExternalSource::GetControlRigClass() const
+TSubclassOf<UControlRig> FAnimNode_ControlRig_ExternalSource::GetControlRigClass() const
 {
-	if(UBaseControlRig* CR = GetControlRig())
+	if(UControlRig* CR = GetControlRig())
 	{
 		return CR->GetClass();
 	}
diff --git a/Engine/Plugins/Animation/ControlRig/Source/ControlRig/Private/AnimNode_ControlRig_Library.cpp b/Engine/Plugins/Animation/ControlRig/Source/ControlRig/Private/AnimNode_ControlRig_Library.cpp
--- a/Engine/Plugins/Animation/ControlRig/Source/ControlRig/Private/AnimNode_ControlRig_Library.cpp
+++ b/Engine/Plugins/Animation/ControlRig/Source/ControlRig/Private/AnimNode_ControlRig_Library.cpp
@@ -9,7 +9,7 @@ FControlRigReference UAnimNodeControlRigLibrary::ConvertToControlRig(const FAnim
 	return FAnimNodeReference::ConvertToType<FControlRigReference>(Node, Result);
 }
 
-FControlRigReference UAnimNodeControlRigLibrary::SetControlRigClass(const FControlRigReference& Node, TSubclassOf<UBaseControlRig> ControlRigClass)
+FControlRigReference UAnimNodeControlRigLibrary::SetControlRigClass(const FControlRigReference& Node, TSubclassOf<UControlRig> ControlRigClass)
 {
 	Node.CallAnimNodeFunction<FAnimNode_ControlRig>(
 	TEXT("SetSequence"),
diff --git a/Engine/Plugins/Animation/ControlRig/Source/ControlRig/Private/BaseControlRig.cpp b/Engine/Plugins/Animation/ControlRig/Source/ControlRig/Private/BaseControlRig.cpp
--- a/Engine/Plugins/Animation/ControlRig/Source/ControlRig/Private/BaseControlRig.cpp
+++ b/Engine/Plugins/Animation/ControlRig/Source/ControlRig/Private/BaseControlRig.cpp
@@ -47,7 +47,7 @@ DECLARE_STATS_GROUP(TEXT("ControlRig"), STATGROUP_ControlRig, STATCAT_Advanced);
 DECLARE_CYCLE_STAT_EXTERN(TEXT("Control Rig Execution"), STAT_RigExecution, STATGROUP_ControlRig, );
 DEFINE_STAT(STAT_RigExecution);
 
-const FName UBaseControlRig::OwnerComponent("OwnerComponent");
+const FName UControlRig::OwnerComponent("OwnerComponent");
 
 //CVar to specify if we should create a float control for each curve in the curve container
 //By default we don't but it may be useful to do so for debugging
@@ -66,7 +66,7 @@ static TAutoConsoleVariable<float> CVarControlRigEnableDrawInterfaceInGame(
 	ECVF_Default);
 
 
-UBaseControlRig::UBaseControlRig(const FObjectInitializer& ObjectInitializer)
+UControlRig::UControlRig(const FObjectInitializer& ObjectInitializer)
 	: Super(ObjectInitializer)
 #if WITH_EDITOR
 	, bEnableAnimAttributeTrace(false)
@@ -96,7 +96,7 @@ UBaseControlRig::UBaseControlRig(const FObjectInitializer& ObjectInitializer)
 	SetRigVMExtendedExecuteContext(&RigVMExtendedExecuteContext);
 }
 
-void UBaseControlRig::BeginDestroy()
+void UControlRig::BeginDestroy()
 {
 	Super::BeginDestroy();
 	SetRigVMExtendedExecuteContext(nullptr);
@@ -110,7 +110,7 @@ void UBaseControlRig::BeginDestroy()
 #if WITH_EDITOR
 	if (!HasAnyFlags(RF_ClassDefaultObject))
 	{
-		if(UBaseControlRig* CDO = GetClass()->GetDefaultObject<UBaseControlRig>())
+		if(UControlRig* CDO = GetClass()->GetDefaultObject<UControlRig>())
 		{
 			if (!CDO->HasAnyFlags(RF_BeginDestroyed))
 			{
@@ -126,7 +126,7 @@ void UBaseControlRig::BeginDestroy()
 	TRACE_OBJECT_LIFETIME_END(this);
 }
 
-UWorld* UBaseControlRig::GetWorld() const
+UWorld* UControlRig::GetWorld() const
 {
 	if (ObjectBinding.IsValid())
 	{
@@ -145,7 +145,7 @@ UWorld* UBaseControlRig::GetWorld() const
 	return Super::GetWorld();
 }
 
-void UBaseControlRig::Initialize(bool bRequestInit)
+void UControlRig::Initialize(bool bRequestInit)
 {
 	TRACE_OBJECT_LIFETIME_BEGIN(this);
 
@@ -164,13 +164,13 @@ void UBaseControlRig::Initialize(bool bRequestInit)
 	RequestConstruction();
 	
 	GetHierarchy()->OnModified().RemoveAll(this);
-	GetHierarchy()->OnModified().AddUObject(this, &UBaseControlRig::HandleHierarchyModified);
+	GetHierarchy()->OnModified().AddUObject(this, &UControlRig::HandleHierarchyModified);
 	GetHierarchy()->OnEventReceived().RemoveAll(this);
-	GetHierarchy()->OnEventReceived().AddUObject(this, &UBaseControlRig::HandleHierarchyEvent);
+	GetHierarchy()->OnEventReceived().AddUObject(this, &UControlRig::HandleHierarchyEvent);
 	GetHierarchy()->UpdateVisibilityOnProxyControls();
 }
 
-void UBaseControlRig::OnAddShapeLibrary(const FControlRigExecuteContext* InContext, const FString& InLibraryName, UControlRigShapeLibrary* InShapeLibrary, bool bReplaceExisting, bool bLogResults)
+void UControlRig::OnAddShapeLibrary(const FControlRigExecuteContext* InContext, const FString& InLibraryName, UControlRigShapeLibrary* InShapeLibrary, bool bReplaceExisting, bool bLogResults)
 {
 	// don't ever change the CDO
 	if (HasAnyFlags(RF_ClassDefaultObject))
@@ -195,7 +195,7 @@ void UBaseControlRig::OnAddShapeLibrary(const FControlRigExecuteContext* InConte
 	// if we've removed all shape libraries - let's add the ones from the CDO back
 	if (ShapeLibraries.IsEmpty())
 	{
-		UBaseControlRig* CDO = GetClass()->GetDefaultObject<UBaseControlRig>();
+		UControlRig* CDO = GetClass()->GetDefaultObject<UControlRig>();
 		ShapeLibraries = CDO->ShapeLibraries;
 	}
 
@@ -285,7 +285,7 @@ void UBaseControlRig::OnAddShapeLibrary(const FControlRigExecuteContext* InConte
 #endif
 }
 
-bool UBaseControlRig::OnShapeExists(const FName& InShapeName) const
+bool UControlRig::OnShapeExists(const FName& InShapeName) const
 {
 	const TArray<TSoftObjectPtr<UControlRigShapeLibrary>>& Libraries = GetShapeLibraries();
 	if (UControlRigShapeLibrary::GetShapeByName(InShapeName, GetShapeLibraries(), ShapeLibraryNameMap))
@@ -296,7 +296,7 @@ bool UBaseControlRig::OnShapeExists(const FName& InShapeName) const
 	return false;
 }
 
-bool UBaseControlRig::InitializeVM(const FName& InEventName)
+bool UControlRig::InitializeVM(const FName& InEventName)
 {
 	if(!InitializeVMs(InEventName))
 	{
@@ -327,7 +327,7 @@ bool UBaseControlRig::InitializeVM(const FName& InEventName)
 	return true;
 }
 
-void UBaseControlRig::Evaluate_AnyThread()
+void UControlRig::Evaluate_AnyThread()
 {
 	if (bIsAdditive)
 	{
@@ -431,7 +431,7 @@ void UBaseControlRig::Evaluate_AnyThread()
 	}
 }
 
-bool UBaseControlRig::EvaluateSkeletalMeshComponent(double InDeltaTime)
+bool UControlRig::EvaluateSkeletalMeshComponent(double InDeltaTime)
 {
 	if (USkeletalMeshComponent* SkelMeshComp = Cast<USkeletalMeshComponent>(GetObjectBinding()->GetBoundObject()))
 	{
@@ -448,17 +448,17 @@ bool UBaseControlRig::EvaluateSkeletalMeshComponent(double InDeltaTime)
 	return false;
 }
 
-void UBaseControlRig::ResetControlValues()
+void UControlRig::ResetControlValues()
 {
 	ControlValues.Reset();
 }
 
-void UBaseControlRig::ClearPoseBeforeBackwardsSolve()
+void UControlRig::ClearPoseBeforeBackwardsSolve()
 {
 	PoseBeforeBackwardsSolve.Reset();
 }
 
-TArray<FRigControlElement*> UBaseControlRig::InvertInputPose(const TArray<FRigElementKey>& InElements, EControlRigSetKey InSetKey)
+TArray<FRigControlElement*> UControlRig::InvertInputPose(const TArray<FRigElementKey>& InElements, EControlRigSetKey InSetKey)
 {
 	TArray<FRigControlElement*> ModifiedElements;
 	ModifiedElements.Reserve(ControlsAfterBackwardsSolve.Num());
@@ -484,7 +484,7 @@ TArray<FRigControlElement*> UBaseControlRig::InvertInputPose(const TArray<FRigEl
 	return ModifiedElements;
 }
 
-void UBaseControlRig::InitializeFromCDO()
+void UControlRig::InitializeFromCDO()
 {
 	InitializeVMsFromCDO();
 	
@@ -498,7 +498,7 @@ void UBaseControlRig::InitializeFromCDO()
 			// similar to FControlRigBlueprintCompilerContext::CopyTermDefaultsToDefaultObject,
 			// where CDO is initialized from BP there,
 			// we initialize all other instances of Control Rig from the CDO here
-			UBaseControlRig* CDO = GetClass()->GetDefaultObject<UBaseControlRig>();
+			UControlRig* CDO = GetClass()->GetDefaultObject<UControlRig>();
 			URigHierarchy* Hierarchy = GetHierarchy();
 
 			// copy hierarchy
@@ -534,38 +534,38 @@ void UBaseControlRig::InitializeFromCDO()
 	}
 }
 
-UBaseControlRig::FAnimAttributeContainerPtrScope::FAnimAttributeContainerPtrScope(UBaseControlRig* InControlRig,
+UControlRig::FAnimAttributeContainerPtrScope::FAnimAttributeContainerPtrScope(UControlRig* InControlRig,
 	UE::Anim::FStackAttributeContainer& InExternalContainer)
 {
 	ControlRig = InControlRig;
 	ControlRig->ExternalAnimAttributeContainer = &InExternalContainer;
 }
 
-UBaseControlRig::FAnimAttributeContainerPtrScope::~FAnimAttributeContainerPtrScope()
+UControlRig::FAnimAttributeContainerPtrScope::~FAnimAttributeContainerPtrScope()
 {
 	// control rig should not hold on to this container since it is stack allocated
 	// and should not be used outside of stack, see FPoseContext
 	ControlRig->ExternalAnimAttributeContainer = nullptr;
 }
 
-AActor* UBaseControlRig::GetHostingActor() const
+AActor* UControlRig::GetHostingActor() const
 {
 	return ObjectBinding ? ObjectBinding->GetHostingActor() : nullptr;
 }
 
 #if WITH_EDITOR
-FText UBaseControlRig::GetCategory() const
+FText UControlRig::GetCategory() const
 {
 	return LOCTEXT("DefaultControlRigCategory", "Animation|ControlRigs");
 }
 
-FText UBaseControlRig::GetToolTipText() const
+FText UControlRig::GetToolTipText() const
 {
 	return LOCTEXT("DefaultControlRigTooltip", "ControlRig");
 }
 #endif
 
-bool UBaseControlRig::AllConnectorsAreResolved(FString* OutFailureReason, FRigElementKey* OutConnector) const
+bool UControlRig::AllConnectorsAreResolved(FString* OutFailureReason, FRigElementKey* OutConnector) const
 {
 	if(const URigHierarchy* Hierarchy = GetHierarchy())
 	{
@@ -624,7 +624,7 @@ bool UBaseControlRig::AllConnectorsAreResolved(FString* OutFailureReason, FRigEl
 	return true;
 }
 
-bool UBaseControlRig::Execute(const FName& InEventName)
+bool UControlRig::Execute(const FName& InEventName)
 {
 	if(!CanExecute())
 	{
@@ -690,7 +690,7 @@ bool UBaseControlRig::Execute(const FName& InEventName)
 #if WITH_EDITOR
 	if (IsInDebugMode())
 	{
-		if (UBaseControlRig* CDO = GetClass()->GetDefaultObject<UBaseControlRig>())
+		if (UControlRig* CDO = GetClass()->GetDefaultObject<UControlRig>())
 		{
 			// Copy the breakpoints. This will not override the state of the breakpoints
 			DebugInfo.SetBreakpoints(CDO->DebugInfo.GetBreakpoints());
@@ -753,7 +753,7 @@ bool UBaseControlRig::Execute(const FName& InEventName)
 	// draw container contains persistent draw instructions, 
 	// so we cannot call Reset(), which will clear them,
 	// instead, we re-initialize them from the CDO
-	if (UBaseControlRig* CDO = GetClass()->GetDefaultObject<UBaseControlRig>())
+	if (UControlRig* CDO = GetClass()->GetDefaultObject<UControlRig>())
 	{
 		DrawContainer = CDO->DrawContainer;
 	}
@@ -772,13 +772,13 @@ bool UBaseControlRig::Execute(const FName& InEventName)
 	check(PublicContext.Hierarchy);
 
 	// allow access to the shape libraries
-	PublicContext.OnAddShapeLibraryDelegate.BindUObject(this, &UBaseControlRig::OnAddShapeLibrary);
-	PublicContext.OnShapeExistsDelegate.BindUObject(this, &UBaseControlRig::OnShapeExists);
+	PublicContext.OnAddShapeLibraryDelegate.BindUObject(this, &UControlRig::OnAddShapeLibrary);
+	PublicContext.OnShapeExistsDelegate.BindUObject(this, &UControlRig::OnShapeExists);
 
 	// allow access to the default hierarchy to allow to reset
 	if(!HasAnyFlags(RF_ClassDefaultObject))
 	{
-		if (UBaseControlRig* CDO = Cast<UBaseControlRig>(GetClass()->GetDefaultObject()))
+		if (UControlRig* CDO = Cast<UControlRig>(GetClass()->GetDefaultObject()))
 		{
 			if(URigHierarchy* DefaultHierarchy = CDO->GetHierarchy())
 			{
@@ -928,7 +928,7 @@ bool UBaseControlRig::Execute(const FName& InEventName)
 					// clone the shape libraries again from the CDO 
 					if (!HasAnyFlags(RF_ClassDefaultObject) && ShapeLibraries.IsEmpty())
 					{
-						UBaseControlRig* CDO = GetClass()->GetDefaultObject<UBaseControlRig>();
+						UControlRig* CDO = GetClass()->GetDefaultObject<UControlRig>();
 						ShapeLibraries = CDO->ShapeLibraries;
 					}
 
@@ -1335,28 +1335,149 @@ bool UBaseControlRig::Execute(const FName& InEventName)
 	return bSuccess;
 }
 
-void UBaseControlRig::RequestInit()
+bool UControlRig::Execute_Internal(const FName& InEventName)
+{
+	if(IsRigModule())
+	{
+		FString ConnectorWarning;
+		if(!AllConnectorsAreResolved(&ConnectorWarning))
+		{
+#if WITH_EDITOR
+			LogOnce(EMessageSeverity::Warning, INDEX_NONE, ConnectorWarning);
+#endif
+			return false;
+		}
+	}
+	
+	if (VM)
+	{
+		FRigVMExtendedExecuteContext& Context = GetRigVMExtendedExecuteContext();
+
+		static constexpr TCHAR InvalidatedVMFormat[] = TEXT("%s: Invalidated VM - aborting execution.");
+		if(VM->IsNativized())
+		{
+			if(!IsValidLowLevel() ||
+				!VM->IsValidLowLevel())
+			{
+				UE_LOG(LogControlRig, Warning, InvalidatedVMFormat, *GetClass()->GetName());
+				return false;
+			}
+		}
+		else
+		{
+			// sanity check the validity of the VM to ensure stability.
+			if(!VM->IsContextValidForExecution(Context)
+				|| !IsValidLowLevel()
+				|| !VM->IsValidLowLevel()
+			)
+			{
+				UE_LOG(LogControlRig, Warning, InvalidatedVMFormat, *GetClass()->GetName());
+				return false;
+			}
+		}
+		
+#if UE_RIGVM_PROFILE_EXECUTE_UNITS_NUM
+		const uint64 StartCycles = FPlatformTime::Cycles64();
+		if(ProfilingRunsLeft <= 0)
+		{
+			ProfilingRunsLeft = UE_RIGVM_PROFILE_EXECUTE_UNITS_NUM;
+			AccumulatedCycles = 0;
+		}
+#endif
+		
+		const bool bUseDebuggingSnapshots = !VM->IsNativized();
+		
+#if WITH_EDITOR
+		if(bUseDebuggingSnapshots)
+		{
+			if(URigVM* SnapShotVM = GetSnapshotVM(false)) // don't create it for normal runs
+			{
+				const bool bIsEventFirstInQueue = !EventQueueToRun.IsEmpty() && EventQueueToRun[0] == InEventName; 
+				const bool bIsEventLastInQueue = !EventQueueToRun.IsEmpty() && EventQueueToRun.Last() == InEventName;
+
+				if (GetHaltedAtBreakpoint().IsValid())
+				{
+					if(bIsEventFirstInQueue)
+					{
+						CopyVMMemory(GetRigVMExtendedExecuteContext(), GetSnapshotContext());
+					}
+				}
+				else if(bIsEventLastInQueue)
+				{
+					CopyVMMemory(GetSnapshotContext(), GetRigVMExtendedExecuteContext());
+				}
+			}
+		}
+#endif
+
+		URigHierarchy* Hierarchy = GetHierarchy();
+#if WITH_EDITOR
+
+		bool bRecordTransformsAtRuntime = true;
+		if(const UObject* Outer = GetOuter())
+		{
+			if(Outer->IsA<UControlRigComponent>())
+			{
+				bRecordTransformsAtRuntime = false;
+			}
+		}
+		TGuardValue<bool> RecordTransformsPerInstructionGuard(Hierarchy->bRecordTransformsAtRuntime, bRecordTransformsAtRuntime);
+		
+		if(Hierarchy->bRecordTransformsAtRuntime)
+		{
+			Hierarchy->ReadTransformsAtRuntime.Reset();
+			Hierarchy->WrittenTransformsAtRuntime.Reset();
+		}
+		
+#endif
+		FRigHierarchyExecuteContextBracket HierarchyContextGuard(Hierarchy, &Context);
+
+		// setup the module information
+		FControlRigExecuteContext& PublicContext = Context.GetPublicDataSafe<FControlRigExecuteContext>();
+		FControlRigExecuteContextRigModuleGuard RigModuleGuard(PublicContext, this);
+		FRigHierarchyRedirectorGuard ElementRedirectorGuard(this);
+
+		const bool bSuccess = VM->ExecuteVM(Context, InEventName) != ERigVMExecuteResult::Failed;
+
+#if UE_RIGVM_PROFILE_EXECUTE_UNITS_NUM
+		const uint64 EndCycles = FPlatformTime::Cycles64();
+		const uint64 Cycles = EndCycles - StartCycles;
+		AccumulatedCycles += Cycles;
+		ProfilingRunsLeft--;
+		if(ProfilingRunsLeft == 0)
+		{
+			const double Milliseconds = FPlatformTime::ToMilliseconds64(AccumulatedCycles);
+			UE_LOG(LogControlRig, Display, TEXT("%s: %d runs took %.03lfms."), *GetClass()->GetName(), UE_RIGVM_PROFILE_EXECUTE_UNITS_NUM, Milliseconds);
+		}
+#endif
+
+		return bSuccess;
+	}
+	return false;
+}
+
+void UControlRig::RequestInit()
 {
 	RequestInitVMs();
 	RequestConstruction();
 }
 
-void UBaseControlRig::RequestConstruction()
+void UControlRig::RequestConstruction()
 {
 	RequestRunOnceEvent(FRigUnit_PrepareForExecution::EventName, 0);
 }
 
-bool UBaseControlRig::IsConstructionRequired() const
+bool UControlRig::IsConstructionRequired() const
 {
 	return IsRunOnceEvent(FRigUnit_PrepareForExecution::EventName);
 }
 
-bool UBaseControlRig::SupportsBackwardsSolve() const
+bool UControlRig::SupportsBackwardsSolve() const
 {
 	return SupportsEvent(FRigUnit_InverseExecution::EventName);
 }
 
-void UBaseControlRig::AdaptEventQueueForEvaluate(TArray<FName>& InOutEventQueueToRun)
+void UControlRig::AdaptEventQueueForEvaluate(TArray<FName>& InOutEventQueueToRun)
 {
 	Super::AdaptEventQueueForEvaluate(InOutEventQueueToRun);
 
@@ -1379,7 +1500,7 @@ void UBaseControlRig::AdaptEventQueueForEvaluate(TArray<FName>& InOutEventQueueT
 	}
 }
 
-void UBaseControlRig::GetMappableNodeData(TArray<FName>& OutNames, TArray<FNodeItem>& OutNodeItems) const
+void UControlRig::GetMappableNodeData(TArray<FName>& OutNames, TArray<FNodeItem>& OutNodeItems) const
 {
 	DECLARE_SCOPE_HIERARCHICAL_COUNTER_FUNC()
 
@@ -1404,7 +1525,7 @@ void UBaseControlRig::GetMappableNodeData(TArray<FName>& OutNames, TArray<FNodeI
 	});
 }
 
-UAnimationDataSourceRegistry* UBaseControlRig::GetDataSourceRegistry()
+UAnimationDataSourceRegistry* UControlRig::GetDataSourceRegistry()
 {
 	if (DataSourceRegistry)
 	{
@@ -1428,7 +1549,7 @@ UAnimationDataSourceRegistry* UBaseControlRig::GetDataSourceRegistry()
 
 #if WITH_EDITORONLY_DATA
 
-void UBaseControlRig::PostReinstanceCallback(const UBaseControlRig* Old)
+void UControlRig::PostReinstanceCallback(const UControlRig* Old)
 {
 	DECLARE_SCOPE_HIERARCHICAL_COUNTER_FUNC()
 
@@ -1438,9 +1559,9 @@ void UBaseControlRig::PostReinstanceCallback(const UBaseControlRig* Old)
 
 #endif // WITH_EDITORONLY_DATA
 
-TArray<UBaseControlRig*> UBaseControlRig::FindBaseControlRigs(UObject* Outer, TSubclassOf<UBaseControlRig> OptionalClass)
+TArray<UControlRig*> UControlRig::FindControlRigs(UObject* Outer, TSubclassOf<UControlRig> OptionalClass)
 {
-	TArray<UBaseControlRig*> Result;
+	TArray<UControlRig*> Result;
 	
 	if(Outer == nullptr)
 	{
@@ -1453,9 +1574,9 @@ TArray<UBaseControlRig*> UBaseControlRig::FindBaseControlRigs(UObject* Outer, TS
 		OuterActor = Outer->GetTypedOuter<AActor>();
 	}
 	
-	for (TObjectIterator<UBaseControlRig> Itr; Itr; ++Itr)
+	for (TObjectIterator<UControlRig> Itr; Itr; ++Itr)
 	{
-		UBaseControlRig* RigInstance = *Itr;
+		UControlRig* RigInstance = *Itr;
 		const UClass* RigInstanceClass = RigInstance ? RigInstance->GetClass() : nullptr;
 		if (OptionalClass == nullptr || (RigInstanceClass && RigInstanceClass->IsChildOf(OptionalClass)))
 		{
@@ -1491,7 +1612,7 @@ TArray<UBaseControlRig*> UBaseControlRig::FindBaseControlRigs(UObject* Outer, TS
 	return Result;
 }
 
-void UBaseControlRig::Serialize(FArchive& Ar)
+void UControlRig::Serialize(FArchive& Ar)
 {
 	DECLARE_SCOPE_HIERARCHICAL_COUNTER_FUNC()
 
@@ -1500,7 +1621,7 @@ void UBaseControlRig::Serialize(FArchive& Ar)
 	Ar.UsingCustomVersion(FControlRigObjectVersion::GUID);
 }
 
-void UBaseControlRig::PostLoad()
+void UControlRig::PostLoad()
 {
 	Super::PostLoad();
 
@@ -1534,25 +1655,25 @@ void UBaseControlRig::PostLoad()
 	Influences = NewInfluences;
 }
 
-const FRigModuleSettings& UBaseControlRig::GetRigModuleSettings() const
+const FRigModuleSettings& UControlRig::GetRigModuleSettings() const
 {
 	if(HasAnyFlags(RF_ClassDefaultObject))
 	{
 		return RigModuleSettings;
 	}
-	if (const UBaseControlRig* CDO = Cast<UBaseControlRig>(GetClass()->GetDefaultObject()))
+	if (const UControlRig* CDO = Cast<UControlRig>(GetClass()->GetDefaultObject()))
 	{
 		return CDO->GetRigModuleSettings();
 	}
 	return RigModuleSettings;
 }
 
-bool UBaseControlRig::IsRigModule() const
+bool UControlRig::IsRigModule() const
 {
 	return GetRigModuleSettings().IsValidModule(false);
 }
 
-bool UBaseControlRig::IsRigModuleInstance() const
+bool UControlRig::IsRigModuleInstance() const
 {
 	if(IsRigModule())
 	{
@@ -1561,16 +1682,16 @@ bool UBaseControlRig::IsRigModuleInstance() const
 	return false;
 }
 
-UBaseControlRig* UBaseControlRig::GetParentRig() const
+UControlRig* UControlRig::GetParentRig() const
 {
-	return GetTypedOuter<UBaseControlRig>();
+	return GetTypedOuter<UControlRig>();
 }
 
-const FString& UBaseControlRig::GetRigModuleNameSpace() const
+const FString& UControlRig::GetRigModuleNameSpace() const
 {
 	if(IsRigModule())
 	{
-		if(const UBaseControlRig* ParentRig = GetParentRig())
+		if(const UControlRig* ParentRig = GetParentRig())
 		{
 			const FString& ParentNameSpace = ParentRig->GetRigModuleNameSpace();
 			static constexpr TCHAR JoinFormat[] = TEXT("%s%s:");
@@ -1589,7 +1710,7 @@ const FString& UBaseControlRig::GetRigModuleNameSpace() const
 	return EmptyNameSpace;
 }
 
-FRigElementKeyRedirector& UBaseControlRig::GetElementKeyRedirector()
+FRigElementKeyRedirector& UControlRig::GetElementKeyRedirector()
 {
 	// if we are an instance on a modular rig, use our local info
 	if(IsRigModuleInstance())
@@ -1607,12 +1728,12 @@ FRigElementKeyRedirector& UBaseControlRig::GetElementKeyRedirector()
 	return EmptyRedirector;
 }
 
-void UBaseControlRig::SetElementKeyRedirector(const FRigElementKeyRedirector InElementRedirector)
+void UControlRig::SetElementKeyRedirector(const FRigElementKeyRedirector InElementRedirector)
 {
 	ElementKeyRedirector = InElementRedirector;
 }
 
-TArray<FRigControlElement*> UBaseControlRig::AvailableControls() const
+TArray<FRigControlElement*> UControlRig::AvailableControls() const
 {
 	if(DynamicHierarchy)
 	{
@@ -1621,7 +1742,7 @@ TArray<FRigControlElement*> UBaseControlRig::AvailableControls() const
 	return TArray<FRigControlElement*>();
 }
 
-FRigControlElement* UBaseControlRig::FindControl(const FName& InControlName) const
+FRigControlElement* UControlRig::FindControl(const FName& InControlName) const
 {
 	if(DynamicHierarchy == nullptr)
 	{
@@ -1630,12 +1751,12 @@ FRigControlElement* UBaseControlRig::FindControl(const FName& InControlName) con
 	return DynamicHierarchy->Find<FRigControlElement>(FRigElementKey(InControlName, ERigElementType::Control));
 }
 
-bool UBaseControlRig::IsConstructionModeEnabled() const
+bool UControlRig::IsConstructionModeEnabled() const
 {
 	return EventQueueToRun.Num() == 1 && EventQueueToRun.Contains(FRigUnit_PrepareForExecution::EventName);
 }
 
-FTransform UBaseControlRig::SetupControlFromGlobalTransform(const FName& InControlName, const FTransform& InGlobalTransform)
+FTransform UControlRig::SetupControlFromGlobalTransform(const FName& InControlName, const FTransform& InGlobalTransform)
 {
 	if (IsConstructionModeEnabled())
 	{
@@ -1660,7 +1781,7 @@ FTransform UBaseControlRig::SetupControlFromGlobalTransform(const FName& InContr
 	return InGlobalTransform;
 }
 
-void UBaseControlRig::CreateRigControlsForCurveContainer()
+void UControlRig::CreateRigControlsForCurveContainer()
 {
 	const bool bCreateFloatControls = CVarControlRigCreateFloatControlsForCurves->GetInt() == 0 ? false : true;
 	if(bCreateFloatControls && DynamicHierarchy)
@@ -1693,11 +1814,11 @@ void UBaseControlRig::CreateRigControlsForCurveContainer()
 			return true;
 		});
 
-		ControlModified().AddUObject(this, &UBaseControlRig::HandleOnControlModified);
+		ControlModified().AddUObject(this, &UControlRig::HandleOnControlModified);
 	}
 }
 
-void UBaseControlRig::HandleOnControlModified(UBaseControlRig* Subject, FRigControlElement* Control, const FRigControlModifiedContext& Context)
+void UControlRig::HandleOnControlModified(UControlRig* Subject, FRigControlElement* Control, const FRigControlModifiedContext& Context)
 {
 	if (Control->Settings.bIsCurve && DynamicHierarchy)
 	{
@@ -1706,12 +1827,12 @@ void UBaseControlRig::HandleOnControlModified(UBaseControlRig* Subject, FRigCont
 	}	
 }
 
-bool UBaseControlRig::IsCurveControl(const FRigControlElement* InControlElement) const
+bool UControlRig::IsCurveControl(const FRigControlElement* InControlElement) const
 {
 	return InControlElement->Settings.bIsCurve;
 }
 
-FTransform UBaseControlRig::GetControlGlobalTransform(const FName& InControlName) const
+FTransform UControlRig::GetControlGlobalTransform(const FName& InControlName) const
 {
 	if(DynamicHierarchy == nullptr)
 	{
@@ -1720,7 +1841,7 @@ FTransform UBaseControlRig::GetControlGlobalTransform(const FName& InControlName
 	return DynamicHierarchy->GetGlobalTransform(FRigElementKey(InControlName, ERigElementType::Control), false);
 }
 
-FRigControlValue UBaseControlRig::GetControlValue(FRigControlElement* InControl, const ERigControlValueType& InValueType)
+FRigControlValue UControlRig::GetControlValue(FRigControlElement* InControl, const ERigControlValueType& InValueType)
 {
 	if (bIsAdditive && InValueType == ERigControlValueType::Current)
 	{
@@ -1746,7 +1867,7 @@ FRigControlValue UBaseControlRig::GetControlValue(FRigControlElement* InControl,
 	return GetHierarchy()->GetControlValue(InControl, InValueType);
 }
 
-void UBaseControlRig::SetControlValueImpl(const FName& InControlName, const FRigControlValue& InValue, bool bNotify,
+void UControlRig::SetControlValueImpl(const FName& InControlName, const FRigControlValue& InValue, bool bNotify,
 	const FRigControlModifiedContext& Context, bool bSetupUndo, bool bPrintPythonCommnds, bool bFixEulerFlips)
 {
 	const FRigElementKey Key(InControlName, ERigElementType::Control);
@@ -1773,7 +1894,7 @@ void UBaseControlRig::SetControlValueImpl(const FName& InControlName, const FRig
 	}
 }
 
-bool UBaseControlRig::SetControlGlobalTransform(const FName& InControlName, const FTransform& InGlobalTransform, bool bNotify, const FRigControlModifiedContext& Context, bool bSetupUndo, bool bPrintPythonCommands, bool bFixEulerFlips)
+bool UControlRig::SetControlGlobalTransform(const FName& InControlName, const FTransform& InGlobalTransform, bool bNotify, const FRigControlModifiedContext& Context, bool bSetupUndo, bool bPrintPythonCommands, bool bFixEulerFlips)
 {
 	FTransform GlobalTransform = InGlobalTransform;
 	ERigTransformType::Type TransformType = ERigTransformType::CurrentGlobal;
@@ -1818,7 +1939,7 @@ bool UBaseControlRig::SetControlGlobalTransform(const FName& InControlName, cons
 	return true;
 }
 
-FRigControlValue UBaseControlRig::GetControlValueFromGlobalTransform(const FName& InControlName, const FTransform& InGlobalTransform, ERigTransformType::Type InTransformType)
+FRigControlValue UControlRig::GetControlValueFromGlobalTransform(const FName& InControlName, const FTransform& InGlobalTransform, ERigTransformType::Type InTransformType)
 {
 	FRigControlValue Value;
 
@@ -1861,7 +1982,7 @@ FRigControlValue UBaseControlRig::GetControlValueFromGlobalTransform(const FName
 	return Value;
 }
 
-void UBaseControlRig::SetControlLocalTransform(const FName& InControlName, const FTransform& InLocalTransform, bool bNotify, const FRigControlModifiedContext& Context, bool bSetupUndo, bool bFixEulerFlips)
+void UControlRig::SetControlLocalTransform(const FName& InControlName, const FTransform& InLocalTransform, bool bNotify, const FRigControlModifiedContext& Context, bool bSetupUndo, bool bFixEulerFlips)
 {
 	if (FRigControlElement* ControlElement = FindControl(InControlName))
 	{
@@ -1877,7 +1998,7 @@ void UBaseControlRig::SetControlLocalTransform(const FName& InControlName, const
 	}
 }
 
-FTransform UBaseControlRig::GetControlLocalTransform(const FName& InControlName)
+FTransform UControlRig::GetControlLocalTransform(const FName& InControlName)
 {
 	if(DynamicHierarchy == nullptr)
 	{
@@ -1897,13 +2018,13 @@ FTransform UBaseControlRig::GetControlLocalTransform(const FName& InControlName)
 	return DynamicHierarchy->GetLocalTransform(FRigElementKey(InControlName, ERigElementType::Control));
 }
 
-const TArray<TSoftObjectPtr<UControlRigShapeLibrary>>& UBaseControlRig::GetShapeLibraries() const
+const TArray<TSoftObjectPtr<UControlRigShapeLibrary>>& UControlRig::GetShapeLibraries() const
 {
 	const TArray<TSoftObjectPtr<UControlRigShapeLibrary>>* LibrariesPtr = &ShapeLibraries;
 
 	if(!GetClass()->IsNative() && ShapeLibraries.IsEmpty())
 	{
-		if (UBaseControlRig* CDO = Cast<UBaseControlRig>(GetClass()->GetDefaultObject()))
+		if (UControlRig* CDO = Cast<UControlRig>(GetClass()->GetDefaultObject()))
 		{
 			LibrariesPtr = &CDO->ShapeLibraries;
 		}
@@ -1921,7 +2042,7 @@ const TArray<TSoftObjectPtr<UControlRigShapeLibrary>>& UBaseControlRig::GetShape
 	return Libraries;
 }
 
-void UBaseControlRig::SelectControl(const FName& InControlName, bool bSelect)
+void UControlRig::SelectControl(const FName& InControlName, bool bSelect)
 {
 	if(DynamicHierarchy)
 	{
@@ -1932,7 +2053,7 @@ void UBaseControlRig::SelectControl(const FName& InControlName, bool bSelect)
 	}
 }
 
-bool UBaseControlRig::ClearControlSelection()
+bool UControlRig::ClearControlSelection()
 {
 	if(DynamicHierarchy)
 	{
@@ -1944,7 +2065,7 @@ bool UBaseControlRig::ClearControlSelection()
 	return false;
 }
 
-TArray<FName> UBaseControlRig::CurrentControlSelection() const
+TArray<FName> UControlRig::CurrentControlSelection() const
 {
 	TArray<FName> SelectedControlNames;
 
@@ -1965,7 +2086,7 @@ TArray<FName> UBaseControlRig::CurrentControlSelection() const
 	return SelectedControlNames;
 }
 
-bool UBaseControlRig::IsControlSelected(const FName& InControlName)const
+bool UControlRig::IsControlSelected(const FName& InControlName)const
 {
 	if(DynamicHierarchy)
 	{
@@ -1977,7 +2098,7 @@ bool UBaseControlRig::IsControlSelected(const FName& InControlName)const
 	return false;
 }
 
-void UBaseControlRig::HandleHierarchyModified(ERigHierarchyNotification InNotification, URigHierarchy* InHierarchy,
+void UControlRig::HandleHierarchyModified(ERigHierarchyNotification InNotification, URigHierarchy* InHierarchy,
     const FRigBaseElement* InElement)
 {
 	switch(InNotification)
@@ -2035,7 +2156,7 @@ void UBaseControlRig::HandleHierarchyModified(ERigHierarchyNotification InNotifi
 
 #if WITH_EDITOR
 
-bool UBaseControlRig::CanAddTransientControl(const URigVMUnitNode* InNode, const FRigDirectManipulationTarget& InTarget, FString* OutFailureReason)
+bool UControlRig::CanAddTransientControl(const URigVMUnitNode* InNode, const FRigDirectManipulationTarget& InTarget, FString* OutFailureReason)
 {
 	if (InNode == nullptr)
 	{
@@ -2089,7 +2210,7 @@ bool UBaseControlRig::CanAddTransientControl(const URigVMUnitNode* InNode, const
 	return false;
 }
 
-FName UBaseControlRig::AddTransientControl(const URigVMUnitNode* InNode, const FRigDirectManipulationTarget& InTarget)
+FName UControlRig::AddTransientControl(const URigVMUnitNode* InNode, const FRigDirectManipulationTarget& InTarget)
 {
 	if(!CanAddTransientControl(InNode, InTarget, nullptr))
 	{
@@ -2145,7 +2266,7 @@ FName UBaseControlRig::AddTransientControl(const URigVMUnitNode* InNode, const F
 	return Info->ControlKey.Name;
 }
 
-bool UBaseControlRig::SetTransientControlValue(const URigVMUnitNode* InNode, TSharedPtr<FRigDirectManipulationInfo> InInfo)
+bool UControlRig::SetTransientControlValue(const URigVMUnitNode* InNode, TSharedPtr<FRigDirectManipulationInfo> InInfo)
 {
 	check(InNode);
 	check(DynamicHierarchy);
@@ -2173,7 +2294,7 @@ bool UBaseControlRig::SetTransientControlValue(const URigVMUnitNode* InNode, TSh
 	return bResult;
 }
 
-FName UBaseControlRig::RemoveTransientControl(const URigVMUnitNode* InNode, const FRigDirectManipulationTarget& InTarget)
+FName UControlRig::RemoveTransientControl(const URigVMUnitNode* InNode, const FRigDirectManipulationTarget& InTarget)
 {
 	if ((InNode == nullptr) || (DynamicHierarchy == nullptr))
 	{
@@ -2207,7 +2328,7 @@ FName UBaseControlRig::RemoveTransientControl(const URigVMUnitNode* InNode, cons
 	return NAME_None;
 }
 
-FName UBaseControlRig::AddTransientControl(const FRigElementKey& InElement)
+FName UControlRig::AddTransientControl(const FRigElementKey& InElement)
 {
 	if (!InElement.IsValid())
 	{
@@ -2282,7 +2403,7 @@ FName UBaseControlRig::AddTransientControl(const FRigElementKey& InElement)
 			if(FRigBoneElement* BoneElement = DynamicHierarchy->Find<FRigBoneElement>(InElement))
 			{
 				// add a modify bone AnimNode internally that the transient control controls for imported bones only
-				// for user created bones, refer to UBaseControlRig::TransformOverrideForUserCreatedBones 
+				// for user created bones, refer to UControlRig::TransformOverrideForUserCreatedBones 
 				if (BoneElement->BoneType == ERigBoneType::Imported)
 				{ 
 					if (PreviewInstance)
@@ -2305,7 +2426,7 @@ FName UBaseControlRig::AddTransientControl(const FRigElementKey& InElement)
 	return ControlKey.Name;
 }
 
-bool UBaseControlRig::SetTransientControlValue(const FRigElementKey& InElement)
+bool UControlRig::SetTransientControlValue(const FRigElementKey& InElement)
 {
 	if (!InElement.IsValid())
 	{
@@ -2374,7 +2495,7 @@ bool UBaseControlRig::SetTransientControlValue(const FRigElementKey& InElement)
 	return false;
 }
 
-FName UBaseControlRig::RemoveTransientControl(const FRigElementKey& InElement)
+FName UControlRig::RemoveTransientControl(const FRigElementKey& InElement)
 {
 	if (!InElement.IsValid())
 	{
@@ -2408,7 +2529,7 @@ FName UBaseControlRig::RemoveTransientControl(const FRigElementKey& InElement)
 	return NAME_None;
 }
 
-FName UBaseControlRig::GetNameForTransientControl(const URigVMUnitNode* InNode, const FRigDirectManipulationTarget& InTarget) const
+FName UControlRig::GetNameForTransientControl(const URigVMUnitNode* InNode, const FRigDirectManipulationTarget& InTarget) const
 {
 	check(InNode);
 	check(DynamicHierarchy);
@@ -2417,7 +2538,7 @@ FName UBaseControlRig::GetNameForTransientControl(const URigVMUnitNode* InNode,
 	return DynamicHierarchy->GetSanitizedName(FRigName(FString::Printf(TEXT("ControlForNode|%s|%s"), *NodeName, *InTarget.Name)));
 }
 
-FString UBaseControlRig::GetNodeNameFromTransientControl(const FRigElementKey& InKey)
+FString UControlRig::GetNodeNameFromTransientControl(const FRigElementKey& InKey)
 {
 	FString Name = InKey.Name.ToString();
 	if(Name.StartsWith(TEXT("ControlForNode|")))
@@ -2432,7 +2553,7 @@ FString UBaseControlRig::GetNodeNameFromTransientControl(const FRigElementKey& I
 	return Name;
 }
 
-FString UBaseControlRig::GetTargetFromTransientControl(const FRigElementKey& InKey)
+FString UControlRig::GetTargetFromTransientControl(const FRigElementKey& InKey)
 {
 	FString Name = InKey.Name.ToString();
 	if(Name.StartsWith(TEXT("ControlForNode|")))
@@ -2447,7 +2568,7 @@ FString UBaseControlRig::GetTargetFromTransientControl(const FRigElementKey& InK
 	return Name;
 }
 
-TSharedPtr<FRigDirectManipulationInfo> UBaseControlRig::GetRigUnitManipulationInfoForTransientControl(
+TSharedPtr<FRigDirectManipulationInfo> UControlRig::GetRigUnitManipulationInfoForTransientControl(
 	const FRigElementKey& InKey)
 {
 	const TSharedPtr<FRigDirectManipulationInfo>* InfoPtr = RigUnitManipulationInfos.FindByPredicate(
@@ -2464,7 +2585,7 @@ TSharedPtr<FRigDirectManipulationInfo> UBaseControlRig::GetRigUnitManipulationIn
 	return TSharedPtr<FRigDirectManipulationInfo>();
 }
 
-FName UBaseControlRig::GetNameForTransientControl(const FRigElementKey& InElement)
+FName UControlRig::GetNameForTransientControl(const FRigElementKey& InElement)
 {
 	if (InElement.Type == ERigElementType::Control)
 	{
@@ -2475,7 +2596,7 @@ FName UBaseControlRig::GetNameForTransientControl(const FRigElementKey& InElemen
 	return *FString::Printf(TEXT("ControlForRigElement_%s_%s"), *EnumName.ToString(), *InElement.Name.ToString());
 }
 
-FRigElementKey UBaseControlRig::GetElementKeyFromTransientControl(const FRigElementKey& InKey)
+FRigElementKey UControlRig::GetElementKeyFromTransientControl(const FRigElementKey& InKey)
 {
 	if(InKey.Type != ERigElementType::Control)
 	{
@@ -2508,7 +2629,7 @@ FRigElementKey UBaseControlRig::GetElementKeyFromTransientControl(const FRigElem
 	return FRigElementKey();;
 }
 
-void UBaseControlRig::ClearTransientControls()
+void UControlRig::ClearTransientControls()
 {
 	if(DynamicHierarchy == nullptr)
 	{
@@ -2540,7 +2661,7 @@ void UBaseControlRig::ClearTransientControls()
 	}
 }
 
-void UBaseControlRig::ApplyTransformOverrideForUserCreatedBones()
+void UControlRig::ApplyTransformOverrideForUserCreatedBones()
 {
 	if(DynamicHierarchy == nullptr)
 	{
@@ -2553,7 +2674,7 @@ void UBaseControlRig::ApplyTransformOverrideForUserCreatedBones()
 	}
 }
 
-void UBaseControlRig::ApplySelectionPoseForConstructionMode(const FName& InEventName)
+void UControlRig::ApplySelectionPoseForConstructionMode(const FName& InEventName)
 {
 	FRigControlModifiedContext ControlValueContext;
 	ControlValueContext.EventName = InEventName;
@@ -2569,7 +2690,7 @@ void UBaseControlRig::ApplySelectionPoseForConstructionMode(const FName& InEvent
 
 #endif
 
-void UBaseControlRig::HandleHierarchyEvent(URigHierarchy* InHierarchy, const FRigEventContext& InEvent)
+void UControlRig::HandleHierarchyEvent(URigHierarchy* InHierarchy, const FRigEventContext& InEvent)
 {
 	if (RigEventDelegate.IsBound())
 	{
@@ -2609,7 +2730,7 @@ void UBaseControlRig::HandleHierarchyEvent(URigHierarchy* InHierarchy, const FRi
 	}
 }
 
-void UBaseControlRig::GetControlsInOrder(TArray<FRigControlElement*>& SortedControls) const
+void UControlRig::GetControlsInOrder(TArray<FRigControlElement*>& SortedControls) const
 {
 	SortedControls.Reset();
 
@@ -2621,9 +2742,9 @@ void UBaseControlRig::GetControlsInOrder(TArray<FRigControlElement*>& SortedCont
 	SortedControls = DynamicHierarchy->GetControls(true);
 }
 
-const FRigInfluenceMap* UBaseControlRig::FindInfluenceMap(const FName& InEventName)
+const FRigInfluenceMap* UControlRig::FindInfluenceMap(const FName& InEventName)
 {
-	if (UBaseControlRig* CDO = GetClass()->GetDefaultObject<UBaseControlRig>())
+	if (UControlRig* CDO = GetClass()->GetDefaultObject<UControlRig>())
 	{
 		return CDO->Influences.Find(InEventName);
 	}
@@ -2632,7 +2753,7 @@ const FRigInfluenceMap* UBaseControlRig::FindInfluenceMap(const FName& InEventNa
 
 #if WITH_EDITOR
 
-void UBaseControlRig::PreEditChange(FProperty* PropertyAboutToChange)
+void UControlRig::PreEditChange(FProperty* PropertyAboutToChange)
 {
 	// for BP user authored properties let's ignore changes since they
 	// will be distributed from the BP anyway to all archetype instances.
@@ -2657,13 +2778,13 @@ void UBaseControlRig::PreEditChange(FProperty* PropertyAboutToChange)
 	Super::PreEditChange(PropertyAboutToChange);
 }
 
-void UBaseControlRig::PostEditChangeProperty(FPropertyChangedEvent& PropertyChangedEvent)
+void UControlRig::PostEditChangeProperty(FPropertyChangedEvent& PropertyChangedEvent)
 {
 	Super::PostEditChangeProperty(PropertyChangedEvent);
 }
 #endif
 
-FRigUnit* UBaseControlRig::GetRigUnitInstanceFromScope(TSharedPtr<FStructOnScope> InScope)
+FRigUnit* UControlRig::GetRigUnitInstanceFromScope(TSharedPtr<FStructOnScope> InScope)
 {
 	if(InScope.IsValid())
 	{
@@ -2679,7 +2800,7 @@ FRigUnit* UBaseControlRig::GetRigUnitInstanceFromScope(TSharedPtr<FStructOnScope
 	return (FRigUnit*)DefaultRigUnitInstance.GetStructMemory();
 }
 
-const TArray<UAssetUserData*>* UBaseControlRig::GetAssetUserDataArray() const
+const TArray<UAssetUserData*>* UControlRig::GetAssetUserDataArray() const
 {
 	if(HasAnyFlags(RF_ClassDefaultObject))
 	{
@@ -2696,7 +2817,7 @@ const TArray<UAssetUserData*>* UBaseControlRig::GetAssetUserDataArray() const
 
 	CombinedAssetUserData.Reset();
 
-	if (UBaseControlRig* CDO = Cast<UBaseControlRig>(GetClass()->GetDefaultObject(false)))
+	if (UControlRig* CDO = Cast<UControlRig>(GetClass()->GetDefaultObject(false)))
 	{
 		CombinedAssetUserData.Append(*CDO->GetAssetUserDataArray());
 	}
@@ -2789,7 +2910,7 @@ const TArray<UAssetUserData*>* UBaseControlRig::GetAssetUserDataArray() const
 	return &ToRawPtrTArrayUnsafe(CombinedAssetUserData);
 }
 
-void UBaseControlRig::CopyPoseFromOtherRig(UBaseControlRig* Subject)
+void UControlRig::CopyPoseFromOtherRig(UControlRig* Subject)
 {
 	check(DynamicHierarchy);
 	check(Subject);
@@ -2824,7 +2945,7 @@ void UBaseControlRig::CopyPoseFromOtherRig(UBaseControlRig* Subject)
 	}
 }
 
-void UBaseControlRig::SetBoneInitialTransformsFromAnimInstance(UAnimInstance* InAnimInstance)
+void UControlRig::SetBoneInitialTransformsFromAnimInstance(UAnimInstance* InAnimInstance)
 {
 	FMemMark Mark(FMemStack::Get());
 	FCompactPose OutPose;
@@ -2832,7 +2953,7 @@ void UBaseControlRig::SetBoneInitialTransformsFromAnimInstance(UAnimInstance* In
 	SetBoneInitialTransformsFromCompactPose(&OutPose);
 }
 
-void UBaseControlRig::SetBoneInitialTransformsFromAnimInstanceProxy(const FAnimInstanceProxy* InAnimInstanceProxy)
+void UControlRig::SetBoneInitialTransformsFromAnimInstanceProxy(const FAnimInstanceProxy* InAnimInstanceProxy)
 {
 	FMemMark Mark(FMemStack::Get());
 	FCompactPose OutPose;
@@ -2840,7 +2961,7 @@ void UBaseControlRig::SetBoneInitialTransformsFromAnimInstanceProxy(const FAnimI
 	SetBoneInitialTransformsFromCompactPose(&OutPose);
 }
 
-void UBaseControlRig::SetBoneInitialTransformsFromSkeletalMeshComponent(USkeletalMeshComponent* InSkelMeshComp, bool bUseAnimInstance)
+void UControlRig::SetBoneInitialTransformsFromSkeletalMeshComponent(USkeletalMeshComponent* InSkelMeshComp, bool bUseAnimInstance)
 {
 	check(InSkelMeshComp);
 	check(DynamicHierarchy);
@@ -2856,7 +2977,7 @@ void UBaseControlRig::SetBoneInitialTransformsFromSkeletalMeshComponent(USkeleta
 }
 
 
-void UBaseControlRig::SetBoneInitialTransformsFromSkeletalMesh(USkeletalMesh* InSkeletalMesh)
+void UControlRig::SetBoneInitialTransformsFromSkeletalMesh(USkeletalMesh* InSkeletalMesh)
 {
 	if (ensure(InSkeletalMesh))
 	{ 
@@ -2864,7 +2985,7 @@ void UBaseControlRig::SetBoneInitialTransformsFromSkeletalMesh(USkeletalMesh* In
 	}
 }
 
-void UBaseControlRig::SetBoneInitialTransformsFromRefSkeleton(const FReferenceSkeleton& InReferenceSkeleton)
+void UControlRig::SetBoneInitialTransformsFromRefSkeleton(const FReferenceSkeleton& InReferenceSkeleton)
 {
 	check(DynamicHierarchy);
 
@@ -2886,7 +3007,7 @@ void UBaseControlRig::SetBoneInitialTransformsFromRefSkeleton(const FReferenceSk
 	RequestConstruction();
 }
 
-void UBaseControlRig::SetBoneInitialTransformsFromCompactPose(FCompactPose* InCompactPose)
+void UControlRig::SetBoneInitialTransformsFromCompactPose(FCompactPose* InCompactPose)
 {
 	check(InCompactPose);
 
@@ -2923,7 +3044,7 @@ void UBaseControlRig::SetBoneInitialTransformsFromCompactPose(FCompactPose* InCo
 	RequestConstruction();
 }
 
-const FRigControlElementCustomization* UBaseControlRig::GetControlCustomization(const FRigElementKey& InControl) const
+const FRigControlElementCustomization* UControlRig::GetControlCustomization(const FRigElementKey& InControl) const
 {
 	check(InControl.Type == ERigElementType::Control);
 
@@ -2943,14 +3064,14 @@ const FRigControlElementCustomization* UBaseControlRig::GetControlCustomization(
 	return nullptr;
 }
 
-void UBaseControlRig::SetControlCustomization(const FRigElementKey& InControl, const FRigControlElementCustomization& InCustomization)
+void UControlRig::SetControlCustomization(const FRigElementKey& InControl, const FRigControlElementCustomization& InCustomization)
 {
 	check(InControl.Type == ERigElementType::Control);
 	
 	ControlCustomizations.FindOrAdd(InControl) = InCustomization;
 }
 
-void UBaseControlRig::PostInitInstanceIfRequired()
+void UControlRig::PostInitInstanceIfRequired()
 {
 	if(GetHierarchy() == nullptr || VM == nullptr)
 	{
@@ -2960,7 +3081,7 @@ void UBaseControlRig::PostInitInstanceIfRequired()
 		}
 		else
 		{
-			UBaseControlRig* CDO = GetClass()->GetDefaultObject<UBaseControlRig>();
+			UControlRig* CDO = GetClass()->GetDefaultObject<UControlRig>();
 			ensure(VM == nullptr || VM == CDO->VM);
 			PostInitInstance(CDO);
 		}
@@ -2968,7 +3089,7 @@ void UBaseControlRig::PostInitInstanceIfRequired()
 }
 
 #if WITH_EDITORONLY_DATA
-void UBaseControlRig::DeclareConstructClasses(TArray<FTopLevelAssetPath>& OutConstructClasses, const UClass* SpecificSubclass)
+void UControlRig::DeclareConstructClasses(TArray<FTopLevelAssetPath>& OutConstructClasses, const UClass* SpecificSubclass)
 {
 	Super::DeclareConstructClasses(OutConstructClasses, SpecificSubclass);
 	OutConstructClasses.Add(FTopLevelAssetPath(URigHierarchy::StaticClass()));
@@ -2977,14 +3098,14 @@ void UBaseControlRig::DeclareConstructClasses(TArray<FTopLevelAssetPath>& OutCon
 
 #endif
 
-USceneComponent* UBaseControlRig::GetOwningSceneComponent()
+USceneComponent* UControlRig::GetOwningSceneComponent()
 {
 	if(OuterSceneComponent == nullptr)
 	{
 		const FControlRigExecuteContext& PublicContext = GetRigVMExtendedExecuteContext().GetPublicDataSafe<FControlRigExecuteContext>();
 		const FRigUnitContext& Context = PublicContext.UnitContext;
 
-		USceneComponent* SceneComponentFromRegistry = Context.DataSourceRegistry->RequestSource<USceneComponent>(UBaseControlRig::OwnerComponent);
+		USceneComponent* SceneComponentFromRegistry = Context.DataSourceRegistry->RequestSource<USceneComponent>(UControlRig::OwnerComponent);
 		if (SceneComponentFromRegistry)
 		{
 			OuterSceneComponent = SceneComponentFromRegistry;
@@ -2998,7 +3119,7 @@ USceneComponent* UBaseControlRig::GetOwningSceneComponent()
 	return OuterSceneComponent.Get();
 }
 
-void UBaseControlRig::PostInitInstance(URigVMHost* InCDO)
+void UControlRig::PostInitInstance(URigVMHost* InCDO)
 {
 	const EObjectFlags SubObjectFlags =
 	HasAnyFlags(RF_ClassDefaultObject) ?
@@ -3010,7 +3131,7 @@ void UBaseControlRig::PostInitInstance(URigVMHost* InCDO)
 	Context.SetContextPublicDataStruct(FControlRigExecuteContext::StaticStruct());
 
 	Context.ExecutionReachedExit().RemoveAll(this);
-	Context.ExecutionReachedExit().AddUObject(this, &UBaseControlRig::HandleExecutionReachedExit);
+	Context.ExecutionReachedExit().AddUObject(this, &UControlRig::HandleExecutionReachedExit);
 	UpdateVMSettings();
 
 	// set up the hierarchy
@@ -3032,8 +3153,8 @@ void UBaseControlRig::PostInitInstance(URigVMHost* InCDO)
 	}
 
 #if WITH_EDITOR
-		const TWeakObjectPtr<UBaseControlRig> WeakThis = this;
-		DynamicHierarchy->OnUndoRedo().AddStatic(&UBaseControlRig::OnHierarchyTransformUndoRedoWeak, WeakThis);
+		const TWeakObjectPtr<UControlRig> WeakThis = this;
+		DynamicHierarchy->OnUndoRedo().AddStatic(&UControlRig::OnHierarchyTransformUndoRedoWeak, WeakThis);
 #endif
 
 	if(!HasAnyFlags(RF_ClassDefaultObject))
@@ -3046,7 +3167,7 @@ void UBaseControlRig::PostInitInstance(URigVMHost* InCDO)
 				VM = InCDO->GetVM();
 			}
 
-			DynamicHierarchy->CopyHierarchy(CastChecked<UBaseControlRig>(InCDO)->GetHierarchy());
+			DynamicHierarchy->CopyHierarchy(CastChecked<UControlRig>(InCDO)->GetHierarchy());
 		}
 	}
 	else // we are the CDO
@@ -3073,7 +3194,7 @@ void UBaseControlRig::PostInitInstance(URigVMHost* InCDO)
 		}
 	}
 
-	if(UBaseControlRig* CDOControlRig = Cast<UBaseControlRig>(InCDO))
+	if(UControlRig* CDOControlRig = Cast<UControlRig>(InCDO))
 	{
 		ElementKeyRedirector = FRigElementKeyRedirector(CDOControlRig->ElementKeyRedirector, DynamicHierarchy);
 	}
@@ -3081,7 +3202,7 @@ void UBaseControlRig::PostInitInstance(URigVMHost* InCDO)
 	RequestInit();
 }
 
-void UBaseControlRig::SetDynamicHierarchy(TObjectPtr<URigHierarchy> InHierarchy)
+void UControlRig::SetDynamicHierarchy(TObjectPtr<URigHierarchy> InHierarchy)
 {
 	// Delete any existing hierarchy
 	if (DynamicHierarchy->GetOuter() == this)
@@ -3093,7 +3214,7 @@ void UBaseControlRig::SetDynamicHierarchy(TObjectPtr<URigHierarchy> InHierarchy)
 	DynamicHierarchy = InHierarchy;
 }
 
-UTransformableControlHandle* UBaseControlRig::CreateTransformableControlHandle(
+UTransformableControlHandle* UControlRig::CreateTransformableControlHandle(
 	const FName& InControlName) const
 {
 	auto IsConstrainable = [this](const FName& InControlName)
@@ -3129,7 +3250,7 @@ UTransformableControlHandle* UBaseControlRig::CreateTransformableControlHandle(
 	return CtrlHandle;
 }
 
-void UBaseControlRig::OnHierarchyTransformUndoRedo(URigHierarchy* InHierarchy, const FRigElementKey& InKey, ERigTransformType::Type InTransformType, const FTransform& InTransform, bool bIsUndo)
+void UControlRig::OnHierarchyTransformUndoRedo(URigHierarchy* InHierarchy, const FRigElementKey& InKey, ERigTransformType::Type InTransformType, const FTransform& InTransform, bool bIsUndo)
 {
 	if(InKey.Type == ERigElementType::Control)
 	{
@@ -3140,7 +3261,7 @@ void UBaseControlRig::OnHierarchyTransformUndoRedo(URigHierarchy* InHierarchy, c
 	}
 }
 
-UBaseControlRig::FPoseScope::FPoseScope(UBaseControlRig* InControlRig, ERigElementType InFilter, const TArray<FRigElementKey>& InElements, const ERigTransformType::Type InTransformType)
+UControlRig::FPoseScope::FPoseScope(UControlRig* InControlRig, ERigElementType InFilter, const TArray<FRigElementKey>& InElements, const ERigTransformType::Type InTransformType)
 : ControlRig(InControlRig)
 , Filter(InFilter)
 , TransformType(InTransformType)
@@ -3150,7 +3271,7 @@ UBaseControlRig::FPoseScope::FPoseScope(UBaseControlRig* InControlRig, ERigEleme
 	CachedPose = InControlRig->GetHierarchy()->GetPose(IsInitial(InTransformType), InFilter, ElementView);
 }
 
-UBaseControlRig::FPoseScope::~FPoseScope()
+UControlRig::FPoseScope::~FPoseScope()
 {
 	check(ControlRig);
 
@@ -3159,7 +3280,7 @@ UBaseControlRig::FPoseScope::~FPoseScope()
 
 #if WITH_EDITOR
 
-UBaseControlRig::FTransientControlScope::FTransientControlScope(TObjectPtr<URigHierarchy> InHierarchy)
+UControlRig::FTransientControlScope::FTransientControlScope(TObjectPtr<URigHierarchy> InHierarchy)
 	:Hierarchy(InHierarchy)
 {
 	for (FRigControlElement* Control : Hierarchy->GetTransientControls())
@@ -3177,7 +3298,7 @@ UBaseControlRig::FTransientControlScope::FTransientControlScope(TObjectPtr<URigH
 	}
 }
 
-UBaseControlRig::FTransientControlScope::~FTransientControlScope()
+UControlRig::FTransientControlScope::~FTransientControlScope()
 {
 	if (URigHierarchyController* Controller = Hierarchy->GetController())
 	{
diff --git a/Engine/Plugins/Animation/ControlRig/Source/ControlRig/Private/Constraints/ControlRigTransformableHandle.cpp b/Engine/Plugins/Animation/ControlRig/Source/ControlRig/Private/Constraints/ControlRigTransformableHandle.cpp
--- a/Engine/Plugins/Animation/ControlRig/Source/ControlRig/Private/Constraints/ControlRigTransformableHandle.cpp
+++ b/Engine/Plugins/Animation/ControlRig/Source/ControlRig/Private/Constraints/ControlRigTransformableHandle.cpp
@@ -177,7 +177,7 @@ FTickFunction* UTransformableControlHandle::GetTickFunction() const
 	return BoundComponent ? &BoundComponent->PrimaryComponentTick : nullptr;
 }
 
-uint32 UTransformableControlHandle::ComputeHash(const UBaseControlRig* InControlRig, const FName& InControlName)
+uint32 UTransformableControlHandle::ComputeHash(const UControlRig* InControlRig, const FName& InControlName)
 {
 	return HashCombine(GetTypeHash(InControlRig), GetTypeHash(InControlName));
 }
@@ -408,7 +408,7 @@ void UTransformableControlHandle::OnHierarchyModified(
 }
 
 void UTransformableControlHandle::OnControlModified(
-	UBaseControlRig* InControlRig,
+	UControlRig* InControlRig,
 	FRigControlElement* InControl,
 	const FRigControlModifiedContext& InContext)
 {
@@ -447,7 +447,7 @@ void UTransformableControlHandle::OnControlModified(
 	}
 }
 
-void UTransformableControlHandle::OnControlRigBound(UBaseControlRig* InControlRig)
+void UTransformableControlHandle::OnControlRigBound(UControlRig* InControlRig)
 {
 	if (!InControlRig)
 	{
@@ -512,7 +512,7 @@ bool UTransformableControlHandle::AddTransformKeys(const TArray<FFrameNumber>& I
 	}
 	auto KeyframeFunc = [this, bLocal](const FTransform& InTransform, const FRigControlModifiedContext& InKeyframeContext)
 	{
-		UBaseControlRig* InControlRig = ControlRig.Get();
+		UControlRig* InControlRig = ControlRig.Get();
 		static constexpr bool bNotify = true;
 		static constexpr bool bUndo = false;
 		static constexpr bool bFixEuler = true;
@@ -552,7 +552,7 @@ bool UTransformableControlHandle::AddTransformKeys(const TArray<FFrameNumber>& I
 //for control rig need to check to see if the control rig is different then we may need to update it based upon what we are now bound to
 void UTransformableControlHandle::ResolveBoundObjects(FMovieSceneSequenceID LocalSequenceID, IMovieScenePlayer& Player, UObject* SubObject)
 {
-	if (UBaseControlRig* InControlRig = Cast<UBaseControlRig>(SubObject))
+	if (UControlRig* InControlRig = Cast<UControlRig>(SubObject))
 	{
 		if (ControlRig != InControlRig)
 		{
@@ -610,7 +610,7 @@ void UTransformableControlHandle::OnObjectsReplaced(const TMap<UObject*, UObject
 {
 	if (UObject* NewObject = InOldToNewInstances.FindRef(ControlRig.Get()))
 	{
-		if (UBaseControlRig* NewControlRig = Cast<UBaseControlRig>(NewObject))
+		if (UControlRig* NewControlRig = Cast<UControlRig>(NewObject))
 		{
 			if (URigHierarchy* Hierarchy = ControlRig->GetHierarchy())
 			{
diff --git a/Engine/Plugins/Animation/ControlRig/Source/ControlRig/Private/ControlRigBlueprintGeneratedClass.cpp b/Engine/Plugins/Animation/ControlRig/Source/ControlRig/Private/ControlRigBlueprintGeneratedClass.cpp
--- a/Engine/Plugins/Animation/ControlRig/Source/ControlRig/Private/ControlRigBlueprintGeneratedClass.cpp
+++ b/Engine/Plugins/Animation/ControlRig/Source/ControlRig/Private/ControlRigBlueprintGeneratedClass.cpp
@@ -34,7 +34,7 @@ void UControlRigBlueprintGeneratedClass::Serialize(FArchive& Ar)
 	const FString VMDebugName = FString::Printf(Format, *GetName(), uint32(++NumVMs));
 	URigVM* VM = NewObject<URigVM>(GetTransientPackage(), *VMDebugName);
 
-	if (UBaseControlRig* CDO = Cast<UBaseControlRig>(GetDefaultObject(true)))
+	if (UControlRig* CDO = Cast<UControlRig>(GetDefaultObject(true)))
 	{
 		if (Ar.IsSaving() && CDO->VM)
 		{
@@ -44,7 +44,7 @@ void UControlRigBlueprintGeneratedClass::Serialize(FArchive& Ar)
 	
 	VM->Serialize(Ar);
 
-	if (UBaseControlRig* CDO = Cast<UBaseControlRig>(GetDefaultObject(false)))
+	if (UControlRig* CDO = Cast<UControlRig>(GetDefaultObject(false)))
 	{
 		if (Ar.IsLoading() && CDO->VM)
 		{
diff --git a/Engine/Plugins/Animation/ControlRig/Source/ControlRig/Private/ControlRigComponent.cpp b/Engine/Plugins/Animation/ControlRig/Source/ControlRig/Private/ControlRigComponent.cpp
--- a/Engine/Plugins/Animation/ControlRig/Source/ControlRig/Private/ControlRigComponent.cpp
+++ b/Engine/Plugins/Animation/ControlRig/Source/ControlRig/Private/ControlRigComponent.cpp
@@ -206,7 +206,7 @@ void UControlRigComponent::TickComponent(float DeltaTime, enum ELevelTick TickTy
 		return;
 	}
 
-	if (UBaseControlRig* CR = SetupControlRigIfRequired())
+	if (UControlRig* CR = SetupControlRigIfRequired())
 	{
 		FScopeLock Lock(&gPendingSkeletalMeshesLock);
 		TArray<FSkeletalMeshToMap>* PendingSkeletalMeshes = gPendingSkeletalMeshes.Find(this);
@@ -283,7 +283,7 @@ FBoxSphereBounds UControlRigComponent::CalcBounds(const FTransform& LocalToWorld
 }
 
 
-UBaseControlRig* UControlRigComponent::GetControlRig()
+UControlRig* UControlRigComponent::GetControlRig()
 {
 	return SetupControlRigIfRequired();
 }
@@ -295,7 +295,7 @@ bool UControlRigComponent::CanExecute()
 		return false;
 	}
 	
-	if(UBaseControlRig* CR = GetControlRig())
+	if(UControlRig* CR = GetControlRig())
 	{
 		return CR->CanExecute();
 	}
@@ -370,7 +370,7 @@ void UControlRigComponent::Initialize()
 		OnPreInitialize(this);
 	}
 	
-	if(UBaseControlRig* CR = SetupControlRigIfRequired())
+	if(UControlRig* CR = SetupControlRigIfRequired())
 	{
 		if (CR->IsInitializing())
 		{
@@ -438,7 +438,7 @@ void UControlRigComponent::Initialize()
 
 void UControlRigComponent::Update(float DeltaTime)
 {
-	if(UBaseControlRig* CR = SetupControlRigIfRequired())
+	if(UControlRig* CR = SetupControlRigIfRequired())
 	{
 		if(!CanExecute())
 		{
@@ -531,7 +531,7 @@ TArray<FName> UControlRigComponent::GetElementNames(ERigElementType ElementType)
 {
 	TArray<FName> Names;
 
-	if (UBaseControlRig* CR = SetupControlRigIfRequired())
+	if (UControlRig* CR = SetupControlRigIfRequired())
 	{
 		for (FRigBaseElement* Element : *CR->GetHierarchy())
 		{
@@ -547,7 +547,7 @@ TArray<FName> UControlRigComponent::GetElementNames(ERigElementType ElementType)
 
 bool UControlRigComponent::DoesElementExist(FName Name, ERigElementType ElementType)
 {
-	if (UBaseControlRig* CR = SetupControlRigIfRequired())
+	if (UControlRig* CR = SetupControlRigIfRequired())
 	{
 		return CR->GetHierarchy()->GetIndex(FRigElementKey(Name, ElementType)) != INDEX_NONE;
 	}
@@ -637,7 +637,7 @@ void UControlRigComponent::AddMappedSkeletalMesh(USkeletalMeshComponent* Skeleta
 		return;
 	}
 
-	UBaseControlRig* CR = SetupControlRigIfRequired();
+	UControlRig* CR = SetupControlRigIfRequired();
 	if (CR == nullptr)
 	{
 		// if we don't have a valid rig yet - delay it until tick component
@@ -753,7 +753,7 @@ void UControlRigComponent::SetBoneInitialTransformsFromSkeletalMesh(USkeletalMes
 {
 	if (InSkeletalMesh)
 	{
-		if (UBaseControlRig* CR = SetupControlRigIfRequired())
+		if (UControlRig* CR = SetupControlRigIfRequired())
 		{
 			CR->SetBoneInitialTransformsFromSkeletalMesh(InSkeletalMesh);
 			bResetInitialsBeforeConstruction = false;
@@ -763,7 +763,7 @@ void UControlRigComponent::SetBoneInitialTransformsFromSkeletalMesh(USkeletalMes
 
 FTransform UControlRigComponent::GetBoneTransform(FName BoneName, EControlRigComponentSpace Space)
 {
-	if (UBaseControlRig* CR = SetupControlRigIfRequired())
+	if (UControlRig* CR = SetupControlRigIfRequired())
 	{
 		const int32 BoneIndex = CR->GetHierarchy()->GetIndex(FRigElementKey(BoneName, ERigElementType::Bone));
 		if (BoneIndex != INDEX_NONE)
@@ -786,7 +786,7 @@ FTransform UControlRigComponent::GetBoneTransform(FName BoneName, EControlRigCom
 
 FTransform UControlRigComponent::GetInitialBoneTransform(FName BoneName, EControlRigComponentSpace Space)
 {
-	if (UBaseControlRig* CR = SetupControlRigIfRequired())
+	if (UControlRig* CR = SetupControlRigIfRequired())
 	{
 		const int32 BoneIndex = CR->GetHierarchy()->GetIndex(FRigElementKey(BoneName, ERigElementType::Bone));
 		if (BoneIndex != INDEX_NONE)
@@ -814,7 +814,7 @@ void UControlRigComponent::SetBoneTransform(FName BoneName, FTransform Transform
 		return;
 	}
 
-	if (UBaseControlRig* CR = SetupControlRigIfRequired())
+	if (UControlRig* CR = SetupControlRigIfRequired())
 	{
 		int32 BoneIndex = CR->GetHierarchy()->GetIndex(FRigElementKey(BoneName, ERigElementType::Bone));
 		if (BoneIndex != INDEX_NONE)
@@ -853,7 +853,7 @@ void UControlRigComponent::SetBoneTransform(FName BoneName, FTransform Transform
 
 void UControlRigComponent::SetInitialBoneTransform(FName BoneName, FTransform InitialTransform, EControlRigComponentSpace Space, bool bPropagateToChildren)
 {
-	if (UBaseControlRig* CR = SetupControlRigIfRequired())
+	if (UControlRig* CR = SetupControlRigIfRequired())
 	{
 		const int32 BoneIndex = CR->GetHierarchy()->GetIndex(FRigElementKey(BoneName, ERigElementType::Bone));
 		if (BoneIndex != INDEX_NONE)
@@ -882,7 +882,7 @@ void UControlRigComponent::SetInitialBoneTransform(FName BoneName, FTransform In
 
 bool UControlRigComponent::GetControlBool(FName ControlName)
 {
-	if (UBaseControlRig* CR = SetupControlRigIfRequired())
+	if (UControlRig* CR = SetupControlRigIfRequired())
 	{
 		if(FRigControlElement* ControlElement = CR->GetHierarchy()->Find<FRigControlElement>(FRigElementKey(ControlName, ERigElementType::Control)))
 		{
@@ -898,7 +898,7 @@ bool UControlRigComponent::GetControlBool(FName ControlName)
 
 float UControlRigComponent::GetControlFloat(FName ControlName)
 {
-	if (UBaseControlRig* CR = SetupControlRigIfRequired())
+	if (UControlRig* CR = SetupControlRigIfRequired())
 	{
 		if(FRigControlElement* ControlElement = CR->GetHierarchy()->Find<FRigControlElement>(FRigElementKey(ControlName, ERigElementType::Control)))
 		{
@@ -914,7 +914,7 @@ float UControlRigComponent::GetControlFloat(FName ControlName)
 
 int32 UControlRigComponent::GetControlInt(FName ControlName)
 {
-	if (UBaseControlRig* CR = SetupControlRigIfRequired())
+	if (UControlRig* CR = SetupControlRigIfRequired())
 	{
 		if(FRigControlElement* ControlElement = CR->GetHierarchy()->Find<FRigControlElement>(FRigElementKey(ControlName, ERigElementType::Control)))
 		{
@@ -930,7 +930,7 @@ int32 UControlRigComponent::GetControlInt(FName ControlName)
 
 FVector2D UControlRigComponent::GetControlVector2D(FName ControlName)
 {
-	if (UBaseControlRig* CR = SetupControlRigIfRequired())
+	if (UControlRig* CR = SetupControlRigIfRequired())
 	{
 		if(FRigControlElement* ControlElement = CR->GetHierarchy()->Find<FRigControlElement>(FRigElementKey(ControlName, ERigElementType::Control)))
 		{
@@ -962,7 +962,7 @@ FVector UControlRigComponent::GetControlScale(FName ControlName, EControlRigComp
 
 FTransform UControlRigComponent::GetControlTransform(FName ControlName, EControlRigComponentSpace Space)
 {
-	if (UBaseControlRig* CR = SetupControlRigIfRequired())
+	if (UControlRig* CR = SetupControlRigIfRequired())
 	{
 		if(FRigControlElement* ControlElement = CR->GetHierarchy()->Find<FRigControlElement>(FRigElementKey(ControlName, ERigElementType::Control)))
 		{
@@ -985,7 +985,7 @@ FTransform UControlRigComponent::GetControlTransform(FName ControlName, EControl
 
 void UControlRigComponent::SetControlBool(FName ControlName, bool Value)
 {
-	if (UBaseControlRig* CR = SetupControlRigIfRequired())
+	if (UControlRig* CR = SetupControlRigIfRequired())
 	{
 		CR->SetControlValue<bool>(ControlName, Value);
 	}
@@ -993,7 +993,7 @@ void UControlRigComponent::SetControlBool(FName ControlName, bool Value)
 
 void UControlRigComponent::SetControlFloat(FName ControlName, float Value)
 {
-	if (UBaseControlRig* CR = SetupControlRigIfRequired())
+	if (UControlRig* CR = SetupControlRigIfRequired())
 	{
 		CR->SetControlValue<float>(ControlName, Value);
 	}
@@ -1001,7 +1001,7 @@ void UControlRigComponent::SetControlFloat(FName ControlName, float Value)
 
 void UControlRigComponent::SetControlInt(FName ControlName, int32 Value)
 {
-	if (UBaseControlRig* CR = SetupControlRigIfRequired())
+	if (UControlRig* CR = SetupControlRigIfRequired())
 	{
 		CR->SetControlValue<int32>(ControlName, Value);
 	}
@@ -1009,7 +1009,7 @@ void UControlRigComponent::SetControlInt(FName ControlName, int32 Value)
 
 void UControlRigComponent::SetControlVector2D(FName ControlName, FVector2D Value)
 {
-	if (UBaseControlRig* CR = SetupControlRigIfRequired())
+	if (UControlRig* CR = SetupControlRigIfRequired())
 	{
 		CR->SetControlValue<FVector2D>(ControlName, Value);
 	}
@@ -1017,7 +1017,7 @@ void UControlRigComponent::SetControlVector2D(FName ControlName, FVector2D Value
 
 void UControlRigComponent::SetControlPosition(FName ControlName, FVector Value, EControlRigComponentSpace Space)
 {
-	if (UBaseControlRig* CR = SetupControlRigIfRequired())
+	if (UControlRig* CR = SetupControlRigIfRequired())
 	{
 		if(FRigControlElement* ControlElement = CR->GetHierarchy()->Find<FRigControlElement>(FRigElementKey(ControlName, ERigElementType::Control)))
 		{
@@ -1044,7 +1044,7 @@ void UControlRigComponent::SetControlPosition(FName ControlName, FVector Value,
 
 void UControlRigComponent::SetControlRotator(FName ControlName, FRotator Value, EControlRigComponentSpace Space)
 {
-	if (UBaseControlRig* CR = SetupControlRigIfRequired())
+	if (UControlRig* CR = SetupControlRigIfRequired())
 	{
 		if(FRigControlElement* ControlElement = CR->GetHierarchy()->Find<FRigControlElement>(FRigElementKey(ControlName, ERigElementType::Control)))
 		{
@@ -1071,7 +1071,7 @@ void UControlRigComponent::SetControlRotator(FName ControlName, FRotator Value,
 
 void UControlRigComponent::SetControlScale(FName ControlName, FVector Value, EControlRigComponentSpace Space)
 {
-	if (UBaseControlRig* CR = SetupControlRigIfRequired())
+	if (UControlRig* CR = SetupControlRigIfRequired())
 	{
 		if(FRigControlElement* ControlElement = CR->GetHierarchy()->Find<FRigControlElement>(FRigElementKey(ControlName, ERigElementType::Control)))
 		{
@@ -1098,7 +1098,7 @@ void UControlRigComponent::SetControlScale(FName ControlName, FVector Value, ECo
 
 void UControlRigComponent::SetControlTransform(FName ControlName, FTransform Value, EControlRigComponentSpace Space)
 {
-	if (UBaseControlRig* CR = SetupControlRigIfRequired())
+	if (UControlRig* CR = SetupControlRigIfRequired())
 	{
 		if(FRigControlElement* ControlElement = CR->GetHierarchy()->Find<FRigControlElement>(FRigElementKey(ControlName, ERigElementType::Control)))
 		{
@@ -1117,7 +1117,7 @@ void UControlRigComponent::SetControlTransform(FName ControlName, FTransform Val
 
 FTransform UControlRigComponent::GetControlOffset(FName ControlName, EControlRigComponentSpace Space)
 {
-	if (UBaseControlRig* CR = SetupControlRigIfRequired())
+	if (UControlRig* CR = SetupControlRigIfRequired())
 	{
 		if(FRigControlElement* ControlElement = CR->GetHierarchy()->Find<FRigControlElement>(FRigElementKey(ControlName, ERigElementType::Control)))
 		{
@@ -1139,7 +1139,7 @@ FTransform UControlRigComponent::GetControlOffset(FName ControlName, EControlRig
 
 void UControlRigComponent::SetControlOffset(FName ControlName, FTransform OffsetTransform, EControlRigComponentSpace Space)
 {
-	if (UBaseControlRig* CR = SetupControlRigIfRequired())
+	if (UControlRig* CR = SetupControlRigIfRequired())
 	{
 		if(FRigControlElement* ControlElement = CR->GetHierarchy()->Find<FRigControlElement>(FRigElementKey(ControlName, ERigElementType::Control)))
 		{
@@ -1156,7 +1156,7 @@ void UControlRigComponent::SetControlOffset(FName ControlName, FTransform Offset
 
 FTransform UControlRigComponent::GetSpaceTransform(FName SpaceName, EControlRigComponentSpace Space)
 {
-	if (UBaseControlRig* CR = SetupControlRigIfRequired())
+	if (UControlRig* CR = SetupControlRigIfRequired())
 	{
 		if(FRigNullElement* NullElement = CR->GetHierarchy()->Find<FRigNullElement>(FRigElementKey(SpaceName, ERigElementType::Control)))
 		{
@@ -1178,7 +1178,7 @@ FTransform UControlRigComponent::GetSpaceTransform(FName SpaceName, EControlRigC
 
 FTransform UControlRigComponent::GetInitialSpaceTransform(FName SpaceName, EControlRigComponentSpace Space)
 {
-	if (UBaseControlRig* CR = SetupControlRigIfRequired())
+	if (UControlRig* CR = SetupControlRigIfRequired())
 	{
 		if(FRigNullElement* NullElement = CR->GetHierarchy()->Find<FRigNullElement>(FRigElementKey(SpaceName, ERigElementType::Control)))
 		{
@@ -1200,7 +1200,7 @@ FTransform UControlRigComponent::GetInitialSpaceTransform(FName SpaceName, ECont
 
 void UControlRigComponent::SetInitialSpaceTransform(FName SpaceName, FTransform InitialTransform, EControlRigComponentSpace Space)
 {
-	if (UBaseControlRig* CR = SetupControlRigIfRequired())
+	if (UControlRig* CR = SetupControlRigIfRequired())
 	{
 		if(FRigNullElement* NullElement = CR->GetHierarchy()->Find<FRigNullElement>(FRigElementKey(SpaceName, ERigElementType::Control)))
 		{
@@ -1217,7 +1217,7 @@ void UControlRigComponent::SetInitialSpaceTransform(FName SpaceName, FTransform
 	}
 }
 
-UBaseControlRig* UControlRigComponent::SetupControlRigIfRequired()
+UControlRig* UControlRigComponent::SetupControlRigIfRequired()
 {
 	if(ControlRig != nullptr)
 	{
@@ -1239,7 +1239,7 @@ UBaseControlRig* UControlRigComponent::SetupControlRigIfRequired()
 
 	if(ControlRigClass)
 	{
-		ControlRig = NewObject<UBaseControlRig>(this, ControlRigClass);
+		ControlRig = NewObject<UControlRig>(this, ControlRigClass);
 
 		SetControlRig(ControlRig);
 
@@ -1253,7 +1253,7 @@ UBaseControlRig* UControlRigComponent::SetupControlRigIfRequired()
 
 	return ControlRig;
 }
-void UControlRigComponent::SetControlRig(UBaseControlRig* InControlRig)
+void UControlRigComponent::SetControlRig(UControlRig* InControlRig)
 {
 	if (ControlRig)
 	{
@@ -1272,7 +1272,7 @@ void UControlRigComponent::SetControlRig(UBaseControlRig* InControlRig)
 	ControlRig->OnPostForwardsSolve_AnyThread().AddUObject(this, &UControlRigComponent::HandleControlRigPostForwardsSolveEvent);
 	ControlRig->OnExecuted_AnyThread().AddUObject(this, &UControlRigComponent::HandleControlRigExecutedEvent);
 
-	ControlRig->GetDataSourceRegistry()->RegisterDataSource(UBaseControlRig::OwnerComponent, this);
+	ControlRig->GetDataSourceRegistry()->RegisterDataSource(UControlRig::OwnerComponent, this);
 	if(ObjectBinding.IsValid())
 	{
 		ControlRig->SetObjectBinding(ObjectBinding);
@@ -1281,7 +1281,7 @@ void UControlRigComponent::SetControlRig(UBaseControlRig* InControlRig)
 	ControlRig->Initialize();
 }
 
-void UControlRigComponent::SetControlRigClass(TSubclassOf<UBaseControlRig> InControlRigClass)
+void UControlRigComponent::SetControlRigClass(TSubclassOf<UControlRig> InControlRigClass)
 {
 	ControlRig = nullptr;
 	ControlRigClass = InControlRigClass;
@@ -1296,7 +1296,7 @@ void UControlRigComponent::SetObjectBinding(UObject* InObjectToBind)
 	}
 	ObjectBinding->BindToObject(InObjectToBind);
 
-	if(UBaseControlRig* CR = SetupControlRigIfRequired())
+	if(UControlRig* CR = SetupControlRigIfRequired())
 	{
 		CR->SetObjectBinding(ObjectBinding);
 	}
@@ -1727,7 +1727,7 @@ void UControlRigComponent::HandleControlRigInitializedEvent(URigVMHost* InContro
 	}
 }
 
-void UControlRigComponent::HandleControlRigPreConstructionEvent(UBaseControlRig* InControlRig, const FName& InEventName)
+void UControlRigComponent::HandleControlRigPreConstructionEvent(UControlRig* InControlRig, const FName& InEventName)
 {
 	TArray<USkeletalMeshComponent*> ComponentsToTick;
 
@@ -1777,7 +1777,7 @@ void UControlRigComponent::HandleControlRigPreConstructionEvent(UBaseControlRig*
 	}
 }
 
-void UControlRigComponent::HandleControlRigPostConstructionEvent(UBaseControlRig* InControlRig, const FName& InEventName)
+void UControlRigComponent::HandleControlRigPostConstructionEvent(UControlRig* InControlRig, const FName& InEventName)
 {
 #if WITH_EDITOR
 	if (bUpdateInEditor)
@@ -1792,7 +1792,7 @@ void UControlRigComponent::HandleControlRigPostConstructionEvent(UBaseControlRig
 	}
 }
 
-void UControlRigComponent::HandleControlRigPreForwardsSolveEvent(UBaseControlRig* InControlRig, const FName& InEventName)
+void UControlRigComponent::HandleControlRigPreForwardsSolveEvent(UControlRig* InControlRig, const FName& InEventName)
 {
 #if WITH_EDITOR
 	if (bUpdateInEditor)
@@ -1807,7 +1807,7 @@ void UControlRigComponent::HandleControlRigPreForwardsSolveEvent(UBaseControlRig
 	}
 }
 
-void UControlRigComponent::HandleControlRigPostForwardsSolveEvent(UBaseControlRig* InControlRig, const FName& InEventName)
+void UControlRigComponent::HandleControlRigPostForwardsSolveEvent(UControlRig* InControlRig, const FName& InEventName)
 {
 #if WITH_EDITOR
 	if (bUpdateInEditor)
@@ -1879,7 +1879,7 @@ void UControlRigComponent::ConvertTransformFromRigSpace(FTransform& InOutTransfo
 
 bool UControlRigComponent::EnsureCalledOutsideOfBracket(const TCHAR* InCallingFunctionName)
 {
-	if (UBaseControlRig* CR = SetupControlRigIfRequired())
+	if (UControlRig* CR = SetupControlRigIfRequired())
 	{
 		if (CR->IsRunningPreConstruction())
 		{
diff --git a/Engine/Plugins/Animation/ControlRig/Source/ControlRig/Private/ControlRigControlActor.cpp b/Engine/Plugins/Animation/ControlRig/Source/ControlRig/Private/ControlRigControlActor.cpp
--- a/Engine/Plugins/Animation/ControlRig/Source/ControlRig/Private/ControlRigControlActor.cpp
+++ b/Engine/Plugins/Animation/ControlRig/Source/ControlRig/Private/ControlRigControlActor.cpp
@@ -156,7 +156,7 @@ void AControlRigControlActor::Refresh()
 
 	if (ControlRig == nullptr)
 	{
-		TArray<UBaseControlRig*> Rigs = UBaseControlRig::FindBaseControlRigs(ActorToTrack, ControlRigClass);
+		TArray<UControlRig*> Rigs = UControlRig::FindControlRigs(ActorToTrack, ControlRigClass);
 		if(Rigs.Num() > 0)
 		{
 			ControlRig = Rigs[0];
diff --git a/Engine/Plugins/Animation/ControlRig/Source/ControlRig/Private/ControlRigGizmoActor.cpp b/Engine/Plugins/Animation/ControlRig/Source/ControlRig/Private/ControlRigGizmoActor.cpp
--- a/Engine/Plugins/Animation/ControlRig/Source/ControlRig/Private/ControlRigGizmoActor.cpp
+++ b/Engine/Plugins/Animation/ControlRig/Source/ControlRig/Private/ControlRigGizmoActor.cpp
@@ -113,7 +113,7 @@ void AControlRigShapeActor::SetShapeColor(const FLinearColor& InColor)
 
 bool AControlRigShapeActor::UpdateControlSettings(
 	ERigHierarchyNotification InNotif,
-	UBaseControlRig* InControlRig,
+	UControlRig* InControlRig,
 	const FRigControlElement* InControlElement,
 	bool bHideManipulators,
 	bool bIsInLevelEditor)
diff --git a/Engine/Plugins/Animation/ControlRig/Source/ControlRig/Private/ControlRigModule.cpp b/Engine/Plugins/Animation/ControlRig/Source/ControlRig/Private/ControlRigModule.cpp
--- a/Engine/Plugins/Animation/ControlRig/Source/ControlRig/Private/ControlRigModule.cpp
+++ b/Engine/Plugins/Animation/ControlRig/Source/ControlRig/Private/ControlRigModule.cpp
@@ -36,7 +36,7 @@ void FControlRigModule::RegisterTransformableCustomization() const
 	// register UControlRig and AControlRigShapeActor
 	auto CreateControlHandle = [](UObject* InObject, const FName& InControlName)->UTransformableHandle*
 	{
-		if (const UBaseControlRig* ControlRig = Cast<UBaseControlRig>(InObject))
+		if (const UControlRig* ControlRig = Cast<UControlRig>(InObject))
 		{
 			return ControlRig->CreateTransformableControlHandle(InControlName);
 		}
@@ -45,7 +45,7 @@ void FControlRigModule::RegisterTransformableCustomization() const
 
 	auto GetControlHash = [](const UObject* InObject, const FName& InControlName)->uint32
 	{
-		if (const UBaseControlRig* ControlRig = Cast<UBaseControlRig>(InObject))
+		if (const UControlRig* ControlRig = Cast<UControlRig>(InObject))
 		{
 			return UTransformableControlHandle::ComputeHash(ControlRig, InControlName);
 		}
@@ -74,7 +74,7 @@ void FControlRigModule::RegisterTransformableCustomization() const
 	// as AControlRigShapeActor is only available if the ControlRig plugin is loaded.
 	FTransformableRegistry& Registry = FTransformableRegistry::Get();
 	Registry.Register(AControlRigShapeActor::StaticClass(), CreateControlHandleFromActor, GetControlHashFromActor);
-	Registry.Register(UBaseControlRig::StaticClass(), CreateControlHandle, GetControlHash);
+	Registry.Register(UControlRig::StaticClass(), CreateControlHandle, GetControlHash);
 }
 
 IMPLEMENT_MODULE(FControlRigModule, ControlRig)
