Control Rig: Rename BaseControlRig to ControlRig

diff --git a/Engine/Plugins/Animation/ControlRig/Source/ControlRig/Public/Tools/ControlRigPoseMirrorTable.h b/Engine/Plugins/Animation/ControlRig/Source/ControlRig/Public/Tools/ControlRigPoseMirrorTable.h
--- a/Engine/Plugins/Animation/ControlRig/Source/ControlRig/Public/Tools/ControlRigPoseMirrorTable.h
+++ b/Engine/Plugins/Animation/ControlRig/Source/ControlRig/Public/Tools/ControlRigPoseMirrorTable.h
@@ -10,7 +10,7 @@
 #include "UObject/ObjectMacros.h"
 #include "UObject/Object.h"
 
-class UBaseControlRig;
+class UControlRig;
 struct FRigControl;
 struct FRigControlCopy;
 struct FControlRigControlPose;
@@ -22,7 +22,7 @@ public:
 	~FControlRigPoseMirrorTable() {};
 
 	/*Set up the Mirror Table*/
-	void SetUpMirrorTable(const UBaseControlRig* ControlRig);
+	void SetUpMirrorTable(const UControlRig* ControlRig);
 
 	/*Get the matched control with the given name*/
 	FRigControlCopy* GetControl(FControlRigControlPose& Pose, FName ControlrigName);
diff --git a/Engine/Plugins/Animation/ControlRig/Source/ControlRig/Public/Tools/ControlRigToolAsset.h b/Engine/Plugins/Animation/ControlRig/Source/ControlRig/Public/Tools/ControlRigToolAsset.h
--- a/Engine/Plugins/Animation/ControlRig/Source/ControlRig/Public/Tools/ControlRigToolAsset.h
+++ b/Engine/Plugins/Animation/ControlRig/Source/ControlRig/Public/Tools/ControlRigToolAsset.h
@@ -19,7 +19,7 @@ Create One of the Control Rig Assets, Currently Only support Pose Assets but may
 class FControlRigToolAsset
 {
 public:
-	template< typename Type> static  UObject* SaveAsset(UBaseControlRig* InControlRig, const FString& CurrentPath, const FString& InString, bool bUseAllControls)
+	template< typename Type> static  UObject* SaveAsset(UControlRig* InControlRig, const FString& CurrentPath, const FString& InString, bool bUseAllControls)
 	{
 		if (InControlRig)
 		{
diff --git a/Engine/Plugins/Animation/ControlRig/Source/ControlRig/Public/Units/RigUnitContext.h b/Engine/Plugins/Animation/ControlRig/Source/ControlRig/Public/Units/RigUnitContext.h
--- a/Engine/Plugins/Animation/ControlRig/Source/ControlRig/Public/Units/RigUnitContext.h
+++ b/Engine/Plugins/Animation/ControlRig/Source/ControlRig/Public/Units/RigUnitContext.h
@@ -11,7 +11,7 @@
 #include "RigVMCore/RigVMExecuteContext.h"
 #include "RigUnitContext.generated.h"
 
-class UBaseControlRig;
+class UControlRig;
 class UControlRigShapeLibrary;
 
 /**
@@ -197,7 +197,7 @@ private:
 class CONTROLRIG_API FControlRigExecuteContextRigModuleGuard
 {
 public:
-	FControlRigExecuteContextRigModuleGuard(FControlRigExecuteContext& InContext, const UBaseControlRig* InControlRig);
+	FControlRigExecuteContextRigModuleGuard(FControlRigExecuteContext& InContext, const UControlRig* InControlRig);
 	FControlRigExecuteContextRigModuleGuard(FControlRigExecuteContext& InContext, const FString& InNewModuleNameSpace);
 	~FControlRigExecuteContextRigModuleGuard();
 
diff --git a/Engine/Plugins/Animation/ControlRig/Source/ControlRigDeveloper/Private/AnimGraphNode_ControlRig.cpp b/Engine/Plugins/Animation/ControlRig/Source/ControlRigDeveloper/Private/AnimGraphNode_ControlRig.cpp
--- a/Engine/Plugins/Animation/ControlRig/Source/ControlRigDeveloper/Private/AnimGraphNode_ControlRig.cpp
+++ b/Engine/Plugins/Animation/ControlRig/Source/ControlRigDeveloper/Private/AnimGraphNode_ControlRig.cpp
@@ -200,7 +200,7 @@ void UAnimGraphNode_ControlRig::CreateCustomPins(TArray<UEdGraphPin*>* OldPins)
 
 	if(const UClass* ControlRigClass = GetTargetClass())
 	{
-		if(UBaseControlRig* CDO = ControlRigClass->GetDefaultObject<UBaseControlRig>())
+		if(UControlRig* CDO = ControlRigClass->GetDefaultObject<UControlRig>())
 		{
 			if(const URigHierarchy* Hierarchy = CDO->GetHierarchy())
 			{
@@ -370,7 +370,7 @@ void UAnimGraphNode_ControlRig::RebuildExposedProperties()
 	// also add all of the controls
 	if(const UClass* TargetClass = GetTargetClass())
 	{
-		if(UBaseControlRig* CDO = TargetClass->GetDefaultObject<UBaseControlRig>())
+		if(UControlRig* CDO = TargetClass->GetDefaultObject<UControlRig>())
 		{
 			if(const URigHierarchy* Hierarchy = CDO->GetHierarchy())
 			{
@@ -400,7 +400,7 @@ FRigControlElement* UAnimGraphNode_ControlRig::FindControlElement(const FName& I
 {
 	if(const UClass* ControlRigClass = GetTargetClass())
 	{
-		if(UBaseControlRig* CDO = ControlRigClass->GetDefaultObject<UBaseControlRig>())
+		if(UControlRig* CDO = ControlRigClass->GetDefaultObject<UControlRig>())
 		{
 			if(const URigHierarchy* Hierarchy = CDO->GetHierarchy())
 			{
@@ -583,7 +583,7 @@ void UAnimGraphNode_ControlRig::GetVariables(bool bInput, TMap<FName, FRigVMExte
 		if (UControlRigBlueprint* RigBlueprint = Cast<UControlRigBlueprint>(TargetClass->ClassGeneratedBy))
 		{
 			//RigBlueprint->CleanupVariables();
-			UBaseControlRig* ControlRig = TargetClass->GetDefaultObject<UBaseControlRig>();
+			UControlRig* ControlRig = TargetClass->GetDefaultObject<UControlRig>();
 			if (ControlRig)
 			{
 				const TArray<FRigVMExternalVariable>& PublicVariables = ControlRig->GetPublicVariables();
@@ -636,7 +636,7 @@ void UAnimGraphNode_ControlRig::GetAvailableMapping(const FName& PathName, TArra
 		// also add all controls
 		if(const UClass* ControlRigClass = GetTargetClass())
 		{
-			if(UBaseControlRig* CDO = ControlRigClass->GetDefaultObject<UBaseControlRig>())
+			if(UControlRig* CDO = ControlRigClass->GetDefaultObject<UControlRig>())
 			{
 				if(const URigHierarchy* Hierarchy = CDO->GetHierarchy())
 				{
@@ -710,7 +710,7 @@ void UAnimGraphNode_ControlRig::CreateVariableMapping(const FString& FilteredTex
 		// add all controls
 		if(const UClass* ControlRigClass = GetTargetClass())
 		{
-			if(UBaseControlRig* CDO = ControlRigClass->GetDefaultObject<UBaseControlRig>())
+			if(UControlRig* CDO = ControlRigClass->GetDefaultObject<UControlRig>())
 			{
 				if(const URigHierarchy* Hierarchy = CDO->GetHierarchy())
 				{
diff --git a/Engine/Plugins/Animation/ControlRig/Source/ControlRigDeveloper/Private/ControlRigBlueprint.cpp b/Engine/Plugins/Animation/ControlRig/Source/ControlRigDeveloper/Private/ControlRigBlueprint.cpp
--- a/Engine/Plugins/Animation/ControlRig/Source/ControlRigDeveloper/Private/ControlRigBlueprint.cpp
+++ b/Engine/Plugins/Animation/ControlRig/Source/ControlRigDeveloper/Private/ControlRigBlueprint.cpp
@@ -110,13 +110,13 @@ void UControlRigBlueprint::PostEditChangeProperty(FPropertyChangedEvent& Propert
 		return Variable.VarName == PropertyChangedEvent.GetMemberPropertyName();
 	}))
 	{
-		if(UBaseControlRig* DebuggedControlRig = Cast<UBaseControlRig>(GetObjectBeingDebugged()))
+		if(UControlRig* DebuggedControlRig = Cast<UControlRig>(GetObjectBeingDebugged()))
 		{
 			if(const FProperty* PropertyOnRig = DebuggedControlRig->GetClass()->FindPropertyByName(PropertyChangedEvent.MemberProperty->GetFName()))
 			{
 				if(PropertyOnRig->SameType(PropertyChangedEvent.MemberProperty))
 				{
-					UBaseControlRig* CDO = DebuggedControlRig->GetClass()->GetDefaultObject<UBaseControlRig>();
+					UControlRig* CDO = DebuggedControlRig->GetClass()->GetDefaultObject<UControlRig>();
 					const uint8* SourceMemory = PropertyOnRig->ContainerPtrToValuePtr<uint8>(CDO);
 					uint8* TargetMemory = PropertyOnRig->ContainerPtrToValuePtr<uint8>(DebuggedControlRig);
 					PropertyOnRig->CopyCompleteValue(TargetMemory, SourceMemory);
@@ -135,7 +135,7 @@ void UControlRigBlueprint::PostEditChangeChainProperty(FPropertyChangedChainEven
 	if (PropertyChangedEvent.Property && PropertyChangedEvent.Property->GetFName() == GET_MEMBER_NAME_CHECKED(UControlRigBlueprint, ShapeLibraries))
 	{
 		URigVMBlueprintGeneratedClass* RigClass = GetRigVMBlueprintGeneratedClass();
-		UBaseControlRig* CDO = Cast<UBaseControlRig>(RigClass->GetDefaultObject(false /* create if needed */));
+		UControlRig* CDO = Cast<UControlRig>(RigClass->GetDefaultObject(false /* create if needed */));
 
 		TArray<UObject*> ArchetypeInstances;
 		CDO->GetArchetypeInstances(ArchetypeInstances);
@@ -144,7 +144,7 @@ void UControlRigBlueprint::PostEditChangeChainProperty(FPropertyChangedChainEven
 		// Propagate libraries to archetypes
 		for (UObject* Instance : ArchetypeInstances)
 		{
-			if (UBaseControlRig* InstanceRig = Cast<UBaseControlRig>(Instance))
+			if (UControlRig* InstanceRig = Cast<UControlRig>(Instance))
 			{
 				InstanceRig->ShapeLibraries = ShapeLibraries;
 			}
@@ -975,7 +975,7 @@ bool UControlRigBlueprint::ResolveConnector(const FRigElementKey& DraggedKey, co
 
 	PropagateHierarchyFromBPToInstances();
 
-	if(UBaseControlRig* ControlRig = Cast<UBaseControlRig>(GetObjectBeingDebugged()))
+	if(UControlRig* ControlRig = Cast<UControlRig>(GetObjectBeingDebugged()))
 	{
 		for (UEdGraph* Graph : UbergraphPages)
 		{
@@ -1090,7 +1090,7 @@ void UControlRigBlueprint::PostLoad()
 		}
 
 		URigVMBlueprintGeneratedClass* RigClass = GetRigVMBlueprintGeneratedClass();
-		UBaseControlRig* CDO = Cast<UBaseControlRig>(RigClass->GetDefaultObject(false /* create if needed */));
+		UControlRig* CDO = Cast<UControlRig>(RigClass->GetDefaultObject(false /* create if needed */));
 
 		TArray<UObject*> ArchetypeInstances;
 		CDO->GetArchetypeInstances(ArchetypeInstances);
@@ -1098,7 +1098,7 @@ void UControlRigBlueprint::PostLoad()
 
 		for (UObject* Instance : ArchetypeInstances)
 		{
-			if (UBaseControlRig* InstanceRig = Cast<UBaseControlRig>(Instance))
+			if (UControlRig* InstanceRig = Cast<UControlRig>(Instance))
 			{
 				InstanceRig->ShapeLibraries.Reset();
 				InstanceRig->GizmoLibrary_DEPRECATED.Reset();
@@ -1119,7 +1119,7 @@ void UControlRigBlueprint::HandlePackageDone()
 		}
 
 		URigVMBlueprintGeneratedClass* RigClass = GetRigVMBlueprintGeneratedClass();
-		UBaseControlRig* CDO = Cast<UBaseControlRig>(RigClass->GetDefaultObject(false /* create if needed */));
+		UControlRig* CDO = Cast<UControlRig>(RigClass->GetDefaultObject(false /* create if needed */));
 
 		TArray<UObject*> ArchetypeInstances;
 		CDO->GetArchetypeInstances(ArchetypeInstances);
@@ -1127,7 +1127,7 @@ void UControlRigBlueprint::HandlePackageDone()
 
 		for (UObject* Instance : ArchetypeInstances)
 		{
-			if (UBaseControlRig* InstanceRig = Cast<UBaseControlRig>(Instance))
+			if (UControlRig* InstanceRig = Cast<UControlRig>(Instance))
 			{
 				InstanceRig->ShapeLibraries = ShapeLibraries;
 			}
@@ -1156,7 +1156,7 @@ void UControlRigBlueprint::HandleConfigureRigVMController(const FRigVMClient* In
 			
 			if(const URigVMBlueprint* StrongThis = WeakThis.Get())
 			{
-				if(UBaseControlRig* ControlRig = Cast<UBaseControlRig>(StrongThis->GetObjectBeingDebugged()))
+				if(UControlRig* ControlRig = Cast<UControlRig>(StrongThis->GetObjectBeingDebugged()))
 				{
 					ControlRigNodeWorkflowOptions->Hierarchy = ControlRig->GetHierarchy();
 					ControlRigNodeWorkflowOptions->Selection = ControlRig->GetHierarchy()->GetSelectedKeys();
@@ -1324,7 +1324,7 @@ TArray<UControlRigBlueprint*> UControlRigBlueprint::GetCurrentlyOpenRigBlueprint
 const FControlRigShapeDefinition* UControlRigBlueprint::GetControlShapeByName(const FName& InName) const
 {
 	TMap<FString, FString> LibraryNameMap;
-	if(UBaseControlRig* ControlRig = Cast<UBaseControlRig>(GetObjectBeingDebugged()))
+	if(UControlRig* ControlRig = Cast<UControlRig>(GetObjectBeingDebugged()))
 	{
 		LibraryNameMap = ControlRig->ShapeLibraryNameMap;
 	}
@@ -1343,14 +1343,14 @@ FName UControlRigBlueprint::AddTransientControl(const URigVMUnitNode* InNode, co
 	ClearTransientControls();
 
 	URigVMBlueprintGeneratedClass* RigClass = GetRigVMBlueprintGeneratedClass();
-	UBaseControlRig* CDO = Cast<UBaseControlRig>(RigClass->GetDefaultObject(true /* create if needed */));
+	UControlRig* CDO = Cast<UControlRig>(RigClass->GetDefaultObject(true /* create if needed */));
 
 	FName ReturnName = NAME_None;
 	TArray<UObject*> ArchetypeInstances;
 	CDO->GetArchetypeInstances(ArchetypeInstances);
 	for (UObject* ArchetypeInstance : ArchetypeInstances)
 	{
-		UBaseControlRig* InstancedControlRig = Cast<UBaseControlRig>(ArchetypeInstance);
+		UControlRig* InstancedControlRig = Cast<UControlRig>(ArchetypeInstance);
 		if (InstancedControlRig)
 		{
 			FName ControlName = InstancedControlRig->AddTransientControl(InNode, InTarget);
@@ -1373,14 +1373,14 @@ FName UControlRigBlueprint::RemoveTransientControl(const URigVMUnitNode* InNode,
 	}
 
 	URigVMBlueprintGeneratedClass* RigClass = GetRigVMBlueprintGeneratedClass();
-	UBaseControlRig* CDO = Cast<UBaseControlRig>(RigClass->GetDefaultObject(true /* create if needed */));
+	UControlRig* CDO = Cast<UControlRig>(RigClass->GetDefaultObject(true /* create if needed */));
 
 	FName RemovedName = NAME_None;
 	TArray<UObject*> ArchetypeInstances;
 	CDO->GetArchetypeInstances(ArchetypeInstances);
 	for (UObject* ArchetypeInstance : ArchetypeInstances)
 	{
-		UBaseControlRig* InstancedControlRig = Cast<UBaseControlRig>(ArchetypeInstance);
+		UControlRig* InstancedControlRig = Cast<UControlRig>(ArchetypeInstance);
 		if (InstancedControlRig)
 		{
 			FName Name = InstancedControlRig->RemoveTransientControl(InNode, InTarget);
@@ -1402,7 +1402,7 @@ FName UControlRigBlueprint::AddTransientControl(const FRigElementKey& InElement)
 		ValueScope = MakeUnique<FControlValueScope>(this);
 	}
 	URigVMBlueprintGeneratedClass* RigClass = GetRigVMBlueprintGeneratedClass();
-	UBaseControlRig* CDO = Cast<UBaseControlRig>(RigClass->GetDefaultObject(true /* create if needed */));
+	UControlRig* CDO = Cast<UControlRig>(RigClass->GetDefaultObject(true /* create if needed */));
 
 	FName ReturnName = NAME_None;
 	TArray<UObject*> ArchetypeInstances;
@@ -1413,7 +1413,7 @@ FName UControlRigBlueprint::AddTransientControl(const FRigElementKey& InElement)
 	TMap<UObject*, FTransform> SavedElementLocalTransforms;
 	for (UObject* ArchetypeInstance : ArchetypeInstances)
 	{
-		UBaseControlRig* InstancedControlRig = Cast<UBaseControlRig>(ArchetypeInstance);
+		UControlRig* InstancedControlRig = Cast<UControlRig>(ArchetypeInstance);
 		if (InstancedControlRig)
 		{
 			if (InstancedControlRig->DynamicHierarchy)
@@ -1428,7 +1428,7 @@ FName UControlRigBlueprint::AddTransientControl(const FRigElementKey& InElement)
 	
 	for (UObject* ArchetypeInstance : ArchetypeInstances)
 	{
-		UBaseControlRig* InstancedControlRig = Cast<UBaseControlRig>(ArchetypeInstance);
+		UControlRig* InstancedControlRig = Cast<UControlRig>(ArchetypeInstance);
 		if (InstancedControlRig)
 		{
 			// restore the element transforms so that transient controls are created at the right place
@@ -1461,14 +1461,14 @@ FName UControlRigBlueprint::RemoveTransientControl(const FRigElementKey& InEleme
 	}
 
 	URigVMBlueprintGeneratedClass* RigClass = GetRigVMBlueprintGeneratedClass();
-	UBaseControlRig* CDO = Cast<UBaseControlRig>(RigClass->GetDefaultObject(true /* create if needed */));
+	UControlRig* CDO = Cast<UControlRig>(RigClass->GetDefaultObject(true /* create if needed */));
 
 	FName RemovedName = NAME_None;
 	TArray<UObject*> ArchetypeInstances;
 	CDO->GetArchetypeInstances(ArchetypeInstances);
 	for (UObject* ArchetypeInstance : ArchetypeInstances)
 	{
-		UBaseControlRig* InstancedControlRig = Cast<UBaseControlRig>(ArchetypeInstance);
+		UControlRig* InstancedControlRig = Cast<UControlRig>(ArchetypeInstance);
 		if (InstancedControlRig)
 		{
 			FName Name = InstancedControlRig->RemoveTransientControl(InElement);
@@ -1492,13 +1492,13 @@ void UControlRigBlueprint::ClearTransientControls()
 
 	if (URigVMBlueprintGeneratedClass* RigClass = GetRigVMBlueprintGeneratedClass())
 	{
-		UBaseControlRig* CDO = Cast<UBaseControlRig>(RigClass->GetDefaultObject(true /* create if needed */));
+		UControlRig* CDO = Cast<UControlRig>(RigClass->GetDefaultObject(true /* create if needed */));
 
 		TArray<UObject*> ArchetypeInstances;
 		CDO->GetArchetypeInstances(ArchetypeInstances);
 		for (UObject* ArchetypeInstance : ArchetypeInstances)
 		{
-			UBaseControlRig* InstancedControlRig = Cast<UBaseControlRig>(ArchetypeInstance);
+			UControlRig* InstancedControlRig = Cast<UControlRig>(ArchetypeInstance);
 			if (InstancedControlRig)
 			{
 				InstancedControlRig->ClearTransientControls();
@@ -1576,7 +1576,7 @@ void UControlRigBlueprint::RecompileModularRig()
 void UControlRigBlueprint::SetupDefaultObjectDuringCompilation(URigVMHost* InCDO)
 {
 	Super::SetupDefaultObjectDuringCompilation(InCDO);
-	CastChecked<UBaseControlRig>(InCDO)->GetHierarchy()->CopyHierarchy(Hierarchy);
+	CastChecked<UControlRig>(InCDO)->GetHierarchy()->CopyHierarchy(Hierarchy);
 }
 
 void UControlRigBlueprint::SetupPinRedirectorsForBackwardsCompatibility()
@@ -2014,14 +2014,14 @@ void UControlRigBlueprint::PatchVariableNodesOnLoad()
 	Super::PatchVariableNodesOnLoad();
 }
 
-void UControlRigBlueprint::UpdateElementKeyRedirector(UBaseControlRig* InControlRig) const
+void UControlRigBlueprint::UpdateElementKeyRedirector(UControlRig* InControlRig) const
 {
 	InControlRig->HierarchySettings = HierarchySettings;
 	InControlRig->RigModuleSettings = RigModuleSettings;
 	InControlRig->ElementKeyRedirector = FRigElementKeyRedirector(ConnectionMap, InControlRig->GetHierarchy());
 }
 
-void UControlRigBlueprint::PropagatePoseFromInstanceToBP(UBaseControlRig* InControlRig) const
+void UControlRigBlueprint::PropagatePoseFromInstanceToBP(UControlRig* InControlRig) const
 {
 	check(InControlRig);
 	// current transforms in BP and CDO are meaningless, no need to copy them
@@ -2034,7 +2034,7 @@ void UControlRigBlueprint::PropagatePoseFromBPToInstances() const
 {
 	if (UClass* MyControlRigClass = GeneratedClass)
 	{
-		if (UBaseControlRig* DefaultObject = Cast<UBaseControlRig>(MyControlRigClass->GetDefaultObject(false)))
+		if (UControlRig* DefaultObject = Cast<UControlRig>(MyControlRigClass->GetDefaultObject(false)))
 		{
 			DefaultObject->PostInitInstanceIfRequired();
 			DefaultObject->GetHierarchy()->CopyPose(Hierarchy, true, true, true);
@@ -2043,7 +2043,7 @@ void UControlRigBlueprint::PropagatePoseFromBPToInstances() const
 			DefaultObject->GetArchetypeInstances(ArchetypeInstances);
 			for (UObject* ArchetypeInstance : ArchetypeInstances)
 			{
-				if (UBaseControlRig* InstanceRig = Cast<UBaseControlRig>(ArchetypeInstance))
+				if (UControlRig* InstanceRig = Cast<UControlRig>(ArchetypeInstance))
 				{
 					InstanceRig->PostInitInstanceIfRequired();
 					InstanceRig->GetHierarchy()->CopyPose(Hierarchy, true, true, true);
@@ -2057,7 +2057,7 @@ void UControlRigBlueprint::PropagateHierarchyFromBPToInstances() const
 {
 	if (UClass* MyControlRigClass = GeneratedClass)
 	{
-		if (UBaseControlRig* DefaultObject = Cast<UBaseControlRig>(MyControlRigClass->GetDefaultObject(false)))
+		if (UControlRig* DefaultObject = Cast<UControlRig>(MyControlRigClass->GetDefaultObject(false)))
 		{
 			DefaultObject->PostInitInstanceIfRequired();
 			DefaultObject->GetHierarchy()->CopyHierarchy(Hierarchy);
@@ -2073,7 +2073,7 @@ void UControlRigBlueprint::PropagateHierarchyFromBPToInstances() const
 			DefaultObject->GetArchetypeInstances(ArchetypeInstances);
 			for (UObject* ArchetypeInstance : ArchetypeInstances)
 			{
-				if (UBaseControlRig* InstanceRig = Cast<UBaseControlRig>(ArchetypeInstance))
+				if (UControlRig* InstanceRig = Cast<UControlRig>(ArchetypeInstance))
 				{
 					InstanceRig->PostInitInstanceIfRequired();
 					InstanceRig->GetHierarchy()->CopyHierarchy(Hierarchy);
@@ -2090,7 +2090,7 @@ void UControlRigBlueprint::PropagateDrawInstructionsFromBPToInstances() const
 {
 	if (UClass* MyControlRigClass = GeneratedClass)
 	{
-		if (UBaseControlRig* DefaultObject = Cast<UBaseControlRig>(MyControlRigClass->GetDefaultObject(false)))
+		if (UControlRig* DefaultObject = Cast<UControlRig>(MyControlRigClass->GetDefaultObject(false)))
 	{
 			DefaultObject->DrawContainer = DrawContainer;
 
@@ -2099,7 +2099,7 @@ void UControlRigBlueprint::PropagateDrawInstructionsFromBPToInstances() const
 
 			for (UObject* ArchetypeInstance : ArchetypeInstances)
 			{
-				if (UBaseControlRig* InstanceRig = Cast<UBaseControlRig>(ArchetypeInstance))
+				if (UControlRig* InstanceRig = Cast<UControlRig>(ArchetypeInstance))
 	{
 					InstanceRig->DrawContainer = DrawContainer;
 				}
@@ -2128,7 +2128,7 @@ void UControlRigBlueprint::PropagatePropertyFromBPToInstances(FRigElementKey InR
 
 	if (UClass* MyControlRigClass = GeneratedClass)
 	{
-		if (UBaseControlRig* DefaultObject = Cast<UBaseControlRig>(MyControlRigClass->GetDefaultObject(false)))
+		if (UControlRig* DefaultObject = Cast<UControlRig>(MyControlRigClass->GetDefaultObject(false)))
 		{
 			TArray<UObject*> ArchetypeInstances;
 			DefaultObject->GetArchetypeInstances(ArchetypeInstances);
@@ -2139,7 +2139,7 @@ void UControlRigBlueprint::PropagatePropertyFromBPToInstances(FRigElementKey InR
 			uint8* Source = ((uint8*)Hierarchy->Get(ElementIndex)) + PropertyOffset;
 			for (UObject* ArchetypeInstance : ArchetypeInstances)
 			{
-				if (UBaseControlRig* InstanceRig = Cast<UBaseControlRig>(ArchetypeInstance))
+				if (UControlRig* InstanceRig = Cast<UControlRig>(ArchetypeInstance))
 				{
 					InstanceRig->PostInitInstanceIfRequired();
 					uint8* Dest = ((uint8*)InstanceRig->GetHierarchy()->Get(ElementIndex)) + PropertyOffset;
@@ -2150,7 +2150,7 @@ void UControlRigBlueprint::PropagatePropertyFromBPToInstances(FRigElementKey InR
 	}
 }
 
-void UControlRigBlueprint::PropagatePropertyFromInstanceToBP(FRigElementKey InRigElement, const FProperty* InProperty, UBaseControlRig* InInstance) const
+void UControlRigBlueprint::PropagatePropertyFromInstanceToBP(FRigElementKey InRigElement, const FProperty* InProperty, UControlRig* InInstance) const
 {
 	const int32 ElementIndex = Hierarchy->GetIndex(InRigElement);
 	ensure(ElementIndex != INDEX_NONE);
@@ -2211,9 +2211,9 @@ void UControlRigBlueprint::HandleHierarchyModified(ERigHierarchyNotification InN
 
 			if(bClearTransientControls)
 			{
-				if(UBaseControlRig* RigBeingDebugged = Cast<UBaseControlRig>(GetObjectBeingDebugged()))
+				if(UControlRig* RigBeingDebugged = Cast<UControlRig>(GetObjectBeingDebugged()))
 				{
-					const FName TransientControlName = UBaseControlRig::GetNameForTransientControl(InElement->GetKey());
+					const FName TransientControlName = UControlRig::GetNameForTransientControl(InElement->GetKey());
 					const FRigElementKey TransientControlKey(TransientControlName, ERigElementType::Control);
 					if (const FRigControlElement* ControlElement = RigBeingDebugged->GetHierarchy()->Find<FRigControlElement>(TransientControlKey))
 					{
@@ -2264,7 +2264,7 @@ UControlRigBlueprint::FControlValueScope::FControlValueScope(UControlRigBlueprin
 #if WITH_EDITOR
 	check(Blueprint);
 
-	if (UBaseControlRig* CR = Cast<UBaseControlRig>(Blueprint->GetObjectBeingDebugged()))
+	if (UControlRig* CR = Cast<UControlRig>(Blueprint->GetObjectBeingDebugged()))
 	{
 		TArray<FRigControlElement*> Controls = CR->AvailableControls();
 		for (FRigControlElement* ControlElement : Controls)
@@ -2280,7 +2280,7 @@ UControlRigBlueprint::FControlValueScope::~FControlValueScope()
 #if WITH_EDITOR
 	check(Blueprint);
 
-	if (UBaseControlRig* CR = Cast<UBaseControlRig>(Blueprint->GetObjectBeingDebugged()))
+	if (UControlRig* CR = Cast<UControlRig>(Blueprint->GetObjectBeingDebugged()))
 	{
 		for (const TPair<FName, FRigControlValue>& Pair : ControlValues)
 		{
diff --git a/Engine/Plugins/Animation/ControlRig/Source/ControlRigDeveloper/Private/ControlRigBlueprintCompiler.cpp b/Engine/Plugins/Animation/ControlRig/Source/ControlRigDeveloper/Private/ControlRigBlueprintCompiler.cpp
--- a/Engine/Plugins/Animation/ControlRig/Source/ControlRigDeveloper/Private/ControlRigBlueprintCompiler.cpp
+++ b/Engine/Plugins/Animation/ControlRig/Source/ControlRigDeveloper/Private/ControlRigBlueprintCompiler.cpp
@@ -9,7 +9,7 @@ bool FControlRigBlueprintCompiler::CanCompile(const UBlueprint* Blueprint)
 {
 	DECLARE_SCOPE_HIERARCHICAL_COUNTER_FUNC()
 
-	if (Blueprint && Blueprint->ParentClass && Blueprint->ParentClass->IsChildOf(UBaseControlRig::StaticClass()))
+	if (Blueprint && Blueprint->ParentClass && Blueprint->ParentClass->IsChildOf(UControlRig::StaticClass()))
 	{
 		return true;
 	}
@@ -53,9 +53,9 @@ void FControlRigBlueprintCompilerContext::CopyTermDefaultsToDefaultObject(UObjec
 	if (ControlRigBlueprint)
 	{
 		// here, CDO is initialized from BP,
-		// and in UBaseControlRig::InitializeFromCDO,
+		// and in UControlRig::InitializeFromCDO,
 		// other Control Rig Instances are then initialized From the CDO 
-		UBaseControlRig* ControlRig = CastChecked<UBaseControlRig>(DefaultObject);
+		UControlRig* ControlRig = CastChecked<UControlRig>(DefaultObject);
 		
 		// copy hierarchy
 		{
diff --git a/Engine/Plugins/Animation/ControlRig/Source/ControlRigDeveloper/Private/ControlRigVisualGraphUtils.cpp b/Engine/Plugins/Animation/ControlRig/Source/ControlRigDeveloper/Private/ControlRigVisualGraphUtils.cpp
--- a/Engine/Plugins/Animation/ControlRig/Source/ControlRigDeveloper/Private/ControlRigVisualGraphUtils.cpp
+++ b/Engine/Plugins/Animation/ControlRig/Source/ControlRigDeveloper/Private/ControlRigVisualGraphUtils.cpp
@@ -66,7 +66,7 @@ FAutoConsoleCommandWithWorldAndArgs FCmdControlRigVisualGraphUtilsDumpHierarchy
 		{
 			if(UObject* Object = FindObject<UObject>(nullptr, *ObjectPathName, false))
 			{
-				if(UBaseControlRig* CR = Cast<UBaseControlRig>(Object))
+				if(UControlRig* CR = Cast<UControlRig>(Object))
 				{
 					Hierarchies.Add(CR->GetHierarchy());
 				}
@@ -227,7 +227,7 @@ FString FControlRigVisualGraphUtils::DumpRigHierarchyToDotGraph(URigHierarchy* I
 	
 	if(!InEventName.IsNone())
 	{
-		if(UBaseControlRig* CR = InHierarchy->GetTypedOuter<UBaseControlRig>())
+		if(UControlRig* CR = InHierarchy->GetTypedOuter<UControlRig>())
 		{
 			URigVM* VM = CR->GetVM();
 
diff --git a/Engine/Plugins/Animation/ControlRig/Source/ControlRigDeveloper/Private/Graph/ControlRigGraph.cpp b/Engine/Plugins/Animation/ControlRig/Source/ControlRigDeveloper/Private/Graph/ControlRigGraph.cpp
--- a/Engine/Plugins/Animation/ControlRig/Source/ControlRigDeveloper/Private/Graph/ControlRigGraph.cpp
+++ b/Engine/Plugins/Animation/ControlRig/Source/ControlRigDeveloper/Private/Graph/ControlRigGraph.cpp
@@ -27,7 +27,7 @@ void UControlRigGraph::InitializeFromBlueprint(URigVMBlueprint* InBlueprint)
 	const UControlRigBlueprint* ControlRigBlueprint = CastChecked<UControlRigBlueprint>(InBlueprint);
 	URigHierarchy* Hierarchy = ControlRigBlueprint->Hierarchy;
 
-	if(UBaseControlRig* ControlRig = Cast<UBaseControlRig>(ControlRigBlueprint->GetObjectBeingDebugged()))
+	if(UControlRig* ControlRig = Cast<UControlRig>(ControlRigBlueprint->GetObjectBeingDebugged()))
 	{
 		Hierarchy = ControlRig->GetHierarchy();
 	}
@@ -73,7 +73,7 @@ void UControlRigGraph::CacheNameLists(URigHierarchy* InHierarchy, const FRigVMDr
 	check(InHierarchy);
 	check(DrawContainer);
 
-	UBaseControlRig* ControlRig = InHierarchy->GetTypedOuter<UBaseControlRig>();
+	UControlRig* ControlRig = InHierarchy->GetTypedOuter<UControlRig>();
 
 	if(LastHierarchyTopologyVersion != InHierarchy->GetTopologyVersion())
 	{
@@ -175,7 +175,7 @@ const TArray<TSharedPtr<FRigVMStringWithTag>>* UControlRigGraph::GetElementNameL
 
 		UControlRigGraph* MutableThis = (UControlRigGraph*)this;
 		URigHierarchy* Hierarchy = Blueprint->Hierarchy;
-		if(UBaseControlRig* ControlRig = Cast<UBaseControlRig>(Blueprint->GetObjectBeingDebugged()))
+		if(UControlRig* ControlRig = Cast<UControlRig>(Blueprint->GetObjectBeingDebugged()))
 		{
 			Hierarchy = ControlRig->GetHierarchy();
 		}	
diff --git a/Engine/Plugins/Animation/ControlRig/Source/ControlRigDeveloper/Private/ModularRigController.cpp b/Engine/Plugins/Animation/ControlRig/Source/ControlRigDeveloper/Private/ModularRigController.cpp
--- a/Engine/Plugins/Animation/ControlRig/Source/ControlRigDeveloper/Private/ModularRigController.cpp
+++ b/Engine/Plugins/Animation/ControlRig/Source/ControlRigDeveloper/Private/ModularRigController.cpp
@@ -2,7 +2,7 @@
 
 #include "ModularRigController.h"
 
-#include "ControlRig.h"
+#include "BaseControlRig.h"
 #include "ModularRig.h"
 #include "ControlRigBlueprint.h"
 #include "ModularRigModel.h"
@@ -15,9 +15,9 @@ UModularRigController::UModularRigController(const FObjectInitializer& ObjectIni
 {
 }
 
-bool UModularRigController::AddModule(const FName& InModuleName, TSubclassOf<UBaseControlRig> InClass, const FString& InParentModulePath)
+bool UModularRigController::AddModule(const FName& InModuleName, TSubclassOf<UControlRig> InClass, const FString& InParentModulePath)
 {
-	if (!InClass->GetDefaultObject<UBaseControlRig>()->IsRigModule())
+	if (!InClass->GetDefaultObject<UControlRig>()->IsRigModule())
 	{
 		return false;
 	}
@@ -69,7 +69,7 @@ bool UModularRigController::AddModule(const FName& InModuleName, TSubclassOf<UBa
 	UControlRigBlueprint* Blueprint = Cast<UControlRigBlueprint>(GetOuter());
 	if (URigHierarchyController* Controller = Blueprint->GetHierarchyController())
 	{
-		if (UBaseControlRig* CDO = Blueprint->GetControlRigClass()->GetDefaultObject<UBaseControlRig>())
+		if (UControlRig* CDO = Blueprint->GetControlRigClass()->GetDefaultObject<UControlRig>())
 		{
 			FRigVMExtendedExecuteContext& Context = CDO->GetRigVMExtendedExecuteContext();
 			FRigHierarchyExecuteContextBracket HierarchyContextGuard(Controller->GetHierarchy(), &Context);
diff --git a/Engine/Plugins/Animation/ControlRig/Source/ControlRigDeveloper/Public/ControlRigBlueprint.h b/Engine/Plugins/Animation/ControlRig/Source/ControlRigDeveloper/Public/ControlRigBlueprint.h
--- a/Engine/Plugins/Animation/ControlRig/Source/ControlRigDeveloper/Public/ControlRigBlueprint.h
+++ b/Engine/Plugins/Animation/ControlRig/Source/ControlRigDeveloper/Public/ControlRigBlueprint.h
@@ -90,10 +90,10 @@ public:
 	UClass* GetControlRigClass();
 
 	UFUNCTION(BlueprintCallable, Category = "Control Rig Blueprint")
-	UBaseControlRig* CreateControlRig() { return Cast<UBaseControlRig>(CreateRigVMHost()); }
+	UControlRig* CreateControlRig() { return Cast<UControlRig>(CreateRigVMHost()); }
 
 	UFUNCTION(BlueprintCallable, Category = "Control Rig Blueprint")
-	UBaseControlRig* GetDebuggedControlRig() { return Cast<UBaseControlRig>(GetDebuggedRigVMHost()); } 
+	UControlRig* GetDebuggedControlRig() { return Cast<UControlRig>(GetDebuggedRigVMHost()); } 
 
 	/** IInterface_PreviewMeshProvider interface */
 	UFUNCTION(BlueprintCallable, Category = "Control Rig Blueprint")
@@ -275,13 +275,13 @@ protected:
 	virtual void PatchVariableNodesOnLoad() override;
 
 public:
-	void UpdateElementKeyRedirector(UBaseControlRig* InControlRig) const;
-	void PropagatePoseFromInstanceToBP(UBaseControlRig* InControlRig) const;
+	void UpdateElementKeyRedirector(UControlRig* InControlRig) const;
+	void PropagatePoseFromInstanceToBP(UControlRig* InControlRig) const;
 	void PropagatePoseFromBPToInstances() const;
 	void PropagateHierarchyFromBPToInstances() const;
 	void PropagateDrawInstructionsFromBPToInstances() const;
 	void PropagatePropertyFromBPToInstances(FRigElementKey InRigElement, const FProperty* InProperty) const;
-	void PropagatePropertyFromInstanceToBP(FRigElementKey InRigElement, const FProperty* InProperty, UBaseControlRig* InInstance) const;
+	void PropagatePropertyFromInstanceToBP(FRigElementKey InRigElement, const FProperty* InProperty, UControlRig* InInstance) const;
 
 	/**
 	* Returns the modified event, which can be used to 
diff --git a/Engine/Plugins/Animation/ControlRig/Source/ControlRigDeveloper/Public/Graph/ControlRigGraph.h b/Engine/Plugins/Animation/ControlRig/Source/ControlRigDeveloper/Public/Graph/ControlRigGraph.h
--- a/Engine/Plugins/Animation/ControlRig/Source/ControlRigDeveloper/Public/Graph/ControlRigGraph.h
+++ b/Engine/Plugins/Animation/ControlRig/Source/ControlRigDeveloper/Public/Graph/ControlRigGraph.h
@@ -14,7 +14,7 @@
 
 class UControlRigBlueprint;
 class UControlRigGraphSchema;
-class UBaseControlRig;
+class UControlRig;
 class URigVMController;
 struct FRigCurveContainer;
 
@@ -83,7 +83,7 @@ private:
 	}
 
 	template<class T>
-	void CacheNameListForHierarchy(UBaseControlRig* InControlRig, URigHierarchy* InHierarchy, TArray<TSharedPtr<FRigVMStringWithTag>>& OutNameList, bool bFilter = true)
+	void CacheNameListForHierarchy(UControlRig* InControlRig, URigHierarchy* InHierarchy, TArray<TSharedPtr<FRigVMStringWithTag>>& OutNameList, bool bFilter = true)
 	{
         TArray<FRigVMStringWithTag> Names;
 		for (auto Element : *InHierarchy)
diff --git a/Engine/Plugins/Animation/ControlRig/Source/ControlRigDeveloper/Public/ModularRigController.h b/Engine/Plugins/Animation/ControlRig/Source/ControlRigDeveloper/Public/ModularRigController.h
--- a/Engine/Plugins/Animation/ControlRig/Source/ControlRigDeveloper/Public/ModularRigController.h
+++ b/Engine/Plugins/Animation/ControlRig/Source/ControlRigDeveloper/Public/ModularRigController.h
@@ -41,7 +41,7 @@ class CONTROLRIGDEVELOPER_API UModularRigController : public UObject
 
 
 	UFUNCTION(BlueprintCallable, Category = "Control Rig | Modules")
-	bool AddModule(const FName& InModuleName, TSubclassOf<UBaseControlRig> InClass, const FString& InParentModulePath);
+	bool AddModule(const FName& InModuleName, TSubclassOf<UControlRig> InClass, const FString& InParentModulePath);
 
 	UFUNCTION(BlueprintCallable, Category = "Control Rig | Modules")
 	bool ConnectModuleToElement(const FRigElementKey& InConnectorKey, const FRigElementKey& InTargetKey);
diff --git a/Engine/Plugins/Animation/ControlRig/Source/ControlRigDeveloper/Public/ModularRigModel.h b/Engine/Plugins/Animation/ControlRig/Source/ControlRigDeveloper/Public/ModularRigModel.h
--- a/Engine/Plugins/Animation/ControlRig/Source/ControlRigDeveloper/Public/ModularRigModel.h
+++ b/Engine/Plugins/Animation/ControlRig/Source/ControlRigDeveloper/Public/ModularRigModel.h
@@ -17,7 +17,7 @@ struct CONTROLRIGDEVELOPER_API FRigModuleReference
 		, Class(nullptr)
 	{}
 	
-	FRigModuleReference(const FName& InName, TSubclassOf<UBaseControlRig> InClass, const FString& InParentPath)
+	FRigModuleReference(const FName& InName, TSubclassOf<UControlRig> InClass, const FString& InParentPath)
 		: Name(InName)
 		, ParentNamespace(InParentPath)
 		, Class(InClass)
@@ -30,7 +30,7 @@ struct CONTROLRIGDEVELOPER_API FRigModuleReference
 	FString ParentNamespace;
 
 	UPROPERTY()
-	TSoftClassPtr<UBaseControlRig> Class;
+	TSoftClassPtr<UControlRig> Class;
 
 	UPROPERTY()
 	TMap<FRigElementKey, FRigElementKey> Connections; // Connectors to Connection element
diff --git a/Engine/Plugins/Animation/ControlRig/Source/ControlRigEditor/Private/ControlRigAnimGraphDetails.cpp b/Engine/Plugins/Animation/ControlRig/Source/ControlRigEditor/Private/ControlRigAnimGraphDetails.cpp
--- a/Engine/Plugins/Animation/ControlRig/Source/ControlRigEditor/Private/ControlRigAnimGraphDetails.cpp
+++ b/Engine/Plugins/Animation/ControlRig/Source/ControlRigEditor/Private/ControlRigAnimGraphDetails.cpp
@@ -132,7 +132,7 @@ void FControlRigAnimNodeEventNameDetails::UpdateEntryNameList()
 	{
 		if(const UClass* Class = AnimNodeBeingCustomized->GetControlRigClass())
 		{
-			if(const UBaseControlRig* CDO = Cast<UBaseControlRig>(Class->GetDefaultObject(true)))
+			if(const UControlRig* CDO = Cast<UControlRig>(Class->GetDefaultObject(true)))
 			{
 				Algo::Transform(CDO->GetSupportedEvents(), EntryNameList,[](const FName& InEntryName)
 				{
diff --git a/Engine/Plugins/Animation/ControlRig/Source/ControlRigEditor/Private/ControlRigBlueprintActions.cpp b/Engine/Plugins/Animation/ControlRig/Source/ControlRigEditor/Private/ControlRigBlueprintActions.cpp
--- a/Engine/Plugins/Animation/ControlRig/Source/ControlRigEditor/Private/ControlRigBlueprintActions.cpp
+++ b/Engine/Plugins/Animation/ControlRig/Source/ControlRigEditor/Private/ControlRigBlueprintActions.cpp
@@ -3,14 +3,14 @@
 #include "ControlRigBlueprintActions.h"
 #include "ControlRigBlueprintFactory.h"
 #include "ControlRigBlueprint.h"
-#include "ControlRig.h"
+#include "BaseControlRig.h"
 #include "Editor/RigVMEditorStyle.h"
 #include "IControlRigEditorModule.h"
 
 #include "Styling/SlateIconFinder.h"
 #include "Widgets/Layout/SBorder.h"
 #include "Widgets/Images/SImage.h"
-#include "Styling/AppStyle.h"
+#include "Styling/AppStyle.h" 
 #include "Subsystems/AssetEditorSubsystem.h"
 
 #include "AssetRegistry/AssetRegistryModule.h"
@@ -393,7 +393,7 @@ void FControlRigBlueprintActions::OnSpawnedSkeletalMeshActorChanged(UObject* InO
 					Track = MovieScene->AddTrack<UMovieSceneControlRigParameterTrack>(ActorTrackGuid);
 				}
 
-				UBaseControlRig* ControlRig = Track->GetControlRig();
+				UControlRig* ControlRig = Track->GetControlRig();
 
 				FString ObjectName = (ControlRigClass->GetName());
 
@@ -405,10 +405,10 @@ void FControlRigBlueprintActions::OnSpawnedSkeletalMeshActorChanged(UObject* InO
 					ObjectName.RemoveFromEnd(TEXT("_C"));
 
 					// This is either a UControlRig or a UModularRig
-					ControlRig = NewObject<UBaseControlRig>(Track, ControlRigClass, FName(*ObjectName), RF_Transactional);
+					ControlRig = NewObject<UControlRig>(Track, ControlRigClass, FName(*ObjectName), RF_Transactional);
 					ControlRig->SetObjectBinding(MakeShared<FControlRigObjectBinding>());
 					ControlRig->GetObjectBinding()->BindToObject(MeshActor->GetSkeletalMeshComponent());
-					ControlRig->GetDataSourceRegistry()->RegisterDataSource(UBaseControlRig::OwnerComponent, ControlRig->GetObjectBinding()->GetBoundObject());
+					ControlRig->GetDataSourceRegistry()->RegisterDataSource(UControlRig::OwnerComponent, ControlRig->GetObjectBinding()->GetBoundObject());
 					ControlRig->Initialize();
 					ControlRig->Evaluate_AnyThread();
 					ControlRig->CreateRigControlsForCurveContainer();
diff --git a/Engine/Plugins/Animation/ControlRig/Source/ControlRigEditor/Private/ControlRigBlueprintFactory.cpp b/Engine/Plugins/Animation/ControlRig/Source/ControlRigEditor/Private/ControlRigBlueprintFactory.cpp
--- a/Engine/Plugins/Animation/ControlRig/Source/ControlRigEditor/Private/ControlRigBlueprintFactory.cpp
+++ b/Engine/Plugins/Animation/ControlRig/Source/ControlRigEditor/Private/ControlRigBlueprintFactory.cpp
@@ -137,7 +137,7 @@ private:
 			// If it appears on the allowed child-of classes list (or there is nothing on that list)
 			if (InClass)
 			{
-				if (InClass == UBaseControlRig::StaticClass())
+				if (InClass == UControlRig::StaticClass())
 				{
 					return false;
 				}
@@ -189,8 +189,8 @@ private:
 		TSharedPtr<FControlRigBlueprintParentFilter> Filter = MakeShareable(new FControlRigBlueprintParentFilter());
 		Options.ClassFilters.Add(Filter.ToSharedRef());
 
-		// All child child classes of UBaseControlRig are valid.
-		Filter->AllowedChildrenOfClasses.Add(UBaseControlRig::StaticClass());
+		// All child child classes of UControlRig are valid.
+		Filter->AllowedChildrenOfClasses.Add(UControlRig::StaticClass());
 
 		ParentClassContainer->ClearChildren();
 		ParentClassContainer->AddSlot()
@@ -292,7 +292,7 @@ UObject* UControlRigBlueprintFactory::FactoryCreateNew(UClass* Class, UObject* I
 	// Make sure we are trying to factory a Control Rig Blueprint, then create and init one
 	check(Class->IsChildOf(UControlRigBlueprint::StaticClass()));
 
-	if ((ParentClass == nullptr) || !FKismetEditorUtilities::CanCreateBlueprintOfClass(ParentClass) || !ParentClass->IsChildOf(UBaseControlRig::StaticClass()))
+	if ((ParentClass == nullptr) || !FKismetEditorUtilities::CanCreateBlueprintOfClass(ParentClass) || !ParentClass->IsChildOf(UControlRig::StaticClass()))
 	{
 		FFormatNamedArguments Args;
 		Args.Add( TEXT("ClassName"), (ParentClass != nullptr) ? FText::FromString( ParentClass->GetName() ) : LOCTEXT("Null", "(null)") );
diff --git a/Engine/Plugins/Animation/ControlRig/Source/ControlRigEditor/Private/ControlRigBlueprintFactory.h b/Engine/Plugins/Animation/ControlRig/Source/ControlRigEditor/Private/ControlRigBlueprintFactory.h
--- a/Engine/Plugins/Animation/ControlRig/Source/ControlRigEditor/Private/ControlRigBlueprintFactory.h
+++ b/Engine/Plugins/Animation/ControlRig/Source/ControlRigEditor/Private/ControlRigBlueprintFactory.h
@@ -8,7 +8,7 @@
 #include "Engine/Blueprint.h"
 #include "Factories/Factory.h"
 #include "ControlRigBlueprint.h"
-#include "ControlRig.h"
+#include "BaseControlRig.h"
 #include "ControlRigBlueprintFactory.generated.h"
 
 UCLASS(HideCategories=Object)
@@ -21,7 +21,7 @@ public:
 
 	// The parent class of the created blueprint
 	UPROPERTY(EditAnywhere, Category="Control Rig Factory", meta=(AllowAbstract = ""))
-	TSubclassOf<UBaseControlRig> ParentClass;
+	TSubclassOf<UControlRig> ParentClass;
 
 	// UFactory Interface
 	virtual bool ConfigureProperties() override;
diff --git a/Engine/Plugins/Animation/ControlRig/Source/ControlRigEditor/Private/ControlRigCompilerDetails.cpp b/Engine/Plugins/Animation/ControlRig/Source/ControlRigEditor/Private/ControlRigCompilerDetails.cpp
--- a/Engine/Plugins/Animation/ControlRig/Source/ControlRigEditor/Private/ControlRigCompilerDetails.cpp
+++ b/Engine/Plugins/Animation/ControlRig/Source/ControlRigEditor/Private/ControlRigCompilerDetails.cpp
@@ -144,7 +144,7 @@ void FRigVMCompileSettingsDetails::CustomizeChildren(TSharedRef<IPropertyHandle>
 					.ContentPadding(FMargin(2))
 					.Visibility_Lambda([]()
 					{
-						return UBaseControlRig::AreNativizedVMsDisabled() ? EVisibility::Collapsed : EVisibility::Visible;
+						return UControlRig::AreNativizedVMsDisabled() ? EVisibility::Collapsed : EVisibility::Visible;
 					})
 					.Content()
 					[
@@ -190,7 +190,7 @@ FReply FRigVMCompileSettingsDetails::OnInspectMemory(ERigVMMemoryType InMemoryTy
 {
 	if (BlueprintBeingCustomized)
 	{
-		if(UBaseControlRig* DebuggedRig = Cast<UBaseControlRig>(BlueprintBeingCustomized->GetObjectBeingDebugged()))
+		if(UControlRig* DebuggedRig = Cast<UControlRig>(BlueprintBeingCustomized->GetObjectBeingDebugged()))
 		{
 			if(FRigVMMemoryStorageStruct* MemoryStorage = DebuggedRig->GetMemoryByType(InMemoryType))
 			{
@@ -221,7 +221,7 @@ FReply FRigVMCompileSettingsDetails::OnCopyByteCodeClicked()
 	{
 		if (BlueprintBeingCustomized->GetDefaultModel())
 		{
-			if(UBaseControlRig* ControlRig = Cast<UBaseControlRig>(BlueprintBeingCustomized->GetObjectBeingDebugged()))
+			if(UControlRig* ControlRig = Cast<UControlRig>(BlueprintBeingCustomized->GetObjectBeingDebugged()))
 			{
 				FString ByteCodeContent = ControlRig->GetVM()->DumpByteCodeAsText(ControlRig->GetRigVMExtendedExecuteContext());
 				FPlatformApplicationMisc::ClipboardCopy(*ByteCodeContent);
@@ -235,7 +235,7 @@ FReply FRigVMCompileSettingsDetails::OnCopyHierarchyGraphClicked()
 {
 	if (BlueprintBeingCustomized)
 	{
-		if(UBaseControlRig* ControlRig = Cast<UBaseControlRig>(BlueprintBeingCustomized->GetObjectBeingDebugged()))
+		if(UControlRig* ControlRig = Cast<UControlRig>(BlueprintBeingCustomized->GetObjectBeingDebugged()))
 		{
 			FName EventName = FRigUnit_BeginExecution::EventName;
 			if(!ControlRig->GetEventQueue().IsEmpty())
@@ -257,7 +257,7 @@ FReply FRigVMCompileSettingsDetails::OnCopyGeneratedCodeClicked()
 		const FString ClassName = FString::Printf(TEXT("%sVM"), *BlueprintBeingCustomized->GetName());
 		if (BlueprintBeingCustomized->GeneratedClass)
 		{
-			if (UBaseControlRig* CDO = Cast<UBaseControlRig>(BlueprintBeingCustomized->GeneratedClass->GetDefaultObject()))
+			if (UControlRig* CDO = Cast<UControlRig>(BlueprintBeingCustomized->GeneratedClass->GetDefaultObject()))
 			{
 				if(CDO->GetVM())
 				{
@@ -287,7 +287,7 @@ FReply FRigVMCompileSettingsDetails::OnCopyVMExecutionClicked()
 {
 	if (BlueprintBeingCustomized)
 	{
-		if (UBaseControlRig* ControlRig = Cast<UBaseControlRig>(BlueprintBeingCustomized->GetObjectBeingDebugged()))
+		if (UControlRig* ControlRig = Cast<UControlRig>(BlueprintBeingCustomized->GetObjectBeingDebugged()))
 		{
 			FString DebugString = ControlRig->GetDebugExecutionString();
 			FPlatformApplicationMisc::ClipboardCopy(*DebugString);
diff --git a/Engine/Plugins/Animation/ControlRig/Source/ControlRigEditor/Private/ControlRigEditorModule.cpp b/Engine/Plugins/Animation/ControlRig/Source/ControlRigEditor/Private/ControlRigEditorModule.cpp
--- a/Engine/Plugins/Animation/ControlRig/Source/ControlRigEditor/Private/ControlRigEditorModule.cpp
+++ b/Engine/Plugins/Animation/ControlRig/Source/ControlRigEditor/Private/ControlRigEditorModule.cpp
@@ -121,7 +121,7 @@ void FControlRigEditorModule::StartupModule()
 	ClassesToUnregisterOnShutdown.Add(UControlRigBlueprint::StaticClass()->GetFName());
 	PropertyEditorModule.RegisterCustomClassLayout(ClassesToUnregisterOnShutdown.Last(), FOnGetDetailCustomizationInstance::CreateStatic(&FControlRigBlueprintDetails::MakeInstance));
 
-	ClassesToUnregisterOnShutdown.Add(UBaseControlRig::StaticClass()->GetFName());
+	ClassesToUnregisterOnShutdown.Add(UControlRig::StaticClass()->GetFName());
 
 	// same as ClassesToUnregisterOnShutdown but for properties, there is none right now
 	PropertiesToUnregisterOnShutdown.Reset();
@@ -632,7 +632,7 @@ void FControlRigEditorModule::GetDirectManipulationMenuActions(IRigVMClientHost*
     // Add direct manipulation context menu entries
 	if(UControlRigBlueprint* ControlRigBlueprint = Cast<UControlRigBlueprint>(RigVMClientHost))
 	{
-		UBaseControlRig* DebuggedRig = Cast<UBaseControlRig>(ControlRigBlueprint->GetObjectBeingDebugged());
+		UControlRig* DebuggedRig = Cast<UControlRig>(ControlRigBlueprint->GetObjectBeingDebugged());
 		if(DebuggedRig == nullptr)
 		{
 			return;
@@ -657,7 +657,7 @@ void FControlRigEditorModule::GetDirectManipulationMenuActions(IRigVMClientHost*
 				return;
 			}
 
-			const FRigUnit* UnitInstance = UBaseControlRig::GetRigUnitInstanceFromScope(NodeInstance);
+			const FRigUnit* UnitInstance = UControlRig::GetRigUnitInstanceFromScope(NodeInstance);
 			TArray<FRigDirectManipulationTarget> Targets;
 			if(UnitInstance->GetDirectManipulationTargets(UnitNode, NodeInstance, DebuggedRig->GetHierarchy(), Targets, nullptr))
 			{
@@ -687,7 +687,7 @@ void FControlRigEditorModule::GetDirectManipulationMenuActions(IRigVMClientHost*
 						
 						auto HasNoUnconstrainedAffectedPin = [NodeInstance, UnitNode, Target]() -> bool
 						{
-							const FRigUnit* UnitInstance = UBaseControlRig::GetRigUnitInstanceFromScope(NodeInstance);
+							const FRigUnit* UnitInstance = UControlRig::GetRigUnitInstanceFromScope(NodeInstance);
 							const TArray<const URigVMPin*> AffectedPins = UnitInstance->GetPinsForDirectManipulation(UnitNode, Target);
 
 							int32 NumAffectedPinsWithRootLinks = 0;
@@ -1056,10 +1056,10 @@ void FControlRigEditorModule::BakeToControlRig(UClass* ControlRigClass, UAnimSeq
 				FString ObjectName = (ControlRigClass->GetName());
 				ObjectName.RemoveFromEnd(TEXT("_C"));
 
-				UBaseControlRig* ControlRig = NewObject<UBaseControlRig>(Track, ControlRigClass, FName(*ObjectName), RF_Transactional);
+				UControlRig* ControlRig = NewObject<UControlRig>(Track, ControlRigClass, FName(*ObjectName), RF_Transactional);
 				ControlRig->SetObjectBinding(MakeShared<FControlRigObjectBinding>());
 				ControlRig->GetObjectBinding()->BindToObject(MeshActor);
-				ControlRig->GetDataSourceRegistry()->RegisterDataSource(UBaseControlRig::OwnerComponent, ControlRig->GetObjectBinding()->GetBoundObject());
+				ControlRig->GetDataSourceRegistry()->RegisterDataSource(UControlRig::OwnerComponent, ControlRig->GetObjectBinding()->GetBoundObject());
 				ControlRig->Initialize();
 				ControlRig->Evaluate_AnyThread();
 
@@ -1421,7 +1421,7 @@ bool FControlRigClassFilter::IsClassAllowed(const FClassViewerInitializationOpti
 {
 	if(InClass)
 	{
-		const bool bChildOfObjectClass = InClass->IsChildOf(UBaseControlRig::StaticClass());
+		const bool bChildOfObjectClass = InClass->IsChildOf(UControlRig::StaticClass());
 		const bool bMatchesFlags = !InClass->HasAnyClassFlags(CLASS_Hidden | CLASS_HideDropDown | CLASS_Deprecated | CLASS_Abstract);
 		const bool bNotNative = !InClass->IsNative();
 
@@ -1442,7 +1442,7 @@ bool FControlRigClassFilter::IsClassAllowed(const FClassViewerInitializationOpti
 
 bool FControlRigClassFilter::IsUnloadedClassAllowed(const FClassViewerInitializationOptions& InInitOptions, const TSharedRef< const IUnloadedBlueprintData > InUnloadedClassData, TSharedRef< FClassViewerFilterFuncs > InFilterFuncs)
 {
-	const bool bChildOfObjectClass = InUnloadedClassData->IsChildOf(UBaseControlRig::StaticClass());
+	const bool bChildOfObjectClass = InUnloadedClassData->IsChildOf(UControlRig::StaticClass());
 	const bool bMatchesFlags = !InUnloadedClassData->HasAnyClassFlags(CLASS_Hidden | CLASS_HideDropDown | CLASS_Deprecated | CLASS_Abstract);
 	if (bChildOfObjectClass && bMatchesFlags)
 	{
diff --git a/Engine/Plugins/Animation/ControlRig/Source/ControlRigEditor/Private/ControlRigElementDetails.cpp b/Engine/Plugins/Animation/ControlRig/Source/ControlRigEditor/Private/ControlRigElementDetails.cpp
--- a/Engine/Plugins/Animation/ControlRig/Source/ControlRigEditor/Private/ControlRigElementDetails.cpp
+++ b/Engine/Plugins/Animation/ControlRig/Source/ControlRigEditor/Private/ControlRigElementDetails.cpp
@@ -91,7 +91,7 @@ void RigElementKeyDetails_GetCustomizedInfo(TSharedRef<IPropertyHandle> InStruct
 			break;
 		}
 
-		if(const UBaseControlRig* ControlRig = Object->GetTypedOuter<UBaseControlRig>())
+		if(const UControlRig* ControlRig = Object->GetTypedOuter<UControlRig>())
 		{
 			OutBlueprint = Cast<UControlRigBlueprint>(ControlRig->GetClass()->ClassGeneratedBy);
 			if(OutBlueprint)
@@ -116,7 +116,7 @@ void RigElementKeyDetails_GetCustomizedInfo(TSharedRef<IPropertyHandle> InStruct
 			Package->GetDefaultSubobjects(SubObjects);
 			for (UObject* SubObject : SubObjects)
 			{
-				if (UBaseControlRig* Rig = Cast<UBaseControlRig>(SubObject))
+				if (UControlRig* Rig = Cast<UControlRig>(SubObject))
 				{
 					UControlRigBlueprint* Blueprint = Cast<UControlRigBlueprint>(Rig->GetClass()->ClassGeneratedBy);
 					if (Blueprint)
@@ -148,7 +148,7 @@ UControlRigBlueprint* RigElementDetails_GetBlueprintFromHierarchy(URigHierarchy*
 	UControlRigBlueprint* Blueprint = InHierarchy->GetTypedOuter<UControlRigBlueprint>();
 	if(Blueprint == nullptr)
 	{
-		UBaseControlRig* Rig = InHierarchy->GetTypedOuter<UBaseControlRig>();
+		UControlRig* Rig = InHierarchy->GetTypedOuter<UControlRig>();
 		if(Rig)
 		{
 			Blueprint = Cast<UControlRigBlueprint>(Rig->GetClass()->ClassGeneratedBy);
@@ -965,7 +965,7 @@ bool FRigBaseElementDetails::IsConstructionModeEnabled() const
 	
 	if(const UControlRigBlueprint* Blueprint = PerElementInfos[0].GetBlueprint())
 	{
-		if (const UBaseControlRig* DebuggedRig = Cast<UBaseControlRig>(Blueprint ->GetObjectBeingDebugged()))
+		if (const UControlRig* DebuggedRig = Cast<UControlRig>(Blueprint ->GetObjectBeingDebugged()))
 		{
 			return DebuggedRig->IsConstructionModeEnabled();
 		}
@@ -2095,7 +2095,7 @@ FDetailWidgetRow& FRigTransformElementDetails::CreateEulerTransformValueWidgetRo
 		const bool bInitial = CurrentTransformType == ERigTransformElementDetailsTransform::Initial; 
 
 		bool bConstructionModeEnabled = false;
-		if (UBaseControlRig* DebuggedRig = Cast<UBaseControlRig>(PerElementInfos[0].GetBlueprint()->GetObjectBeingDebugged()))
+		if (UControlRig* DebuggedRig = Cast<UControlRig>(PerElementInfos[0].GetBlueprint()->GetObjectBeingDebugged()))
 		{
 			bConstructionModeEnabled = DebuggedRig->IsConstructionModeEnabled();
 		}
@@ -2334,7 +2334,7 @@ FDetailWidgetRow& FRigTransformElementDetails::CreateEulerTransformValueWidgetRo
 			ESlateRotationRepresentation::Type Representation,
 			ESlateTransformSubComponent::Type SubComponent)
 		{
-			if (UBaseControlRig* DebuggedRig = Cast<UBaseControlRig>(PerElementInfos[0].GetBlueprint()->GetObjectBeingDebugged()))
+			if (UControlRig* DebuggedRig = Cast<UControlRig>(PerElementInfos[0].GetBlueprint()->GetObjectBeingDebugged()))
 			{
 				EControlRigInteractionType Type = EControlRigInteractionType::None;
 				switch (Component)
@@ -2364,7 +2364,7 @@ FDetailWidgetRow& FRigTransformElementDetails::CreateEulerTransformValueWidgetRo
 			ESlateTransformSubComponent::Type SubComponent,
 			FVector::FReal InNumericValue)
 		{
-			if (UBaseControlRig* DebuggedRig = Cast<UBaseControlRig>(PerElementInfos[0].GetBlueprint()->GetObjectBeingDebugged()))
+			if (UControlRig* DebuggedRig = Cast<UControlRig>(PerElementInfos[0].GetBlueprint()->GetObjectBeingDebugged()))
 			{
 				DebuggedRig->InteractionType = (uint8)EControlRigInteractionType::None;
 				DebuggedRig->ElementsBeingInteracted.Reset();
@@ -3134,7 +3134,7 @@ void FRigControlElementDetails::CustomizeControl(IDetailLayoutBuilder& DetailBui
 						Info.GetDefaultHierarchy()->SetControlValue(ControlBeingCustomized, InitialValue, ERigControlValueType::Initial, false, false, true);
 						Info.GetDefaultHierarchy()->SetControlValue(ControlBeingCustomized, CurrentValue, ERigControlValueType::Current, false, false, true);
 
-						if (UBaseControlRig* DebuggedRig = Cast<UBaseControlRig>(Info.GetBlueprint()->GetObjectBeingDebugged()))
+						if (UControlRig* DebuggedRig = Cast<UControlRig>(Info.GetBlueprint()->GetObjectBeingDebugged()))
 						{
 							URigHierarchy* DebuggedHierarchy = DebuggedRig->GetHierarchy();
 							if(FRigControlElement* DebuggedControlElement = DebuggedHierarchy->Find<FRigControlElement>(ControlBeingCustomized->GetKey()))
diff --git a/Engine/Plugins/Animation/ControlRig/Source/ControlRigEditor/Private/ControlRigElementDetails.h b/Engine/Plugins/Animation/ControlRig/Source/ControlRigEditor/Private/ControlRigElementDetails.h
--- a/Engine/Plugins/Animation/ControlRig/Source/ControlRigEditor/Private/ControlRigElementDetails.h
+++ b/Engine/Plugins/Animation/ControlRig/Source/ControlRigEditor/Private/ControlRigElementDetails.h
@@ -196,7 +196,7 @@ public:
 
 		UControlRigBlueprint* GetBlueprint() const
 		{
-			if(const UBaseControlRig* ControlRig = GetHierarchy()->GetTypedOuter<UBaseControlRig>())
+			if(const UControlRig* ControlRig = GetHierarchy()->GetTypedOuter<UControlRig>())
 			{
 				return Cast<UControlRigBlueprint>(ControlRig->GetClass()->ClassGeneratedBy);
 			}
diff --git a/Engine/Plugins/Animation/ControlRig/Source/ControlRigEditor/Private/ControlRigInfluenceMapDetails.cpp b/Engine/Plugins/Animation/ControlRig/Source/ControlRigEditor/Private/ControlRigInfluenceMapDetails.cpp
--- a/Engine/Plugins/Animation/ControlRig/Source/ControlRigEditor/Private/ControlRigInfluenceMapDetails.cpp
+++ b/Engine/Plugins/Animation/ControlRig/Source/ControlRigEditor/Private/ControlRigInfluenceMapDetails.cpp
@@ -40,7 +40,7 @@ UControlRigBlueprint* FRigInfluenceMapPerEventDetails::GetBlueprintFromDetailBui
 
 			for (UObject* SubObject : SubObjects)
 			{
-				if (UBaseControlRig* Rig = Cast<UBaseControlRig>(SubObject))
+				if (UControlRig* Rig = Cast<UControlRig>(SubObject))
 				{
 					if(UControlRigBlueprint* Blueprint = Cast<UControlRigBlueprint>(Rig->GetClass()->ClassGeneratedBy))
 					{
diff --git a/Engine/Plugins/Animation/ControlRig/Source/ControlRigEditor/Private/ControlRigSequencerEditorLibrary.cpp b/Engine/Plugins/Animation/ControlRig/Source/ControlRigEditor/Private/ControlRigSequencerEditorLibrary.cpp
--- a/Engine/Plugins/Animation/ControlRig/Source/ControlRigEditor/Private/ControlRigSequencerEditorLibrary.cpp
+++ b/Engine/Plugins/Animation/ControlRig/Source/ControlRigEditor/Private/ControlRigSequencerEditorLibrary.cpp
@@ -17,7 +17,7 @@
 #include "MovieSceneToolHelpers.h"
 #include "Rigs/FKControlRig.h"
 #include "Units/Execution/RigUnit_InverseExecution.h"
-#include "ControlRig.h"
+#include "BaseControlRig.h"
 #include "EditMode/ControlRigEditMode.h"
 #include "EditorModeManager.h"
 #include "Engine/Selection.h"
@@ -57,18 +57,10 @@
 TArray<UControlRig*> UControlRigSequencerEditorLibrary::GetVisibleControlRigs()
 {
 	TArray<UControlRig*> ControlRigs;
-	TArray<UBaseControlRig*> BaseControlRigs;
 	FControlRigEditMode* ControlRigEditMode = static_cast<FControlRigEditMode*>(GLevelEditorModeTools().GetActiveMode(FControlRigEditMode::ModeName));
 	if (ControlRigEditMode)
 	{
-		BaseControlRigs = ControlRigEditMode->GetControlRigsArray(true /*bIsVisible*/);
-	}
-	for (UBaseControlRig* BaseRig : BaseControlRigs)
-	{
-		if (UControlRig* ControlRig = Cast<UControlRig>(BaseRig))
-		{
-			ControlRigs.Add(ControlRig);
-		}
+		ControlRigs = ControlRigEditMode->GetControlRigsArray(true /*bIsVisible*/);
 	}
 	return ControlRigs;
 }
@@ -250,7 +242,7 @@ static UMovieSceneControlRigParameterTrack* AddControlRig(ULevelSequence* LevelS
 		}
 
 		ControlRig->Modify();
-		if (UFKControlRig* FKControlRig = Cast<UFKControlRig>(Cast<UBaseControlRig>(ControlRig)))
+		if (UFKControlRig* FKControlRig = Cast<UFKControlRig>(Cast<UControlRig>(ControlRig)))
 		{
 			if (bIsAdditiveControlRig)
 			{
@@ -445,7 +437,7 @@ bool UControlRigSequencerEditorLibrary::TweenControlRig(ULevelSequence* LevelSeq
 	{
 		FControlsToTween ControlsToTween;
 		LevelSequence->GetMovieScene()->Modify();
-		TArray<UBaseControlRig*> SelectedControlRigs;
+		TArray<UControlRig*> SelectedControlRigs;
 		SelectedControlRigs.Add(ControlRig);
 		ControlsToTween.Setup(SelectedControlRigs, WeakSequencer);
 		ControlsToTween.Blend(WeakSequencer, TweenValue);
@@ -3086,7 +3078,7 @@ bool UControlRigSequencerEditorLibrary::IsAdditiveControlRig(UControlRig* InCont
 EControlRigFKRigExecuteMode UControlRigSequencerEditorLibrary::GetFKControlRigApplyMode(UControlRig* InControlRig)
 {
 	EControlRigFKRigExecuteMode ApplyMode = EControlRigFKRigExecuteMode::Direct;
-	if (UFKControlRig* FKRig = Cast<UFKControlRig>(Cast<UBaseControlRig>(InControlRig)))
+	if (UFKControlRig* FKRig = Cast<UFKControlRig>(Cast<UControlRig>(InControlRig)))
 	{
 		ApplyMode = FKRig->GetApplyMode();
 	}
@@ -3095,7 +3087,7 @@ EControlRigFKRigExecuteMode UControlRigSequencerEditorLibrary::GetFKControlRigAp
 
 bool UControlRigSequencerEditorLibrary::SetControlRigApplyMode(UControlRig* InControlRig, EControlRigFKRigExecuteMode InApplyMode)
 {
-	if (UFKControlRig* FKRig = Cast<UFKControlRig>(Cast<UBaseControlRig>(InControlRig)))
+	if (UFKControlRig* FKRig = Cast<UFKControlRig>(Cast<UControlRig>(InControlRig)))
 	{
 		FKRig->SetApplyMode(InApplyMode);
 		return true;
diff --git a/Engine/Plugins/Animation/ControlRig/Source/ControlRigEditor/Private/ControlRigThumbnailRenderer.cpp b/Engine/Plugins/Animation/ControlRig/Source/ControlRigEditor/Private/ControlRigThumbnailRenderer.cpp
--- a/Engine/Plugins/Animation/ControlRig/Source/ControlRigEditor/Private/ControlRigThumbnailRenderer.cpp
+++ b/Engine/Plugins/Animation/ControlRig/Source/ControlRigEditor/Private/ControlRigThumbnailRenderer.cpp
@@ -100,16 +100,16 @@ void UControlRigThumbnailRenderer::AddAdditionalPreviewSceneContent(UObject* Obj
 			return;
 		}
 
-		UBaseControlRig* ControlRig = nullptr;
+		UControlRig* ControlRig = nullptr;
 
 		// reuse the current control rig if possible
-		UBaseControlRig* CDO = Cast<UBaseControlRig>(RigBlueprint->GeneratedClass->GetDefaultObject(true /* create if needed */));
+		UControlRig* CDO = Cast<UControlRig>(RigBlueprint->GeneratedClass->GetDefaultObject(true /* create if needed */));
 
 		TArray<UObject*> ArchetypeInstances;
 		CDO->GetArchetypeInstances(ArchetypeInstances);
 		for (UObject* ArchetypeInstance : ArchetypeInstances)
 		{
-			ControlRig = Cast<UBaseControlRig>(ArchetypeInstance);
+			ControlRig = Cast<UControlRig>(ArchetypeInstance);
 			break;
 		}
 
diff --git a/Engine/Plugins/Animation/ControlRig/Source/ControlRigEditor/Private/EditMode/ControlRigBaseDockableView.cpp b/Engine/Plugins/Animation/ControlRig/Source/ControlRigEditor/Private/EditMode/ControlRigBaseDockableView.cpp
--- a/Engine/Plugins/Animation/ControlRig/Source/ControlRigEditor/Private/EditMode/ControlRigBaseDockableView.cpp
+++ b/Engine/Plugins/Animation/ControlRig/Source/ControlRigEditor/Private/EditMode/ControlRigBaseDockableView.cpp
@@ -25,9 +25,9 @@ FControlRigBaseDockableView::~FControlRigBaseDockableView()
 	}
 }
 
-TArray<UBaseControlRig*> FControlRigBaseDockableView::GetControlRigs()  const
+TArray<UControlRig*> FControlRigBaseDockableView::GetControlRigs()  const
 {
-	TArray<UBaseControlRig*> ControlRigs;
+	TArray<UControlRig*> ControlRigs;
 	if (FControlRigEditMode* EditMode = static_cast<FControlRigEditMode*>(ModeTools->GetActiveMode(FControlRigEditMode::ModeName)))
 	{
 		ControlRigs = EditMode->GetControlRigsArray(false);
@@ -48,15 +48,15 @@ void FControlRigBaseDockableView::SetEditMode(FControlRigEditMode& InEditMode)
 	}
 }
 
-void FControlRigBaseDockableView::HandleControlAdded(UBaseControlRig* ControlRig, bool bIsAdded)
+void FControlRigBaseDockableView::HandleControlAdded(UControlRig* ControlRig, bool bIsAdded)
 {
 }
 
-void FControlRigBaseDockableView::HandleControlSelected(UBaseControlRig* ControlRig, FRigControlElement* InControl, bool bSelected)
+void FControlRigBaseDockableView::HandleControlSelected(UControlRig* ControlRig, FRigControlElement* InControl, bool bSelected)
 {
 }
 
-void FControlRigBaseDockableView::HandlElementSelected(UBaseControlRig* ControlRig, const FRigElementKey& Key, bool bSelected)
+void FControlRigBaseDockableView::HandlElementSelected(UControlRig* ControlRig, const FRigElementKey& Key, bool bSelected)
 {
 	if (ControlRig)
 	{
diff --git a/Engine/Plugins/Animation/ControlRig/Source/ControlRigEditor/Private/EditMode/ControlRigBaseDockableView.h b/Engine/Plugins/Animation/ControlRig/Source/ControlRigEditor/Private/EditMode/ControlRigBaseDockableView.h
--- a/Engine/Plugins/Animation/ControlRig/Source/ControlRigEditor/Private/EditMode/ControlRigBaseDockableView.h
+++ b/Engine/Plugins/Animation/ControlRig/Source/ControlRigEditor/Private/EditMode/ControlRigBaseDockableView.h
@@ -8,7 +8,7 @@
 #include "EditorModeManager.h"
 
 
-class UBaseControlRig;
+class UControlRig;
 class ISequencer;
 class FControlRigEditMode;
 struct FRigControlElement;
@@ -19,15 +19,15 @@ class FControlRigBaseDockableView
 public:
 	FControlRigBaseDockableView();
 	virtual ~FControlRigBaseDockableView();
-	TArray<UBaseControlRig*> GetControlRigs() const;
+	TArray<UControlRig*> GetControlRigs() const;
 
 	virtual void SetEditMode(FControlRigEditMode& InEditMode);
 
 protected:
-	virtual void HandleControlSelected(UBaseControlRig* Subject, FRigControlElement* InControl, bool bSelected);
-	virtual void HandleControlAdded(UBaseControlRig* ControlRig, bool bIsAdded);
+	virtual void HandleControlSelected(UControlRig* Subject, FRigControlElement* InControl, bool bSelected);
+	virtual void HandleControlAdded(UControlRig* ControlRig, bool bIsAdded);
 
-	void HandlElementSelected(UBaseControlRig* Subject, const FRigElementKey& Key, bool bSelected);
+	void HandlElementSelected(UControlRig* Subject, const FRigElementKey& Key, bool bSelected);
 
 	ISequencer* GetSequencer() const;
 
diff --git a/Engine/Plugins/Animation/ControlRig/Source/ControlRigEditor/Private/EditMode/ControlRigControlsProxy.cpp b/Engine/Plugins/Animation/ControlRig/Source/ControlRigEditor/Private/EditMode/ControlRigControlsProxy.cpp
--- a/Engine/Plugins/Animation/ControlRig/Source/ControlRigEditor/Private/EditMode/ControlRigControlsProxy.cpp
+++ b/Engine/Plugins/Animation/ControlRig/Source/ControlRigEditor/Private/EditMode/ControlRigControlsProxy.cpp
@@ -820,7 +820,7 @@ void UControlRigBoolControlProxy::SetKey(const IPropertyHandle& KeyedPropertyHan
 
 //////UControlDetailPanelControlProxies////////
 
-UControlRigControlsProxy* UControlRigDetailPanelControlProxies::FindProxy(UBaseControlRig* ControlRig, const FName& Name) const
+UControlRigControlsProxy* UControlRigDetailPanelControlProxies::FindProxy(UControlRig* ControlRig, const FName& Name) const
 {
 	const FControlToProxyMap* ControlRigProxies = AllProxies.Find(ControlRig);
 	if (ControlRigProxies)
@@ -834,7 +834,7 @@ UControlRigControlsProxy* UControlRigDetailPanelControlProxies::FindProxy(UBaseC
 	return nullptr;
 }
 
-void UControlRigDetailPanelControlProxies::AddProxy(UBaseControlRig* ControlRig, const FName& Name,  FRigControlElement* ControlElement)
+void UControlRigDetailPanelControlProxies::AddProxy(UControlRig* ControlRig, const FName& Name,  FRigControlElement* ControlElement)
 {
 	UControlRigControlsProxy* Proxy = FindProxy(ControlRig,Name);
 	if (!Proxy && ControlElement != nullptr)
@@ -929,7 +929,7 @@ void UControlRigDetailPanelControlProxies::AddProxy(UBaseControlRig* ControlRig,
 	}
 }
 
-void UControlRigDetailPanelControlProxies::RemoveProxy(UBaseControlRig* ControlRig, const FName& Name)
+void UControlRigDetailPanelControlProxies::RemoveProxy(UControlRig* ControlRig, const FName& Name)
 {
 	UControlRigControlsProxy* ExistingProxy = FindProxy(ControlRig,Name);
 	if (ExistingProxy)
@@ -944,12 +944,12 @@ void UControlRigDetailPanelControlProxies::RemoveProxy(UBaseControlRig* ControlR
 	}
 }
 
-void UControlRigDetailPanelControlProxies::RemoveAllProxies(UBaseControlRig* ControlRig)
+void UControlRigDetailPanelControlProxies::RemoveAllProxies(UControlRig* ControlRig)
 {
 	//no control rig remove all
 	if (ControlRig == nullptr)
 	{
-		for (TPair<TObjectPtr<UBaseControlRig>, FControlToProxyMap>& ControlRigProxies : AllProxies)
+		for (TPair<TObjectPtr<UControlRig>, FControlToProxyMap>& ControlRigProxies : AllProxies)
 		{
 			for (TPair<FName, TObjectPtr<UControlRigControlsProxy> >& Pair : ControlRigProxies.Value.ControlToProxy)
 			{
@@ -985,7 +985,7 @@ void UControlRigDetailPanelControlProxies::RemoveAllProxies(UBaseControlRig* Con
 	}
 }
 
-void UControlRigDetailPanelControlProxies::RecreateAllProxies(UBaseControlRig* ControlRig)
+void UControlRigDetailPanelControlProxies::RecreateAllProxies(UControlRig* ControlRig)
 {
 	RemoveAllProxies(ControlRig);
 	TArray<FRigControlElement*> Controls = ControlRig->AvailableControls();
@@ -998,7 +998,7 @@ void UControlRigDetailPanelControlProxies::RecreateAllProxies(UBaseControlRig* C
 	}
 }
 
-void UControlRigDetailPanelControlProxies::ProxyChanged(UBaseControlRig* ControlRig, const FName& Name, bool bModify )
+void UControlRigDetailPanelControlProxies::ProxyChanged(UControlRig* ControlRig, const FName& Name, bool bModify )
 {
 	if (IsInGameThread())
 	{
@@ -1015,7 +1015,7 @@ void UControlRigDetailPanelControlProxies::ProxyChanged(UBaseControlRig* Control
 	}
 }
 
-void UControlRigDetailPanelControlProxies::SelectProxy(UBaseControlRig* ControlRig,const FName& Name, bool bSelected)
+void UControlRigDetailPanelControlProxies::SelectProxy(UControlRig* ControlRig,const FName& Name, bool bSelected)
 {
 	UControlRigControlsProxy* Proxy = FindProxy(ControlRig,Name);
 	if (Proxy)
@@ -1036,7 +1036,7 @@ void UControlRigDetailPanelControlProxies::SelectProxy(UBaseControlRig* ControlR
 	}
 }
 
-bool UControlRigDetailPanelControlProxies::IsSelected(UBaseControlRig* InControlRig, const FName& Name) const
+bool UControlRigDetailPanelControlProxies::IsSelected(UControlRig* InControlRig, const FName& Name) const
 {
 	if (UControlRigControlsProxy* Proxy = FindProxy(InControlRig, Name))
 	{
diff --git a/Engine/Plugins/Animation/ControlRig/Source/ControlRigEditor/Private/EditMode/ControlRigControlsProxy.h b/Engine/Plugins/Animation/ControlRig/Source/ControlRigEditor/Private/EditMode/ControlRigControlsProxy.h
--- a/Engine/Plugins/Animation/ControlRig/Source/ControlRigEditor/Private/EditMode/ControlRigControlsProxy.h
+++ b/Engine/Plugins/Animation/ControlRig/Source/ControlRigEditor/Private/EditMode/ControlRigControlsProxy.h
@@ -14,7 +14,6 @@
 #include "ControlRigControlsProxy.generated.h"
 
 struct FRigControlElement;
-class UBaseControlRig;
 class UControlRig;
 class IPropertyHandle;
 class FControlRigInteractionScope;
@@ -42,7 +41,7 @@ public:
 public:
 
 	FRigControlElement* GetControlElement() const;
-	TWeakObjectPtr<UBaseControlRig> ControlRig;
+	TWeakObjectPtr<UControlRig> ControlRig;
 
 	UPROPERTY()
 	bool bSelected;
@@ -320,21 +319,21 @@ class UControlRigDetailPanelControlProxies :public UObject
 protected:
 
 	UPROPERTY()
-	TMap<TObjectPtr<UBaseControlRig>, FControlToProxyMap> AllProxies; //proxies themselves contain weakobjectptr to the controlrig
+	TMap<TObjectPtr<UControlRig>, FControlToProxyMap> AllProxies; //proxies themselves contain weakobjectptr to the controlrig
 
 	UPROPERTY()
 	TArray< TObjectPtr<UControlRigControlsProxy>> SelectedProxies;
 
 
 public:
-	void SelectProxy(UBaseControlRig* InControlRig, const FName& Name, bool bSelected);
-	UControlRigControlsProxy* FindProxy(UBaseControlRig* InControlRig, const FName& Name) const;
-	void AddProxy(UBaseControlRig* InControlRig, const FName& Name,  FRigControlElement* ControlElement);
-	void RemoveProxy(UBaseControlRig* InControlRig, const FName& Name );
-	void ProxyChanged(UBaseControlRig* InControlRig, const FName& Name, bool bModify = true);
-	void RemoveAllProxies(UBaseControlRig* InControlRig);
-	void RecreateAllProxies(UBaseControlRig* InControlRig);
+	void SelectProxy(UControlRig* InControlRig, const FName& Name, bool bSelected);
+	UControlRigControlsProxy* FindProxy(UControlRig* InControlRig, const FName& Name) const;
+	void AddProxy(UControlRig* InControlRig, const FName& Name,  FRigControlElement* ControlElement);
+	void RemoveProxy(UControlRig* InControlRig, const FName& Name );
+	void ProxyChanged(UControlRig* InControlRig, const FName& Name, bool bModify = true);
+	void RemoveAllProxies(UControlRig* InControlRig);
+	void RecreateAllProxies(UControlRig* InControlRig);
 	const TArray<UControlRigControlsProxy*>& GetSelectedProxies() const { return SelectedProxies;}
-	bool IsSelected(UBaseControlRig* InControlRig, const FName& Name) const;
+	bool IsSelected(UControlRig* InControlRig, const FName& Name) const;
 
 };
