Control Rig: Rename BaseControlRig to ControlRig

diff --git a/Engine/Plugins/Animation/ControlRig/Source/ControlRigEditor/Private/EditMode/ControlRigEditMode.cpp b/Engine/Plugins/Animation/ControlRig/Source/ControlRigEditor/Private/EditMode/ControlRigEditMode.cpp
--- a/Engine/Plugins/Animation/ControlRig/Source/ControlRigEditor/Private/EditMode/ControlRigEditMode.cpp
+++ b/Engine/Plugins/Animation/ControlRig/Source/ControlRigEditor/Private/EditMode/ControlRigEditMode.cpp
@@ -8,7 +8,7 @@
 #include "Toolkits/ToolkitManager.h"
 #include "EditMode/SControlRigEditModeTools.h"
 #include "Algo/Transform.h"
-#include "ControlRig.h"
+#include "BaseControlRig.h"
 #include "HitProxies.h"
 #include "EditMode/ControlRigEditModeSettings.h"
 #include "ISequencer.h"
@@ -186,7 +186,7 @@ FControlRigEditMode::~FControlRigEditMode()
 	OnControlRigAddedOrRemovedDelegate.Clear();
 	OnControlRigSelectedDelegate.Clear();
 
-	TArray<TWeakObjectPtr<UBaseControlRig>> PreviousRuntimeRigs = RuntimeControlRigs;
+	TArray<TWeakObjectPtr<UControlRig>> PreviousRuntimeRigs = RuntimeControlRigs;
 	for (int32 PreviousRuntimeRigIndex = 0; PreviousRuntimeRigIndex < PreviousRuntimeRigs.Num(); PreviousRuntimeRigIndex++)
 	{
 		if (PreviousRuntimeRigs[PreviousRuntimeRigIndex].IsValid())
@@ -208,7 +208,7 @@ bool FControlRigEditMode:: SetSequencer(TWeakPtr<ISequencer> InSequencer)
 	{
 		WeakSequencer = InSequencer;
 		DestroyShapesActors(nullptr);
-		TArray<TWeakObjectPtr<UBaseControlRig>> PreviousRuntimeRigs = RuntimeControlRigs;
+		TArray<TWeakObjectPtr<UControlRig>> PreviousRuntimeRigs = RuntimeControlRigs;
 		for (int32 PreviousRuntimeRigIndex = 0; PreviousRuntimeRigIndex < PreviousRuntimeRigs.Num(); PreviousRuntimeRigIndex++)
 		{
 			if (PreviousRuntimeRigs[PreviousRuntimeRigIndex].IsValid())
@@ -225,7 +225,7 @@ bool FControlRigEditMode:: SetSequencer(TWeakPtr<ISequencer> InSequencer)
 				TArray<FControlRigSequencerBindingProxy> Proxies = UControlRigSequencerEditorLibrary::GetControlRigs(LevelSequence);
 				for (FControlRigSequencerBindingProxy& Proxy : Proxies)
 				{
-					if (UBaseControlRig* ControlRig = Proxy.ControlRig.Get())
+					if (UControlRig* ControlRig = Proxy.ControlRig.Get())
 					{
 						AddControlRigInternal(ControlRig);
 					}
@@ -238,7 +238,7 @@ bool FControlRigEditMode:: SetSequencer(TWeakPtr<ISequencer> InSequencer)
 	return false;
 }
 
-bool FControlRigEditMode::AddControlRigObject(UBaseControlRig* ControlRig, TWeakPtr<ISequencer> InSequencer)
+bool FControlRigEditMode::AddControlRigObject(UControlRig* ControlRig, TWeakPtr<ISequencer> InSequencer)
 {
 	if (ControlRig)
 	{
@@ -258,9 +258,9 @@ bool FControlRigEditMode::AddControlRigObject(UBaseControlRig* ControlRig, TWeak
 	return false;
 }
 
-void FControlRigEditMode::SetObjects(UBaseControlRig* ControlRig,  UObject* BindingObject, TWeakPtr<ISequencer> InSequencer)
+void FControlRigEditMode::SetObjects(UControlRig* ControlRig,  UObject* BindingObject, TWeakPtr<ISequencer> InSequencer)
 {
-	TArray<TWeakObjectPtr<UBaseControlRig>> PreviousRuntimeRigs = RuntimeControlRigs;
+	TArray<TWeakObjectPtr<UControlRig>> PreviousRuntimeRigs = RuntimeControlRigs;
 	for (int32 PreviousRuntimeRigIndex = 0; PreviousRuntimeRigIndex < PreviousRuntimeRigs.Num(); PreviousRuntimeRigIndex++)
 	{
 		if (PreviousRuntimeRigs[PreviousRuntimeRigIndex].IsValid())
@@ -313,9 +313,9 @@ void FControlRigEditMode::SetUpDetailPanel()
 void FControlRigEditMode::SetObjects_Internal()
 {
 	bool bHasValidRuntimeControlRig = false;
-	for (TWeakObjectPtr<UBaseControlRig>& RuntimeRigPtr : RuntimeControlRigs)
+	for (TWeakObjectPtr<UControlRig>& RuntimeRigPtr : RuntimeControlRigs)
 	{
-		if (UBaseControlRig* RuntimeControlRig = RuntimeRigPtr.Get())
+		if (UControlRig* RuntimeControlRig = RuntimeRigPtr.Get())
 		{
 			RuntimeControlRig->ControlModified().RemoveAll(this);
 			RuntimeControlRig->GetHierarchy()->OnModified().RemoveAll(this);
@@ -435,9 +435,9 @@ void FControlRigEditMode::Exit()
 	ClearOutAnyActiveTools();
 	OnControlRigAddedOrRemovedDelegate.Clear();
 	OnControlRigSelectedDelegate.Clear();
-	for (TWeakObjectPtr<UBaseControlRig>& RuntimeRigPtr : RuntimeControlRigs)
+	for (TWeakObjectPtr<UControlRig>& RuntimeRigPtr : RuntimeControlRigs)
 	{
-		if (UBaseControlRig* ControlRig = RuntimeRigPtr.Get())
+		if (UControlRig* ControlRig = RuntimeRigPtr.Get())
 		{
 			ControlRig->ClearControlSelection();
 		}
@@ -451,7 +451,7 @@ void FControlRigEditMode::Exit()
 			GEditor->EndTransaction();
 		}
 
-		for (TPair<UBaseControlRig*,FControlRigInteractionScope*>& InteractionScope : InteractionScopes)
+		for (TPair<UControlRig*,FControlRigInteractionScope*>& InteractionScope : InteractionScopes)
 		{
 			if (InteractionScope.Value)
 			{
@@ -471,7 +471,7 @@ void FControlRigEditMode::Exit()
 	DestroyShapesActors(nullptr);
 
 
-	TArray<TWeakObjectPtr<UBaseControlRig>> PreviousRuntimeRigs = RuntimeControlRigs;
+	TArray<TWeakObjectPtr<UControlRig>> PreviousRuntimeRigs = RuntimeControlRigs;
 	for (int32 PreviousRuntimeRigIndex = 0; PreviousRuntimeRigIndex < PreviousRuntimeRigs.Num(); PreviousRuntimeRigIndex++)
 	{
 		if (PreviousRuntimeRigs[PreviousRuntimeRigIndex].IsValid())
@@ -540,9 +540,9 @@ void FControlRigEditMode::Tick(FEditorViewportClient* ViewportClient, float Delt
 	if(bShowControlsAsOverlay != Settings->bShowControlsAsOverlay)
 	{
 		bShowControlsAsOverlay = Settings->bShowControlsAsOverlay;
-		for (TWeakObjectPtr<UBaseControlRig>& RuntimeRigPtr : RuntimeControlRigs)
+		for (TWeakObjectPtr<UControlRig>& RuntimeRigPtr : RuntimeControlRigs)
 		{
-			if (UBaseControlRig* RuntimeControlRig = RuntimeRigPtr.Get())
+			if (UControlRig* RuntimeControlRig = RuntimeRigPtr.Get())
 			{
 				UpdateSelectabilityOnSkeletalMeshes(RuntimeControlRig, !bShowControlsAsOverlay);
 			}
@@ -554,9 +554,9 @@ void FControlRigEditMode::Tick(FEditorViewportClient* ViewportClient, float Delt
 	if (RecreateControlShapesRequired != ERecreateControlRigShape::RecreateNone && !(FSlateApplication::Get().HasAnyMouseCaptor() || GUnrealEd->IsUserInteracting()))
 	{
 		RecreateControlShapeActors();
-		for (TWeakObjectPtr<UBaseControlRig>& RuntimeRigPtr : RuntimeControlRigs)
+		for (TWeakObjectPtr<UControlRig>& RuntimeRigPtr : RuntimeControlRigs)
 		{
-			if (UBaseControlRig* ControlRig = RuntimeRigPtr.Get())
+			if (UControlRig* ControlRig = RuntimeRigPtr.Get())
 			{
 				TArray<FRigElementKey> SelectedRigElements = GetSelectedRigElements(ControlRig);
 				for (const FRigElementKey& SelectedKey : SelectedRigElements)
@@ -598,9 +598,9 @@ void FControlRigEditMode::Tick(FEditorViewportClient* ViewportClient, float Delt
 
 		if (!AreEditingControlRigDirectly() == false) //only do this check if not in level editor
 		{
-			for (TWeakObjectPtr<UBaseControlRig>& RuntimeRigPtr : RuntimeControlRigs)
+			for (TWeakObjectPtr<UControlRig>& RuntimeRigPtr : RuntimeControlRigs)
 			{
-				if (UBaseControlRig* ControlRig = RuntimeRigPtr.Get())
+				if (UControlRig* ControlRig = RuntimeRigPtr.Get())
 				{
 					TArray<FRigElementKey> SelectedRigElements = GetSelectedRigElements(ControlRig);
 					UE::Widget::EWidgetMode CurrentWidgetMode = ViewportClient->GetWidgetMode();
@@ -664,7 +664,7 @@ struct  HFKRigBoneProxy : public HHitProxy
 	DECLARE_HIT_PROXY()
 
 	FName BoneName;
-	UBaseControlRig* ControlRig;
+	UControlRig* ControlRig;
 
 	HFKRigBoneProxy()
 		: HHitProxy(HPP_Foreground)
@@ -672,7 +672,7 @@ struct  HFKRigBoneProxy : public HHitProxy
 		, ControlRig(nullptr)
 	{}
 
-	HFKRigBoneProxy(FName InBoneName, UBaseControlRig *InControlRig)
+	HFKRigBoneProxy(FName InBoneName, UControlRig *InControlRig)
 		: HHitProxy(HPP_Foreground)
 		, BoneName(InBoneName)
 		, ControlRig(InControlRig)
@@ -687,7 +687,7 @@ struct  HFKRigBoneProxy : public HHitProxy
 IMPLEMENT_HIT_PROXY(HFKRigBoneProxy, HHitProxy)
 
 
-TSet<FName> FControlRigEditMode::GetActiveControlsFromSequencer(UBaseControlRig* ControlRig)
+TSet<FName> FControlRigEditMode::GetActiveControlsFromSequencer(UControlRig* ControlRig)
 {
 	TSet<FName> ActiveControls;
 	if (WeakSequencer.IsValid() == false)
@@ -759,9 +759,9 @@ void FControlRigEditMode::Render(const FSceneView* View, FViewport* Viewport, FP
 	const UControlRigEditModeSettings* Settings = GetDefault<UControlRigEditModeSettings>();
 	const bool bIsInGameView = !AreEditingControlRigDirectly() ? (ViewportToGameView.Find(Viewport) && ViewportToGameView[Viewport]) : false;
 	bool bRender = !Settings->bHideControlShapes;
-	for (TWeakObjectPtr<UBaseControlRig>& ControlRigPtr : RuntimeControlRigs)
+	for (TWeakObjectPtr<UControlRig>& ControlRigPtr : RuntimeControlRigs)
 	{
-		UBaseControlRig* ControlRig = ControlRigPtr.Get();
+		UControlRig* ControlRig = ControlRigPtr.Get();
 		//actor game view drawing is handled by not drawing in game via SetActorHiddenInGame().
 		if (bRender && ControlRig && ControlRig->GetControlsVisible())
 		{
@@ -1014,7 +1014,7 @@ bool FControlRigEditMode::EndTracking(FEditorViewportClient* InViewportClient, F
 			GEditor->EndTransaction();
 		}
 
-		for (TPair<UBaseControlRig*, FControlRigInteractionScope*>& InteractionScope : InteractionScopes)
+		for (TPair<UControlRig*, FControlRigInteractionScope*>& InteractionScope : InteractionScopes)
 		{
 			if (InteractionScope.Value)
 			{
@@ -1066,9 +1066,9 @@ bool FControlRigEditMode::StartTracking(FEditorViewportClient* InViewportClient,
 		{
 			if (AreEditingControlRigDirectly())
 			{
-				for (TWeakObjectPtr<UBaseControlRig>& RuntimeRigPtr : RuntimeControlRigs)
+				for (TWeakObjectPtr<UControlRig>& RuntimeRigPtr : RuntimeControlRigs)
 				{
-					if (UBaseControlRig* ControlRig = RuntimeRigPtr.Get())
+					if (UControlRig* ControlRig = RuntimeRigPtr.Get())
 					{
 						TArray<FRigElementKey> SelectedRigElements = GetSelectedRigElements(ControlRig);
 						for (const FRigElementKey& Key : SelectedRigElements)
@@ -1087,9 +1087,9 @@ bool FControlRigEditMode::StartTracking(FEditorViewportClient* InViewportClient,
 
 		if (AreEditingControlRigDirectly())
 		{
-			for (TWeakObjectPtr<UBaseControlRig>& RuntimeRigPtr : RuntimeControlRigs)
+			for (TWeakObjectPtr<UControlRig>& RuntimeRigPtr : RuntimeControlRigs)
 			{
-				if (UBaseControlRig* ControlRig = RuntimeRigPtr.Get())
+				if (UControlRig* ControlRig = RuntimeRigPtr.Get())
 				{
 					UObject* Blueprint = ControlRig->GetClass()->ClassGeneratedBy;
 					if (Blueprint)
@@ -1115,9 +1115,9 @@ bool FControlRigEditMode::StartTracking(FEditorViewportClient* InViewportClient,
 
 	if (!AreEditingControlRigDirectly())
 	{
-		for (TWeakObjectPtr<UBaseControlRig>& RuntimeRigPtr : RuntimeControlRigs)
+		for (TWeakObjectPtr<UControlRig>& RuntimeRigPtr : RuntimeControlRigs)
 		{
-			if (UBaseControlRig* ControlRig = RuntimeRigPtr.Get())
+			if (UControlRig* ControlRig = RuntimeRigPtr.Get())
 			{
 				if (AreRigElementSelectedAndMovable(ControlRig))
 				{
@@ -1135,7 +1135,7 @@ bool FControlRigEditMode::StartTracking(FEditorViewportClient* InViewportClient,
 	}
 	else if(UControlRigEditorSettings::Get()->bEnableUndoForPoseInteraction)
 	{
-		UBaseControlRig* ControlRig = RuntimeControlRigs[0].Get();
+		UControlRig* ControlRig = RuntimeControlRigs[0].Get();
 		FControlRigInteractionScope* InteractionScope = new FControlRigInteractionScope(ControlRig);
 		InteractionScopes.Add(ControlRig,InteractionScope);
 	}
@@ -1289,7 +1289,7 @@ bool FControlRigEditMode::HandleClick(FEditorViewportClient* InViewportClient, H
 					}
 					else if(Click.IsControlDown()) //if ctrl we toggle selection
 					{
-						if (UBaseControlRig* ControlRig = ShapeActor->ControlRig.Get())
+						if (UControlRig* ControlRig = ShapeActor->ControlRig.Get())
 						{
 							bool bIsSelected = ControlRig->IsControlSelected(ControlName);
 							SetRigElementSelection(ControlRig, ERigElementType::Control, ControlName, !bIsSelected);
@@ -1331,9 +1331,9 @@ bool FControlRigEditMode::HandleClick(FEditorViewportClient* InViewportClient, H
 			}
 			else 
 			{ 
-				for (TWeakObjectPtr<UBaseControlRig>& RuntimeRigPtr : RuntimeControlRigs)
+				for (TWeakObjectPtr<UControlRig>& RuntimeRigPtr : RuntimeControlRigs)
 				{
-					if (UBaseControlRig* ControlRig = RuntimeRigPtr.Get())
+					if (UControlRig* ControlRig = RuntimeRigPtr.Get())
 					{
 
 						//if we have an additive or fk control rig active select the control based upon the selected bone.
@@ -1396,9 +1396,9 @@ bool FControlRigEditMode::HandleClick(FEditorViewportClient* InViewportClient, H
 			}
 			else if (Click.IsControlDown()) //if ctrl we toggle selection
 			{
-				for (TWeakObjectPtr<UBaseControlRig>& RuntimeRigPtr : RuntimeControlRigs)
+				for (TWeakObjectPtr<UControlRig>& RuntimeRigPtr : RuntimeControlRigs)
 				{
-					if (UBaseControlRig* ControlRig = RuntimeRigPtr.Get())
+					if (UControlRig* ControlRig = RuntimeRigPtr.Get())
 					{
 						bool bIsSelected = ControlRig->IsControlSelected(ControlName);
 						SetRigElementSelection(FKBoneProxy->ControlRig, ERigElementType::Control, ControlName, !bIsSelected);
@@ -1417,7 +1417,7 @@ bool FControlRigEditMode::HandleClick(FEditorViewportClient* InViewportClient, H
 	{
 		if (RuntimeControlRigs.Num() > 0)
 		{
-			if (UBaseControlRig* DebuggedControlRig = RuntimeControlRigs[0].Get())
+			if (UControlRig* DebuggedControlRig = RuntimeControlRigs[0].Get())
 			{
 				URigHierarchy* Hierarchy = DebuggedControlRig->GetHierarchy();
 
@@ -1970,9 +1970,9 @@ bool FControlRigEditMode::InputDelta(FEditorViewportClient* InViewportClient, FV
 
 bool FControlRigEditMode::ShouldDrawWidget() const
 {
-	for (const TWeakObjectPtr<UBaseControlRig>& RuntimeRigPtr : RuntimeControlRigs)
+	for (const TWeakObjectPtr<UControlRig>& RuntimeRigPtr : RuntimeControlRigs)
 	{
-		if (UBaseControlRig* ControlRig = RuntimeRigPtr.Get())
+		if (UControlRig* ControlRig = RuntimeRigPtr.Get())
 		{
 			if (AreRigElementSelectedAndMovable(ControlRig))
 			{
@@ -2005,9 +2005,9 @@ void FControlRigEditMode::AddReferencedObjects( FReferenceCollector& Collector )
 
 void FControlRigEditMode::ClearRigElementSelection(uint32 InTypes)
 {
-	for (TWeakObjectPtr<UBaseControlRig>& RuntimeRigPtr : RuntimeControlRigs)
+	for (TWeakObjectPtr<UControlRig>& RuntimeRigPtr : RuntimeControlRigs)
 	{
-		if (UBaseControlRig* ControlRig = RuntimeRigPtr.Get())
+		if (UControlRig* ControlRig = RuntimeRigPtr.Get())
 		{
 			if (!AreEditingControlRigDirectly())
 			{
@@ -2025,7 +2025,7 @@ void FControlRigEditMode::ClearRigElementSelection(uint32 InTypes)
 }
 
 // internal private function that doesn't use guarding.
-void FControlRigEditMode::SetRigElementSelectionInternal(UBaseControlRig* ControlRig, ERigElementType Type, const FName& InRigElementName, bool bSelected)
+void FControlRigEditMode::SetRigElementSelectionInternal(UControlRig* ControlRig, ERigElementType Type, const FName& InRigElementName, bool bSelected)
 {
 	if(URigHierarchyController* Controller = ControlRig->GetHierarchy()->GetController())
 	{
@@ -2033,7 +2033,7 @@ void FControlRigEditMode::SetRigElementSelectionInternal(UBaseControlRig* Contro
 	}
 }
 
-void FControlRigEditMode::SetRigElementSelection(UBaseControlRig* ControlRig, ERigElementType Type, const FName& InRigElementName, bool bSelected)
+void FControlRigEditMode::SetRigElementSelection(UControlRig* ControlRig, ERigElementType Type, const FName& InRigElementName, bool bSelected)
 {
 	if (!bSelecting)
 	{
@@ -2045,7 +2045,7 @@ void FControlRigEditMode::SetRigElementSelection(UBaseControlRig* ControlRig, ER
 	}
 }
 
-void FControlRigEditMode::SetRigElementSelection(UBaseControlRig* ControlRig, ERigElementType Type, const TArray<FName>& InRigElementNames, bool bSelected)
+void FControlRigEditMode::SetRigElementSelection(UControlRig* ControlRig, ERigElementType Type, const TArray<FName>& InRigElementNames, bool bSelected)
 {
 	if (!bSelecting)
 	{
@@ -2060,7 +2060,7 @@ void FControlRigEditMode::SetRigElementSelection(UBaseControlRig* ControlRig, ER
 	}
 }
 
-TArray<FRigElementKey> FControlRigEditMode::GetSelectedRigElements(UBaseControlRig* ControlRig) const
+TArray<FRigElementKey> FControlRigEditMode::GetSelectedRigElements(UControlRig* ControlRig) const
 {
 	if (ControlRig == nullptr && GetControlRigs().Num() > 0)
 	{
@@ -2082,7 +2082,7 @@ TArray<FRigElementKey> FControlRigEditMode::GetSelectedRigElements(UBaseControlR
 	{
 		if (SelectedKeys[0].Type == ERigElementType::Bone || SelectedKeys[0].Type == ERigElementType::Null)
 		{
-			const FName ControlName = UBaseControlRig::GetNameForTransientControl(SelectedKeys[0]);
+			const FName ControlName = UControlRig::GetNameForTransientControl(SelectedKeys[0]);
 			const FRigElementKey TransientControlKey = FRigElementKey(ControlName, ERigElementType::Control);
 			if(ControlRig->GetHierarchy()->Contains(TransientControlKey))
 			{
@@ -2108,7 +2108,7 @@ TArray<FRigElementKey> FControlRigEditMode::GetSelectedRigElements(UBaseControlR
 	return SelectedKeys;
 }
 
-bool FControlRigEditMode::AreRigElementsSelected(uint32 InTypes, UBaseControlRig* InControlRig) const
+bool FControlRigEditMode::AreRigElementsSelected(uint32 InTypes, UControlRig* InControlRig) const
 {
 	TArray<FRigElementKey> SelectedRigElements = GetSelectedRigElements(InControlRig);
 
@@ -2123,7 +2123,7 @@ bool FControlRigEditMode::AreRigElementsSelected(uint32 InTypes, UBaseControlRig
 	return false;
 }
 
-int32 FControlRigEditMode::GetNumSelectedRigElements(uint32 InTypes, UBaseControlRig* InControlRig) const
+int32 FControlRigEditMode::GetNumSelectedRigElements(uint32 InTypes, UControlRig* InControlRig) const
 {
 	TArray<FRigElementKey> SelectedRigElements = GetSelectedRigElements(InControlRig);
 	if (FRigElementTypeHelper::DoesHave(InTypes, ERigElementType::All))
@@ -2186,7 +2186,7 @@ ECoordSystem FControlRigEditMode::GetCoordSystemSpace() const
 	return GetModeManager()->GetCoordSystem();	
 }
 
-void FControlRigEditMode::UpdatePivotFromEditedShape(UBaseControlRig* InControlRig)
+void FControlRigEditMode::UpdatePivotFromEditedShape(UControlRig* InControlRig)
 {
 	const URigHierarchy* Hierarchy = InControlRig ? InControlRig->GetHierarchy() : nullptr;
 	if (!Hierarchy)
@@ -2221,7 +2221,7 @@ void FControlRigEditMode::UpdatePivotFromEditedShape(UBaseControlRig* InControlR
 	PivotTransforms.Add(InControlRig, MoveTemp(PivotTransform));
 }
 
-void FControlRigEditMode::UpdatePivotFromShapeActors(UBaseControlRig* InControlRig, const bool bEachLocalSpace, const bool bIsParentSpace)
+void FControlRigEditMode::UpdatePivotFromShapeActors(UControlRig* InControlRig, const bool bEachLocalSpace, const bool bIsParentSpace)
 {
 	if (!ensure(!bIsChangingControlShapeTransform))
 	{
@@ -2277,7 +2277,7 @@ void FControlRigEditMode::UpdatePivotFromShapeActors(UBaseControlRig* InControlR
 	PivotTransforms.Add(InControlRig, MoveTemp(PivotTransform));
 }
 
-void FControlRigEditMode::UpdatePivotFromElements(UBaseControlRig* InControlRig)
+void FControlRigEditMode::UpdatePivotFromElements(UControlRig* InControlRig)
 {
 	if (!ensure(!bIsChangingControlShapeTransform))
 	{
@@ -2331,9 +2331,9 @@ void FControlRigEditMode::UpdatePivotTransforms()
 
 	PivotTransforms.Reset();
 
-	for (const TWeakObjectPtr<UBaseControlRig>& RuntimeRigPtr : RuntimeControlRigs)
+	for (const TWeakObjectPtr<UControlRig>& RuntimeRigPtr : RuntimeControlRigs)
 	{
-		if (UBaseControlRig* ControlRig = RuntimeRigPtr.Get())
+		if (UControlRig* ControlRig = RuntimeRigPtr.Get())
 		{
 			if (AreRigElementsSelected(ValidControlTypeMask(), ControlRig))
 			{
@@ -2394,7 +2394,7 @@ bool FControlRigEditMode::HasPivotTransformsChanged() const
 	{
 		return true;
 	}
-	for (const TPair<UBaseControlRig*, FTransform>& Transform : PivotTransforms)
+	for (const TPair<UControlRig*, FTransform>& Transform : PivotTransforms)
 	{
 		if (const FTransform* LastTransform = LastPivotTransforms.Find(Transform.Key))
 		{
@@ -2508,9 +2508,9 @@ void FControlRigEditMode::BindCommands()
 bool FControlRigEditMode::IsControlSelected() const
 {
 	static uint32 TypeFlag = (uint32)ERigElementType::Control;
-	for (const TWeakObjectPtr<UBaseControlRig>& RuntimeRigPtr : RuntimeControlRigs)
+	for (const TWeakObjectPtr<UControlRig>& RuntimeRigPtr : RuntimeControlRigs)
 	{
-		if (UBaseControlRig* ControlRig = RuntimeRigPtr.Get())
+		if (UControlRig* ControlRig = RuntimeRigPtr.Get())
 		{
 			if (AreRigElementsSelected(TypeFlag,ControlRig))
 			{
@@ -2523,9 +2523,9 @@ bool FControlRigEditMode::IsControlSelected() const
 
 bool FControlRigEditMode::CanFrameSelection()
 {
-	for (TWeakObjectPtr<UBaseControlRig>& RuntimeRigPtr : RuntimeControlRigs)
+	for (TWeakObjectPtr<UControlRig>& RuntimeRigPtr : RuntimeControlRigs)
 	{
-		if (UBaseControlRig* ControlRig = RuntimeRigPtr.Get())
+		if (UControlRig* ControlRig = RuntimeRigPtr.Get())
 		{
 			if (GetSelectedRigElements(ControlRig).Num() > 0)
 			{
@@ -2562,9 +2562,9 @@ void FControlRigEditMode::FrameSelection()
     }
 
 	TArray<AActor*> Actors;
-	for (TWeakObjectPtr<UBaseControlRig>& RuntimeRigPtr : RuntimeControlRigs)
+	for (TWeakObjectPtr<UControlRig>& RuntimeRigPtr : RuntimeControlRigs)
 	{
-		if (UBaseControlRig* ControlRig = RuntimeRigPtr.Get())
+		if (UControlRig* ControlRig = RuntimeRigPtr.Get())
 		{
 			TArray<FRigElementKey> SelectedRigElements = GetSelectedRigElements(ControlRig);
 			for (const FRigElementKey& SelectedKey : SelectedRigElements)
@@ -2682,11 +2682,11 @@ void FControlRigEditMode::ToggleControlShapeTransformEdit()
 	}
 }
 
-void FControlRigEditMode::GetAllSelectedControls(TMap<UBaseControlRig*, TArray<FRigElementKey>>& OutSelectedControls) const
+void FControlRigEditMode::GetAllSelectedControls(TMap<UControlRig*, TArray<FRigElementKey>>& OutSelectedControls) const
 {
-	for (const TWeakObjectPtr<UBaseControlRig>& RuntimeRigPtr : RuntimeControlRigs)
+	for (const TWeakObjectPtr<UControlRig>& RuntimeRigPtr : RuntimeControlRigs)
 	{
-		if (UBaseControlRig* ControlRig = RuntimeRigPtr.Get())
+		if (UControlRig* ControlRig = RuntimeRigPtr.Get())
 		{
 			if (const URigHierarchy* Hierarchy = ControlRig->GetHierarchy())
 			{
@@ -2772,7 +2772,7 @@ bool FControlRigEditMode::IsDragAnimSliderToolPressed(FViewport* InViewport)
 
 void FControlRigEditMode::OpenSpacePickerWidget()
 {
-	TMap<UBaseControlRig*, TArray<FRigElementKey>> SelectedControlRigsAndControls;
+	TMap<UControlRig*, TArray<FRigElementKey>> SelectedControlRigsAndControls;
 	GetAllSelectedControls(SelectedControlRigsAndControls);
 
 	if (SelectedControlRigsAndControls.Num() < 1)
@@ -2780,14 +2780,14 @@ void FControlRigEditMode::OpenSpacePickerWidget()
 		return;
 	}
 
-	TArray<UBaseControlRig*> ControlRigs;
+	TArray<UControlRig*> ControlRigs;
 	TArray<TArray<FRigElementKey>> AllSelectedControls;
 	SelectedControlRigsAndControls.GenerateKeyArray(ControlRigs);
 	SelectedControlRigsAndControls.GenerateValueArray(AllSelectedControls);
 
 
 	//mz todo handle multiple control rigs with space picker
-	UBaseControlRig* RuntimeRig = ControlRigs[0];
+	UControlRig* RuntimeRig = ControlRigs[0];
 	TArray<FRigElementKey>& SelectedControls = AllSelectedControls[0];
 
 	URigHierarchy* Hierarchy = RuntimeRig->GetHierarchy();
@@ -2914,11 +2914,11 @@ void FControlRigEditMode::ToggleManipulators()
 {
 	if (!AreEditingControlRigDirectly())
 	{
-		TMap<UBaseControlRig*, TArray<FRigElementKey>> SelectedControls;
+		TMap<UControlRig*, TArray<FRigElementKey>> SelectedControls;
 		GetAllSelectedControls(SelectedControls);
-		TArray<UBaseControlRig*> ControlRigs;
+		TArray<UControlRig*> ControlRigs;
 		SelectedControls.GenerateKeyArray(ControlRigs);
-		for (UBaseControlRig* ControlRig : ControlRigs)
+		for (UControlRig* ControlRig : ControlRigs)
 		{
 			if (ControlRig)
 			{
@@ -2943,9 +2943,9 @@ void FControlRigEditMode::ToggleAllManipulators()
 	//turn on all if in level editor in case any where off
 	if (!AreEditingControlRigDirectly() && Settings->bHideControlShapes)
 	{
-		for (TWeakObjectPtr<UBaseControlRig>& RuntimeRigPtr : RuntimeControlRigs)
+		for (TWeakObjectPtr<UControlRig>& RuntimeRigPtr : RuntimeControlRigs)
 		{
-			if (UBaseControlRig* ControlRig = RuntimeRigPtr.Get())
+			if (UControlRig* ControlRig = RuntimeRigPtr.Get())
 			{
 				ControlRig->SetControlsVisible(true);
 			}
@@ -2956,10 +2956,10 @@ void FControlRigEditMode::ToggleAllManipulators()
 void FControlRigEditMode::ZeroTransforms(bool bSelectionOnly)
 {
 	// Gather up the control rigs for the selected controls
-	TArray<UBaseControlRig*> ControlRigs;
-	for (TWeakObjectPtr<UBaseControlRig>& RuntimeRigPtr : RuntimeControlRigs)
+	TArray<UControlRig*> ControlRigs;
+	for (TWeakObjectPtr<UControlRig>& RuntimeRigPtr : RuntimeControlRigs)
 	{
-		if (UBaseControlRig* ControlRig = RuntimeRigPtr.Get())
+		if (UControlRig* ControlRig = RuntimeRigPtr.Get())
 		{
 			if (!bSelectionOnly || ControlRig->CurrentControlSelection().Num() > 0)
 			{
@@ -2974,7 +2974,7 @@ void FControlRigEditMode::ZeroTransforms(bool bSelectionOnly)
 
 	FScopedTransaction Transaction(LOCTEXT("HierarchyZeroTransforms", "Zero Transforms"));
 
-	for (UBaseControlRig* ControlRig : ControlRigs)
+	for (UControlRig* ControlRig : ControlRigs)
 	{
 		TArray<FRigElementKey> SelectedRigElements = GetSelectedRigElements(ControlRig);
 		if (ControlRig->IsAdditive())
@@ -3050,7 +3050,7 @@ void FControlRigEditMode::ZeroTransforms(bool bSelectionOnly)
 				ControlElement = ControlRig->FindControl(ElementToReset.Name);
 				if (ControlElement->Settings.bIsTransientControl)
 				{
-					if(UBaseControlRig::GetNodeNameFromTransientControl(ControlElement->GetKey()).IsEmpty())
+					if(UControlRig::GetNodeNameFromTransientControl(ControlElement->GetKey()).IsEmpty())
 					{
 						ControlElement = nullptr;
 					}
@@ -3181,10 +3181,10 @@ void FControlRigEditMode::ZeroTransforms(bool bSelectionOnly)
 void FControlRigEditMode::InvertInputPose(bool bSelectionOnly)
 {
 	// Gather up the control rigs for the selected controls
-	TArray<UBaseControlRig*> ControlRigs;
-	for (TWeakObjectPtr<UBaseControlRig>& RuntimeRigPtr : RuntimeControlRigs)
+	TArray<UControlRig*> ControlRigs;
+	for (TWeakObjectPtr<UControlRig>& RuntimeRigPtr : RuntimeControlRigs)
 	{
-		if (UBaseControlRig* ControlRig = RuntimeRigPtr.Get())
+		if (UControlRig* ControlRig = RuntimeRigPtr.Get())
 		{
 			if (!bSelectionOnly || ControlRig->CurrentControlSelection().Num() > 0)
 			{
@@ -3199,7 +3199,7 @@ void FControlRigEditMode::InvertInputPose(bool bSelectionOnly)
 
 	FScopedTransaction Transaction(LOCTEXT("HierarchyInvertTransformsToRestPose", "Invert Transforms to Rest Pose"));
 
-	for (UBaseControlRig* ControlRig : ControlRigs)
+	for (UControlRig* ControlRig : ControlRigs)
 	{
 		if (!ControlRig->IsAdditive())
 		{
@@ -3278,11 +3278,11 @@ bool FControlRigEditMode::CheckMovieSceneSig()
 			{
 				if (ULevelSequence* LevelSequence = Cast<ULevelSequence>(Sequencer->GetFocusedMovieSceneSequence()))
 				{
-					TArray<TWeakObjectPtr<UBaseControlRig>> CurrentControlRigs;
+					TArray<TWeakObjectPtr<UControlRig>> CurrentControlRigs;
 					TArray<FControlRigSequencerBindingProxy> Proxies = UControlRigSequencerEditorLibrary::GetControlRigs(LevelSequence);
 					for (FControlRigSequencerBindingProxy& Proxy : Proxies)
 					{
-						if (UBaseControlRig* ControlRig = Proxy.ControlRig.Get())
+						if (UControlRig* ControlRig = Proxy.ControlRig.Get())
 						{
 							CurrentControlRigs.Add(ControlRig);
 							if (RuntimeControlRigs.Contains(ControlRig) == false)
@@ -3292,15 +3292,15 @@ bool FControlRigEditMode::CheckMovieSceneSig()
 							}
 						}
 					}
-					TArray<TWeakObjectPtr<UBaseControlRig>> ControlRigsToRemove;
-					for (TWeakObjectPtr<UBaseControlRig>& RuntimeRigPtr : RuntimeControlRigs)
+					TArray<TWeakObjectPtr<UControlRig>> ControlRigsToRemove;
+					for (TWeakObjectPtr<UControlRig>& RuntimeRigPtr : RuntimeControlRigs)
 					{
 						if (CurrentControlRigs.Contains(RuntimeRigPtr) == false)
 						{
 							ControlRigsToRemove.Add(RuntimeRigPtr);
 						}
 					}
-					for (TWeakObjectPtr<UBaseControlRig>& OldRuntimeRigPtr : ControlRigsToRemove)
+					for (TWeakObjectPtr<UControlRig>& OldRuntimeRigPtr : ControlRigsToRemove)
 					{
 						RemoveControlRig(OldRuntimeRigPtr.Get());
 					}
@@ -3322,7 +3322,7 @@ void FControlRigEditMode::PostUndo()
 	if (WeakSequencer.IsValid())
 	{
 		bool bHaveInvalidControlRig = false;
-		for (TWeakObjectPtr<UBaseControlRig>& RuntimeRigPtr : RuntimeControlRigs)
+		for (TWeakObjectPtr<UControlRig>& RuntimeRigPtr : RuntimeControlRigs)
 		{
 			if (RuntimeRigPtr.IsValid() == false)
 			{				
@@ -3333,7 +3333,7 @@ void FControlRigEditMode::PostUndo()
 		//if one is invalid we need to clear everything,since no longer have ptr to selectively delete
 		if (bHaveInvalidControlRig == true)
 		{
-			TArray<TWeakObjectPtr<UBaseControlRig>> PreviousRuntimeRigs = RuntimeControlRigs;
+			TArray<TWeakObjectPtr<UControlRig>> PreviousRuntimeRigs = RuntimeControlRigs;
 			for (int32 PreviousRuntimeRigIndex = 0; PreviousRuntimeRigIndex < PreviousRuntimeRigs.Num(); PreviousRuntimeRigIndex++)
 			{
 				if (PreviousRuntimeRigs[PreviousRuntimeRigIndex].IsValid())
@@ -3353,7 +3353,7 @@ void FControlRigEditMode::PostUndo()
 			TArray<FControlRigSequencerBindingProxy> Proxies = UControlRigSequencerEditorLibrary::GetControlRigs(LevelSequence);
 			for (FControlRigSequencerBindingProxy& Proxy : Proxies)
 			{
-				if (UBaseControlRig* ControlRig = Proxy.ControlRig.Get())
+				if (UControlRig* ControlRig = Proxy.ControlRig.Get())
 				{
 					if (RuntimeControlRigs.Contains(ControlRig) == false)
 					{
@@ -3373,7 +3373,7 @@ void FControlRigEditMode::PostUndo()
 	}
 	else
 	{
-		for (TWeakObjectPtr<UBaseControlRig>& RuntimeRigPtr : RuntimeControlRigs)
+		for (TWeakObjectPtr<UControlRig>& RuntimeRigPtr : RuntimeControlRigs)
 		{
 			if (RuntimeRigPtr.IsValid() == false)
 			{
@@ -3397,7 +3397,7 @@ void FControlRigEditMode::PostUndo()
 
 }
 
-void FControlRigEditMode::RequestToRecreateControlShapeActors(UBaseControlRig* ControlRig)
+void FControlRigEditMode::RequestToRecreateControlShapeActors(UControlRig* ControlRig)
 { 
 	if (ControlRig)
 	{
@@ -3447,9 +3447,9 @@ void FControlRigEditMode::RecreateControlShapeActors(const TArray<FRigElementKey
 {
 	if (RecreateControlShapesRequired == ERecreateControlRigShape::RecreateAll)
 	{
-		for (TWeakObjectPtr<UBaseControlRig>& RuntimeRigPtr : RuntimeControlRigs)
+		for (TWeakObjectPtr<UControlRig>& RuntimeRigPtr : RuntimeControlRigs)
 		{
-			if (UBaseControlRig* RuntimeControlRig = RuntimeRigPtr.Get())
+			if (UControlRig* RuntimeControlRig = RuntimeRigPtr.Get())
 			{
 				DestroyShapesActors(RuntimeControlRig);
 				CreateShapeActors(RuntimeControlRig);
@@ -3458,8 +3458,8 @@ void FControlRigEditMode::RecreateControlShapeActors(const TArray<FRigElementKey
 	}
 	else if (ControlRigsToRecreate.Num() > 0)
 	{
-		TArray < UBaseControlRig*> ControlRigsCopy = ControlRigsToRecreate;
-		for (UBaseControlRig* ControlRig : ControlRigsCopy)
+		TArray < UControlRig*> ControlRigsCopy = ControlRigsToRecreate;
+		for (UControlRig* ControlRig : ControlRigsCopy)
 		{
 			//check to see if actors have really changed, if not don't do it
 			bool bRecreateThem = true;
@@ -3538,7 +3538,7 @@ void FControlRigEditMode::RecreateControlShapeActors(const TArray<FRigElementKey
 	}
 }
 
-void FControlRigEditMode::CreateShapeActors(UBaseControlRig* ControlRig)
+void FControlRigEditMode::CreateShapeActors(UControlRig* ControlRig)
 {
 	// create gizmo actors
 	FActorSpawnParameters ActorSpawnParameters;
@@ -3686,7 +3686,7 @@ void FControlRigEditMode::OnObjectsReplaced(const TMap<UObject*, UObject*>& OldT
 		UObject* NewObject = OldToNewInstanceMap.FindRef(OldObject);
 		if (NewObject)
 		{
-			TArray<TWeakObjectPtr<UBaseControlRig>> PreviousRuntimeRigs = RuntimeControlRigs;
+			TArray<TWeakObjectPtr<UControlRig>> PreviousRuntimeRigs = RuntimeControlRigs;
 			for (int32 PreviousRuntimeRigIndex = 0; PreviousRuntimeRigIndex < PreviousRuntimeRigs.Num(); PreviousRuntimeRigIndex++)
 			{
 				if (PreviousRuntimeRigs[PreviousRuntimeRigIndex].IsValid())
@@ -3696,7 +3696,7 @@ void FControlRigEditMode::OnObjectsReplaced(const TMap<UObject*, UObject*>& OldT
 			}
 			RuntimeControlRigs.Reset();
 
-			UBaseControlRig* NewRig = Cast<UBaseControlRig>(NewObject);
+			UControlRig* NewRig = Cast<UControlRig>(NewObject);
 			AddControlRigInternal(NewRig);
 
 			NewRig->Initialize();
@@ -3711,7 +3711,7 @@ bool FControlRigEditMode::IsTransformDelegateAvailable() const
 	return (OnGetRigElementTransformDelegate.IsBound() && OnSetRigElementTransformDelegate.IsBound());
 }
 
-bool FControlRigEditMode::AreRigElementSelectedAndMovable(UBaseControlRig* ControlRig) const
+bool FControlRigEditMode::AreRigElementSelectedAndMovable(UControlRig* ControlRig) const
 {
 	const UControlRigEditModeSettings* Settings = GetDefault<UControlRigEditModeSettings>();
 
@@ -3769,7 +3769,7 @@ bool FControlRigEditMode::AreRigElementSelectedAndMovable(UBaseControlRig* Contr
 	return true;
 }
 
-void FControlRigEditMode::ReplaceControlRig(UBaseControlRig* OldControlRig, UBaseControlRig* NewControlRig)
+void FControlRigEditMode::ReplaceControlRig(UControlRig* OldControlRig, UControlRig* NewControlRig)
 {
 	if (OldControlRig != nullptr)
 	{
@@ -3796,7 +3796,7 @@ void FControlRigEditMode::OnHierarchyModified(ERigHierarchyNotification InNotif,
 		case ERigHierarchyNotification::ElementReordered:
 		case ERigHierarchyNotification::HierarchyReset:
 		{
-			UBaseControlRig* ControlRig = InHierarchy->GetTypedOuter<UBaseControlRig>();
+			UControlRig* ControlRig = InHierarchy->GetTypedOuter<UControlRig>();
 			RequestToRecreateControlShapeActors(ControlRig);
 			break;
 		}
@@ -3805,7 +3805,7 @@ void FControlRigEditMode::OnHierarchyModified(ERigHierarchyNotification InNotif,
 		case ERigHierarchyNotification::ControlShapeTransformChanged:
 		{
 			const FRigElementKey Key = InElement->GetKey();
-			UBaseControlRig* ControlRig = InHierarchy->GetTypedOuter<UBaseControlRig>();
+			UControlRig* ControlRig = InHierarchy->GetTypedOuter<UControlRig>();
 			if (Key.Type == ERigElementType::Control)
 			{
 				if (const FRigControlElement* ControlElement = Cast<FRigControlElement>(InElement))
@@ -3839,7 +3839,7 @@ void FControlRigEditMode::OnHierarchyModified(ERigHierarchyNotification InNotif,
 				// are correctly selected in sequencer to match what the proxy control is driving.
 				if (FRigControlElement* ControlElement = InHierarchy->Find<FRigControlElement>(InElement->GetKey()))
 				{
-					UBaseControlRig* ControlRig = InHierarchy->GetTypedOuter<UBaseControlRig>();
+					UControlRig* ControlRig = InHierarchy->GetTypedOuter<UControlRig>();
 					if(ControlProxy->IsSelected(ControlRig, ControlElement->GetFName()))
 					{
 						// reselect the control - to affect the details panel / sequencer
@@ -3883,7 +3883,7 @@ void FControlRigEditMode::OnHierarchyModified(ERigHierarchyNotification InNotif,
 					const bool bSelected = InNotif == ERigHierarchyNotification::ElementSelected;
 					// users may select gizmo and control rig units, so we have to let them go through both of them if they do
 						// first go through gizmo actor
-					UBaseControlRig* ControlRig = InHierarchy->GetTypedOuter<UBaseControlRig>();
+					UControlRig* ControlRig = InHierarchy->GetTypedOuter<UControlRig>();
 					if (ControlRig == nullptr)
 					{
 						if (RuntimeControlRigs.Num() > 0)
@@ -4040,7 +4040,7 @@ void FControlRigEditMode::OnHierarchyModified_AnyThread(ERigHierarchyNotificatio
 	}, TStatId(), NULL, ENamedThreads::GameThread);
 }
 
-void FControlRigEditMode::OnControlModified(UBaseControlRig* Subject, FRigControlElement* InControlElement, const FRigControlModifiedContext& Context)
+void FControlRigEditMode::OnControlModified(UControlRig* Subject, FRigControlElement* InControlElement, const FRigControlModifiedContext& Context)
 {
 	//this makes sure the details panel ui get's updated, don't remove
 	const bool bModify = Context.SetKey != EControlRigSetKey::Never;
@@ -4051,7 +4051,7 @@ void FControlRigEditMode::OnControlModified(UBaseControlRig* Subject, FRigContro
 	ControlProxy->Modify();
 	RecalcPivotTransform();
 
-	if (UBaseControlRig* ControlRig = static_cast<UBaseControlRig*>(Subject))
+	if (UControlRig* ControlRig = static_cast<UControlRig*>(Subject))
 	{
 		FTransform ComponentTransform = GetHostingSceneComponentTransform();
 		if (AControlRigShapeActor* const* Actor = GizmoToControlMap.FindKey(InControl.Index))
@@ -4062,12 +4062,12 @@ void FControlRigEditMode::OnControlModified(UBaseControlRig* Subject, FRigContro
 	*/
 }
 
-void FControlRigEditMode::OnPreConstruction_AnyThread(UBaseControlRig* InRig, const FName& InEventName)
+void FControlRigEditMode::OnPreConstruction_AnyThread(UControlRig* InRig, const FName& InEventName)
 {
 	bIsConstructionEventRunning = true;
 }
 
-void FControlRigEditMode::OnPostConstruction_AnyThread(UBaseControlRig* InRig, const FName& InEventName)
+void FControlRigEditMode::OnPostConstruction_AnyThread(UControlRig* InRig, const FName& InEventName)
 {
 	bIsConstructionEventRunning = false;
 
@@ -4136,9 +4136,9 @@ bool FControlRigEditMode::CanChangeControlShapeTransform()
 {
 	if (AreEditingControlRigDirectly())
 	{
-		for (TWeakObjectPtr<UBaseControlRig> RuntimeRigPtr : RuntimeControlRigs)
+		for (TWeakObjectPtr<UControlRig> RuntimeRigPtr : RuntimeControlRigs)
 		{
-			if (UBaseControlRig* ControlRig = RuntimeRigPtr.Get())
+			if (UControlRig* ControlRig = RuntimeRigPtr.Get())
 			{
 				TArray<FRigElementKey> SelectedRigElements = GetSelectedRigElements(ControlRig);
 				// do not allow multi-select
@@ -4177,7 +4177,7 @@ void FControlRigEditMode::SetControlShapeTransform(
 	const FRigControlModifiedContext& InContext,
 	const bool bPrintPython) const
 {
-	UBaseControlRig* ControlRig = InShapeActor->ControlRig.Get();
+	UControlRig* ControlRig = InShapeActor->ControlRig.Get();
 	if (!ControlRig)
 	{
 		return;
@@ -4234,7 +4234,7 @@ void FControlRigEditMode::SetControlShapeTransform(
 
 FTransform FControlRigEditMode::GetControlShapeTransform(const AControlRigShapeActor* ShapeActor)
 {
-	if (const UBaseControlRig* ControlRig = ShapeActor->ControlRig.Get())
+	if (const UControlRig* ControlRig = ShapeActor->ControlRig.Get())
 	{
 		return ControlRig->GetControlGlobalTransform(ShapeActor->ControlName);
 	}
@@ -4249,7 +4249,7 @@ void FControlRigEditMode::MoveControlShape(AControlRigShapeActor* ShapeActor, co
 	//first case is where we do all controls by the local diff.
 	if (bUseLocal)
 	{
-		if (UBaseControlRig* ControlRig = ShapeActor->ControlRig.Get())
+		if (UControlRig* ControlRig = ShapeActor->ControlRig.Get())
 		{
 			FRigControlModifiedContext Context;
 			Context.EventName = FRigUnit_BeginExecution::EventName;
@@ -4324,7 +4324,7 @@ void FControlRigEditMode::MoveControlShape(AControlRigShapeActor* ShapeActor, co
 
 		if (bTransformChanged)
 		{
-			if (UBaseControlRig* ControlRig = ShapeActor->ControlRig.Get())
+			if (UControlRig* ControlRig = ShapeActor->ControlRig.Get())
 			{
 				ControlRig->InteractionType = InteractionType;
 				ControlRig->ElementsBeingInteracted.AddUnique(ShapeActor->GetElementKey());
@@ -4387,7 +4387,7 @@ void FControlRigEditMode::ChangeControlShapeTransform(AControlRigShapeActor* Sha
 
 	FTransform CurrentTransform = FTransform::Identity;
 
-	if (UBaseControlRig* ControlRig = ShapeActor->ControlRig.Get())
+	if (UControlRig* ControlRig = ShapeActor->ControlRig.Get())
 	{
 		if (FRigControlElement* ControlElement = ControlRig->GetHierarchy()->Find<FRigControlElement>(ShapeActor->GetElementKey()))
 		{
@@ -4422,7 +4422,7 @@ void FControlRigEditMode::ChangeControlShapeTransform(AControlRigShapeActor* Sha
 
 	if (bTransformChanged)
 	{
-		if (UBaseControlRig* ControlRig = ShapeActor->ControlRig.Get())
+		if (UControlRig* ControlRig = ShapeActor->ControlRig.Get())
 		{
 
 			FTransform NewTransform = CurrentTransform.GetRelativeTransform(ToWorldTransform);
@@ -4451,7 +4451,7 @@ void FControlRigEditMode::ChangeControlShapeTransform(AControlRigShapeActor* Sha
 
 bool FControlRigEditMode::ModeSupportedByShapeActor(const AControlRigShapeActor* ShapeActor, UE::Widget::EWidgetMode InMode) const
 {
-	if (UBaseControlRig* ControlRig = ShapeActor->ControlRig.Get())
+	if (UControlRig* ControlRig = ShapeActor->ControlRig.Get())
 	{
 		const FRigControlElement* ControlElement = ControlRig->FindControl(ShapeActor->ControlName);
 		if (ControlElement)
@@ -4548,7 +4548,7 @@ bool FControlRigEditMode::ModeSupportedByShapeActor(const AControlRigShapeActor*
 	return false;
 }
 
-bool FControlRigEditMode::IsControlRigSkelMeshVisible(UBaseControlRig* ControlRig) const
+bool FControlRigEditMode::IsControlRigSkelMeshVisible(UControlRig* ControlRig) const
 {
 	if (IsInLevelEditor())
 	{
@@ -4570,7 +4570,7 @@ void FControlRigEditMode::TickControlShape(AControlRigShapeActor* ShapeActor, co
 	const UControlRigEditModeSettings* Settings = GetDefault<UControlRigEditModeSettings>();
 	if (ShapeActor)
 	{
-		if (UBaseControlRig* ControlRig = ShapeActor->ControlRig.Get())
+		if (UControlRig* ControlRig = ShapeActor->ControlRig.Get())
 		{
 			const FTransform Transform = ControlRig->GetControlGlobalTransform(ShapeActor->ControlName);
 			ShapeActor->SetActorTransform(Transform * ComponentTransform);
@@ -4608,7 +4608,7 @@ void FControlRigEditMode::TickControlShape(AControlRigShapeActor* ShapeActor, co
 	}
 }
 
-AControlRigShapeActor* FControlRigEditMode::GetControlShapeFromControlName(UBaseControlRig* InControlRig,const FName& ControlName) const
+AControlRigShapeActor* FControlRigEditMode::GetControlShapeFromControlName(UControlRig* InControlRig,const FName& ControlName) const
 {
 	const auto* ShapeActors = ControlRigShapeActors.Find(InControlRig);
 	if (ShapeActors)
@@ -4625,7 +4625,7 @@ AControlRigShapeActor* FControlRigEditMode::GetControlShapeFromControlName(UBase
 	return nullptr;
 }
 
-void FControlRigEditMode::AddControlRigInternal(UBaseControlRig* InControlRig)
+void FControlRigEditMode::AddControlRigInternal(UControlRig* InControlRig)
 {
 	RuntimeControlRigs.AddUnique(InControlRig);
 	LastHierarchyHash.Add(INDEX_NONE);
@@ -4647,20 +4647,20 @@ void FControlRigEditMode::AddControlRigInternal(UBaseControlRig* InControlRig)
 	UpdateSelectabilityOnSkeletalMeshes(InControlRig, !bShowControlsAsOverlay);
 }
 
-TArrayView<const TWeakObjectPtr<UBaseControlRig>> FControlRigEditMode::GetControlRigs() const
+TArrayView<const TWeakObjectPtr<UControlRig>> FControlRigEditMode::GetControlRigs() const
 {
 	return MakeArrayView(RuntimeControlRigs);
 }
 
-TArrayView<TWeakObjectPtr<UBaseControlRig>> FControlRigEditMode::GetControlRigs() 
+TArrayView<TWeakObjectPtr<UControlRig>> FControlRigEditMode::GetControlRigs() 
 {
 	return MakeArrayView(RuntimeControlRigs);
 }
 
-TArray<UBaseControlRig*> FControlRigEditMode::GetControlRigsArray(bool bIsVisible)
+TArray<UControlRig*> FControlRigEditMode::GetControlRigsArray(bool bIsVisible)
 {
-	TArray < UBaseControlRig*> ControlRigs;
-	for (TWeakObjectPtr<UBaseControlRig> ControlRigPtr : RuntimeControlRigs)
+	TArray < UControlRig*> ControlRigs;
+	for (TWeakObjectPtr<UControlRig> ControlRigPtr : RuntimeControlRigs)
 	{
 		if (ControlRigPtr.IsValid() && ControlRigPtr.Get() != nullptr && (bIsVisible == false ||ControlRigPtr.Get()->GetControlsVisible()))
 		{
@@ -4670,10 +4670,10 @@ TArray<UBaseControlRig*> FControlRigEditMode::GetControlRigsArray(bool bIsVisibl
 	return ControlRigs;
 }
 
-TArray<const UBaseControlRig*> FControlRigEditMode::GetControlRigsArray(bool bIsVisible) const
+TArray<const UControlRig*> FControlRigEditMode::GetControlRigsArray(bool bIsVisible) const
 {
-	TArray<const UBaseControlRig*> ControlRigs;
-	for (const TWeakObjectPtr<UBaseControlRig> ControlRigPtr : RuntimeControlRigs)
+	TArray<const UControlRig*> ControlRigs;
+	for (const TWeakObjectPtr<UControlRig> ControlRigPtr : RuntimeControlRigs)
 	{
 		if (ControlRigPtr.IsValid() && ControlRigPtr.Get() != nullptr && (bIsVisible == false || ControlRigPtr.Get()->GetControlsVisible()))
 		{
@@ -4683,7 +4683,7 @@ TArray<const UBaseControlRig*> FControlRigEditMode::GetControlRigsArray(bool bIs
 	return ControlRigs;
 }
 
-void FControlRigEditMode::RemoveControlRig(UBaseControlRig* InControlRig)
+void FControlRigEditMode::RemoveControlRig(UControlRig* InControlRig)
 {
 	if (InControlRig == nullptr)
 	{
@@ -4725,9 +4725,9 @@ void FControlRigEditMode::RemoveControlRig(UBaseControlRig* InControlRig)
 
 void FControlRigEditMode::TickManipulatableObjects(float DeltaTime)
 {
-	for (TWeakObjectPtr<UBaseControlRig> RuntimeRigPtr : RuntimeControlRigs)
+	for (TWeakObjectPtr<UControlRig> RuntimeRigPtr : RuntimeControlRigs)
 	{
-		if (UBaseControlRig* ControlRig = RuntimeRigPtr.Get())
+		if (UControlRig* ControlRig = RuntimeRigPtr.Get())
 		{
 			// tick skeletalmeshcomponent, that's how they update their transform from rig change
 			USceneComponent* SceneComponent = GetHostingSceneComponent(ControlRig);
@@ -4795,7 +4795,7 @@ void FControlRigEditMode::OnEditorClosed()
 	ControlRigsToRecreate.Reset();
 }
 
-void FControlRigEditMode::DestroyShapesActors(UBaseControlRig* ControlRig)
+void FControlRigEditMode::DestroyShapesActors(UControlRig* ControlRig)
 {
 	if (ControlRig == nullptr)
 	{
@@ -4844,7 +4844,7 @@ void FControlRigEditMode::DestroyShapesActors(UBaseControlRig* ControlRig)
 	}
 }
 
-USceneComponent* FControlRigEditMode::GetHostingSceneComponent(const UBaseControlRig* ControlRig) const
+USceneComponent* FControlRigEditMode::GetHostingSceneComponent(const UControlRig* ControlRig) const
 {
 	if (ControlRig == nullptr && GetControlRigs().Num() > 0)
 	{
@@ -4885,7 +4885,7 @@ USceneComponent* FControlRigEditMode::GetHostingSceneComponent(const UBaseContro
 	return nullptr;
 }
 
-FTransform FControlRigEditMode::GetHostingSceneComponentTransform(const UBaseControlRig* ControlRig) const
+FTransform FControlRigEditMode::GetHostingSceneComponentTransform(const UControlRig* ControlRig) const
 {
 	if (ControlRig == nullptr && GetControlRigs().Num() > 0)
 	{
@@ -4917,7 +4917,7 @@ void FControlRigEditMode::PostPoseUpdate()
 
 }
 
-void FControlRigEditMode::NotifyDrivenControls(UBaseControlRig* InControlRig, const FRigElementKey& InKey)
+void FControlRigEditMode::NotifyDrivenControls(UControlRig* InControlRig, const FRigElementKey& InKey)
 {
 	// if we are changing a proxy control - we also need to notify the change for the driven controls
 	if (FRigControlElement* ControlElement = InControlRig->GetHierarchy()->Find<FRigControlElement>(InKey))
@@ -4939,7 +4939,7 @@ void FControlRigEditMode::NotifyDrivenControls(UBaseControlRig* InControlRig, co
 	}
 }
 
-void FControlRigEditMode::UpdateSelectabilityOnSkeletalMeshes(UBaseControlRig* InControlRig, bool bEnabled)
+void FControlRigEditMode::UpdateSelectabilityOnSkeletalMeshes(UControlRig* InControlRig, bool bEnabled)
 {
 	if(const USceneComponent* HostingComponent = GetHostingSceneComponent(InControlRig))
 	{
diff --git a/Engine/Plugins/Animation/ControlRig/Source/ControlRigEditor/Private/EditMode/SControlPicker.cpp b/Engine/Plugins/Animation/ControlRig/Source/ControlRigEditor/Private/EditMode/SControlPicker.cpp
--- a/Engine/Plugins/Animation/ControlRig/Source/ControlRigEditor/Private/EditMode/SControlPicker.cpp
+++ b/Engine/Plugins/Animation/ControlRig/Source/ControlRigEditor/Private/EditMode/SControlPicker.cpp
@@ -30,7 +30,7 @@ void SControlPicker::Construct(const FArguments& InArgs, UWorld* InWorld)
 	];
 }
 
-void SControlPicker::SetControlRig(UBaseControlRig* InRig)
+void SControlPicker::SetControlRig(UControlRig* InRig)
 {
 	if (InRig != RigPtr.Get())
 	{
@@ -38,7 +38,7 @@ void SControlPicker::SetControlRig(UBaseControlRig* InRig)
 	}
 }
 
-UBaseControlRig* SControlPicker::GetRig() const
+UControlRig* SControlPicker::GetRig() const
 {
 	return RigPtr.Get();
 }
diff --git a/Engine/Plugins/Animation/ControlRig/Source/ControlRigEditor/Private/EditMode/SControlPicker.h b/Engine/Plugins/Animation/ControlRig/Source/ControlRigEditor/Private/EditMode/SControlPicker.h
--- a/Engine/Plugins/Animation/ControlRig/Source/ControlRigEditor/Private/EditMode/SControlPicker.h
+++ b/Engine/Plugins/Animation/ControlRig/Source/ControlRigEditor/Private/EditMode/SControlPicker.h
@@ -8,7 +8,7 @@
 #include "Delegates/Delegate.h"
 #include "UObject/WeakObjectPtr.h"
 
-class UBaseControlRig;
+class UControlRig;
 class SControlPicker;
 class SCanvas;
 class SScaleBox;
@@ -27,7 +27,7 @@ public:
 	void Construct(const FArguments& InArgs, UWorld* InWorld);
 
 	/** Set the rig to display manipulators for */
-	void SetControlRig(UBaseControlRig* InRig);
+	void SetControlRig(UControlRig* InRig);
 
 	/** Set the manipulators that are currently selected */
 	void SetSelectedManipulators(const TArray<FName>& Manipulators);
@@ -43,7 +43,7 @@ public:
 	bool IsManipulatorSelected(FName ManipulatorName);
 
 	/** Returns the rig we are displaying controls for */
-	UBaseControlRig* GetRig() const;
+	UControlRig* GetRig() const;
 
 	/** See if the limb/spine is in IK mode, or FK */
 	bool IsControlIK(FName ControlName) const;
@@ -59,5 +59,5 @@ protected:
 	TSharedPtr<SEditorUserWidgetHost> EditorUserWidgetHost;
 
 	/** Rig we are showing controls for */
-	TWeakObjectPtr<UBaseControlRig> RigPtr;
+	TWeakObjectPtr<UControlRig> RigPtr;
 };
diff --git a/Engine/Plugins/Animation/ControlRig/Source/ControlRigEditor/Private/EditMode/SControlRigBaseListWidget.cpp b/Engine/Plugins/Animation/ControlRig/Source/ControlRigEditor/Private/EditMode/SControlRigBaseListWidget.cpp
--- a/Engine/Plugins/Animation/ControlRig/Source/ControlRigEditor/Private/EditMode/SControlRigBaseListWidget.cpp
+++ b/Engine/Plugins/Animation/ControlRig/Source/ControlRigEditor/Private/EditMode/SControlRigBaseListWidget.cpp
@@ -17,7 +17,7 @@
 
 #include "ScopedTransaction.h"
 
-#include "ControlRig.h"
+#include "BaseControlRig.h"
 #include "UnrealEdGlobals.h"
 #include "EditMode/ControlRigEditMode.h"
 #include "Tools/ControlRigPose.h"
@@ -423,7 +423,7 @@ void SControlRigPoseAnimSelectionToolbar::MakeControlRigAssetDialog(FControlRigA
 		return;
 	}
 
-	TMap<UBaseControlRig*, TArray<FRigElementKey>> AllSelectedControls;
+	TMap<UControlRig*, TArray<FRigElementKey>> AllSelectedControls;
 	ControlRigEditMode->GetAllSelectedControls(AllSelectedControls);
 	
 	if (AllSelectedControls.Num() > 1)
@@ -455,13 +455,13 @@ void SControlRigPoseAnimSelectionToolbar::MakeControlRigAssetDialog(FControlRigA
 			FControlRigEditMode* ControlRigEditMode =  OwningControlRigWidget->GetEditMode();
 			if (ControlRigEditMode )
 			{
-				TMap<UBaseControlRig*, TArray<FRigElementKey>> AllSelectedControls;
+				TMap<UControlRig*, TArray<FRigElementKey>> AllSelectedControls;
 				ControlRigEditMode->GetAllSelectedControls(AllSelectedControls);
 				if (AllSelectedControls.Num() == 1)
 				{
-					TArray<UBaseControlRig*> ControlRigs;
+					TArray<UControlRig*> ControlRigs;
 					AllSelectedControls.GenerateKeyArray(ControlRigs);
-					UBaseControlRig* ControlRig =ControlRigs[0];
+					UControlRig* ControlRig =ControlRigs[0];
 					UObject* NewAsset = nullptr;
 					switch (Type)
 					{
@@ -497,7 +497,7 @@ bool SControlRigPoseAnimSelectionToolbar::CanExecuteMakeControlRigAsset()
 		return false;
 	}
 	
-	TMap<UBaseControlRig*, TArray<FRigElementKey>> AllSelectedControls;
+	TMap<UControlRig*, TArray<FRigElementKey>> AllSelectedControls;
 	ControlRigEditMode->GetAllSelectedControls(AllSelectedControls);
 
 	if (AllSelectedControls.Num() != 1)
@@ -990,11 +990,11 @@ void SControlRigBaseListWidget::OnAssetsActivated(const TArray<FAssetData>& Sele
 					FControlRigEditMode* ControlRigEditMode = GetEditMode();
 					if (ControlRigEditMode)
 					{
-						TArray<UBaseControlRig*> ControlRigs = ControlRigEditMode->GetControlRigsArray(true /*bIsVisible*/);
+						TArray<UControlRig*> ControlRigs = ControlRigEditMode->GetControlRigsArray(true /*bIsVisible*/);
 						if (ControlRigs.Num() > 0)
 						{
 							const FScopedTransaction Transaction(LOCTEXT("SelectControls", "Select Controls"));
-							for (UBaseControlRig* ControlRig : ControlRigs)
+							for (UControlRig* ControlRig : ControlRigs)
 							{
 								ControlRig->Modify();
 								PoseAsset->SelectControls(ControlRig, SControlRigPoseView::IsMirror());
@@ -1368,11 +1368,11 @@ void SControlRigBaseListWidget::ExecutePastePose(UControlRigPoseAsset* PoseAsset
 		FControlRigEditMode* ControlRigEditMode = GetEditMode();
 		if (ControlRigEditMode)
 		{
-			TArray<UBaseControlRig*> ControlRigs = ControlRigEditMode->GetControlRigsArray(true /*bIsVisible*/);
+			TArray<UControlRig*> ControlRigs = ControlRigEditMode->GetControlRigsArray(true /*bIsVisible*/);
 			if (ControlRigs.Num() > 0)
 			{
 				const FScopedTransaction Transaction(LOCTEXT("PastePose", "Paste Pose"));
-				for (UBaseControlRig* ControlRig : ControlRigs)
+				for (UControlRig* ControlRig : ControlRigs)
 				{
 					ControlRig->Modify();
 					PoseAsset->PastePose(ControlRig);
@@ -1397,11 +1397,11 @@ void SControlRigBaseListWidget::ExecuteSelectControls(UControlRigPoseAsset* Pose
 	FControlRigEditMode* ControlRigEditMode = GetEditMode();
 	if (ControlRigEditMode)
 	{
-		TArray<UBaseControlRig*> ControlRigs = ControlRigEditMode->GetControlRigsArray(true /*bIsVisible*/);
+		TArray<UControlRig*> ControlRigs = ControlRigEditMode->GetControlRigsArray(true /*bIsVisible*/);
 		if (ControlRigs.Num() > 0)
 		{
 			const FScopedTransaction Transaction(LOCTEXT("SelectControls", "Select Controls"));
-			for (UBaseControlRig* ControlRig : ControlRigs)
+			for (UControlRig* ControlRig : ControlRigs)
 			{
 				ControlRig->Modify();
 				PoseAsset->SelectControls(ControlRig);
@@ -1434,11 +1434,11 @@ void SControlRigBaseListWidget::ExecutePasteMirrorPose(UControlRigPoseAsset* Pos
 		FControlRigEditMode* ControlRigEditMode = GetEditMode();
 		if (ControlRigEditMode)
 		{
-			TArray<UBaseControlRig*> ControlRigs = ControlRigEditMode->GetControlRigsArray(true /*bIsVisible*/);
+			TArray<UControlRig*> ControlRigs = ControlRigEditMode->GetControlRigsArray(true /*bIsVisible*/);
 			if (ControlRigs.Num() > 0)
 			{
 				const FScopedTransaction Transaction(LOCTEXT("PasteMirrorPose", "Paste Mirror Pose"));
-				for (UBaseControlRig* ControlRig : ControlRigs)
+				for (UControlRig* ControlRig : ControlRigs)
 				{
 					ControlRig->Modify();
 					PoseAsset->PastePose(ControlRig, false, true);
diff --git a/Engine/Plugins/Animation/ControlRig/Source/ControlRigEditor/Private/EditMode/SControlRigBaseListWidget.h b/Engine/Plugins/Animation/ControlRig/Source/ControlRigEditor/Private/EditMode/SControlRigBaseListWidget.h
--- a/Engine/Plugins/Animation/ControlRig/Source/ControlRigEditor/Private/EditMode/SControlRigBaseListWidget.h
+++ b/Engine/Plugins/Animation/ControlRig/Source/ControlRigEditor/Private/EditMode/SControlRigBaseListWidget.h
@@ -20,7 +20,7 @@
 #include "ContentBrowserDelegates.h"
 #include "EditMode/SControlRigControlViews.h"
 
-class UBaseControlRig;
+class UControlRig;
 class FControlRigEditMode;
 class UControlRigPoseAsset;
 class FControlRigEditModeToolkit;
diff --git a/Engine/Plugins/Animation/ControlRig/Source/ControlRigEditor/Private/EditMode/SControlRigControlViews.cpp b/Engine/Plugins/Animation/ControlRig/Source/ControlRigEditor/Private/EditMode/SControlRigControlViews.cpp
--- a/Engine/Plugins/Animation/ControlRig/Source/ControlRigEditor/Private/EditMode/SControlRigControlViews.cpp
+++ b/Engine/Plugins/Animation/ControlRig/Source/ControlRigEditor/Private/EditMode/SControlRigControlViews.cpp
@@ -446,8 +446,8 @@ void SControlRigPoseView::Construct(const FArguments& InArgs)
 	if (FControlRigEditMode* EditMode = OwningWidget->GetEditMode())
 	{
 		EditMode->OnControlRigAddedOrRemoved().AddRaw(this, &SControlRigPoseView::HandleControlAdded);
-		TArray<UBaseControlRig*> ControlRigs = GetControlRigs();
-		for (UBaseControlRig* ControlRig : ControlRigs)
+		TArray<UControlRig*> ControlRigs = GetControlRigs();
+		for (UControlRig* ControlRig : ControlRigs)
 		{
 			HandleControlAdded(ControlRig, true);
 		}
@@ -459,8 +459,8 @@ SControlRigPoseView::~SControlRigPoseView()
 	if (FControlRigEditMode* EditMode = OwningWidget->GetEditMode())
 	{
 		EditMode->OnControlRigAddedOrRemoved().RemoveAll(this);
-		TArray<UBaseControlRig*> EditModeRigs = EditMode->GetControlRigsArray(false /*bIsVisible*/);
-		for (UBaseControlRig* ControlRig : EditModeRigs)
+		TArray<UControlRig*> EditModeRigs = EditMode->GetControlRigsArray(false /*bIsVisible*/);
+		for (UControlRig* ControlRig : EditModeRigs)
 		{
 			if (ControlRig)
 			{
@@ -470,7 +470,7 @@ SControlRigPoseView::~SControlRigPoseView()
 	}
 	else
 	{
-		for (TWeakObjectPtr<UBaseControlRig>& CurrentControlRig: CurrentControlRigs)
+		for (TWeakObjectPtr<UControlRig>& CurrentControlRig: CurrentControlRigs)
 		{
 			if (CurrentControlRig.IsValid())
 			{
@@ -538,8 +538,8 @@ FReply SControlRigPoseView::OnPastePose()
 {
 	if (PoseAsset.IsValid())
 	{
-		TArray<UBaseControlRig*> ControlRigs = GetControlRigs();
-		for (UBaseControlRig* ControlRig : ControlRigs)
+		TArray<UControlRig*> ControlRigs = GetControlRigs();
+		for (UControlRig* ControlRig : ControlRigs)
 		{
 			if (ControlRig)
 			{
@@ -554,8 +554,8 @@ FReply SControlRigPoseView::OnSelectControls()
 {	
 	if (PoseAsset.IsValid())
 	{
-		TArray<UBaseControlRig*> ControlRigs = GetControlRigs();
-		for (UBaseControlRig* ControlRig : ControlRigs)
+		TArray<UControlRig*> ControlRigs = GetControlRigs();
+		for (UControlRig* ControlRig : ControlRigs)
 		{
 			if (ControlRig)
 			{
@@ -570,8 +570,8 @@ void SControlRigPoseView::OnPoseBlendChanged(float ChangedVal)
 {
 	if (PoseAsset.IsValid())
 	{
-		TArray<UBaseControlRig*> ControlRigs = GetControlRigs();
-		for (UBaseControlRig* ControlRig : ControlRigs)
+		TArray<UControlRig*> ControlRigs = GetControlRigs();
+		for (UControlRig* ControlRig : ControlRigs)
 		{
 			if (ControlRig)
 			{
@@ -610,11 +610,11 @@ void SControlRigPoseView::OnPoseBlendCommited(float ChangedVal, ETextCommit::Typ
 {
 	if (PoseAsset.IsValid())
 	{
-		TArray<UBaseControlRig*> ControlRigs = GetControlRigs();
+		TArray<UControlRig*> ControlRigs = GetControlRigs();
 		if (ControlRigs.Num() > 0)
 		{
 			FScopedTransaction ScopedTransaction(LOCTEXT("PastePoseTransaction", "Paste Pose"));
-			for (UBaseControlRig* ControlRig : ControlRigs)
+			for (UControlRig* ControlRig : ControlRigs)
 			{
 				if (ControlRig)
 				{
@@ -662,15 +662,15 @@ TSharedRef<SWidget> SControlRigPoseView::GetThumbnailWidget()
 		];
 }
 
-TArray<UBaseControlRig*> SControlRigPoseView::GetControlRigs()
+TArray<UControlRig*> SControlRigPoseView::GetControlRigs()
 {
 	FControlRigEditMode* EditMode = OwningWidget->GetEditMode();
-	TArray<UBaseControlRig*> NewControlRigs;
+	TArray<UControlRig*> NewControlRigs;
 	if (EditMode)
 	{
 		NewControlRigs =  EditMode->GetControlRigsArray(false /*bIsVisible*/);
 	}
-	for (TWeakObjectPtr<UBaseControlRig> ControlRigPtr : CurrentControlRigs)
+	for (TWeakObjectPtr<UControlRig> ControlRigPtr : CurrentControlRigs)
 	{
 		if (ControlRigPtr.IsValid())
 		{
@@ -701,7 +701,7 @@ void SControlRigPoseView::CreateControlList()
 	}
 }
 */
-void SControlRigPoseView::HandleControlAdded(UBaseControlRig* ControlRig, bool bIsAdded)
+void SControlRigPoseView::HandleControlAdded(UControlRig* ControlRig, bool bIsAdded)
 {
 	if (ControlRig)
 	{
@@ -718,7 +718,7 @@ void SControlRigPoseView::HandleControlAdded(UBaseControlRig* ControlRig, bool b
 	UpdateStatusBlocks();
 }
 
-void SControlRigPoseView::HandleControlSelected(UBaseControlRig* Subject, FRigControlElement* InControl, bool bSelected)
+void SControlRigPoseView::HandleControlSelected(UControlRig* Subject, FRigControlElement* InControl, bool bSelected)
 {
 	UpdateStatusBlocks();
 }
@@ -727,7 +727,7 @@ void SControlRigPoseView::UpdateStatusBlocks()
 {
 	FText StatusText1;
 	FText StatusText2;
-	TArray<UBaseControlRig*> ControlRigs = GetControlRigs();
+	TArray<UControlRig*> ControlRigs = GetControlRigs();
 	if (PoseAsset.IsValid() && ControlRigs.Num() > 0)
 	{
 
@@ -737,7 +737,7 @@ void SControlRigPoseView::UpdateStatusBlocks()
 		int32 TotalSelected = 0;
 		uint32 Matching = 0;
 		uint32 MirrorMatching = 0;
-		for (UBaseControlRig* ControlRig : ControlRigs)
+		for (UControlRig* ControlRig : ControlRigs)
 		{
 
 			TArray<FName> SelectedNames = ControlRig->CurrentControlSelection();
diff --git a/Engine/Plugins/Animation/ControlRig/Source/ControlRigEditor/Private/EditMode/SControlRigControlViews.h b/Engine/Plugins/Animation/ControlRig/Source/ControlRigEditor/Private/EditMode/SControlRigControlViews.h
--- a/Engine/Plugins/Animation/ControlRig/Source/ControlRigEditor/Private/EditMode/SControlRigControlViews.h
+++ b/Engine/Plugins/Animation/ControlRig/Source/ControlRigEditor/Private/EditMode/SControlRigControlViews.h
@@ -19,7 +19,7 @@
 
 class FAssetThumbnail;
 class IDetailsView;
-class UBaseControlRig;
+class UControlRig;
 class SControlRigBaseListWidget;
 
 //Class to Hold Statics that are shared and externally callable
@@ -62,9 +62,9 @@ private:
 	FReply OnPastePose();
 	FReply OnSelectControls();
 	FReply OnCaptureThumbnail();
-	void HandleControlSelected(UBaseControlRig* Subject, FRigControlElement* InControl, bool bSelected);
-	void HandleControlAdded(UBaseControlRig* ControlRig, bool bIsAdded);
-	TArray<UBaseControlRig*> GetControlRigs();
+	void HandleControlSelected(UControlRig* Subject, FRigControlElement* InControl, bool bSelected);
+	void HandleControlAdded(UControlRig* ControlRig, bool bIsAdded);
+	TArray<UControlRig*> GetControlRigs();
 
 	static bool bIsKey;
 	static bool bIsMirror;
@@ -76,7 +76,7 @@ private:
 	TSharedRef<SWidget> GetThumbnailWidget();
 
 	TWeakObjectPtr<UControlRigPoseAsset> PoseAsset;
-	TArray<TWeakObjectPtr<UBaseControlRig>> CurrentControlRigs;
+	TArray<TWeakObjectPtr<UControlRig>> CurrentControlRigs;
 
 	TSharedPtr<STextBlock> TextStatusBlock1;
 	TSharedPtr<STextBlock> TextStatusBlock2;
diff --git a/Engine/Plugins/Animation/ControlRig/Source/ControlRigEditor/Private/EditMode/SControlRigDetails.cpp b/Engine/Plugins/Animation/ControlRig/Source/ControlRigEditor/Private/EditMode/SControlRigDetails.cpp
--- a/Engine/Plugins/Animation/ControlRig/Source/ControlRigEditor/Private/EditMode/SControlRigDetails.cpp
+++ b/Engine/Plugins/Animation/ControlRig/Source/ControlRigEditor/Private/EditMode/SControlRigDetails.cpp
@@ -157,7 +157,7 @@ void FControlRigEditModeGenericDetails::CustomizeDetails(class IDetailLayoutBuil
 				{
 					if (FControlRigEditMode* EditMode = static_cast<FControlRigEditMode*>(ModeTools->GetActiveMode(FControlRigEditMode::ModeName)))
 					{
-						UBaseControlRig* ControlRig = Proxy->ControlRig.Get();
+						UControlRig* ControlRig = Proxy->ControlRig.Get();
 						if (UObject* NestedProxy = EditMode->ControlProxy->FindProxy(ControlRig,ControlElement->GetFName()))
 						{
 							FName PropertyName(NAME_None);
@@ -379,7 +379,7 @@ SControlRigDetails::~SControlRigDetails()
 	//base class handles control rig related cleanup
 }
 
-void SControlRigDetails::HandleControlSelected(UBaseControlRig* Subject, FRigControlElement* InControl, bool bSelected)
+void SControlRigDetails::HandleControlSelected(UControlRig* Subject, FRigControlElement* InControl, bool bSelected)
 {
 	FControlRigBaseDockableView::HandleControlSelected(Subject, InControl, bSelected);
 	UpdateProxies();
@@ -420,7 +420,7 @@ void SControlRigDetails::UpdateProxies()
 		TArray<TWeakObjectPtr<>> IndividualBools;
 		TArray<TWeakObjectPtr<>> IndividualIntegers;
 		TArray<TWeakObjectPtr<>> IndividualEnums;
-		TArray<UBaseControlRig*> ControlRigs = StrongThis->GetControlRigs();
+		TArray<UControlRig*> ControlRigs = StrongThis->GetControlRigs();
 	
 		if (FControlRigEditMode* EditMode = static_cast<FControlRigEditMode*>(StrongThis->ModeTools->GetActiveMode(FControlRigEditMode::ModeName)))
 		{
@@ -662,7 +662,7 @@ bool SControlRigDetails::IsPropertyKeyable(const UClass* InObjectClass, const IP
 	{
 		if (UControlRigControlsProxy* Proxy = Cast< UControlRigControlsProxy>(OuterObjects[0]))
 		{
-			if(UBaseControlRig* ControlRig = Proxy->ControlRig.Get())
+			if(UControlRig* ControlRig = Proxy->ControlRig.Get())
 			{
 				const FRigElementKey Key = FRigElementKey(Proxy->ControlName, ERigElementType::Control);
 				if(const FRigControlElement* ControlElement = ControlRig->GetHierarchy()->Find<FRigControlElement>(Key))
diff --git a/Engine/Plugins/Animation/ControlRig/Source/ControlRigEditor/Private/EditMode/SControlRigDetails.h b/Engine/Plugins/Animation/ControlRig/Source/ControlRigEditor/Private/EditMode/SControlRigDetails.h
--- a/Engine/Plugins/Animation/ControlRig/Source/ControlRigEditor/Private/EditMode/SControlRigDetails.h
+++ b/Engine/Plugins/Animation/ControlRig/Source/ControlRigEditor/Private/EditMode/SControlRigDetails.h
@@ -19,7 +19,7 @@
 #include "DetailWidgetRow.h"
 
 class ISequencer;
-class UBaseControlRig;
+class UControlRig;
 
 class FControlRigEditModeGenericDetails : public IDetailCustomization
 {
@@ -75,7 +75,7 @@ private:
 	void SetEnumDetailsObjects(const TArray<TWeakObjectPtr<>>& InObjects, bool bIsIndividual);
 
 	void UpdateProxies();
-	virtual void HandleControlSelected(UBaseControlRig* Subject, FRigControlElement* InControl, bool bSelected) override;
+	virtual void HandleControlSelected(UControlRig* Subject, FRigControlElement* InControl, bool bSelected) override;
 
 	//these views will hold more than one of the same tyhpe
 	TSharedPtr<IDetailsView> ControlEulerTransformDetailsView;
diff --git a/Engine/Plugins/Animation/ControlRig/Source/ControlRigEditor/Private/EditMode/SControlRigEditModeTools.cpp b/Engine/Plugins/Animation/ControlRig/Source/ControlRigEditor/Private/EditMode/SControlRigEditModeTools.cpp
--- a/Engine/Plugins/Animation/ControlRig/Source/ControlRigEditor/Private/EditMode/SControlRigEditModeTools.cpp
+++ b/Engine/Plugins/Animation/ControlRig/Source/ControlRigEditor/Private/EditMode/SControlRigEditModeTools.cpp
@@ -46,9 +46,9 @@
 //statics to reuse in the UI
 FRigSpacePickerBakeSettings SControlRigEditModeTools::BakeSpaceSettings;
 
-void SControlRigEditModeTools::SetControlRigs(const TArrayView<TWeakObjectPtr<UBaseControlRig>>& InControlRigs)
+void SControlRigEditModeTools::SetControlRigs(const TArrayView<TWeakObjectPtr<UControlRig>>& InControlRigs)
 {
-	for (TWeakObjectPtr<UBaseControlRig>& ControlRig : ControlRigs)
+	for (TWeakObjectPtr<UControlRig>& ControlRig : ControlRigs)
 	{
 		if (ControlRig.IsValid())
 		{
@@ -56,7 +56,7 @@ void SControlRigEditModeTools::SetControlRigs(const TArrayView<TWeakObjectPtr<UB
 		}
 	}
 	ControlRigs = InControlRigs;
-	for (TWeakObjectPtr<UBaseControlRig>& InControlRig : InControlRigs)
+	for (TWeakObjectPtr<UControlRig>& InControlRig : InControlRigs)
 	{
 		if (InControlRig.IsValid())
 		{
@@ -65,7 +65,7 @@ void SControlRigEditModeTools::SetControlRigs(const TArrayView<TWeakObjectPtr<UB
 	}
 
 	//mz todo handle multiple rigs
-	UBaseControlRig* Rig = ControlRigs.Num() > 0 ? ControlRigs[0].Get() : nullptr;
+	UControlRig* Rig = ControlRigs.Num() > 0 ? ControlRigs[0].Get() : nullptr;
 	TArray<TWeakObjectPtr<>> Objects;
 	Objects.Add(Rig);
 	RigOptionsDetailsView->SetObjects(Objects);
@@ -79,7 +79,7 @@ void SControlRigEditModeTools::SetControlRigs(const TArrayView<TWeakObjectPtr<UB
 const URigHierarchy* SControlRigEditModeTools::GetHierarchy() const
 {
 	//mz todo handle multiple rigs
-	UBaseControlRig* Rig = ControlRigs.Num() > 0 ? ControlRigs[0].Get() : nullptr;
+	UControlRig* Rig = ControlRigs.Num() > 0 ? ControlRigs[0].Get() : nullptr;
 	if (Rig)
 	{
 		Rig->GetHierarchy();
@@ -810,7 +810,7 @@ void SControlRigEditModeTools::HandleSelectionChanged(TSharedPtr<FRigTreeElement
 }
 #endif
 
-void SControlRigEditModeTools::OnRigElementSelected(UBaseControlRig* Subject, FRigControlElement* ControlElement, bool bSelected)
+void SControlRigEditModeTools::OnRigElementSelected(UControlRig* Subject, FRigControlElement* ControlElement, bool bSelected)
 {
 #if USE_LOCAL_DETAILS
 	const FRigElementKey Key = ControlElement->GetKey();
@@ -850,8 +850,8 @@ void SControlRigEditModeTools::OnRigElementSelected(UBaseControlRig* Subject, FR
 
 const FRigControlElementCustomization* SControlRigEditModeTools::HandleGetControlElementCustomization(URigHierarchy* InHierarchy, const FRigElementKey& InControlKey)
 {
-	UBaseControlRig* Rig = ControlRigs.Num() > 0 ? ControlRigs[0].Get() : nullptr;
-	for (TWeakObjectPtr<UBaseControlRig>& ControlRig : ControlRigs)
+	UControlRig* Rig = ControlRigs.Num() > 0 ? ControlRigs[0].Get() : nullptr;
+	for (TWeakObjectPtr<UControlRig>& ControlRig : ControlRigs)
 	{
 		if (ControlRig.IsValid() && ControlRig->GetHierarchy() == InHierarchy)
 		{
@@ -867,7 +867,7 @@ void SControlRigEditModeTools::HandleActiveSpaceChanged(URigHierarchy* InHierarc
 
 	if (WeakSequencer.IsValid())
 	{
-		for (TWeakObjectPtr<UBaseControlRig>& ControlRig : ControlRigs)
+		for (TWeakObjectPtr<UControlRig>& ControlRig : ControlRigs)
 		{
 			if (ControlRig.IsValid() && ControlRig->GetHierarchy() == InHierarchy)
 			{
@@ -914,7 +914,7 @@ void SControlRigEditModeTools::HandleSpaceListChanged(URigHierarchy* InHierarchy
 {
 	FScopedTransaction Transaction(LOCTEXT("ChangeControlRigSpace", "Change Control Rig Space"));
 
-	for (TWeakObjectPtr<UBaseControlRig>& ControlRig : ControlRigs)
+	for (TWeakObjectPtr<UControlRig>& ControlRig : ControlRigs)
 	{
 		if (ControlRig.IsValid() && ControlRig->GetHierarchy() == InHierarchy)
 		{
@@ -980,7 +980,7 @@ FReply SControlRigEditModeTools::OnBakeControlsToNewSpaceButtonClicked()
 	}
 
 	bool bNoValidControlRig = true;
-	for (TWeakObjectPtr<UBaseControlRig>& ControlRig : ControlRigs)
+	for (TWeakObjectPtr<UControlRig>& ControlRig : ControlRigs)
 	{
 		if (ControlRig.IsValid() && SpacePickerWidget->GetHierarchy() == ControlRig->GetHierarchy())
 		{
@@ -998,7 +998,7 @@ FReply SControlRigEditModeTools::OnBakeControlsToNewSpaceButtonClicked()
 	{
 		return FReply::Unhandled();
 	}
-	for (TWeakObjectPtr<UBaseControlRig>& ControlRig : ControlRigs)
+	for (TWeakObjectPtr<UControlRig>& ControlRig : ControlRigs)
 	{
 		if (ControlRig.IsValid() && SpacePickerWidget->GetHierarchy() == ControlRig->GetHierarchy())
 		{
@@ -1099,7 +1099,7 @@ FReply SControlRigEditModeTools::HandleAddConstraintClicked()
 
 EVisibility SControlRigEditModeTools::GetRigOptionExpanderVisibility() const
 {
-	for (const TWeakObjectPtr<UBaseControlRig>& ControlRig : ControlRigs)
+	for (const TWeakObjectPtr<UControlRig>& ControlRig : ControlRigs)
 	{
 		if (ControlRig.IsValid())
 		{
@@ -1114,7 +1114,7 @@ EVisibility SControlRigEditModeTools::GetRigOptionExpanderVisibility() const
 
 void SControlRigEditModeTools::OnRigOptionFinishedChange(const FPropertyChangedEvent& PropertyChangedEvent)
 {
-	TArray<TWeakObjectPtr<UBaseControlRig>> ControlRigsCopy = ControlRigs;
+	TArray<TWeakObjectPtr<UControlRig>> ControlRigsCopy = ControlRigs;
 	SetControlRigs(ControlRigsCopy);
 
 	if (FControlRigEditMode* EditMode = static_cast<FControlRigEditMode*>(ModeTools->GetActiveMode(FControlRigEditMode::ModeName)))
diff --git a/Engine/Plugins/Animation/ControlRig/Source/ControlRigEditor/Private/EditMode/SControlRigEditModeTools.h b/Engine/Plugins/Animation/ControlRig/Source/ControlRigEditor/Private/EditMode/SControlRigEditModeTools.h
--- a/Engine/Plugins/Animation/ControlRig/Source/ControlRigEditor/Private/EditMode/SControlRigEditModeTools.h
+++ b/Engine/Plugins/Animation/ControlRig/Source/ControlRigEditor/Private/EditMode/SControlRigEditModeTools.h
@@ -17,7 +17,7 @@ class ISequencer;
 class SControlPicker;
 class SExpandableArea;
 class SRigHierarchyTreeView;
-class UBaseControlRig;
+class UControlRig;
 class URigHierarchy;
 class FToolBarBuilder;
 class FEditorModeTools;
@@ -50,7 +50,7 @@ public:
 	void SetSequencer(TWeakPtr<ISequencer> InSequencer);
 
 	/** Set The Control Rig we are using*/
-	void SetControlRigs(const TArrayView<TWeakObjectPtr<UBaseControlRig>>& InControlRigs);
+	void SetControlRigs(const TArrayView<TWeakObjectPtr<UControlRig>>& InControlRigs);
 
 	/** Returns the hierarchy currently being used */
 	const URigHierarchy* GetHierarchy() const;
@@ -95,7 +95,7 @@ private:
 	TSharedPtr<SRigSpacePickerWidget> SpacePickerWidget;
 
 	/** Storage for control rigs */
-	TArray<TWeakObjectPtr<UBaseControlRig>> ControlRigs;
+	TArray<TWeakObjectPtr<UControlRig>> ControlRigs;
 
 	/** Constraint edition widget. */
 	TSharedPtr<SExpandableArea> ConstraintPickerExpander = nullptr;
@@ -110,7 +110,7 @@ private:
 
 	void HandleModifiedEvent(ERigVMGraphNotifType InNotifType, URigVMGraph* InGraph, UObject* InSubject);
 	void HandleSelectionChanged(TSharedPtr<FRigTreeElement> Selection, ESelectInfo::Type SelectInfo);
-	void OnRigElementSelected(UBaseControlRig* Subject, FRigControlElement* ControlElement, bool bSelected);
+	void OnRigElementSelected(UControlRig* Subject, FRigControlElement* ControlElement, bool bSelected);
 
 	const FRigControlElementCustomization* HandleGetControlElementCustomization(URigHierarchy* InHierarchy, const FRigElementKey& InControlKey);
 	void HandleActiveSpaceChanged(URigHierarchy* InHierarchy, const FRigElementKey& InControlKey, const FRigElementKey& InSpaceKey);
diff --git a/Engine/Plugins/Animation/ControlRig/Source/ControlRigEditor/Private/EditMode/SControlRigOutliner.cpp b/Engine/Plugins/Animation/ControlRig/Source/ControlRigEditor/Private/EditMode/SControlRigOutliner.cpp
--- a/Engine/Plugins/Animation/ControlRig/Source/ControlRigEditor/Private/EditMode/SControlRigOutliner.cpp
+++ b/Engine/Plugins/Animation/ControlRig/Source/ControlRigEditor/Private/EditMode/SControlRigOutliner.cpp
@@ -51,7 +51,7 @@ FRigTreeDisplaySettings FMultiRigTreeDelegates::DefaultDisplaySettings;
 
 uint32 GetTypeHash(const FMultiRigData& Data)
 {
-	return GetTypeHash(TTuple<const UBaseControlRig*, FRigElementKey>(Data.ControlRig.Get(), (Data.Key.IsSet() ? Data.Key.GetValue() : FRigElementKey())));
+	return GetTypeHash(TTuple<const UControlRig*, FRigElementKey>(Data.ControlRig.Get(), (Data.Key.IsSet() ? Data.Key.GetValue() : FRigElementKey())));
 }
 
 FText FMultiRigData::GetName() const
@@ -233,7 +233,7 @@ void SMultiRigHierarchyItem::Construct(const FArguments& InArgs, const TSharedRe
 						//if no key is set then it's the control rig so we get that based upon it's state
 						if (WeakRigTreeElement.Pin()->Data.Key.IsSet() == false)
 						{
-							if (UBaseControlRig* ControlRig = WeakRigTreeElement.Pin()->Data.ControlRig.Get())
+							if (UControlRig* ControlRig = WeakRigTreeElement.Pin()->Data.ControlRig.Get())
 							{
 								if (ControlRig->GetControlsVisible())
 								{
@@ -280,7 +280,7 @@ FReply SMultiRigHierarchyItem::OnGetSelectedClicked()
 	{
 		if (Element->Data.Key.IsSet() == false)
 		{
-			if (UBaseControlRig* ControlRig = Element->Data.ControlRig.Get())
+			if (UControlRig* ControlRig = Element->Data.ControlRig.Get())
 			{
 				FScopedTransaction ScopedTransaction(LOCTEXT("ToggleControlsVisibility", "Toggle Controls Visibility"), !GIsTransacting);
 				ControlRig->Modify();
@@ -313,7 +313,7 @@ TPair<const FSlateBrush*, FSlateColor> SMultiRigHierarchyItem::GetBrushForElemen
 	}
 	else
 	{
-		if (UBaseControlRig* ControlRig = InData.ControlRig.Get())
+		if (UControlRig* ControlRig = InData.ControlRig.Get())
 		{
 			if (ControlRig->GetControlsVisible())
 			{
@@ -466,7 +466,7 @@ bool SMultiRigHierarchyTreeView::AddElement(const FMultiRigData& InData, const F
 	return true;
 }
 
-bool SMultiRigHierarchyTreeView::AddElement(UBaseControlRig* InControlRig, const FRigBaseElement* InElement)
+bool SMultiRigHierarchyTreeView::AddElement(UControlRig* InControlRig, const FRigBaseElement* InElement)
 {
 	check(InControlRig);
 	check(InElement);
@@ -725,9 +725,9 @@ void SMultiRigHierarchyTreeView::RefreshTreeView(bool bRebuildContent)
 
 	if (bRebuildContent)
 	{
-		for (const TWeakObjectPtr<UBaseControlRig>& ControlRigPtr : ControlRigs)
+		for (const TWeakObjectPtr<UControlRig>& ControlRigPtr : ControlRigs)
 		{
-			if (UBaseControlRig* ControlRig  = ControlRigPtr.Get())
+			if (UControlRig* ControlRig  = ControlRigPtr.Get())
 			{
 				FMultiRigData Empty(nullptr, FRigElementKey());
 				FMultiRigData CRData;
@@ -787,9 +787,9 @@ void SMultiRigHierarchyTreeView::RefreshTreeView(bool bRebuildContent)
 	{
 		ClearSelection();
 
-		for (const TWeakObjectPtr<UBaseControlRig>& ControlRigPtr : ControlRigs)
+		for (const TWeakObjectPtr<UControlRig>& ControlRigPtr : ControlRigs)
 		{
-			if (UBaseControlRig* ControlRig = ControlRigPtr.Get())
+			if (UControlRig* ControlRig = ControlRigPtr.Get())
 			{
 				if (const URigHierarchy* Hierarchy = ControlRig->GetHierarchy())
 				{
@@ -864,7 +864,7 @@ TArray<FMultiRigData> SMultiRigHierarchyTreeView::GetSelectedData() const
 TArray<URigHierarchy*> SMultiRigHierarchyTreeView::GetHierarchy() const
 {
 	TArray<URigHierarchy*> RigHierarchy;
-	for (const TWeakObjectPtr<UBaseControlRig>& ControlRig : ControlRigs)
+	for (const TWeakObjectPtr<UControlRig>& ControlRig : ControlRigs)
 	{
 		if (ControlRig.IsValid())
 		{
@@ -874,10 +874,10 @@ TArray<URigHierarchy*> SMultiRigHierarchyTreeView::GetHierarchy() const
 	return RigHierarchy;
 }
 
-void SMultiRigHierarchyTreeView::SetControlRigs(TArrayView < TWeakObjectPtr<UBaseControlRig>>& InControlRigs)
+void SMultiRigHierarchyTreeView::SetControlRigs(TArrayView < TWeakObjectPtr<UControlRig>>& InControlRigs)
 {
 	ControlRigs.SetNum(0);
-	for (TWeakObjectPtr<UBaseControlRig>& ControlRig : InControlRigs)
+	for (TWeakObjectPtr<UControlRig>& ControlRig : InControlRigs)
 	{
 		if (ControlRig.IsValid())
 		{
@@ -1006,7 +1006,7 @@ void SControlRigOutliner::OnObjectsReplaced(const TMap<UObject*, UObject*>& OldT
 	{
 		if(Pair.Key && Pair.Value)
 		{
-			if (Pair.Key->IsA<UBaseControlRig>() && Pair.Value->IsA<UBaseControlRig>())
+			if (Pair.Key->IsA<UControlRig>() && Pair.Value->IsA<UControlRig>())
 			{
 				bNewControlRig = false;
 				break;
@@ -1029,7 +1029,7 @@ SControlRigOutliner::~SControlRigOutliner()
 	FCoreUObjectDelegates::OnObjectsReplaced.RemoveAll(this);
 }
 
-void SControlRigOutliner::HandleControlSelected(UBaseControlRig* Subject, FRigControlElement* ControlElement, bool bSelected)
+void SControlRigOutliner::HandleControlSelected(UControlRig* Subject, FRigControlElement* ControlElement, bool bSelected)
 {
 	FControlRigBaseDockableView::HandleControlSelected(Subject, ControlElement, bSelected);
 	const FRigElementKey Key = ControlElement->GetKey();
@@ -1065,7 +1065,7 @@ void SControlRigOutliner::HandleSelectionChanged(TSharedPtr<FMultiRigTreeElement
 		return;
 	}
 	const TArray<FMultiRigData> NewSelection = HierarchyTreeView->GetTreeView()->GetSelectedData();
-	TMap<UBaseControlRig*, TArray<FRigElementKey>> SelectedRigAndKeys;
+	TMap<UControlRig*, TArray<FRigElementKey>> SelectedRigAndKeys;
 	for (const FMultiRigData& Data : NewSelection)
 	{
 		if (Data.ControlRig.IsValid() && (Data.Key.IsSet() && Data.Key.GetValue() != FRigElementKey()))
@@ -1090,9 +1090,9 @@ void SControlRigOutliner::HandleSelectionChanged(TSharedPtr<FMultiRigTreeElement
 	{
 		if (EditMode)
 		{		
-			TMap<UBaseControlRig*, TArray<FRigElementKey>> SelectedControls;
+			TMap<UControlRig*, TArray<FRigElementKey>> SelectedControls;
 			EditMode->GetAllSelectedControls(SelectedControls);
-			for (TPair<UBaseControlRig*, TArray<FRigElementKey>>& CurrentSelection : SelectedControls)
+			for (TPair<UControlRig*, TArray<FRigElementKey>>& CurrentSelection : SelectedControls)
 			{
 				if (CurrentSelection.Key)
 				{
@@ -1102,7 +1102,7 @@ void SControlRigOutliner::HandleSelectionChanged(TSharedPtr<FMultiRigTreeElement
 		}
 	}
 
-	for(TPair<UBaseControlRig*, TArray<FRigElementKey>>& RigAndKeys: SelectedRigAndKeys)
+	for(TPair<UControlRig*, TArray<FRigElementKey>>& RigAndKeys: SelectedRigAndKeys)
 	{ 
 		const URigHierarchy* Hierarchy = RigAndKeys.Key->GetHierarchy();
 		if (Hierarchy)
@@ -1131,17 +1131,17 @@ void SControlRigOutliner::SetEditMode(FControlRigEditMode& InEditMode)
 	ModeTools = InEditMode.GetModeManager();
 	if (FControlRigEditMode* EditMode = static_cast<FControlRigEditMode*>(ModeTools->GetActiveMode(FControlRigEditMode::ModeName)))
 	{
-		TArrayView<TWeakObjectPtr<UBaseControlRig>> ControlRigs = EditMode->GetControlRigs();
+		TArrayView<TWeakObjectPtr<UControlRig>> ControlRigs = EditMode->GetControlRigs();
 		HierarchyTreeView->GetTreeView()->SetControlRigs(ControlRigs); //will refresh tree
 	}
 }
 
-void SControlRigOutliner::HandleControlAdded(UBaseControlRig* ControlRig, bool bIsAdded)
+void SControlRigOutliner::HandleControlAdded(UControlRig* ControlRig, bool bIsAdded)
 {
 	FControlRigBaseDockableView::HandleControlAdded(ControlRig, bIsAdded);
 	if (FControlRigEditMode* EditMode = static_cast<FControlRigEditMode*>(ModeTools->GetActiveMode(FControlRigEditMode::ModeName)))
 	{
-		TArrayView<TWeakObjectPtr<UBaseControlRig>> ControlRigs = EditMode->GetControlRigs();
+		TArrayView<TWeakObjectPtr<UControlRig>> ControlRigs = EditMode->GetControlRigs();
 		HierarchyTreeView->GetTreeView()->SetControlRigs(ControlRigs); //will refresh tree
 	}
 }
diff --git a/Engine/Plugins/Animation/ControlRig/Source/ControlRigEditor/Private/EditMode/SControlRigOutliner.h b/Engine/Plugins/Animation/ControlRig/Source/ControlRigEditor/Private/EditMode/SControlRigOutliner.h
--- a/Engine/Plugins/Animation/ControlRig/Source/ControlRigEditor/Private/EditMode/SControlRigOutliner.h
+++ b/Engine/Plugins/Animation/ControlRig/Source/ControlRigEditor/Private/EditMode/SControlRigOutliner.h
@@ -77,14 +77,14 @@ class FMultiRigData
 
 public:
 	FMultiRigData() {};
-	FMultiRigData(UBaseControlRig* InControlRig, FRigElementKey InKey) : ControlRig(InControlRig), Key(InKey) {};
+	FMultiRigData(UControlRig* InControlRig, FRigElementKey InKey) : ControlRig(InControlRig), Key(InKey) {};
 	FText GetName() const;
 	FText GetDisplayName() const;
 	bool operator == (const FMultiRigData & Other) const;
 	bool IsValid() const;
 	URigHierarchy* GetHierarchy() const;
 public:
-	TWeakObjectPtr<UBaseControlRig> ControlRig;
+	TWeakObjectPtr<UControlRig> ControlRig;
 	TOptional<FRigElementKey> Key;
 };
 
@@ -177,7 +177,7 @@ public:
 	static TSharedPtr<FMultiRigTreeElement> FindElement(const FMultiRigData& InData, TSharedPtr<FMultiRigTreeElement> CurrentItem);
 
 	bool AddElement(const FMultiRigData& InData, const FMultiRigData& InParentData);
-	bool AddElement(UBaseControlRig* InControlRig, const FRigBaseElement* InElement);
+	bool AddElement(UControlRig* InControlRig, const FRigBaseElement* InElement);
 	bool ReparentElement(const FMultiRigData& InData, const FMultiRigData& InParentData);
 	bool RemoveElement(const FMultiRigData& InData);
 	void RefreshTreeView(bool bRebuildContent = true);
@@ -190,7 +190,7 @@ public:
 	FMultiRigTreeDelegates& GetTreeDelegates() { return Delegates; }
 
 	TArray<URigHierarchy*> GetHierarchy() const;
-	void SetControlRigs(TArrayView < TWeakObjectPtr<UBaseControlRig>>& InControlRigs);
+	void SetControlRigs(TArrayView < TWeakObjectPtr<UControlRig>>& InControlRigs);
 
 private:
 
@@ -210,7 +210,7 @@ private:
 
 	friend class SRigHierarchy;
 
-	TArray <TWeakObjectPtr<UBaseControlRig>> ControlRigs;
+	TArray <TWeakObjectPtr<UControlRig>> ControlRigs;
 };
 
 class SSearchableMultiRigHierarchyTreeView : public SCompoundWidget
@@ -249,8 +249,8 @@ class SControlRigOutliner : public FControlRigBaseDockableView, public SCompound
 	//FControlRigBaseDockableView overrides
 	virtual void SetEditMode(FControlRigEditMode& InEditMode) override;
 private:
-	virtual void HandleControlAdded(UBaseControlRig* ControlRig, bool bIsAdded) override;
-	virtual void HandleControlSelected(UBaseControlRig* Subject, FRigControlElement* InControl, bool bSelected) override;
+	virtual void HandleControlAdded(UControlRig* ControlRig, bool bIsAdded) override;
+	virtual void HandleControlSelected(UControlRig* Subject, FRigControlElement* InControl, bool bSelected) override;
 
 	void OnObjectsReplaced(const TMap<UObject*, UObject*>& OldToNewInstanceMap);
 
diff --git a/Engine/Plugins/Animation/ControlRig/Source/ControlRigEditor/Private/EditMode/SControlRigSnapper.cpp b/Engine/Plugins/Animation/ControlRig/Source/ControlRigEditor/Private/EditMode/SControlRigSnapper.cpp
--- a/Engine/Plugins/Animation/ControlRig/Source/ControlRigEditor/Private/EditMode/SControlRigSnapper.cpp
+++ b/Engine/Plugins/Animation/ControlRig/Source/ControlRigEditor/Private/EditMode/SControlRigSnapper.cpp
@@ -501,9 +501,9 @@ void SControlRigSnapper::SetStartEndFrames()
 FControlRigSnapperSelection SControlRigSnapper::GetSelection(bool bGetAll) 
 {
 	FControlRigSnapperSelection Selection;
-	TArray<UBaseControlRig*> ControlRigs;
+	TArray<UControlRig*> ControlRigs;
 	GetControlRigs(ControlRigs);
-	for (UBaseControlRig* ControlRig : ControlRigs)
+	for (UControlRig* ControlRig : ControlRigs)
 	{
 		if (ControlRig)
 		{
@@ -569,15 +569,15 @@ FControlRigSnapperSelection SControlRigSnapper::GetSelection(bool bGetAll)
 	return Selection;
 }
 
-void SControlRigSnapper::GetControlRigs(TArray<UBaseControlRig*>& OutControlRigs) const
+void SControlRigSnapper::GetControlRigs(TArray<UControlRig*>& OutControlRigs) const
 {
 	OutControlRigs.SetNum(0);
 	if (FControlRigEditMode* EditMode = static_cast<FControlRigEditMode*>(GLevelEditorModeTools().GetActiveMode(FControlRigEditMode::ModeName)))
 	{
-		TArrayView<TWeakObjectPtr<UBaseControlRig>> ControlRigPtrs = EditMode->GetControlRigs();
-		for (TWeakObjectPtr<UBaseControlRig>& ControlRigPtr : ControlRigPtrs)
+		TArrayView<TWeakObjectPtr<UControlRig>> ControlRigPtrs = EditMode->GetControlRigs();
+		for (TWeakObjectPtr<UControlRig>& ControlRigPtr : ControlRigPtrs)
 		{
-			if (UBaseControlRig* ControlRig = ControlRigPtr.Get())
+			if (UControlRig* ControlRig = ControlRigPtr.Get())
 			{
 				OutControlRigs.Add(ControlRig);
 			}
diff --git a/Engine/Plugins/Animation/ControlRig/Source/ControlRigEditor/Private/EditMode/SControlRigSnapper.h b/Engine/Plugins/Animation/ControlRig/Source/ControlRigEditor/Private/EditMode/SControlRigSnapper.h
--- a/Engine/Plugins/Animation/ControlRig/Source/ControlRigEditor/Private/EditMode/SControlRigSnapper.h
+++ b/Engine/Plugins/Animation/ControlRig/Source/ControlRigEditor/Private/EditMode/SControlRigSnapper.h
@@ -13,7 +13,7 @@
 #include "IDetailsView.h"
 #include "MovieSceneSequenceID.h"
 
-class UBaseControlRig;
+class UControlRig;
 class ISequencer;
 class AActor;
 
@@ -30,7 +30,7 @@ class SControlRigSnapper : public SCompoundWidget
 
 private:
 
-	void GetControlRigs(TArray<UBaseControlRig*>& OutControlRigs) const;
+	void GetControlRigs(TArray<UControlRig*>& OutControlRigs) const;
 
 	/*
 	* Delegates and Helpers
diff --git a/Engine/Plugins/Animation/ControlRig/Source/ControlRigEditor/Private/EditMode/SControlRigSpacePicker.cpp b/Engine/Plugins/Animation/ControlRig/Source/ControlRigEditor/Private/EditMode/SControlRigSpacePicker.cpp
--- a/Engine/Plugins/Animation/ControlRig/Source/ControlRigEditor/Private/EditMode/SControlRigSpacePicker.cpp
+++ b/Engine/Plugins/Animation/ControlRig/Source/ControlRigEditor/Private/EditMode/SControlRigSpacePicker.cpp
@@ -108,10 +108,10 @@ SControlRigSpacePicker::~SControlRigSpacePicker()
 	//base class handles control rig related cleanup
 }
 
-UBaseControlRig* SControlRigSpacePicker::GetControlRig()
+UControlRig* SControlRigSpacePicker::GetControlRig()
 {
-	TArray<UBaseControlRig*> ControlRigs = GetControlRigs();
-	for(UBaseControlRig* ControlRig: ControlRigs)
+	TArray<UControlRig*> ControlRigs = GetControlRigs();
+	for(UControlRig* ControlRig: ControlRigs)
 	{
 		if (ControlRig)
 		{
@@ -125,10 +125,10 @@ UBaseControlRig* SControlRigSpacePicker::GetControlRig()
 	return nullptr;
 }
 
-void SControlRigSpacePicker::HandleControlSelected(UBaseControlRig* Subject, FRigControlElement* ControlElement, bool bSelected)
+void SControlRigSpacePicker::HandleControlSelected(UControlRig* Subject, FRigControlElement* ControlElement, bool bSelected)
 {
 	FControlRigBaseDockableView::HandleControlSelected(Subject, ControlElement, bSelected);
-	if (UBaseControlRig* ControlRig = GetControlRig())
+	if (UControlRig* ControlRig = GetControlRig())
 	{
 		// get the selected controls
 		TArray<FRigElementKey> SelectedControls = ControlRig->GetHierarchy()->GetSelectedKeys(ERigElementType::Control);
@@ -143,7 +143,7 @@ void SControlRigSpacePicker::HandleControlSelected(UBaseControlRig* Subject, FRi
 
 const FRigControlElementCustomization* SControlRigSpacePicker::HandleGetControlElementCustomization(URigHierarchy* InHierarchy, const FRigElementKey& InControlKey)
 {
-	if (UBaseControlRig* ControlRig = GetControlRig())
+	if (UControlRig* ControlRig = GetControlRig())
 	{
 		return ControlRig->GetControlCustomization(InControlKey);
 	}
@@ -156,7 +156,7 @@ void SControlRigSpacePicker::HandleActiveSpaceChanged(URigHierarchy* InHierarchy
 
 	if (ISequencer* Sequencer = GetSequencer())
 	{
-		if (UBaseControlRig* ControlRig = GetControlRig())
+		if (UControlRig* ControlRig = GetControlRig())
 		{
 			FString FailureReason;
 			URigHierarchy::TElementDependencyMap DependencyMap = InHierarchy->GetDependenciesForVM(ControlRig->GetVM());
@@ -194,7 +194,7 @@ void SControlRigSpacePicker::HandleActiveSpaceChanged(URigHierarchy* InHierarchy
 void SControlRigSpacePicker::HandleSpaceListChanged(URigHierarchy* InHierarchy, const FRigElementKey& InControlKey,
 	const TArray<FRigElementKey>& InSpaceList)
 {
-	if (UBaseControlRig* ControlRig = GetControlRig())
+	if (UControlRig* ControlRig = GetControlRig())
 	{
 		if (const FRigControlElement* ControlElement = InHierarchy->Find<FRigControlElement>(InControlKey))
 		{
@@ -253,7 +253,7 @@ FReply SControlRigSpacePicker::OnBakeControlsToNewSpaceButtonClicked()
 	{
 		return FReply::Unhandled();
 	}
-	UBaseControlRig* ControlRig = GetControlRig();
+	UControlRig* ControlRig = GetControlRig();
 
 	FRigSpacePickerBakeSettings Settings;
 	//Find default target space, just use first control and find space at current sequencer time
diff --git a/Engine/Plugins/Animation/ControlRig/Source/ControlRigEditor/Private/EditMode/SControlRigSpacePicker.h b/Engine/Plugins/Animation/ControlRig/Source/ControlRigEditor/Private/EditMode/SControlRigSpacePicker.h
--- a/Engine/Plugins/Animation/ControlRig/Source/ControlRigEditor/Private/EditMode/SControlRigSpacePicker.h
+++ b/Engine/Plugins/Animation/ControlRig/Source/ControlRigEditor/Private/EditMode/SControlRigSpacePicker.h
@@ -29,7 +29,7 @@ class SControlRigSpacePicker : public SCompoundWidget, public FControlRigBaseDoc
 	void Construct(const FArguments& InArgs, FControlRigEditMode& InEditMode);
 
 private:
-	virtual void HandleControlSelected(UBaseControlRig* Subject, FRigControlElement* InControl, bool bSelected) override;
+	virtual void HandleControlSelected(UControlRig* Subject, FRigControlElement* InControl, bool bSelected) override;
 
 	/** Space picker widget*/
 	TSharedPtr<SRigSpacePickerWidget> SpacePickerWidget;
@@ -43,7 +43,7 @@ private:
 	EVisibility GetAddSpaceButtonVisibility() const;
 
 	//for now picker works off of one ControlRig, this function gets the first control rig with a selection
-	UBaseControlRig* GetControlRig();
+	UControlRig* GetControlRig();
 
 };
 
diff --git a/Engine/Plugins/Animation/ControlRig/Source/ControlRigEditor/Private/EditMode/SControlRigUpdatePose.cpp b/Engine/Plugins/Animation/ControlRig/Source/ControlRigEditor/Private/EditMode/SControlRigUpdatePose.cpp
--- a/Engine/Plugins/Animation/ControlRig/Source/ControlRigEditor/Private/EditMode/SControlRigUpdatePose.cpp
+++ b/Engine/Plugins/Animation/ControlRig/Source/ControlRigEditor/Private/EditMode/SControlRigUpdatePose.cpp
@@ -81,7 +81,7 @@ private:
 		FControlRigEditMode* ControlRigEditMode = static_cast<FControlRigEditMode*>(GLevelEditorModeTools().GetActiveMode(FControlRigEditMode::ModeName));
 		if (PoseAsset.IsValid() && ControlRigEditMode)
 		{
-			TMap<UBaseControlRig*, TArray<FRigElementKey>> AllSelectedControls;
+			TMap<UControlRig*, TArray<FRigElementKey>> AllSelectedControls;
 			ControlRigEditMode->GetAllSelectedControls(AllSelectedControls);
 			return (AllSelectedControls.Num() > 0);
 		}
@@ -95,11 +95,11 @@ private:
 		FControlRigEditMode* ControlRigEditMode = static_cast<FControlRigEditMode*>(GLevelEditorModeTools().GetActiveMode(FControlRigEditMode::ModeName));
 		if (PoseAsset.IsValid() && ControlRigEditMode)
 		{
-			TMap<UBaseControlRig*, TArray<FRigElementKey>> AllSelectedControls;
+			TMap<UControlRig*, TArray<FRigElementKey>> AllSelectedControls;
 			ControlRigEditMode->GetAllSelectedControls(AllSelectedControls);
 			if (AllSelectedControls.Num() == 1)
 			{
-				TArray<UBaseControlRig*> ControlRigs;
+				TArray<UControlRig*> ControlRigs;
 				AllSelectedControls.GenerateKeyArray(ControlRigs);
 				const FScopedTransaction Transaction(NSLOCTEXT("ControlRig", "UpdatePose", "Update Pose"));
 				PoseAsset->Modify();
diff --git a/Engine/Plugins/Animation/ControlRig/Source/ControlRigEditor/Private/EditMode/SCurveControlContainer.cpp b/Engine/Plugins/Animation/ControlRig/Source/ControlRigEditor/Private/EditMode/SCurveControlContainer.cpp
--- a/Engine/Plugins/Animation/ControlRig/Source/ControlRigEditor/Private/EditMode/SCurveControlContainer.cpp
+++ b/Engine/Plugins/Animation/ControlRig/Source/ControlRigEditor/Private/EditMode/SCurveControlContainer.cpp
@@ -134,7 +134,7 @@ float SCurveControlListRow::GetValue() const
 //////////////////////////////////////////////////////////////////////////
 // SCurveControlContainer
 
-void SCurveControlContainer::Construct(const FArguments& InArgs, UBaseControlRig* InControlRig)
+void SCurveControlContainer::Construct(const FArguments& InArgs, UControlRig* InControlRig)
 {
 	ChildSlot
 	[
@@ -183,7 +183,7 @@ void SCurveControlContainer::Construct(const FArguments& InArgs, UBaseControlRig
 }
 
 
-void SCurveControlContainer::SetControlRig(UBaseControlRig* InControlRig)
+void SCurveControlContainer::SetControlRig(UControlRig* InControlRig)
 {
 
 	if (ControlRig.IsValid())
@@ -350,7 +350,7 @@ void SCurveControlContainer::OnSelectionChanged(FDisplayedCurveControlInfoPtr Se
 		}
 	}
 }
-void SCurveControlContainer::OnRigElementSelected(UBaseControlRig* Subject, FRigControlElement* ControlElement, bool bSelected)
+void SCurveControlContainer::OnRigElementSelected(UControlRig* Subject, FRigControlElement* ControlElement, bool bSelected)
 {
 	for(const FDisplayedCurveControlInfoPtr& Item : CurveControlList)
 	{
diff --git a/Engine/Plugins/Animation/ControlRig/Source/ControlRigEditor/Private/EditMode/SCurveControlContainer.h b/Engine/Plugins/Animation/ControlRig/Source/ControlRigEditor/Private/EditMode/SCurveControlContainer.h
--- a/Engine/Plugins/Animation/ControlRig/Source/ControlRigEditor/Private/EditMode/SCurveControlContainer.h
+++ b/Engine/Plugins/Animation/ControlRig/Source/ControlRigEditor/Private/EditMode/SCurveControlContainer.h
@@ -14,7 +14,7 @@
 
 class SInlineEditableTextBlock;
 struct FAssetData;
-class UBaseControlRig;
+class UControlRig;
 
 class FDisplayedCurveControlInfo
 {
@@ -129,7 +129,7 @@ public:
 	* @param InArgs - Arguments passed from Slate
 	*
 	*/
-	void Construct( const FArguments& InArgs, UBaseControlRig* InControlRig);
+	void Construct( const FArguments& InArgs, UControlRig* InControlRig);
 
 	/**
 	* Destructor - resets the animation curve
@@ -173,7 +173,7 @@ public:
 	FText& GetFilterText() { return FilterText; }
 
 	void RefreshCurveList();
-	void SetControlRig(UBaseControlRig* InControlRig);
+	void SetControlRig(UControlRig* InControlRig);
 private:
 
 	/**
@@ -190,7 +190,7 @@ private:
 
 	void OnSelectionChanged(FDisplayedCurveControlInfoPtr Selection, ESelectInfo::Type SelectInfo);
 
-	void OnRigElementSelected(UBaseControlRig* Subject, FRigControlElement* ControlElement, bool bSelected);
+	void OnRigElementSelected(UControlRig* Subject, FRigControlElement* ControlElement, bool bSelected);
 
 	/** Box to filter to a specific morph target name */
 	TSharedPtr<SSearchBox>	NameFilterBox;
@@ -202,7 +202,7 @@ private:
 	TSharedPtr<SCurveControlListType> CurveControlListView;
 
 	/** ControlRig */
-	TWeakObjectPtr<UBaseControlRig> ControlRig;
+	TWeakObjectPtr<UControlRig> ControlRig;
 
 	/** Current text typed into NameFilterBox */
 	FText FilterText;
diff --git a/Engine/Plugins/Animation/ControlRig/Source/ControlRigEditor/Private/Editor/ControlRigContextMenuContext.cpp b/Engine/Plugins/Animation/ControlRig/Source/ControlRigEditor/Private/Editor/ControlRigContextMenuContext.cpp
--- a/Engine/Plugins/Animation/ControlRig/Source/ControlRigEditor/Private/Editor/ControlRigContextMenuContext.cpp
+++ b/Engine/Plugins/Animation/ControlRig/Source/ControlRigEditor/Private/Editor/ControlRigContextMenuContext.cpp
@@ -39,11 +39,11 @@ UControlRigBlueprint* UControlRigContextMenuContext::GetControlRigBlueprint() co
 	return nullptr;
 }
 
-UBaseControlRig* UControlRigContextMenuContext::GetControlRig() const
+UControlRig* UControlRigContextMenuContext::GetControlRig() const
 {
 	if (UControlRigBlueprint* RigBlueprint = GetControlRigBlueprint())
 	{
-		if (UBaseControlRig* ControlRig = Cast<UBaseControlRig>(RigBlueprint->GetObjectBeingDebugged()))
+		if (UControlRig* ControlRig = Cast<UControlRig>(RigBlueprint->GetObjectBeingDebugged()))
 		{
 			return ControlRig;
 		}
diff --git a/Engine/Plugins/Animation/ControlRig/Source/ControlRigEditor/Private/Editor/ControlRigContextMenuContext.h b/Engine/Plugins/Animation/ControlRig/Source/ControlRigEditor/Private/Editor/ControlRigContextMenuContext.h
--- a/Engine/Plugins/Animation/ControlRig/Source/ControlRigEditor/Private/Editor/ControlRigContextMenuContext.h
+++ b/Engine/Plugins/Animation/ControlRig/Source/ControlRigEditor/Private/Editor/ControlRigContextMenuContext.h
@@ -8,7 +8,7 @@
 
 #include "ControlRigContextMenuContext.generated.h"
 
-class UBaseControlRig;
+class UControlRig;
 class UControlRigBlueprint;
 class FControlRigEditor;
 class URigVMGraph;
@@ -121,7 +121,7 @@ public:
 	
 	/** Get the active control rig instance in the viewport */
 	UFUNCTION(BlueprintCallable, Category = ControlRigEditorExtensions)
-    UBaseControlRig* GetControlRig() const;
+    UControlRig* GetControlRig() const;
 
 	/** Returns true if either alt key is down */
 	UFUNCTION(BlueprintCallable, Category = ControlRigEditorExtensions)
