Control Rig: Rename BaseControlRig to ControlRig

diff --git a/Engine/Plugins/Animation/ControlRig/Source/ControlRigEditor/Private/Sequencer/ControlRigSpaceChannelEditors.cpp b/Engine/Plugins/Animation/ControlRig/Source/ControlRigEditor/Private/Sequencer/ControlRigSpaceChannelEditors.cpp
--- a/Engine/Plugins/Animation/ControlRig/Source/ControlRigEditor/Private/Sequencer/ControlRigSpaceChannelEditors.cpp
+++ b/Engine/Plugins/Animation/ControlRig/Source/ControlRigEditor/Private/Sequencer/ControlRigSpaceChannelEditors.cpp
@@ -69,7 +69,7 @@
 //do not compensate via the FControlRigParameterTrackEditor setting a key.
 static bool bDoNotCompensate = false;
 
-static FKeyHandle SequencerOpenSpaceSwitchDialog(UBaseControlRig* ControlRig, TArray<FRigElementKey> SelectedControls,ISequencer* Sequencer, FMovieSceneControlRigSpaceChannel* Channel, UMovieSceneSection* SectionToKey, FFrameNumber Time)
+static FKeyHandle SequencerOpenSpaceSwitchDialog(UControlRig* ControlRig, TArray<FRigElementKey> SelectedControls,ISequencer* Sequencer, FMovieSceneControlRigSpaceChannel* Channel, UMovieSceneSection* SectionToKey, FFrameNumber Time)
 {
 	FKeyHandle Handle = FKeyHandle::Invalid();
 	if (ControlRig == nullptr || Sequencer == nullptr)
@@ -143,7 +143,7 @@ FKeyHandle AddOrUpdateKey(FMovieSceneControlRigSpaceChannel* Channel, UMovieScen
 	/**
 	if (UMovieSceneControlRigParameterSection* Section = Cast<UMovieSceneControlRigParameterSection>(SectionToKey))
 	{
-		if (UBaseControlRig* ControlRig = Section->GetControlRig())
+		if (UControlRig* ControlRig = Section->GetControlRig())
 		{
 			FName ControlName = Section->FindControlNameFromSpaceChannel(Channel);
 			if (ControlName != NAME_None)
@@ -185,7 +185,7 @@ TSharedRef<SWidget> CreateKeyEditor(const TMovieSceneChannelHandle<FMovieSceneCo
 **********************************************************************/
 
 
-FSpaceChannelAndSection FControlRigSpaceChannelHelpers::FindSpaceChannelAndSectionForControl(UBaseControlRig* ControlRig, FName ControlName, ISequencer* Sequencer, bool bCreateIfNeeded)
+FSpaceChannelAndSection FControlRigSpaceChannelHelpers::FindSpaceChannelAndSectionForControl(UControlRig* ControlRig, FName ControlName, ISequencer* Sequencer, bool bCreateIfNeeded)
 {
 	FSpaceChannelAndSection SpaceChannelAndSection;
 	SpaceChannelAndSection.SpaceChannel = nullptr;
@@ -201,7 +201,7 @@ FSpaceChannelAndSection FControlRigSpaceChannelHelpers::FindSpaceChannelAndSecti
 	}
 	const TArray<FMovieSceneBinding>& Bindings = MovieScene->GetBindings();
 	bool bRecreateCurves = false;
-	TArray<TPair<UBaseControlRig*, FName>> ControlRigPairsToReselect;
+	TArray<TPair<UControlRig*, FName>> ControlRigPairsToReselect;
 	for (const FMovieSceneBinding& Binding : Bindings)
 	{
 		UMovieSceneControlRigParameterTrack* ControlRigParameterTrack = Cast<UMovieSceneControlRigParameterTrack>(MovieScene->FindTrack(UMovieSceneControlRigParameterTrack::StaticClass(), Binding.GetObjectGuid(), NAME_None));
@@ -222,7 +222,7 @@ FSpaceChannelAndSection FControlRigSpaceChannelHelpers::FindSpaceChannelAndSecti
 				{
 					if (ControlRig->IsControlSelected(ControlName))
 					{
-						TPair<UBaseControlRig*, FName> Pair;
+						TPair<UControlRig*, FName> Pair;
 						Pair.Key = ControlRig;
 						Pair.Value = ControlName;
 						ControlRigPairsToReselect.Add(Pair);
@@ -242,7 +242,7 @@ FSpaceChannelAndSection FControlRigSpaceChannelHelpers::FindSpaceChannelAndSecti
 	if (bRecreateCurves)
 	{
 		Sequencer->RecreateCurveEditor(); //this will require the curve editor to get recreated so the ordering is correct
-		for (TPair<UBaseControlRig*, FName>& Pair : ControlRigPairsToReselect)
+		for (TPair<UControlRig*, FName>& Pair : ControlRigPairsToReselect)
 		{
 
 			GEditor->GetTimerManager()->SetTimerForNextTick([Pair]()
@@ -255,7 +255,7 @@ FSpaceChannelAndSection FControlRigSpaceChannelHelpers::FindSpaceChannelAndSecti
 	return SpaceChannelAndSection;
 }
 
-static TTuple<FRigControlElement*, FChannelMapInfo*, int32, int32> GetControlAndChannelInfos(UBaseControlRig* ControlRig, UMovieSceneControlRigParameterSection* ControlRigSection, FName ControlName)
+static TTuple<FRigControlElement*, FChannelMapInfo*, int32, int32> GetControlAndChannelInfos(UControlRig* ControlRig, UMovieSceneControlRigParameterSection* ControlRigSection, FName ControlName)
 {
 	FRigControlElement* ControlElement = nullptr;
 	FChannelMapInfo* pChannelIndex = nullptr;
@@ -285,7 +285,7 @@ static TTuple<FRigControlElement*, FChannelMapInfo*, int32, int32> GetControlAnd
 *  4) For 1-3 we always compensate any transform keys to the new space
 */
 
-FKeyHandle FControlRigSpaceChannelHelpers::SequencerKeyControlRigSpaceChannel(UBaseControlRig* ControlRig, ISequencer* Sequencer, FMovieSceneControlRigSpaceChannel* Channel, UMovieSceneSection* SectionToKey, FFrameNumber Time, URigHierarchy* RigHierarchy, const FRigElementKey& ControlKey, const FRigElementKey& SpaceKey)
+FKeyHandle FControlRigSpaceChannelHelpers::SequencerKeyControlRigSpaceChannel(UControlRig* ControlRig, ISequencer* Sequencer, FMovieSceneControlRigSpaceChannel* Channel, UMovieSceneSection* SectionToKey, FFrameNumber Time, URigHierarchy* RigHierarchy, const FRigElementKey& ControlKey, const FRigElementKey& SpaceKey)
 {
 	FKeyHandle Handle = FKeyHandle::Invalid();
 	if (bDoNotCompensate == true || ControlRig == nullptr || Sequencer == nullptr || Sequencer->GetFocusedMovieSceneSequence() == nullptr || Sequencer->GetFocusedMovieSceneSequence()->GetMovieScene() == nullptr)
@@ -515,7 +515,7 @@ FKeyHandle FControlRigSpaceChannelHelpers::SequencerKeyControlRigSpaceChannel(UB
 }
 
 //get the mask of the transform keys at the specified time
-static EControlRigContextChannelToKey GetCurrentTransformKeysAtThisTime(UBaseControlRig* ControlRig, FName ControlName, UMovieSceneSection* SectionToKey, FFrameNumber Time)
+static EControlRigContextChannelToKey GetCurrentTransformKeysAtThisTime(UControlRig* ControlRig, FName ControlName, UMovieSceneSection* SectionToKey, FFrameNumber Time)
 {
 	EControlRigContextChannelToKey ChannelToKey = EControlRigContextChannelToKey::None;
 	if (ControlRig && SectionToKey)
@@ -634,7 +634,7 @@ static EControlRigContextChannelToKey GetCurrentTransformKeysAtThisTime(UBaseCon
 	return ChannelToKey;
 }
 
-void  FControlRigSpaceChannelHelpers::SequencerSpaceChannelKeyDeleted(UBaseControlRig* ControlRig, ISequencer* Sequencer, FName ControlName, FMovieSceneControlRigSpaceChannel* Channel, UMovieSceneControlRigParameterSection* SectionToKey,
+void  FControlRigSpaceChannelHelpers::SequencerSpaceChannelKeyDeleted(UControlRig* ControlRig, ISequencer* Sequencer, FName ControlName, FMovieSceneControlRigSpaceChannel* Channel, UMovieSceneControlRigParameterSection* SectionToKey,
 	FFrameNumber TimeOfDeletion)
 {
 	if (bDoNotCompensate == true)
@@ -716,7 +716,7 @@ void  FControlRigSpaceChannelHelpers::SequencerSpaceChannelKeyDeleted(UBaseContr
 	}
 }
 
-void FControlRigSpaceChannelHelpers::DeleteTransformKeysAtThisTime(UBaseControlRig* ControlRig, UMovieSceneControlRigParameterSection* Section, FName ControlName, FFrameNumber Time)
+void FControlRigSpaceChannelHelpers::DeleteTransformKeysAtThisTime(UControlRig* ControlRig, UMovieSceneControlRigParameterSection* Section, FName ControlName, FFrameNumber Time)
 {
 	if (Section && ControlRig)
 	{
@@ -782,7 +782,7 @@ void FControlRigSpaceChannelHelpers::DeleteTransformKeysAtThisTime(UBaseControlR
 }
 
 void FControlRigSpaceChannelHelpers::GetFramesInThisSpaceAfterThisTime(
-	UBaseControlRig* ControlRig,
+	UControlRig* ControlRig,
 	FName ControlName,
 	FMovieSceneControlRigSpaceBaseKey CurrentValue,
 	FMovieSceneControlRigSpaceChannel* Channel,
@@ -859,7 +859,7 @@ void FControlRigSpaceChannelHelpers::GetFramesInThisSpaceAfterThisTime(
 	}
 }
 
-void FControlRigSpaceChannelHelpers::SequencerBakeControlInSpace(UBaseControlRig* ControlRig, ISequencer* Sequencer, FMovieSceneControlRigSpaceChannel* Channel, UMovieSceneSection* SectionToKey,
+void FControlRigSpaceChannelHelpers::SequencerBakeControlInSpace(UControlRig* ControlRig, ISequencer* Sequencer, FMovieSceneControlRigSpaceChannel* Channel, UMovieSceneSection* SectionToKey,
 	URigHierarchy* RigHierarchy, const FRigElementKey& ControlKey, FRigSpacePickerBakeSettings Settings)
 {
 	if (bDoNotCompensate == true)
@@ -1175,7 +1175,7 @@ void FControlRigSpaceChannelHelpers::SequencerBakeControlInSpace(UBaseControlRig
 	}
 }
 
-void FControlRigSpaceChannelHelpers::HandleSpaceKeyTimeChanged(UBaseControlRig* ControlRig, FName ControlName, FMovieSceneControlRigSpaceChannel* Channel, UMovieSceneSection* SectionToKey,
+void FControlRigSpaceChannelHelpers::HandleSpaceKeyTimeChanged(UControlRig* ControlRig, FName ControlName, FMovieSceneControlRigSpaceChannel* Channel, UMovieSceneSection* SectionToKey,
 	FFrameNumber CurrentFrame, FFrameNumber NextFrame)
 {
 	if (ControlRig && Channel && SectionToKey && (CurrentFrame != NextFrame))
@@ -1259,7 +1259,7 @@ void FControlRigSpaceChannelHelpers::HandleSpaceKeyTimeChanged(UBaseControlRig*
 
 
 void FControlRigSpaceChannelHelpers::CompensateIfNeeded(
-	UBaseControlRig* ControlRig,
+	UControlRig* ControlRig,
 	ISequencer* Sequencer,
 	UMovieSceneControlRigParameterSection* Section,
 	FName ControlName,
@@ -1420,7 +1420,7 @@ FReply FControlRigSpaceChannelHelpers::OpenBakeDialog(ISequencer* Sequencer, FMo
 	}
 	if (UMovieSceneControlRigParameterSection* Section = Cast<UMovieSceneControlRigParameterSection>(SectionToKey))
 	{
-		if (UBaseControlRig* ControlRig = Section->GetControlRig())
+		if (UControlRig* ControlRig = Section->GetControlRig())
 		{
 			FName ControlName = Section->FindControlNameFromSpaceChannel(Channel);
 
@@ -1529,7 +1529,7 @@ TArray<FKeyBarCurveModel::FBarRange> FControlRigSpaceChannelHelpers::FindRanges(
 }
 
 // NOTE use this function in HandleSpaceKeyTimeChanged, DeleteTransformKeysAtThisTime, ...
-TPair<FRigControlElement*, FChannelMapInfo*> FControlRigSpaceChannelHelpers::GetControlAndChannelInfo(UBaseControlRig* ControlRig, UMovieSceneControlRigParameterSection* Section, FName ControlName)
+TPair<FRigControlElement*, FChannelMapInfo*> FControlRigSpaceChannelHelpers::GetControlAndChannelInfo(UControlRig* ControlRig, UMovieSceneControlRigParameterSection* Section, FName ControlName)
 {
 	FRigControlElement* ControlElement = ControlRig ? ControlRig->FindControl(ControlName) : nullptr;
 	FChannelMapInfo* pChannelIndex = Section ? Section->ControlChannelMap.Find(ControlName) : nullptr;
diff --git a/Engine/Plugins/Animation/ControlRig/Source/ControlRigEditor/Private/Tools/AssetDefinition_ControlRigPose.cpp b/Engine/Plugins/Animation/ControlRig/Source/ControlRigEditor/Private/Tools/AssetDefinition_ControlRigPose.cpp
--- a/Engine/Plugins/Animation/ControlRig/Source/ControlRigEditor/Private/Tools/AssetDefinition_ControlRigPose.cpp
+++ b/Engine/Plugins/Animation/ControlRig/Source/ControlRigEditor/Private/Tools/AssetDefinition_ControlRigPose.cpp
@@ -32,14 +32,14 @@ namespace MenuExtension_AnimationAsset
 			return;
 		}
 		
-		TMap<UBaseControlRig*, TArray<FRigElementKey>> AllSelectedControls;
+		TMap<UControlRig*, TArray<FRigElementKey>> AllSelectedControls;
 		ControlRigEditMode->GetAllSelectedControls(AllSelectedControls);
-		TArray<UBaseControlRig*> ControlRigs;
+		TArray<UControlRig*> ControlRigs;
 		AllSelectedControls.GenerateKeyArray(ControlRigs);
 			
 		for (UControlRigPoseAsset* ControlRigPoseAsset : ControlRigPoseAssets)
 		{
-			for (UBaseControlRig* ControlRig : ControlRigs)
+			for (UControlRig* ControlRig : ControlRigs)
 			{
 				ControlRigPoseAsset->PastePose(ControlRig, false, false);
 			}
@@ -57,14 +57,14 @@ namespace MenuExtension_AnimationAsset
 			return;
 		}
 
-		TMap<UBaseControlRig*, TArray<FRigElementKey>> AllSelectedControls;
+		TMap<UControlRig*, TArray<FRigElementKey>> AllSelectedControls;
 		ControlRigEditMode->GetAllSelectedControls(AllSelectedControls);
-		TArray<UBaseControlRig*> ControlRigs;
+		TArray<UControlRig*> ControlRigs;
 		AllSelectedControls.GenerateKeyArray(ControlRigs);
 
 		for (UControlRigPoseAsset* ControlRigPoseAsset : ControlRigPoseAssets)
 		{
-			for (UBaseControlRig* ControlRig : ControlRigs)
+			for (UControlRig* ControlRig : ControlRigs)
 			{
 				ControlRigPoseAsset->SelectControls(ControlRig);
 			}
@@ -82,14 +82,14 @@ namespace MenuExtension_AnimationAsset
 			return;
 		}
 
-		TMap<UBaseControlRig*, TArray<FRigElementKey>> AllSelectedControls;
+		TMap<UControlRig*, TArray<FRigElementKey>> AllSelectedControls;
 		ControlRigEditMode->GetAllSelectedControls(AllSelectedControls);
-		TArray<UBaseControlRig*> ControlRigs;
+		TArray<UControlRig*> ControlRigs;
 		AllSelectedControls.GenerateKeyArray(ControlRigs);
 
 		for (UControlRigPoseAsset* ControlRigPoseAsset : ControlRigPoseAssets)
 		{
-			for (UBaseControlRig* ControlRig : ControlRigs)
+			for (UControlRig* ControlRig : ControlRigs)
 			{
 				ControlRigPoseAsset->SavePose(ControlRig, false);
 			}
diff --git a/Engine/Plugins/Animation/ControlRig/Source/ControlRigEditor/Private/Tools/ControlRigPoseThumbnailRenderer.cpp b/Engine/Plugins/Animation/ControlRig/Source/ControlRigEditor/Private/Tools/ControlRigPoseThumbnailRenderer.cpp
--- a/Engine/Plugins/Animation/ControlRig/Source/ControlRigEditor/Private/Tools/ControlRigPoseThumbnailRenderer.cpp
+++ b/Engine/Plugins/Animation/ControlRig/Source/ControlRigEditor/Private/Tools/ControlRigPoseThumbnailRenderer.cpp
@@ -67,7 +67,7 @@ bool FControlRigPoseThumbnailScene::SetControlRigPoseAsset(UControlRigPoseAsset*
 	if (PoseAsset)
 	{
 		USkeletalMesh* SkeletalMesh = nullptr; // PoseAsset->GetSkeletalMeshAsset();
-		UBaseControlRig* ControlRig = nullptr; // PoseAsset->GetControlRig();
+		UControlRig* ControlRig = nullptr; // PoseAsset->GetControlRig();
 		PreviewActor->GetSkeletalMeshComponent()->OverrideMaterials.Empty();
 
 
@@ -80,10 +80,10 @@ bool FControlRigPoseThumbnailScene::SetControlRigPoseAsset(UControlRigPoseAsset*
 
 			bSetSucessfully = true;
 			//now set up the control rig and the anim instance fo rit
-			UBaseControlRig* TempControlRig = NewObject<UBaseControlRig>(GetTransientPackage(), ControlRig->GetClass(),NAME_None, RF_Transient);
+			UControlRig* TempControlRig = NewObject<UControlRig>(GetTransientPackage(), ControlRig->GetClass(),NAME_None, RF_Transient);
 			TempControlRig->SetObjectBinding(MakeShared<FControlRigObjectBinding>());
 			TempControlRig->GetObjectBinding()->BindToObject(SkeletalMesh);
-			TempControlRig->GetDataSourceRegistry()->RegisterDataSource(UBaseControlRig::OwnerComponent, TempControlRig->GetObjectBinding()->GetBoundObject());
+			TempControlRig->GetDataSourceRegistry()->RegisterDataSource(UControlRig::OwnerComponent, TempControlRig->GetObjectBinding()->GetBoundObject());
 
 
 			bool bWasCreated;
diff --git a/Engine/Plugins/Animation/ControlRig/Source/ControlRigEditor/Private/Tools/ControlRigSnapper.cpp b/Engine/Plugins/Animation/ControlRig/Source/ControlRigEditor/Private/Tools/ControlRigSnapper.cpp
--- a/Engine/Plugins/Animation/ControlRig/Source/ControlRigEditor/Private/Tools/ControlRigSnapper.cpp
+++ b/Engine/Plugins/Animation/ControlRig/Source/ControlRigEditor/Private/Tools/ControlRigSnapper.cpp
@@ -112,7 +112,7 @@ TWeakPtr<ISequencer> FControlRigSnapper::GetSequencer()
 }
 
 static bool LocalGetControlRigControlTransforms(IMovieScenePlayer* Player, const TOptional<FFrameNumber>& CurrentFrame, UMovieSceneSequence* MovieSceneSequence, FMovieSceneSequenceIDRef Template, FMovieSceneSequenceTransform& RootToLocalTransform,
-	UBaseControlRig* ControlRig, const FName& ControlName,
+	UControlRig* ControlRig, const FName& ControlName,
 	const TArray<FFrameNumber>& Frames, const TArray<FTransform>& ParentTransforms, TArray<FTransform>& OutTransforms)
 {
 	if (Frames.Num() > ParentTransforms.Num())
@@ -160,7 +160,7 @@ static bool LocalGetControlRigControlTransforms(IMovieScenePlayer* Player, const
 	return true;
 }
 
-bool FControlRigSnapper::GetControlRigControlTransforms(ISequencer* Sequencer,  UBaseControlRig* ControlRig, const FName& ControlName,
+bool FControlRigSnapper::GetControlRigControlTransforms(ISequencer* Sequencer,  UControlRig* ControlRig, const FName& ControlName,
 	const TArray<FFrameNumber> &Frames, const TArray<FTransform>& ParentTransforms,TArray<FTransform>& OutTransforms)
 {
 	if (Sequencer->GetFocusedMovieSceneSequence())
@@ -175,7 +175,7 @@ bool FControlRigSnapper::GetControlRigControlTransforms(ISequencer* Sequencer,
 	return false;
 }
 
-bool FControlRigSnapper::GetControlRigControlTransforms(UWorld* World,ULevelSequence* LevelSequence, UBaseControlRig* ControlRig, const FName& ControlName,
+bool FControlRigSnapper::GetControlRigControlTransforms(UWorld* World,ULevelSequence* LevelSequence, UControlRig* ControlRig, const FName& ControlName,
 	const TArray<FFrameNumber>& Frames, const TArray<FTransform>& ParentTransforms, TArray<FTransform>& OutTransforms)
 {
 	if (LevelSequence)
@@ -297,7 +297,7 @@ static bool CalculateWorldTransformsFromParents(ISequencer* Sequencer, const FCo
 	for (const FControlRigForWorldTransforms& ControlRigAndSelection : ParentToSnap.ControlRigs)
 	{
 		//get actor transform...
-		UBaseControlRig* ControlRig = ControlRigAndSelection.ControlRig.Get();
+		UControlRig* ControlRig = ControlRigAndSelection.ControlRig.Get();
 
 		if (ControlRig)
 		{
@@ -341,7 +341,7 @@ static void GetControlRigParents(const FControlRigForWorldTransforms& ControlRig
 	TArray<FActorForWorldTransforms>& OutParentActors, TArray<FControlRigForWorldTransforms>& OutControlRigs)
 {
 
-	if (UBaseControlRig* ControlRig = ControlRigAndSelection.ControlRig.Get())
+	if (UControlRig* ControlRig = ControlRigAndSelection.ControlRig.Get())
 	{
 		URigHierarchy* Hierarchy = ControlRig->GetHierarchy();
 		for (const FName& ControlName : ControlRigAndSelection.ControlNames)
@@ -381,7 +381,7 @@ static void GetControlRigParents(const FControlRigForWorldTransforms& ControlRig
 	}
 }
 
-static UMovieSceneControlRigParameterSection* GetControlRigSection(ISequencer* Sequencer, const UBaseControlRig* ControlRig)
+static UMovieSceneControlRigParameterSection* GetControlRigSection(ISequencer* Sequencer, const UControlRig* ControlRig)
 {
 
 	if (ControlRig == nullptr || Sequencer == nullptr)
@@ -587,7 +587,7 @@ bool FControlRigSnapper::SnapIt(FFrameNumber StartFrame, FFrameNumber EndFrame,c
 		for (const FControlRigForWorldTransforms& ControlRigAndSelection : ActorToSnap.ControlRigs)
 		{
 			//get actor transform...
-			UBaseControlRig* ControlRig = ControlRigAndSelection.ControlRig.Get();
+			UControlRig* ControlRig = ControlRigAndSelection.ControlRig.Get();
 
 			if (ControlRig)
 			{
diff --git a/Engine/Plugins/Animation/ControlRig/Source/ControlRigEditor/Private/Tools/ControlRigTweener.cpp b/Engine/Plugins/Animation/ControlRig/Source/ControlRigEditor/Private/Tools/ControlRigTweener.cpp
--- a/Engine/Plugins/Animation/ControlRig/Source/ControlRigEditor/Private/Tools/ControlRigTweener.cpp
+++ b/Engine/Plugins/Animation/ControlRig/Source/ControlRigEditor/Private/Tools/ControlRigTweener.cpp
@@ -107,7 +107,7 @@ bool FAnimSliderKeySelection::Setup(const TWeakPtr<ISequencer>& InSequencer)
 bool FAnimSliderObjectSelection::Setup(TWeakPtr<ISequencer>& InSequencer, TWeakPtr<FControlRigEditMode>& InEditMode)
 {
 	ChannelsArray.Reset();
-	const TArray<UBaseControlRig*> ControlRigs = GetControlRigs(InEditMode);
+	const TArray<UControlRig*> ControlRigs = GetControlRigs(InEditMode);
 	return Setup(ControlRigs, InSequencer);
 }
 
@@ -183,14 +183,14 @@ void FAnimSliderKeySelection::GetMapOfContiguousKeys()
 * FAnimSliderObjectSelection
 *
 */
-TArray<UBaseControlRig*> FAnimSliderObjectSelection::GetControlRigs(TWeakPtr<FControlRigEditMode>& InEditMode)
+TArray<UControlRig*> FAnimSliderObjectSelection::GetControlRigs(TWeakPtr<FControlRigEditMode>& InEditMode)
 {
-	TArray<UBaseControlRig*> ControlRigs;
+	TArray<UControlRig*> ControlRigs;
 	if (const FControlRigEditMode* EditMode = InEditMode.Pin().Get())
 	{
-		TMap<UBaseControlRig*, TArray<FRigElementKey>> SelectedControls;
+		TMap<UControlRig*, TArray<FRigElementKey>> SelectedControls;
 		EditMode->GetAllSelectedControls(SelectedControls);
-		for (TPair<UBaseControlRig*, TArray<FRigElementKey>>& Selected : SelectedControls)
+		for (TPair<UControlRig*, TArray<FRigElementKey>>& Selected : SelectedControls)
 		{
 			ControlRigs.Add(Selected.Key);
 		}
@@ -273,7 +273,7 @@ void FAnimSliderObjectSelection::SetupChannel(FFrameNumber CurrentFrame, TArray<
 		&& KeyBounds.PreviousIndex != KeyBounds.NextIndex) ? true : false;
 }
 
-bool FAnimSliderObjectSelection::Setup(const TArray<UBaseControlRig*>& SelectedControlRigs, TWeakPtr<ISequencer>& InSequencer)
+bool FAnimSliderObjectSelection::Setup(const TArray<UControlRig*>& SelectedControlRigs, TWeakPtr<ISequencer>& InSequencer)
 {
 	ChannelsArray.Reset();
 	if (InSequencer.IsValid() == false)
@@ -299,7 +299,7 @@ bool FAnimSliderObjectSelection::Setup(const TArray<UBaseControlRig*>& SelectedC
 			if (Section && Section->IsActive() && Section->GetRange().Contains(CurrentTime.Time.GetFrame()) && !HandledSections.Contains(Section))
 			{
 				HandledSections.Add(Section);				
-				UBaseControlRig* ControlRig = Track->GetControlRig();
+				UControlRig* ControlRig = Track->GetControlRig();
 				TArray<FRigControlElement*> CurrentControls;
 				ControlRig->GetControlsInOrder(CurrentControls);
 
@@ -613,7 +613,7 @@ FText FControlsToTween::GetTooltipText() const
 	return LOCTEXT("TweenControllerTooltip", "Tween between the next and previous keys");
 }
 
-bool FControlsToTween::Setup(const TArray<UBaseControlRig*>& SelectedControlRigs, TWeakPtr<ISequencer>& InSequencer)
+bool FControlsToTween::Setup(const TArray<UControlRig*>& SelectedControlRigs, TWeakPtr<ISequencer>& InSequencer)
 {
 	return ObjectSelection.Setup(SelectedControlRigs, InSequencer);
 }
diff --git a/Engine/Plugins/Animation/ControlRig/Source/ControlRigEditor/Public/ControlRigSpaceChannelEditors.h b/Engine/Plugins/Animation/ControlRig/Source/ControlRigEditor/Public/ControlRigSpaceChannelEditors.h
--- a/Engine/Plugins/Animation/ControlRig/Source/ControlRigEditor/Public/ControlRigSpaceChannelEditors.h
+++ b/Engine/Plugins/Animation/ControlRig/Source/ControlRigEditor/Public/ControlRigSpaceChannelEditors.h
@@ -21,7 +21,7 @@
 
 struct FKeyHandle;
 struct FKeyDrawParams;
-class UBaseControlRig;
+class UControlRig;
 class ISequencer;
 class UMovieSceneSection;
 class URigHierarchy;
@@ -44,23 +44,23 @@ struct FSpaceChannelAndSection
 */
 struct FControlRigSpaceChannelHelpers
 {
-	static FKeyHandle SequencerKeyControlRigSpaceChannel(UBaseControlRig* ControlRig, ISequencer* Sequencer, FMovieSceneControlRigSpaceChannel* Channel, UMovieSceneSection* SectionToKey, FFrameNumber Time, URigHierarchy* RigHierarchy, const FRigElementKey& ControlKey, const FRigElementKey& SpaceKey);
-	static void SequencerSpaceChannelKeyDeleted(UBaseControlRig* ControlRig, ISequencer* Sequencer, FName ControlName, FMovieSceneControlRigSpaceChannel* Channel, UMovieSceneControlRigParameterSection* SectionToKey, FFrameNumber TimeOfDeletion);
-	static void CompensateIfNeeded(UBaseControlRig* ControlRig, ISequencer* Sequencer, UMovieSceneControlRigParameterSection* Section, FName ControlName, TOptional<FFrameNumber>& Time);
-	static FSpaceChannelAndSection FindSpaceChannelAndSectionForControl(UBaseControlRig* ControlRig, FName ControlName, ISequencer* Sequencer, bool bCreateIfNeeded);
-	static void SequencerBakeControlInSpace(UBaseControlRig* ControlRig, ISequencer* Sequencer, FMovieSceneControlRigSpaceChannel* Channel, UMovieSceneSection* SectionToKey,
+	static FKeyHandle SequencerKeyControlRigSpaceChannel(UControlRig* ControlRig, ISequencer* Sequencer, FMovieSceneControlRigSpaceChannel* Channel, UMovieSceneSection* SectionToKey, FFrameNumber Time, URigHierarchy* RigHierarchy, const FRigElementKey& ControlKey, const FRigElementKey& SpaceKey);
+	static void SequencerSpaceChannelKeyDeleted(UControlRig* ControlRig, ISequencer* Sequencer, FName ControlName, FMovieSceneControlRigSpaceChannel* Channel, UMovieSceneControlRigParameterSection* SectionToKey, FFrameNumber TimeOfDeletion);
+	static void CompensateIfNeeded(UControlRig* ControlRig, ISequencer* Sequencer, UMovieSceneControlRigParameterSection* Section, FName ControlName, TOptional<FFrameNumber>& Time);
+	static FSpaceChannelAndSection FindSpaceChannelAndSectionForControl(UControlRig* ControlRig, FName ControlName, ISequencer* Sequencer, bool bCreateIfNeeded);
+	static void SequencerBakeControlInSpace(UControlRig* ControlRig, ISequencer* Sequencer, FMovieSceneControlRigSpaceChannel* Channel, UMovieSceneSection* SectionToKey,
 		URigHierarchy* RigHierarchy, const FRigElementKey& ControlKey, FRigSpacePickerBakeSettings InSettings);
-	static void GetFramesInThisSpaceAfterThisTime(UBaseControlRig* ControlRig, FName ControlName, FMovieSceneControlRigSpaceBaseKey CurrentValue,
+	static void GetFramesInThisSpaceAfterThisTime(UControlRig* ControlRig, FName ControlName, FMovieSceneControlRigSpaceBaseKey CurrentValue,
 		FMovieSceneControlRigSpaceChannel* Channel, UMovieSceneSection* SectionToKey,
 		FFrameNumber Time, TSortedMap<FFrameNumber,FFrameNumber>& OutMoreFrames);
-	static void HandleSpaceKeyTimeChanged(UBaseControlRig* ControlRig, FName ControlName,FMovieSceneControlRigSpaceChannel* Channel, UMovieSceneSection* SectionToKey,
+	static void HandleSpaceKeyTimeChanged(UControlRig* ControlRig, FName ControlName,FMovieSceneControlRigSpaceChannel* Channel, UMovieSceneSection* SectionToKey,
 		FFrameNumber CurrentFrame, FFrameNumber NextFrame);
-	static void DeleteTransformKeysAtThisTime(UBaseControlRig* ControlRig, UMovieSceneControlRigParameterSection* Section, FName ControlName, FFrameNumber Time);
+	static void DeleteTransformKeysAtThisTime(UControlRig* ControlRig, UMovieSceneControlRigParameterSection* Section, FName ControlName, FFrameNumber Time);
 	static FLinearColor GetColor(const FMovieSceneControlRigSpaceBaseKey& Key);
 	static FReply OpenBakeDialog(ISequencer* Sequencer, FMovieSceneControlRigSpaceChannel* Channel, int32 KeyIndex, UMovieSceneSection* SectionToKey);
 	static TArray<FKeyBarCurveModel::FBarRange> FindRanges(FMovieSceneControlRigSpaceChannel* Channel, const UMovieSceneSection* Section);
 	// retrieve the control and the channel infos for that ControlRig/Section.
-	static TPair<FRigControlElement*, FChannelMapInfo*> GetControlAndChannelInfo(UBaseControlRig* ControlRig, UMovieSceneControlRigParameterSection* Section, FName ControlName);
+	static TPair<FRigControlElement*, FChannelMapInfo*> GetControlAndChannelInfo(UControlRig* ControlRig, UMovieSceneControlRigParameterSection* Section, FName ControlName);
 	// retrieve the number of float channels based on the control type.
 	static int32 GetNumFloatChannels(const ERigControlType InControlType);
 };
diff --git a/Engine/Plugins/Animation/ControlRig/Source/ControlRigEditor/Public/EditMode/ControlRigEditMode.h b/Engine/Plugins/Animation/ControlRig/Source/ControlRigEditor/Public/EditMode/ControlRigEditMode.h
--- a/Engine/Plugins/Animation/ControlRig/Source/ControlRigEditor/Public/EditMode/ControlRigEditMode.h
+++ b/Engine/Plugins/Animation/ControlRig/Source/ControlRigEditor/Public/EditMode/ControlRigEditMode.h
@@ -22,7 +22,7 @@ class FRigVMEditor;
 class FViewport;
 class UActorFactory;
 struct FViewportClick;
-class UBaseControlRig;
+class UControlRig;
 class FControlRigInteractionScope;
 class ISequencer;
 class UControlManipulator;
@@ -44,8 +44,8 @@ class UToolMenu;
 DECLARE_DELEGATE_RetVal_ThreeParams(FTransform, FOnGetRigElementTransform, const FRigElementKey& /*RigElementKey*/, bool /*bLocal*/, bool /*bOnDebugInstance*/);
 DECLARE_DELEGATE_ThreeParams(FOnSetRigElementTransform, const FRigElementKey& /*RigElementKey*/, const FTransform& /*Transform*/, bool /*bLocal*/);
 DECLARE_DELEGATE_RetVal(TSharedPtr<FUICommandList>, FNewMenuCommandsDelegate);
-DECLARE_MULTICAST_DELEGATE_TwoParams(FControlRigAddedOrRemoved, UBaseControlRig*, bool /*true if added, false if removed*/);
-DECLARE_MULTICAST_DELEGATE_ThreeParams(FControlRigSelected, UBaseControlRig*, const FRigElementKey& /*RigElementKey*/,const bool /*bIsSelected*/);
+DECLARE_MULTICAST_DELEGATE_TwoParams(FControlRigAddedOrRemoved, UControlRig*, bool /*true if added, false if removed*/);
+DECLARE_MULTICAST_DELEGATE_ThreeParams(FControlRigSelected, UControlRig*, const FRigElementKey& /*RigElementKey*/,const bool /*bIsSelected*/);
 DECLARE_DELEGATE_RetVal(UToolMenu*, FOnGetContextMenu);
 
 class FControlRigEditMode;
@@ -91,16 +91,16 @@ public:
 
 	/** Set the Control Rig Object to be active in the edit mode. You set both the Control Rig and a possible binding together with an optional Sequencer
 	 This will remove all other control rigs present and should be called for stand alone editors, like the Control Rig Editor*/
-	void SetObjects(UBaseControlRig* InControlRig, UObject* BindingObject, TWeakPtr<ISequencer> InSequencer);
+	void SetObjects(UControlRig* InControlRig, UObject* BindingObject, TWeakPtr<ISequencer> InSequencer);
 
 	/** Add a Control Rig object if it doesn't exist, will return true if it was added, false if it wasn't since it's already there. You can also set the Sequencer.*/
-	bool AddControlRigObject(UBaseControlRig* InControlRig, TWeakPtr<ISequencer> InSequencer);
+	bool AddControlRigObject(UControlRig* InControlRig, TWeakPtr<ISequencer> InSequencer);
 
 	/* Remove control rig */
-	void RemoveControlRig(UBaseControlRig* InControlRig);
+	void RemoveControlRig(UControlRig* InControlRig);
 
 	/*Replace old Control Rig with the New Control Rig, perhaps from a recompile in the level editor*/
-	void ReplaceControlRig(UBaseControlRig* OldControlRig, UBaseControlRig* NewControlRig);
+	void ReplaceControlRig(UControlRig* OldControlRig, UControlRig* NewControlRig);
 
 	/** This edit mode is re-used between the level editor and the asset editors (control rig editor etc.). Calling this indicates which context we are in */
 	virtual bool IsInLevelEditor() const;
@@ -166,9 +166,9 @@ public:
 	// callback that gets called when rig element is selected in other view
 	void OnHierarchyModified(ERigHierarchyNotification InNotif, URigHierarchy* InHierarchy, const FRigBaseElement* InElement);
 	void OnHierarchyModified_AnyThread(ERigHierarchyNotification InNotif, URigHierarchy* InHierarchy, const FRigBaseElement* InElement);
-	void OnControlModified(UBaseControlRig* Subject, FRigControlElement* InControlElement, const FRigControlModifiedContext& Context);
-	void OnPreConstruction_AnyThread(UBaseControlRig* InRig, const FName& InEventName);
-	void OnPostConstruction_AnyThread(UBaseControlRig* InRig, const FName& InEventName);
+	void OnControlModified(UControlRig* Subject, FRigControlElement* InControlElement, const FRigControlModifiedContext& Context);
+	void OnPreConstruction_AnyThread(UControlRig* InRig, const FName& InEventName);
+	void OnPostConstruction_AnyThread(UControlRig* InRig, const FName& InEventName);
 
 	/** return true if it can be removed from preview scene 
 	- this is to ensure preview scene doesn't remove shape actors */
@@ -178,7 +178,7 @@ public:
 
 	/** Requests to recreate the shape actors in the next tick. Will recreate only the ones for the specified
 	Control Rig, otherwise will recreate all of them*/
-	void RequestToRecreateControlShapeActors(UBaseControlRig* ControlRig = nullptr); 
+	void RequestToRecreateControlShapeActors(UControlRig* ControlRig = nullptr); 
 
 	static uint32 ValidControlTypeMask()
 	{
@@ -189,7 +189,7 @@ protected:
 
 	// shape related functions wrt enable/selection
 	/** Get the node name from the property path */
-	AControlRigShapeActor* GetControlShapeFromControlName(UBaseControlRig* InControlRig,const FName& InControlName) const;
+	AControlRigShapeActor* GetControlShapeFromControlName(UControlRig* InControlRig,const FName& InControlName) const;
 
 protected:
 	/** Helper function: set ControlRigs array to the details panel */
@@ -200,9 +200,9 @@ protected:
 
 	/** Updates cached pivot transforms */
 	void UpdatePivotTransforms();
-	void UpdatePivotFromEditedShape(UBaseControlRig* InControlRig);
-	void UpdatePivotFromShapeActors(UBaseControlRig* InControlRig, const bool bEachLocalSpace, const bool bIsParentSpace);
-	void UpdatePivotFromElements(UBaseControlRig* InControlRig);
+	void UpdatePivotFromEditedShape(UControlRig* InControlRig);
+	void UpdatePivotFromShapeActors(UControlRig* InControlRig, const bool bEachLocalSpace, const bool bIsParentSpace);
+	void UpdatePivotFromElements(UControlRig* InControlRig);
 	
 	/** Get the current coordinate system space */
 	ECoordSystem GetCoordSystemSpace() const;
@@ -296,7 +296,7 @@ protected:
 	FGuid LastMovieSceneSig;
 
 	/** The scope for the interaction, one per manipulated Control rig */
-	TMap<UBaseControlRig*,FControlRigInteractionScope*> InteractionScopes;
+	TMap<UControlRig*,FControlRigInteractionScope*> InteractionScopes;
 
 	/** True if there's tracking going on right now */
 	bool bIsTracking;
@@ -311,10 +311,10 @@ protected:
 	bool bSelectionChanged;
 
 	/** Cached transform of pivot point for selected objects for each Control Rig */
-	TMap<UBaseControlRig*,FTransform> PivotTransforms;
+	TMap<UControlRig*,FTransform> PivotTransforms;
 
 	/** Previous cached transforms, need this to check on tick if any transform changed, gizmo may have changed*/
-	TMap<UBaseControlRig*, FTransform> LastPivotTransforms;
+	TMap<UControlRig*, FTransform> LastPivotTransforms;
 
 	/** Command bindings for keyboard shortcuts */
 	TSharedPtr<FUICommandList> CommandBindings;
@@ -334,23 +334,23 @@ protected:
 	FControlRigSelected OnControlRigSelectedDelegate;
 
 	/** GetSelectedRigElements, if InControlRig is nullptr get the first one */
-	TArray<FRigElementKey> GetSelectedRigElements(UBaseControlRig* InControlRig = nullptr) const;
+	TArray<FRigElementKey> GetSelectedRigElements(UControlRig* InControlRig = nullptr) const;
 
 	/* Flag to recreate shapes during tick */
 	ERecreateControlRigShape RecreateControlShapesRequired;
 	/* List of Control Rigs we should recreate*/
-	TArray<UBaseControlRig*> ControlRigsToRecreate;
+	TArray<UControlRig*> ControlRigsToRecreate;
 
 	/* Flag to temporarily disable handling notifs from the hierarchy */
 	bool bSuspendHierarchyNotifs;
 
 	/** Shape actors */
-	TMap<UBaseControlRig*,TArray<TObjectPtr<AControlRigShapeActor>>> ControlRigShapeActors;
+	TMap<UControlRig*,TArray<TObjectPtr<AControlRigShapeActor>>> ControlRigShapeActors;
 	TObjectPtr<UControlRigDetailPanelControlProxies> ControlProxy;
 
 	/** Utility functions for UI/Some other viewport manipulation*/
 	bool IsControlSelected() const;
-	bool AreRigElementSelectedAndMovable(UBaseControlRig* InControlRig) const;
+	bool AreRigElementSelectedAndMovable(UControlRig* InControlRig) const;
 	
 	/** Set initial transform handlers */
 	void OpenContextMenu(FEditorViewportClient* InViewportClient);
@@ -368,26 +368,26 @@ public:
 	void ClearRigElementSelection(uint32 InTypes);
 
 	/** Set a RigElement's selection state */
-	void SetRigElementSelection(UBaseControlRig* ControlRig, ERigElementType Type, const FName& InRigElementName, bool bSelected);
+	void SetRigElementSelection(UControlRig* ControlRig, ERigElementType Type, const FName& InRigElementName, bool bSelected);
 
 	/** Set multiple RigElement's selection states */
-	void SetRigElementSelection(UBaseControlRig* ControlRig, ERigElementType Type, const TArray<FName>& InRigElementNames, bool bSelected);
+	void SetRigElementSelection(UControlRig* ControlRig, ERigElementType Type, const TArray<FName>& InRigElementNames, bool bSelected);
 
 	/** Check if any RigElements are selected */
-	bool AreRigElementsSelected(uint32 InTypes, UBaseControlRig* InControlRig) const;
+	bool AreRigElementsSelected(uint32 InTypes, UControlRig* InControlRig) const;
 
 	/** Get the number of selected RigElements */
-	int32 GetNumSelectedRigElements(uint32 InTypes, UBaseControlRig* InControlRig) const;
+	int32 GetNumSelectedRigElements(uint32 InTypes, UControlRig* InControlRig) const;
 
 	/** Get all of the selected Controls*/
-	void GetAllSelectedControls(TMap<UBaseControlRig*, TArray<FRigElementKey>>& OutSelectedControls) const;
+	void GetAllSelectedControls(TMap<UControlRig*, TArray<FRigElementKey>>& OutSelectedControls) const;
 
 	/** Get all of the ControlRigs, maybe not valid anymore */
-	TArrayView<const TWeakObjectPtr<UBaseControlRig>> GetControlRigs() const;
-	TArrayView<TWeakObjectPtr<UBaseControlRig>> GetControlRigs();
+	TArrayView<const TWeakObjectPtr<UControlRig>> GetControlRigs() const;
+	TArrayView<TWeakObjectPtr<UControlRig>> GetControlRigs();
 	/* Get valid  Control Rigs possibly just visible*/
-	TArray<UBaseControlRig*> GetControlRigsArray(bool bIsVisible);
-	TArray<const UBaseControlRig*> GetControlRigsArray(bool bIsVisible) const;
+	TArray<UControlRig*> GetControlRigsArray(bool bIsVisible);
+	TArray<const UControlRig*> GetControlRigsArray(bool bIsVisible) const;
 
 	/** Get the detail proxies control rig*/
 	UControlRigDetailPanelControlProxies* GetDetailProxies() { return ControlProxy; }
@@ -405,7 +405,7 @@ private:
 	/** Whether or not Pivot Transforms have changed, in which case we need to redraw viewport*/
 	bool HasPivotTransformsChanged() const;
 	/** Set a RigElement's selection state */
-	void SetRigElementSelectionInternal(UBaseControlRig* ControlRig, ERigElementType Type, const FName& InRigElementName, bool bSelected);
+	void SetRigElementSelectionInternal(UControlRig* ControlRig, ERigElementType Type, const FName& InRigElementName, bool bSelected);
 	
 	FEditorViewportClient* CurrentViewportClient;
 	TArray<UE::Widget::EWidgetMode> RequestedWidgetModes;
@@ -424,15 +424,15 @@ public:
 	bool GetOnlySelectRigControls()const;
 
 private:
-	TSet<FName> GetActiveControlsFromSequencer(UBaseControlRig* ControlRig);
+	TSet<FName> GetActiveControlsFromSequencer(UControlRig* ControlRig);
 	bool SetSequencer(TWeakPtr<ISequencer> InSequencer);
 
 	/** Create/Delete for the specified ControlRig*/
-	void CreateShapeActors(UBaseControlRig* InControlRig);
-	void DestroyShapesActors(UBaseControlRig* InControlRig);
+	void CreateShapeActors(UControlRig* InControlRig);
+	void DestroyShapesActors(UControlRig* InControlRig);
 
 	/*Internal function for adding ControlRig*/
-	void AddControlRigInternal(UBaseControlRig* InControlRig);
+	void AddControlRigInternal(UControlRig* InControlRig);
 	void TickManipulatableObjects(float DeltaTime);
 
 	/* Check on tick to see if movie scene has changed, returns true if it has*/
@@ -455,19 +455,19 @@ protected:
 	
 	/** Get bindings to a runtime object */
 	//If the passed in ControlRig is nullptr we use the first Control Rig(this can happen from the BP Editors).
-	USceneComponent* GetHostingSceneComponent(const UBaseControlRig* ControlRig = nullptr) const;
-	FTransform	GetHostingSceneComponentTransform(const UBaseControlRig* ControlRig =  nullptr) const;
+	USceneComponent* GetHostingSceneComponent(const UControlRig* ControlRig = nullptr) const;
+	FTransform	GetHostingSceneComponentTransform(const UControlRig* ControlRig =  nullptr) const;
 
 	//Get if the hosted component is visible
-	bool IsControlRigSkelMeshVisible(UBaseControlRig* ControlRig) const;
+	bool IsControlRigSkelMeshVisible(UControlRig* ControlRig) const;
 
 private:
 
 	// Post pose update handler
 	void OnPoseInitialized();
 	void PostPoseUpdate();
-	void NotifyDrivenControls(UBaseControlRig* InControlRig, const FRigElementKey& InKey);
-	void UpdateSelectabilityOnSkeletalMeshes(UBaseControlRig* InControlRig, bool bEnabled);
+	void NotifyDrivenControls(UControlRig* InControlRig, const FRigElementKey& InKey);
+	void UpdateSelectabilityOnSkeletalMeshes(UControlRig* InControlRig, bool bEnabled);
 
 	// world clean up handlers
 	FDelegateHandle OnWorldCleanupHandle;
@@ -476,8 +476,8 @@ private:
 
 	void OnEditorClosed();
 	
-	TArray<TWeakObjectPtr<UBaseControlRig>> RuntimeControlRigs;
-	TMap<UBaseControlRig*,TStrongObjectPtr<UControlRigEditModeDelegateHelper>> DelegateHelpers;
+	TArray<TWeakObjectPtr<UControlRig>> RuntimeControlRigs;
+	TMap<UControlRig*,TStrongObjectPtr<UControlRigEditModeDelegateHelper>> DelegateHelpers;
 
 	//hack since we can't get the viewport client from the viewport, so in the tick we set the gameview bool and then in render/tickcontrolshapes we use it.
 	TMap<FViewport*, bool>  ViewportToGameView;
diff --git a/Engine/Plugins/Animation/ControlRig/Source/ControlRigEditor/Public/Tools/ControlRigSnapper.h b/Engine/Plugins/Animation/ControlRig/Source/ControlRigEditor/Public/Tools/ControlRigSnapper.h
--- a/Engine/Plugins/Animation/ControlRig/Source/ControlRigEditor/Public/Tools/ControlRigSnapper.h
+++ b/Engine/Plugins/Animation/ControlRig/Source/ControlRigEditor/Public/Tools/ControlRigSnapper.h
@@ -9,7 +9,7 @@
 #include "ControlRigSnapper.generated.h"
 
 class AActor;
-class UBaseControlRig;
+class UControlRig;
 class ISequencer;
 class UMovieScene;
 class ULevelSequence;
@@ -23,7 +23,7 @@ struct FControlRigForWorldTransforms
 
 	FControlRigForWorldTransforms() : ControlRig(nullptr) {};
 	UPROPERTY(BlueprintReadWrite, EditAnywhere, Category = "Control Rig")
-	TWeakObjectPtr<UBaseControlRig> ControlRig;
+	TWeakObjectPtr<UControlRig> ControlRig;
 	UPROPERTY(BlueprintReadWrite, EditAnywhere, Category = "Control Rig")
 	TArray<FName> ControlNames;
 };
@@ -75,7 +75,7 @@ struct FControlRigSnapper
 	* @param ParentTransforms List of Parent Transforms for each time that the resulting transforms will be concatenated to.
 	* @param OutTransforms  Results
 	*/
-	bool GetControlRigControlTransforms(ISequencer* Sequencer, UBaseControlRig* ControlRig, const FName& ControlName,
+	bool GetControlRigControlTransforms(ISequencer* Sequencer, UControlRig* ControlRig, const FName& ControlName,
 		const TArray<FFrameNumber>& Frames, const TArray<FTransform>& ParentTransforms, TArray<FTransform>& OutTransforms);
 
 
@@ -90,7 +90,7 @@ struct FControlRigSnapper
 	* @param OutTransforms  Results
 	*/
 	bool GetControlRigControlTransforms(UWorld* World, ULevelSequence* LevelSequence,
-		UBaseControlRig* ControlRig, const FName& ControlName,const TArray<FFrameNumber>& Frames, const TArray<FTransform>& ParentTransforms, TArray<FTransform>& OutTransforms);
+		UControlRig* ControlRig, const FName& ControlName,const TArray<FFrameNumber>& Frames, const TArray<FTransform>& ParentTransforms, TArray<FTransform>& OutTransforms);
 
 };
 
diff --git a/Engine/Plugins/Animation/ControlRig/Source/ControlRigEditor/Public/Tools/ControlRigTweener.h b/Engine/Plugins/Animation/ControlRig/Source/ControlRigEditor/Public/Tools/ControlRigTweener.h
--- a/Engine/Plugins/Animation/ControlRig/Source/ControlRigEditor/Public/Tools/ControlRigTweener.h
+++ b/Engine/Plugins/Animation/ControlRig/Source/ControlRigEditor/Public/Tools/ControlRigTweener.h
@@ -10,7 +10,7 @@
 
 struct FMovieSceneFloatChannel;
 struct FMovieSceneDoubleChannel;
-class UBaseControlRig;
+class UControlRig;
 class ISequencer;
 struct FRigControl;
 class UMovieSceneSection;
@@ -85,7 +85,7 @@ struct FAnimSliderObjectSelection
 	};
 
 	bool Setup(TWeakPtr<ISequencer>& InSequencer, TWeakPtr<class FControlRigEditMode>& InEditMode);
-	bool Setup(const TArray<UBaseControlRig*>& SelectedControlRigs, TWeakPtr<ISequencer>& InSequencer);
+	bool Setup(const TArray<UControlRig*>& SelectedControlRigs, TWeakPtr<ISequencer>& InSequencer);
 
 	TArray<FObjectChannels>  ChannelsArray;
 
@@ -95,7 +95,7 @@ private:
 	void SetupChannel(FFrameNumber CurrentFrame, TArray<FFrameNumber>& KeyTimes, TArray<FKeyHandle>& Handles, FMovieSceneFloatChannel* FloatChannel,
 		FMovieSceneDoubleChannel* DoubleChannel, FChannelKeyBounds& KeyBounds);
 
-	TArray<UBaseControlRig*> GetControlRigs(TWeakPtr<class FControlRigEditMode>& InEditMode);
+	TArray<UControlRig*> GetControlRigs(TWeakPtr<class FControlRigEditMode>& InEditMode);
 };
 
 /*
@@ -230,7 +230,7 @@ struct FControlsToTween :public FBasicBlendSlider
 	//BasicBlendSlider overrides
 	virtual double DoBlend(const FBlendStruct& BlendStruct) override;
 
-	bool Setup(const TArray<UBaseControlRig*>& SelectedControlRigs, TWeakPtr<ISequencer>& InSequencer);
+	bool Setup(const TArray<UControlRig*>& SelectedControlRigs, TWeakPtr<ISequencer>& InSequencer);
 
 };
 
diff --git a/Engine/Plugins/Animation/RigLogic/Source/RigLogicModule/Private/RigUnit_RigLogic.cpp b/Engine/Plugins/Animation/RigLogic/Source/RigLogicModule/Private/RigUnit_RigLogic.cpp
--- a/Engine/Plugins/Animation/RigLogic/Source/RigLogicModule/Private/RigUnit_RigLogic.cpp
+++ b/Engine/Plugins/Animation/RigLogic/Source/RigLogicModule/Private/RigUnit_RigLogic.cpp
@@ -390,7 +390,7 @@ FRigUnit_RigLogic_Execute()
 			//const double startTime = FPlatformTime::Seconds();
 			if (!Data.SkelMeshComponent.IsValid())
 			{
-				Data.SkelMeshComponent = ExecuteContext.UnitContext.DataSourceRegistry->RequestSource<USkeletalMeshComponent>(UBaseControlRig::OwnerComponent);
+				Data.SkelMeshComponent = ExecuteContext.UnitContext.DataSourceRegistry->RequestSource<USkeletalMeshComponent>(UControlRig::OwnerComponent);
 				// In normal execution, Data.SkelMeshComponent will be nullptr at the beginning
 				// however, during unit testing we cannot fetch it from DataSourceRegistry 
 				// in that case, a mock version will be inserted into Data by unit test beforehand
diff --git a/Engine/Plugins/Editor/SequencerAnimTools/Source/SequencerAnimTools/Private/Sequencer/MovieSceneTransformTrail.cpp b/Engine/Plugins/Editor/SequencerAnimTools/Source/SequencerAnimTools/Private/Sequencer/MovieSceneTransformTrail.cpp
--- a/Engine/Plugins/Editor/SequencerAnimTools/Source/SequencerAnimTools/Private/Sequencer/MovieSceneTransformTrail.cpp
+++ b/Engine/Plugins/Editor/SequencerAnimTools/Source/SequencerAnimTools/Private/Sequencer/MovieSceneTransformTrail.cpp
@@ -686,7 +686,7 @@ bool FMovieSceneControlRigTransformTrail::StartTracking()
 	{
 		return false;
 	}
-	UBaseControlRig* ControlRig = Section->GetControlRig();
+	UControlRig* ControlRig = Section->GetControlRig();
 	if (!ControlRig)
 	{
 		return false ;
@@ -720,7 +720,7 @@ bool FMovieSceneControlRigTransformTrail::ApplyDelta(const FVector& Pos, const F
 	{
 		return false;
 	}
-	UBaseControlRig* ControlRig = Section->GetControlRig();
+	UControlRig* ControlRig = Section->GetControlRig();
 	if (!ControlRig)
 	{
 		return false;
@@ -793,7 +793,7 @@ bool FMovieSceneControlRigTransformTrail::EndTracking()
 	{
 		return false;
 	}
-	UBaseControlRig* ControlRig = Section->GetControlRig();
+	UControlRig* ControlRig = Section->GetControlRig();
 	if (!ControlRig)
 	{
 		return false;
@@ -824,7 +824,7 @@ bool FMovieSceneControlRigTransformTrail::HandleAltClick(FEditorViewportClient*
 	{
 		return false;
 	}
-	UBaseControlRig* ControlRig = Section->GetControlRig();
+	UControlRig* ControlRig = Section->GetControlRig();
 	if (!ControlRig)
 	{
 		return false;
diff --git a/Engine/Plugins/Editor/SequencerAnimTools/Source/SequencerAnimTools/Private/Sequencer/SequencerTrailHierarchy.cpp b/Engine/Plugins/Editor/SequencerAnimTools/Source/SequencerAnimTools/Private/Sequencer/SequencerTrailHierarchy.cpp
--- a/Engine/Plugins/Editor/SequencerAnimTools/Source/SequencerAnimTools/Private/Sequencer/SequencerTrailHierarchy.cpp
+++ b/Engine/Plugins/Editor/SequencerAnimTools/Source/SequencerAnimTools/Private/Sequencer/SequencerTrailHierarchy.cpp
@@ -197,7 +197,7 @@ void FSequencerTrailHierarchy::RemoveTrail(const FGuid& Key)
 				return;
 			}
 		}
-		for (TPair<UBaseControlRig*, TMap<FName, FGuid>>& CompMapPair : ControlsTracked)
+		for (TPair<UControlRig*, TMap<FName, FGuid>>& CompMapPair : ControlsTracked)
 		{
 			if (const FName* FoundControl = CompMapPair.Value.FindKey(Key))
 			{
@@ -214,7 +214,7 @@ struct FTrailControlTransforms
 	TArray<FTransform> Transforms;
 };
 
-void FSequencerTrailHierarchy::UpdateControlRig(const TArray<FFrameNumber> &Frames,UBaseControlRig* ControlRig, TMap<FName, FGuid >& CompMapPair, bool bUseEditedTimes)
+void FSequencerTrailHierarchy::UpdateControlRig(const TArray<FFrameNumber> &Frames,UControlRig* ControlRig, TMap<FName, FGuid >& CompMapPair, bool bUseEditedTimes)
 {
 	if (TSharedPtr<IControlRigObjectBinding> ObjectBinding = ControlRig->GetObjectBinding())
 	{
@@ -299,7 +299,7 @@ void FSequencerTrailHierarchy::UpdateControlRig(const TArray<FFrameNumber> &Fram
 	}
 }
 
-void FSequencerTrailHierarchy::UpdateControlRig(const FTrailEvaluateTimes& EvaluateTimes,UBaseControlRig* ControlRig, TMap<FName, FGuid > &CompMapPair)
+void FSequencerTrailHierarchy::UpdateControlRig(const FTrailEvaluateTimes& EvaluateTimes,UControlRig* ControlRig, TMap<FName, FGuid > &CompMapPair)
 {
 	if (TSharedPtr<IControlRigObjectBinding> ObjectBinding = ControlRig->GetObjectBinding())
 	{
@@ -340,7 +340,7 @@ void FSequencerTrailHierarchy::Update()
 	FTrailHierarchy::Update();  //calculates EvalTimesArr
 	FTrailEvaluateTimes EvalTimes = FTrailEvaluateTimes(EvalTimesArr, SecondsPerSegment);
 
-	for (TPair<UBaseControlRig*, TMap<FName, FGuid>>& CompMapPair : ControlsTracked)
+	for (TPair<UControlRig*, TMap<FName, FGuid>>& CompMapPair : ControlsTracked)
 	{
 		bool bNeedToUpdateControlRig = false;
 		FTrail* ForceTrail = nullptr;
@@ -594,7 +594,7 @@ void FSequencerTrailHierarchy::UpdateSequencerBindings(const TArray<FGuid>& Sequ
 				if (!ControlRigDelegateHandles.Contains(CRParameterTrack))
 				{
 					RegisterControlRigDelegates(BoundComponent, CRParameterTrack);
-					UBaseControlRig* ControlRig = CRParameterTrack->GetControlRig();
+					UControlRig* ControlRig = CRParameterTrack->GetControlRig();
 					if (ControlRig)
 					{
 						TArray<FName> Selected = ControlRig->CurrentControlSelection();
@@ -702,7 +702,7 @@ void FSequencerTrailHierarchy::AddSkeletonToHierarchy(class USkeletalMeshCompone
 	TimingStats.Add("FSequencerTrailHierarchy::AddSkeletonToHierarchy", Timespan);
 }
 
-void FSequencerTrailHierarchy::AddControlRigTrail(USkeletalMeshComponent* Component, UBaseControlRig* ControlRig, UMovieSceneControlRigParameterTrack* CRParameterTrack,const FName& ControlName)
+void FSequencerTrailHierarchy::AddControlRigTrail(USkeletalMeshComponent* Component, UControlRig* ControlRig, UMovieSceneControlRigParameterTrack* CRParameterTrack,const FName& ControlName)
 {
 	if (ControlsTracked.Find(ControlRig) == nullptr)
 	{
@@ -737,10 +737,10 @@ void FSequencerTrailHierarchy::AddControlRigTrail(USkeletalMeshComponent* Compon
 
 void FSequencerTrailHierarchy::RegisterControlRigDelegates(USkeletalMeshComponent* Component, UMovieSceneControlRigParameterTrack* CRParameterTrack)
 {
-	UBaseControlRig* ControlRig = CRParameterTrack->GetControlRig();
+	UControlRig* ControlRig = CRParameterTrack->GetControlRig();
 	URigHierarchy* RigHierarchy = ControlRig->GetHierarchy();
 	FControlRigDelegateHandles& DelegateHandles = ControlRigDelegateHandles.Add(CRParameterTrack);
-	DelegateHandles.OnControlSelected = ControlRig->ControlSelected().AddLambda([this, Component,CRParameterTrack](UBaseControlRig* ControlRig, FRigControlElement* ControlElement, bool bSelected)
+	DelegateHandles.OnControlSelected = ControlRig->ControlSelected().AddLambda([this, Component,CRParameterTrack](UControlRig* ControlRig, FRigControlElement* ControlElement, bool bSelected)
 		{
 			if (ControlElement->Settings.ControlType != ERigControlType::Transform &&
 				ControlElement->Settings.ControlType != ERigControlType::TransformNoScale &&
@@ -769,9 +769,9 @@ void FSequencerTrailHierarchy::RegisterControlRigDelegates(USkeletalMeshComponen
 
 			//check to see if the seleced control rig is sill selected
 			TArray<FGuid> TrailsToRemove;
-			for (TPair<UBaseControlRig*, TMap<FName, FGuid>>& CompMapPair : ControlsTracked)
+			for (TPair<UControlRig*, TMap<FName, FGuid>>& CompMapPair : ControlsTracked)
 			{
-				UBaseControlRig* TrackedControlRig = CompMapPair.Key;
+				UControlRig* TrackedControlRig = CompMapPair.Key;
 				for (TPair<FName, FGuid> NameGuid : CompMapPair.Value)
 				{
 					if (TrackedControlRig->IsControlSelected(NameGuid.Key) == false)
diff --git a/Engine/Plugins/Editor/SequencerAnimTools/Source/SequencerAnimTools/Private/Sequencer/SequencerTrailHierarchy.h b/Engine/Plugins/Editor/SequencerAnimTools/Source/SequencerAnimTools/Private/Sequencer/SequencerTrailHierarchy.h
--- a/Engine/Plugins/Editor/SequencerAnimTools/Source/SequencerAnimTools/Private/Sequencer/SequencerTrailHierarchy.h
+++ b/Engine/Plugins/Editor/SequencerAnimTools/Source/SequencerAnimTools/Private/Sequencer/SequencerTrailHierarchy.h
@@ -17,7 +17,7 @@ class UMovieSceneSection;
 class UMovieScene3DTransformTrack;
 class UMovieSceneControlRigParameterTrack;
 class UMovieSceneControlRigParameterSection;
-class UBaseControlRig;
+class UControlRig;
 struct FRigHierarchyContainer;
 
 namespace UE
@@ -62,7 +62,7 @@ public:
 
 	const TMap<UObject*, FGuid>& GetObjectsTracked() const { return ObjectsTracked; }
 	const TMap<USkeletalMeshComponent*, TMap<FName, FGuid>>& GetBonesTracked() const { return BonesTracked; }
-	const TMap<UBaseControlRig*, TMap<FName, FGuid>>& GetControlsTracked() const { return ControlsTracked; }
+	const TMap<UControlRig*, TMap<FName, FGuid>>& GetControlsTracked() const { return ControlsTracked; }
 
 	void OnBoneVisibilityChanged(USkeleton* Skeleton, const FName& BoneName, const bool bIsVisible);
 	void OnBindingVisibilityStateChanged(UObject* BoundObject, const EBindingVisibilityState VisibilityState);
@@ -71,10 +71,10 @@ public:
 	void OnActorsChangedSomehow(TArray<AActor*>& InActors);	
 
 private:
-	void AddControlRigTrail(USkeletalMeshComponent* Component,UBaseControlRig* ControlRig, UMovieSceneControlRigParameterTrack* CRTrack, const FName& ControlName);
+	void AddControlRigTrail(USkeletalMeshComponent* Component,UControlRig* ControlRig, UMovieSceneControlRigParameterTrack* CRTrack, const FName& ControlName);
 
-	void UpdateControlRig(const TArray<FFrameNumber>& Frames, UBaseControlRig* ControlRig, TMap<FName, FGuid >& CompMapPair, bool bUseEditedTimes);
-	void UpdateControlRig(const FTrailEvaluateTimes& EvalTimes, UBaseControlRig* ControlRig, TMap<FName, FGuid >& CompMapPair);
+	void UpdateControlRig(const TArray<FFrameNumber>& Frames, UControlRig* ControlRig, TMap<FName, FGuid >& CompMapPair, bool bUseEditedTimes);
+	void UpdateControlRig(const FTrailEvaluateTimes& EvalTimes, UControlRig* ControlRig, TMap<FName, FGuid >& CompMapPair);
 
 	void UpdateSequencerBindings(const TArray<FGuid>& SequencerBindings, TFunctionRef<void(UObject*, FTrail*, FGuid)> OnUpdated);
 	void UpdateViewAndEvalRange();
@@ -87,7 +87,7 @@ private:
 	TWeakPtr<ISequencer> WeakSequencer;
 	TMap<UObject*, FGuid> ObjectsTracked;
 	TMap<USkeletalMeshComponent*, TMap<FName, FGuid>> BonesTracked; 
-	TMap<UBaseControlRig*, TMap<FName, FGuid>> ControlsTracked; 
+	TMap<UControlRig*, TMap<FName, FGuid>> ControlsTracked; 
 
 	// TODO: components can have multiple rigs so make this a map from sections to controls instead. However, this is only part of a larger problem of handling blending
 
diff --git a/Engine/Plugins/Editor/SequencerAnimTools/Source/SequencerAnimTools/Private/SequencerAnimEditPivotTool.cpp b/Engine/Plugins/Editor/SequencerAnimTools/Source/SequencerAnimTools/Private/SequencerAnimEditPivotTool.cpp
--- a/Engine/Plugins/Editor/SequencerAnimTools/Source/SequencerAnimTools/Private/SequencerAnimEditPivotTool.cpp
+++ b/Engine/Plugins/Editor/SequencerAnimTools/Source/SequencerAnimTools/Private/SequencerAnimEditPivotTool.cpp
@@ -17,7 +17,7 @@
 #include "EditorModeManager.h"
 #include "Misc/QualifiedFrameTime.h"
 #include "Modules/ModuleManager.h"
-#include "ControlRig.h"
+#include "BaseControlRig.h"
 #include "ControlRigSequencerEditorLibrary.h"
 #include "LevelSequence.h"
 #include "LevelSequenceEditorBlueprintLibrary.h"
@@ -49,7 +49,7 @@ FSavedMappings USequencerPivotTool::SavedPivotLocations;
 FLastSelectedObjects USequencerPivotTool::LastSelectedObjects;
 
 
-static void GetControlRigsAndSequencer(TArray<TWeakObjectPtr<UBaseControlRig>>& ControlRigs, TWeakPtr<ISequencer>& SequencerPtr, ULevelSequence** LevelSequence)
+static void GetControlRigsAndSequencer(TArray<TWeakObjectPtr<UControlRig>>& ControlRigs, TWeakPtr<ISequencer>& SequencerPtr, ULevelSequence** LevelSequence)
 {
 	*LevelSequence = ULevelSequenceEditorBlueprintLibrary::GetFocusedLevelSequence();
 	//if getting sequencer from level sequence need to use the current(master), not the focused
@@ -65,7 +65,7 @@ static void GetControlRigsAndSequencer(TArray<TWeakObjectPtr<UBaseControlRig>>&
 			TempControlRigs = UControlRigSequencerEditorLibrary::GetVisibleControlRigs();
 			for (UControlRig* ControlRig : TempControlRigs)
 			{
-				UBaseControlRig* Base = ControlRig;
+				UControlRig* Base = ControlRig;
 				ControlRigs.Add(ControlRig);
 			}
 		}
@@ -80,11 +80,11 @@ bool USequencerPivotToolBuilder::CanBuildTool(const FToolBuilderState& SceneStat
 		return true;
 	}
 	ULevelSequence* LevelSequence;
-	TArray<TWeakObjectPtr<UBaseControlRig>> ControlRigs;
+	TArray<TWeakObjectPtr<UControlRig>> ControlRigs;
 	TWeakPtr<ISequencer> SequencerPtr;
 	GetControlRigsAndSequencer(ControlRigs, SequencerPtr, &LevelSequence);
 
-	for (TWeakObjectPtr<UBaseControlRig> ControlRig : ControlRigs)
+	for (TWeakObjectPtr<UControlRig> ControlRig : ControlRigs)
 	{
 		if (ControlRig.IsValid() && ControlRig->CurrentControlSelection().Num() > 0)
 		{
@@ -290,7 +290,7 @@ void USequencerPivotTool::Setup()
 	UpdateGizmoTransform();
 	UpdateGizmoVisibility();
 	//we get delegates last since we may select something above
-	for (TWeakObjectPtr<UBaseControlRig>& ControlRig : ControlRigs)
+	for (TWeakObjectPtr<UControlRig>& ControlRig : ControlRigs)
 	{
 		if (ControlRig.IsValid())
 		{
@@ -354,7 +354,7 @@ void USequencerPivotTool::SaveLastSelected()
 	LastSelectedObjects.LastSelectedControlRigs.SetNum(0);
 	LastSelectedObjects.LastSelectedActors.SetNum(0);
 
-	for (TWeakObjectPtr<UBaseControlRig>& ControlRig : ControlRigs)
+	for (TWeakObjectPtr<UControlRig>& ControlRig : ControlRigs)
 	{
 		if (ControlRig.IsValid())
 		{
@@ -429,7 +429,7 @@ bool USequencerPivotTool::SetGizmoBasedOnSelection(bool bUseSaved)
 	bool bHaveSomethingSelected = false;
 	FVector AverageLocation(0.0f);
 	int NumLocations = 0;
-	for (TWeakObjectPtr<UBaseControlRig>& ControlRig : ControlRigs)
+	for (TWeakObjectPtr<UControlRig>& ControlRig : ControlRigs)
 	{
 		if (ControlRig.IsValid() == false)
 		{
@@ -584,14 +584,14 @@ void USequencerPivotTool::OnEditorSelectionChanged(UObject* NewSelection)
 	*/
 }
 
-void USequencerPivotTool::HandleControlSelected(UBaseControlRig* Subject, FRigControlElement* InControl, bool bSelected)
+void USequencerPivotTool::HandleControlSelected(UControlRig* Subject, FRigControlElement* InControl, bool bSelected)
 {
 	DeactivateMe();
 }
 
 void USequencerPivotTool::RemoveDelegates()
 {
-	for (TWeakObjectPtr<UBaseControlRig>& ControlRig : ControlRigs)
+	for (TWeakObjectPtr<UControlRig>& ControlRig : ControlRigs)
 	{
 		if (ControlRig.IsValid())
 		{
@@ -616,7 +616,7 @@ void USequencerPivotTool::GizmoTransformStarted(UTransformProxy* Proxy)
 		const FFrameTime FrameTime = Sequencer->GetLocalTime().ConvertTo(TickResolution);
 		const FFrameNumber FrameNumber = FrameTime.RoundToFrame();
 
-		for (TWeakObjectPtr<UBaseControlRig>& ControlRig : ControlRigs)
+		for (TWeakObjectPtr<UControlRig>& ControlRig : ControlRigs)
 		{
 			if (ControlRig.IsValid() == false)
 			{
@@ -696,7 +696,7 @@ void USequencerPivotTool::GizmoTransformStarted(UTransformProxy* Proxy)
 	InteractionScopes.Reset();
 	if(bInPivotMode)
 	{
-		TMap<UBaseControlRig*, int32> RigToScopeIndex;
+		TMap<UControlRig*, int32> RigToScopeIndex;
 		for(int32 IndexA = 0; IndexA < ControlRigDrags.Num(); IndexA++)
 		{
 			const FControlRigSelectionDuringDrag& ControlRigDrag = ControlRigDrags[IndexA];
diff --git a/Engine/Plugins/Editor/SequencerAnimTools/Source/SequencerAnimTools/Public/SequencerAnimEditPivotTool.h b/Engine/Plugins/Editor/SequencerAnimTools/Source/SequencerAnimTools/Public/SequencerAnimEditPivotTool.h
--- a/Engine/Plugins/Editor/SequencerAnimTools/Source/SequencerAnimTools/Public/SequencerAnimEditPivotTool.h
+++ b/Engine/Plugins/Editor/SequencerAnimTools/Source/SequencerAnimTools/Public/SequencerAnimEditPivotTool.h
@@ -13,7 +13,7 @@
 
 class FControlRigInteractionScope;
 class SWidget;
-class UBaseControlRig;
+class UControlRig;
 class UInteractiveGizmoManager;
 class UTransformProxy;
 enum class EToolShutdownType : uint8;
@@ -56,7 +56,7 @@ struct FSelectionDuringDrag
 };
 struct FControlRigSelectionDuringDrag : public FSelectionDuringDrag
 {
-	UBaseControlRig* ControlRig;
+	UControlRig* ControlRig;
 	FName ControlName;
 };
 
@@ -67,7 +67,7 @@ struct FActorSelectonDuringDrag : public FSelectionDuringDrag
 
 struct FControlRigMappings
 {
-	TWeakObjectPtr<UBaseControlRig> ControlRig;
+	TWeakObjectPtr<UControlRig> ControlRig;
 
 	FTransform GetParentTransform() const;
 	TOptional<FTransform> GetWorldTransform(const FName& Name) const;
@@ -92,7 +92,7 @@ private:
 
 struct FSavedMappings
 {
-	TMap<TWeakObjectPtr<UBaseControlRig>, FControlRigMappings> ControlRigMappings;
+	TMap<TWeakObjectPtr<UControlRig>, FControlRigMappings> ControlRigMappings;
 	TMap<TWeakObjectPtr<AActor>, FActorMappings> ActorMappings;
 };
 
@@ -195,7 +195,7 @@ protected:
 
 	//since we are selection based we can cache this
 	ULevelSequence* LevelSequence;
-	TArray<TWeakObjectPtr<UBaseControlRig>> ControlRigs;
+	TArray<TWeakObjectPtr<UControlRig>> ControlRigs;
 	TWeakPtr<ISequencer> SequencerPtr;
 	TArray<TWeakObjectPtr<AActor>> Actors;
 
@@ -220,7 +220,7 @@ protected:
 	// selection delegates
 	void DeactivateMe();
 	void RemoveDelegates();
-	void HandleControlSelected(UBaseControlRig* Subject, FRigControlElement* InControl, bool bSelected);
+	void HandleControlSelected(UControlRig* Subject, FRigControlElement* InControl, bool bSelected);
 	void OnEditorSelectionChanged(UObject* NewSelection);
 	FDelegateHandle OnEditorSelectionChangedHandle;
 
diff --git a/Engine/Plugins/Experimental/AnimNext/Source/AnimNextUncookedOnly/Private/Graph/AnimGraph/AnimGraphNode_AnimNextGraph.cpp b/Engine/Plugins/Experimental/AnimNext/Source/AnimNextUncookedOnly/Private/Graph/AnimGraph/AnimGraphNode_AnimNextGraph.cpp
--- a/Engine/Plugins/Experimental/AnimNext/Source/AnimNextUncookedOnly/Private/Graph/AnimGraph/AnimGraphNode_AnimNextGraph.cpp
+++ b/Engine/Plugins/Experimental/AnimNext/Source/AnimNextUncookedOnly/Private/Graph/AnimGraph/AnimGraphNode_AnimNextGraph.cpp
@@ -157,7 +157,7 @@ void UAnimGraphNode_AnimNextGraph::GetVariables(bool bInput, TMap<FName, FRigVME
 		if (UControlRigBlueprint* RigBlueprint = Cast<UControlRigBlueprint>(TargetClass->ClassGeneratedBy))
 		{
 			//RigBlueprint->CleanupVariables();
-			UBaseControlRig* ControlRig = TargetClass->GetDefaultObject<UBaseControlRig>();
+			UControlRig* ControlRig = TargetClass->GetDefaultObject<UControlRig>();
 			if (ControlRig)
 			{
 				const TArray<FRigVMExternalVariable>& PublicVariables = ControlRig->GetPublicVariables();
diff --git a/Engine/Plugins/Importers/USDImporter/Source/USDStage/Private/USDLevelSequenceHelper.cpp b/Engine/Plugins/Importers/USDImporter/Source/USDStage/Private/USDLevelSequenceHelper.cpp
--- a/Engine/Plugins/Importers/USDImporter/Source/USDStage/Private/USDLevelSequenceHelper.cpp
+++ b/Engine/Plugins/Importers/USDImporter/Source/USDStage/Private/USDLevelSequenceHelper.cpp
@@ -306,7 +306,7 @@ namespace UsdLevelSequenceHelperImpl
 		{
 			FString ObjectName = InClass->GetName();
 			ObjectName.RemoveFromEnd(TEXT("_C"));
-			UBaseControlRig* ControlRig = NewObject<UBaseControlRig>(Track, InClass, FName(*ObjectName), RF_Transactional);
+			UControlRig* ControlRig = NewObject<UControlRig>(Track, InClass, FName(*ObjectName), RF_Transactional);
 			if (InClass != UFKControlRig::StaticClass() && !ControlRig->SupportsEvent(TEXT("Backwards Solve")))
 			{
 				MovieScene->RemoveTrack(*Track);
@@ -316,7 +316,7 @@ namespace UsdLevelSequenceHelperImpl
 			ControlRig->Modify();
 			ControlRig->SetObjectBinding(MakeShared<FControlRigObjectBinding>());
 			ControlRig->GetObjectBinding()->BindToObject(SkeletalMeshComp);
-			ControlRig->GetDataSourceRegistry()->RegisterDataSource(UBaseControlRig::OwnerComponent, ControlRig->GetObjectBinding()->GetBoundObject());
+			ControlRig->GetDataSourceRegistry()->RegisterDataSource(UControlRig::OwnerComponent, ControlRig->GetObjectBinding()->GetBoundObject());
 			ControlRig->Initialize();
 			ControlRig->RequestInit();
 			ControlRig->SetBoneInitialTransformsFromSkeletalMeshComponent(SkeletalMeshComp, true);
diff --git a/Engine/Plugins/Importers/USDImporter/Source/USDUtilities/Private/USDSkeletalDataConversion.cpp b/Engine/Plugins/Importers/USDImporter/Source/USDUtilities/Private/USDSkeletalDataConversion.cpp
--- a/Engine/Plugins/Importers/USDImporter/Source/USDUtilities/Private/USDSkeletalDataConversion.cpp
+++ b/Engine/Plugins/Importers/USDImporter/Source/USDUtilities/Private/USDSkeletalDataConversion.cpp
@@ -3278,7 +3278,7 @@ bool UnrealToUsd::ConvertControlRigSection(
 		return false;
 	}
 
-	UBaseControlRig* ControlRig = InSection->GetControlRig();
+	UControlRig* ControlRig = InSection->GetControlRig();
 	if ( !ControlRig )
 	{
 		return false;
