Control Rig: Per rig type specific layout / UI behavior

#rnx

diff --git a/Engine/Plugins/Animation/ControlRig/Source/ControlRigDeveloper/Private/ControlRigBlueprint.cpp b/Engine/Plugins/Animation/ControlRig/Source/ControlRigDeveloper/Private/ControlRigBlueprint.cpp
--- a/Engine/Plugins/Animation/ControlRig/Source/ControlRigDeveloper/Private/ControlRigBlueprint.cpp
+++ b/Engine/Plugins/Animation/ControlRig/Source/ControlRigDeveloper/Private/ControlRigBlueprint.cpp
@@ -152,11 +152,20 @@ void UControlRigBlueprint::PostEditChangeChainProperty(FPropertyChangedChainEven
 	}
 }
 
-UClass* UControlRigBlueprint::GetControlRigClass()
+UClass* UControlRigBlueprint::GetControlRigClass() const
 {
 	return GetRigVMHostClass();
 }
 
+bool UControlRigBlueprint::IsModularRig() const
+{
+	if(const UClass* Class = GetControlRigClass())
+	{
+		return Class->IsChildOf(UModularRig::StaticClass());
+	}
+	return false;
+}
+
 USkeletalMesh* UControlRigBlueprint::GetPreviewMesh() const
 {
 	DECLARE_SCOPE_HIERARCHICAL_COUNTER_FUNC()
diff --git a/Engine/Plugins/Animation/ControlRig/Source/ControlRigDeveloper/Public/ControlRigBlueprint.h b/Engine/Plugins/Animation/ControlRig/Source/ControlRigDeveloper/Public/ControlRigBlueprint.h
--- a/Engine/Plugins/Animation/ControlRig/Source/ControlRigDeveloper/Public/ControlRigBlueprint.h
+++ b/Engine/Plugins/Animation/ControlRig/Source/ControlRigDeveloper/Public/ControlRigBlueprint.h
@@ -80,6 +80,12 @@ public:
 	virtual void PostDuplicate(bool bDuplicateForPIE) override;
 	virtual bool RequiresForceLoadMembers(UObject* InObject) const override;
 
+	virtual bool SupportsGlobalVariables() const override { return true; }
+	virtual bool SupportsLocalVariables() const override { return !IsModularRig(); }
+	virtual bool SupportsFunctions() const override { return !IsModularRig(); }
+	virtual bool SupportsEventGraphs() const override { return !IsModularRig(); }
+
+
 	// UObject interface
 	virtual void PostEditChangeProperty(FPropertyChangedEvent& PropertyChangedEvent) override;
 	virtual void PostEditChangeChainProperty(struct FPropertyChangedChainEvent& PropertyChangedEvent) override;
@@ -87,7 +93,9 @@ public:
 #endif	// #if WITH_EDITOR
 
 	UFUNCTION(BlueprintCallable, Category = "VM")
-	UClass* GetControlRigClass();
+	UClass* GetControlRigClass() const;
+
+	bool IsModularRig() const;
 
 	UFUNCTION(BlueprintCallable, Category = "Control Rig Blueprint")
 	UControlRig* CreateControlRig() { return Cast<UControlRig>(CreateRigVMHost()); }
diff --git a/Engine/Plugins/Animation/ControlRig/Source/ControlRigEditor/Private/ControlRigEditorModule.cpp b/Engine/Plugins/Animation/ControlRig/Source/ControlRigEditor/Private/ControlRigEditorModule.cpp
--- a/Engine/Plugins/Animation/ControlRig/Source/ControlRigEditor/Private/ControlRigEditorModule.cpp
+++ b/Engine/Plugins/Animation/ControlRig/Source/ControlRigEditor/Private/ControlRigEditorModule.cpp
@@ -191,6 +191,13 @@ void FControlRigEditorModule::StartupModule()
 		false,
 		8500);
 
+	FEditorModeRegistry::Get().RegisterMode<FModularRigEditorEditMode>(
+		FModularRigEditorEditMode::ModeName,
+		NSLOCTEXT("RiggingModeToolkit", "DisplayName", "Rigging"),
+		FSlateIcon(FRigVMEditorStyle::Get().GetStyleSetName(), "RigVMEditMode", "RigVMEditMode.Small"),
+		false,
+		9000);
+
 	ICurveEditorModule& CurveEditorModule = FModuleManager::LoadModuleChecked<ICurveEditorModule>("CurveEditor");
 	FControlRigSpaceChannelCurveModel::ViewID = CurveEditorModule.RegisterView(FOnCreateCurveEditorView::CreateStatic(
 		[](TWeakPtr<FCurveEditor> WeakCurveEditor) -> TSharedRef<SCurveEditorView>
@@ -233,6 +240,7 @@ void FControlRigEditorModule::ShutdownModule()
 	//UThumbnailManager::Get().UnregisterCustomRenderer(UControlRigBlueprint::StaticClass());
 	//UActorFactorySkeletalMesh::UnregisterDelegatesForAssetClass(UControlRigBlueprint::StaticClass());
 
+	FEditorModeRegistry::Get().UnregisterMode(FModularRigEditorEditMode::ModeName);
 	FEditorModeRegistry::Get().UnregisterMode(FControlRigEditorEditMode::ModeName);
 	FEditorModeRegistry::Get().UnregisterMode(FControlRigEditMode::ModeName);
 
diff --git a/Engine/Plugins/Animation/ControlRig/Source/ControlRigEditor/Private/Editor/ControlRigEditor.cpp b/Engine/Plugins/Animation/ControlRig/Source/ControlRigEditor/Private/Editor/ControlRigEditor.cpp
--- a/Engine/Plugins/Animation/ControlRig/Source/ControlRigEditor/Private/Editor/ControlRigEditor.cpp
+++ b/Engine/Plugins/Animation/ControlRig/Source/ControlRigEditor/Private/Editor/ControlRigEditor.cpp
@@ -172,6 +172,69 @@ UClass* FControlRigEditor::GetDetailWrapperClass() const
 	return UControlRigWrapperObject::StaticClass();
 }
 
+bool FControlRigEditor::IsSectionVisible(NodeSectionID::Type InSectionID) const
+{
+	if(!IControlRigEditor::IsSectionVisible(InSectionID))
+	{
+		return false;
+	}
+
+	if(const UControlRigBlueprint* RigBlueprint = GetControlRigBlueprint())
+	{
+		if(IsModularRig())
+		{
+			switch (InSectionID)
+			{
+				case NodeSectionID::GRAPH:
+				{
+					return RigBlueprint->SupportsEventGraphs();
+				}
+				case NodeSectionID::FUNCTION:
+				{
+					return RigBlueprint->SupportsFunctions();
+				}
+				default:
+				{
+					break;
+				}
+			}
+		}
+	}
+	return true;
+}
+
+bool FControlRigEditor::NewDocument_IsVisibleForType(ECreatedDocumentType GraphType) const
+{
+	if(!IControlRigEditor::NewDocument_IsVisibleForType(GraphType))
+	{
+		return false;
+	}
+
+	if(const UControlRigBlueprint* RigBlueprint = GetControlRigBlueprint())
+	{
+		if(IsModularRig())
+		{
+			switch(GraphType)
+			{
+				case CGT_NewEventGraph:
+				{
+					return RigBlueprint->SupportsEventGraphs();
+				}
+				case CGT_NewFunctionGraph:
+				{
+					return RigBlueprint->SupportsFunctions();
+				}
+				default:
+				{
+					break;
+				}
+			}
+		}
+	}
+
+	return true;
+}
+
 void FControlRigEditor::CreateEmptyGraphContent(URigVMController* InController)
 {
 	URigVMNode* Node = InController->AddUnitNode(FRigUnit_BeginExecution::StaticStruct(), FRigUnit::GetMethodName(), FVector2D::ZeroVector, FString(), false);
@@ -286,12 +349,21 @@ const FName FControlRigEditor::GetEditorAppName() const
 
 const FName FControlRigEditor::GetEditorModeName() const
 {
+	if(IsModularRig())
+	{
+		return FModularRigEditorEditMode::ModeName;
+	}
 	return FControlRigEditorEditMode::ModeName;
 }
 
 TSharedPtr<FApplicationMode> FControlRigEditor::CreateEditorMode()
 {
 	CreatePersonaToolKitIfRequired();
+
+	if(IsModularRig())
+	{
+		return MakeShareable(new FModularRigEditorMode(SharedThis(this)));
+	}
 	return MakeShareable(new FControlRigEditorMode(SharedThis(this)));
 }
 
@@ -1273,6 +1345,15 @@ void FControlRigEditor::SaveAssetAs_Execute()
 	ActionDatabase.RefreshClassActions(UControlRigBlueprint::StaticClass());
 }
 
+bool FControlRigEditor::IsModularRig() const
+{
+	if(UControlRigBlueprint* RigBlueprint = Cast<UControlRigBlueprint>(GetBlueprintObj()))
+	{
+		return RigBlueprint->IsModularRig();
+	}
+	return false;
+}
+
 FName FControlRigEditor::GetToolkitFName() const
 {
 	return FName("ControlRigEditor");
@@ -3460,7 +3541,7 @@ void FControlRigEditor::ResetAllBoneModification()
 
 FControlRigEditorEditMode* FControlRigEditor::GetEditMode() const
 {
-	return static_cast<FControlRigEditorEditMode*>(GetEditorModeManager().GetActiveMode(FControlRigEditorEditMode::ModeName));
+	return static_cast<FControlRigEditorEditMode*>(GetEditorModeManager().GetActiveMode(GetEditorModeName()));
 }
 
 
diff --git a/Engine/Plugins/Animation/ControlRig/Source/ControlRigEditor/Private/Editor/ControlRigEditor.h b/Engine/Plugins/Animation/ControlRig/Source/ControlRigEditor/Private/Editor/ControlRigEditor.h
--- a/Engine/Plugins/Animation/ControlRig/Source/ControlRigEditor/Private/Editor/ControlRigEditor.h
+++ b/Engine/Plugins/Animation/ControlRig/Source/ControlRigEditor/Private/Editor/ControlRigEditor.h
@@ -80,12 +80,15 @@ public:
 	UE_DEPRECATED(5.4, "Please use HandleVMCompiledEvent with ExtendedExecuteContext parameter.")
 	virtual void HandleVMCompiledEvent(UObject* InCompiledObject, URigVM* InVM) override {}
 	virtual void HandleVMCompiledEvent(UObject* InCompiledObject, URigVM* InVM, FRigVMExtendedExecuteContext& InContext) override;
+	virtual bool ShouldOpenGraphByDefault() const override { return !IsModularRig(); }
 
 	// allows the editor to fill an empty graph
 	virtual void CreateEmptyGraphContent(URigVMController* InController) override;
 
 	int32 GetRigHierarchyTabCount() const { return RigHierarchyTabCount; }
 
+	bool IsModularRig() const;
+
 public:
 	
 	// IToolkit Interface
@@ -99,7 +102,9 @@ public:
 
 	// BlueprintEditor interface
 	virtual FReply OnSpawnGraphNodeByShortcut(FInputChord InChord, const FVector2D& InPosition, UEdGraph* InGraph) override;
-	
+	virtual bool IsSectionVisible(NodeSectionID::Type InSectionID) const override;
+	virtual bool NewDocument_IsVisibleForType(ECreatedDocumentType GraphType) const override;
+
 	virtual void PostTransaction(bool bSuccess, const FTransaction* Transaction, bool bIsRedo) override;
 
 	void EnsureValidRigElementsInDetailPanel();
diff --git a/Engine/Plugins/Animation/ControlRig/Source/ControlRigEditor/Private/Editor/ControlRigEditorEditMode.cpp b/Engine/Plugins/Animation/ControlRig/Source/ControlRigEditor/Private/Editor/ControlRigEditorEditMode.cpp
--- a/Engine/Plugins/Animation/ControlRig/Source/ControlRigEditor/Private/Editor/ControlRigEditorEditMode.cpp
+++ b/Engine/Plugins/Animation/ControlRig/Source/ControlRigEditor/Private/Editor/ControlRigEditorEditMode.cpp
@@ -265,4 +265,6 @@ IPersonaPreviewScene& FControlRigEditorEditMode::GetAnimPreviewScene() const
 void FControlRigEditorEditMode::GetOnScreenDebugInfo(TArray<FText>& OutDebugInfo) const
 {
 
-}
\ No newline at end of file
+}
+
+FName FModularRigEditorEditMode::ModeName("EditMode.ModularRigEditor");
diff --git a/Engine/Plugins/Animation/ControlRig/Source/ControlRigEditor/Private/Editor/ControlRigEditorEditMode.h b/Engine/Plugins/Animation/ControlRig/Source/ControlRigEditor/Private/Editor/ControlRigEditorEditMode.h
--- a/Engine/Plugins/Animation/ControlRig/Source/ControlRigEditor/Private/Editor/ControlRigEditorEditMode.h
+++ b/Engine/Plugins/Animation/ControlRig/Source/ControlRigEditor/Private/Editor/ControlRigEditorEditMode.h
@@ -28,3 +28,9 @@ public:
 	/** Drawing options */
 	UPersonaOptions* ConfigOption = nullptr;
 };
+
+class FModularRigEditorEditMode : public FControlRigEditorEditMode
+{
+public:
+	static FName ModeName;
+};
diff --git a/Engine/Plugins/Animation/ControlRig/Source/ControlRigEditor/Private/Editor/ControlRigEditorMode.cpp b/Engine/Plugins/Animation/ControlRig/Source/ControlRigEditor/Private/Editor/ControlRigEditorMode.cpp
--- a/Engine/Plugins/Animation/ControlRig/Source/ControlRigEditor/Private/Editor/ControlRigEditorMode.cpp
+++ b/Engine/Plugins/Animation/ControlRig/Source/ControlRigEditor/Private/Editor/ControlRigEditorMode.cpp
@@ -13,7 +13,7 @@
 #include "Editor/RigAnimAttributeTabSummoner.h"
 #include "ToolMenus.h"
 
-FControlRigEditorMode::FControlRigEditorMode(const TSharedRef<FControlRigEditor>& InControlRigEditor)
+FControlRigEditorMode::FControlRigEditorMode(const TSharedRef<FControlRigEditor>& InControlRigEditor, bool bCreateDefaultLayout)
 	: FBlueprintEditorApplicationMode(InControlRigEditor, FControlRigEditorModes::ControlRigEditorMode, FControlRigEditorModes::GetLocalizedMode, false, false)
 {
 
@@ -39,7 +39,111 @@ FControlRigEditorMode::FControlRigEditorMode(const TSharedRef<FControlRigEditor>
  	TabFactories.RegisterFactory(PersonaModule.CreatePersonaViewportTabFactory(InControlRigEditor, ViewportArgs));
 	TabFactories.RegisterFactory(PersonaModule.CreateAdvancedPreviewSceneTabFactory(InControlRigEditor, InControlRigEditor->GetPersonaToolkit()->GetPreviewScene()));
 
-	TabLayout = FTabManager::NewLayout("Standalone_ControlRigEditMode_Layout_v1.5")
+	if(bCreateDefaultLayout)
+	{
+		TabLayout = FTabManager::NewLayout("Standalone_ControlRigEditMode_Layout_v1.5")
+			->AddArea
+			(
+				// Main application area
+				FTabManager::NewPrimaryArea()
+				->SetOrientation(Orient_Vertical)
+				->Split
+				(
+					FTabManager::NewSplitter()
+					->SetOrientation(Orient_Horizontal)
+					->Split
+					(
+						FTabManager::NewSplitter()
+						->SetOrientation(Orient_Vertical)
+						->SetSizeCoefficient(0.2f)
+						->Split
+						(
+							//	Left top - viewport
+							FTabManager::NewStack()
+							->SetSizeCoefficient(0.5f)
+							->SetHideTabWell(true)
+							->AddTab(FPersonaTabs::PreviewViewportID, ETabState::OpenedTab)
+						
+						)
+						->Split
+						(
+							//	Left bottom - rig/hierarchy
+							FTabManager::NewStack()
+							->SetSizeCoefficient(0.5f)
+							->AddTab(FRigHierarchyTabSummoner::TabID, ETabState::OpenedTab)
+							->AddTab(FRigVMExecutionStackTabSummoner::TabID, ETabState::OpenedTab)
+							->AddTab(FRigCurveContainerTabSummoner::TabID, ETabState::OpenedTab)
+							->AddTab(FBlueprintEditorTabs::MyBlueprintID, ETabState::OpenedTab)
+						)
+					)
+					->Split
+					(
+						// Middle 
+						FTabManager::NewSplitter()
+						->SetOrientation(Orient_Vertical)
+						->SetSizeCoefficient(0.6f)
+						->Split
+						(
+							// Middle top - document edit area
+							FTabManager::NewStack()
+							->SetSizeCoefficient(0.8f)
+							->AddTab("Document", ETabState::ClosedTab)
+						)
+						->Split
+						(
+							// Middle bottom - compiler results & find
+							FTabManager::NewStack()
+							->SetSizeCoefficient(0.2f)
+							->AddTab(FBlueprintEditorTabs::CompilerResultsID, ETabState::ClosedTab)
+							->AddTab(FBlueprintEditorTabs::FindResultsID, ETabState::ClosedTab)
+						)
+					)
+					->Split
+					(
+						// Right side
+						FTabManager::NewSplitter()
+						->SetOrientation(Orient_Vertical)
+						->SetSizeCoefficient(0.2f)
+						->Split
+						(
+							// Right top
+							FTabManager::NewStack()
+							->SetHideTabWell(false)
+							->SetSizeCoefficient(1.f)
+							->AddTab(FBlueprintEditorTabs::DetailsID, ETabState::OpenedTab)
+							->AddTab(FPersonaTabs::AdvancedPreviewSceneSettingsID, ETabState::OpenedTab)
+							->AddTab(FRigAnimAttributeTabSummoner::TabID, ETabState::OpenedTab)
+							->SetForegroundTab(FBlueprintEditorTabs::DetailsID)
+						)
+					)
+				)
+			);
+	}
+	
+	if (UToolMenu* Toolbar = InControlRigEditor->RegisterModeToolbarIfUnregistered(GetModeName()))
+	{
+		InControlRigEditor->GetToolbarBuilder()->AddCompileToolbar(Toolbar);
+		InControlRigEditor->GetToolbarBuilder()->AddScriptingToolbar(Toolbar);
+		InControlRigEditor->GetToolbarBuilder()->AddBlueprintGlobalOptionsToolbar(Toolbar);
+	}
+}
+
+void FControlRigEditorMode::RegisterTabFactories(TSharedPtr<FTabManager> InTabManager)
+{
+	TSharedPtr<FBlueprintEditor> BP = MyBlueprintEditor.Pin();
+	
+	BP->RegisterToolbarTab(InTabManager.ToSharedRef());
+
+	// Mode-specific setup
+	BP->PushTabFactories(CoreTabFactories);
+	BP->PushTabFactories(BlueprintEditorTabFactories);
+	BP->PushTabFactories(TabFactories);
+}
+
+FModularRigEditorMode::FModularRigEditorMode(const TSharedRef<FControlRigEditor>& InControlRigEditor)
+	: FControlRigEditorMode(InControlRigEditor, false)
+{
+	TabLayout = FTabManager::NewLayout("Standalone_ModularRigEditMode_Layout_v1.0")
 		->AddArea
 		(
 			// Main application area
@@ -68,35 +172,13 @@ FControlRigEditorMode::FControlRigEditorMode(const TSharedRef<FControlRigEditor>
 						//	Left bottom - rig/hierarchy
 						FTabManager::NewStack()
 						->SetSizeCoefficient(0.5f)
+						->AddTab(FBlueprintEditorTabs::CompilerResultsID, ETabState::ClosedTab)
 						->AddTab(FRigHierarchyTabSummoner::TabID, ETabState::OpenedTab)
-						->AddTab(FRigVMExecutionStackTabSummoner::TabID, ETabState::OpenedTab)
 						->AddTab(FRigCurveContainerTabSummoner::TabID, ETabState::OpenedTab)
 						->AddTab(FBlueprintEditorTabs::MyBlueprintID, ETabState::OpenedTab)
 					)
 				)
 				->Split
-				(
-					// Middle 
-					FTabManager::NewSplitter()
-					->SetOrientation(Orient_Vertical)
-					->SetSizeCoefficient(0.6f)
-					->Split
-					(
-						// Middle top - document edit area
-						FTabManager::NewStack()
-						->SetSizeCoefficient(0.8f)
-						->AddTab("Document", ETabState::ClosedTab)
-					)
-					->Split
-					(
-						// Middle bottom - compiler results & find
-						FTabManager::NewStack()
-						->SetSizeCoefficient(0.2f)
-						->AddTab(FBlueprintEditorTabs::CompilerResultsID, ETabState::ClosedTab)
-						->AddTab(FBlueprintEditorTabs::FindResultsID, ETabState::ClosedTab)
-					)
-				)
-				->Split
 				(
 					// Right side
 					FTabManager::NewSplitter()
@@ -116,23 +198,4 @@ FControlRigEditorMode::FControlRigEditorMode(const TSharedRef<FControlRigEditor>
 				)
 			)
 		);
-
-	if (UToolMenu* Toolbar = InControlRigEditor->RegisterModeToolbarIfUnregistered(GetModeName()))
-	{
-		InControlRigEditor->GetToolbarBuilder()->AddCompileToolbar(Toolbar);
-		InControlRigEditor->GetToolbarBuilder()->AddScriptingToolbar(Toolbar);
-		InControlRigEditor->GetToolbarBuilder()->AddBlueprintGlobalOptionsToolbar(Toolbar);
-	}
-}
-
-void FControlRigEditorMode::RegisterTabFactories(TSharedPtr<FTabManager> InTabManager)
-{
-	TSharedPtr<FBlueprintEditor> BP = MyBlueprintEditor.Pin();
-	
-	BP->RegisterToolbarTab(InTabManager.ToSharedRef());
-
-	// Mode-specific setup
-	BP->PushTabFactories(CoreTabFactories);
-	BP->PushTabFactories(BlueprintEditorTabFactories);
-	BP->PushTabFactories(TabFactories);
-}
+}
\ No newline at end of file
diff --git a/Engine/Plugins/Animation/ControlRig/Source/ControlRigEditor/Private/Editor/ControlRigEditorMode.h b/Engine/Plugins/Animation/ControlRig/Source/ControlRigEditor/Private/Editor/ControlRigEditorMode.h
--- a/Engine/Plugins/Animation/ControlRig/Source/ControlRigEditor/Private/Editor/ControlRigEditorMode.h
+++ b/Engine/Plugins/Animation/ControlRig/Source/ControlRigEditor/Private/Editor/ControlRigEditorMode.h
@@ -13,7 +13,7 @@ class UControlRigBlueprint;
 class FControlRigEditorMode : public FBlueprintEditorApplicationMode
 {
 public:
-	FControlRigEditorMode(const TSharedRef<FControlRigEditor>& InControlRigEditor);
+	FControlRigEditorMode(const TSharedRef<FControlRigEditor>& InControlRigEditor, bool bCreateDefaultLayout = true);
 
 	// FApplicationMode interface
 	virtual void RegisterTabFactories(TSharedPtr<FTabManager> InTabManager) override;
@@ -25,3 +25,13 @@ protected:
 private:
 	TWeakObjectPtr<UControlRigBlueprint> ControlRigBlueprintPtr;
 };
+
+class FModularRigEditorMode : public FControlRigEditorMode
+{
+public:
+	FModularRigEditorMode(const TSharedRef<FControlRigEditor>& InControlRigEditor);
+
+	// for now just don't open up the previous edited documents
+	virtual void PostActivateMode() override {}
+};
+
diff --git a/Engine/Plugins/Runtime/RigVM/Source/RigVMEditor/Private/Editor/RigVMEditor.cpp b/Engine/Plugins/Runtime/RigVM/Source/RigVMEditor/Private/Editor/RigVMEditor.cpp
--- a/Engine/Plugins/Runtime/RigVM/Source/RigVMEditor/Private/Editor/RigVMEditor.cpp
+++ b/Engine/Plugins/Runtime/RigVM/Source/RigVMEditor/Private/Editor/RigVMEditor.cpp
@@ -211,7 +211,7 @@ void FRigVMEditor::InitRigVMEditor(const EToolkitMode::Type Mode, const TSharedP
 	FString ActiveTabNodePath;
 	TArray<FString> OpenedTabNodePaths;
 
-	if (Blueprints.Num() > 0)
+	if (ShouldOpenGraphByDefault() && (Blueprints.Num() > 0))
 	{
 		bool bBroughtGraphToFront = false;
 		for(UEdGraph* Graph : Blueprints[0]->UbergraphPages)
@@ -301,9 +301,12 @@ void FRigVMEditor::InitRigVMEditor(const EToolkitMode::Type Mode, const TSharedP
 		}
 	}
 
-	if (UEdGraph* ActiveGraph = InRigVMBlueprint->GetEdGraph(ActiveTabNodePath))
+	if(ShouldOpenGraphByDefault())
 	{
-		OpenGraphAndBringToFront(ActiveGraph, true);
+		if (UEdGraph* ActiveGraph = InRigVMBlueprint->GetEdGraph(ActiveTabNodePath))
+		{
+			OpenGraphAndBringToFront(ActiveGraph, true);
+		}
 	}
 
 	FRigVMBlueprintUtils::HandleRefreshAllNodes(InRigVMBlueprint);
@@ -1508,6 +1511,50 @@ bool FRigVMEditor::IsSectionVisible(NodeSectionID::Type InSectionID) const
 	return false;
 }
 
+bool FRigVMEditor::AreEventGraphsAllowed() const
+{
+	if(const URigVMBlueprint* RigVMBlueprint = GetRigVMBlueprint())
+	{
+		return RigVMBlueprint->SupportsEventGraphs();
+	}
+	return FBlueprintEditor::AreEventGraphsAllowed();
+}
+
+bool FRigVMEditor::AreMacrosAllowed() const
+{
+	if(const URigVMBlueprint* RigVMBlueprint = GetRigVMBlueprint())
+	{
+		return RigVMBlueprint->SupportsMacros();
+	}
+	return FBlueprintEditor::AreMacrosAllowed();
+}
+
+bool FRigVMEditor::AreDelegatesAllowed() const
+{
+	if(const URigVMBlueprint* RigVMBlueprint = GetRigVMBlueprint())
+	{
+		return RigVMBlueprint->SupportsDelegates();
+	}
+	return FBlueprintEditor::AreDelegatesAllowed();
+}
+
+bool FRigVMEditor::NewDocument_IsVisibleForType(ECreatedDocumentType GraphType) const
+{
+	switch(GraphType)
+	{
+		case ECreatedDocumentType::CGT_NewMacroGraph:
+		case ECreatedDocumentType::CGT_NewAnimationLayer:
+		{
+			return false;
+		}
+		default:
+		{
+			break;
+		}
+	}
+	return FBlueprintEditor::NewDocument_IsVisibleForType(GraphType);
+}
+
 FGraphAppearanceInfo FRigVMEditor::GetGraphAppearance(UEdGraph* InGraph) const
 {
 	FGraphAppearanceInfo AppearanceInfo = FBlueprintEditor::GetGraphAppearance(InGraph);
diff --git a/Engine/Plugins/Runtime/RigVM/Source/RigVMEditor/Public/Editor/RigVMEditor.h b/Engine/Plugins/Runtime/RigVM/Source/RigVMEditor/Public/Editor/RigVMEditor.h
--- a/Engine/Plugins/Runtime/RigVM/Source/RigVMEditor/Public/Editor/RigVMEditor.h
+++ b/Engine/Plugins/Runtime/RigVM/Source/RigVMEditor/Public/Editor/RigVMEditor.h
@@ -87,6 +87,7 @@ public:
 	virtual FReply OnSpawnGraphNodeByShortcut(FInputChord InChord, const FVector2D& InPosition, UEdGraph* InGraph) override;
 	virtual bool ShouldLoadBPLibrariesFromAssetRegistry() override { return false; }
 	virtual void JumpToHyperlink(const UObject* ObjectReference, bool bRequestRename = false) override;
+	virtual bool ShouldOpenGraphByDefault() const { return true; }
 
 	// FEditorUndoClient Interface
 	virtual void PostUndo(bool bSuccess) override;
@@ -150,6 +151,10 @@ protected:
 	virtual void CreateDefaultTabContents(const TArray<UBlueprint*>& InBlueprints) override;
 	virtual void NewDocument_OnClicked(ECreatedDocumentType GraphType) override;
 	virtual bool IsSectionVisible(NodeSectionID::Type InSectionID) const override;
+	virtual bool AreEventGraphsAllowed() const override;
+	virtual bool AreMacrosAllowed() const override;
+	virtual bool AreDelegatesAllowed() const override;
+	virtual bool NewDocument_IsVisibleForType(ECreatedDocumentType GraphType) const;
 	virtual FGraphAppearanceInfo GetGraphAppearance(class UEdGraph* InGraph) const override;
 	virtual bool IsEditable(UEdGraph* InGraph) const override;
 	virtual bool IsCompilingEnabled() const override;
