Control Rig: FIx not able to create a new control rig
#rnx

diff --git a/Engine/Plugins/Animation/ControlRig/Source/ControlRigEditor/Private/ControlRigBlueprintFactory.cpp b/Engine/Plugins/Animation/ControlRig/Source/ControlRigEditor/Private/ControlRigBlueprintFactory.cpp
--- a/Engine/Plugins/Animation/ControlRig/Source/ControlRigEditor/Private/ControlRigBlueprintFactory.cpp
+++ b/Engine/Plugins/Animation/ControlRig/Source/ControlRigEditor/Private/ControlRigBlueprintFactory.cpp
@@ -137,11 +137,6 @@ private:
 			// If it appears on the allowed child-of classes list (or there is nothing on that list)
 			if (InClass)
 			{
-				if (InClass == UControlRig::StaticClass())
-				{
-					return false;
-				}
-				
 				if (InFilterFuncs->IfInChildOfClassesSet(AllowedChildrenOfClasses, InClass) == EFilterReturn::Failed)
 				{
 					return false;
