Control Rig: Improve metadata set/get performance

#rnx

diff --git a/Engine/Plugins/Animation/ControlRig/Source/ControlRig/Private/Rigs/RigHierarchy.cpp b/Engine/Plugins/Animation/ControlRig/Source/ControlRig/Private/Rigs/RigHierarchy.cpp
--- a/Engine/Plugins/Animation/ControlRig/Source/ControlRig/Private/Rigs/RigHierarchy.cpp
+++ b/Engine/Plugins/Animation/ControlRig/Source/ControlRig/Private/Rigs/RigHierarchy.cpp
@@ -210,7 +210,20 @@ void URigHierarchy::Save(FArchive& Ar)
 	Ar << PreviousNameMap;
 	Ar << PreviousParentMap;
 
-	Ar << ElementMetadata;
+	{
+		TMap<FRigElementKey, FMetadataStorage> ElementMetadataToSave;
+		
+		for (const FRigBaseElement* Element: Elements)
+		{
+			if (ElementMetadata.IsValidIndex(Element->MetadataStorageIndex))
+			{
+				ElementMetadataToSave.Add(Element->Key, ElementMetadata[Element->MetadataStorageIndex]);
+			}
+		}
+		
+		Ar << ElementMetadataToSave;
+	}
+	
 }
 
 void URigHierarchy::Load(FArchive& Ar)
@@ -320,7 +333,17 @@ void URigHierarchy::Load(FArchive& Ar)
 
 	if (Ar.CustomVer(FControlRigObjectVersion::GUID) >= FControlRigObjectVersion::RigHierarchyStoresElementMetadata)
 	{
-		Ar << ElementMetadata;
+		ElementMetadata.Reset();
+		ElementMetadataFreeList.Reset();
+		TMap<FRigElementKey, FMetadataStorage> LoadedElementMetadata;
+		
+		Ar << LoadedElementMetadata;
+		for (TPair<FRigElementKey, FMetadataStorage>& Entry: LoadedElementMetadata)
+		{
+			FRigBaseElement* Element = Find(Entry.Key);
+			Element->MetadataStorageIndex = ElementMetadata.Num();
+			ElementMetadata.Add(MoveTemp(Entry.Value));
+		}
 	}
 }
 
@@ -376,14 +399,15 @@ void URigHierarchy::Reset_Impl(bool bResetElements)
 		}
 		IndexLookup.Reset();
 
-		for (TTuple<FRigElementKey, FMetaDataStorage>& Metadata: ElementMetadata)
+		for (FMetadataStorage& MetadataStorage: ElementMetadata)
 		{
-			for (TTuple<FName, FRigBaseMetadata*>& Item: Metadata.Value.MetadataMap)
+			for (TTuple<FName, FRigBaseMetadata*>& Item: MetadataStorage.MetadataMap)
 			{
 				FRigBaseMetadata::DestroyMetadata(&Item.Value);
 			} 
 		}
 		ElementMetadata.Reset();
+		ElementMetadataFreeList.Reset();
 	}
 
 	ResetPoseHash = INDEX_NONE;
@@ -583,7 +607,7 @@ void URigHierarchy::CopyHierarchy(URigHierarchy* InHierarchy)
 
 		Target->CopyFrom(Source);
 
-		CopyAllMetadataFromElement(Target->GetKey(), Source);
+		CopyAllMetadataFromElement(Target, Source);
 	}
 
 	PreviousNameMap.Append(InHierarchy->PreviousNameMap);
@@ -1029,33 +1053,40 @@ TArray<FRigElementKey> URigHierarchy::RestoreConnectorsFromInfos(TArray<FRigConn
 TArray<FName> URigHierarchy::GetMetadataNames(FRigElementKey InItem) const
 {
 	TArray<FName> Names;
-	if(const FMetaDataStorage* Storage = ElementMetadata.Find(InItem))
+	if (const FRigBaseElement* Element = Find(InItem))
 	{
-		Storage->MetadataMap.GetKeys(Names);
+		if (Element->MetadataStorageIndex != INDEX_NONE)
+		{
+			ElementMetadata[Element->MetadataStorageIndex].MetadataMap.GetKeys(Names);
+		}
 	}
 	return Names;
 }
 
 ERigMetadataType URigHierarchy::GetMetadataType(FRigElementKey InItem, FName InMetadataName) const
 {
-	if(const FMetaDataStorage* Storage = ElementMetadata.Find(InItem))
+	if (const FRigBaseElement* Element = Find(InItem))
 	{
-		if (const FRigBaseMetadata* const* MetadataPtrPtr = Storage->MetadataMap.Find(InMetadataName))
+		if (Element->MetadataStorageIndex != INDEX_NONE)
 		{
-			return (*MetadataPtrPtr)->GetType();
+			if (const FRigBaseMetadata* const* MetadataPtrPtr = ElementMetadata[Element->MetadataStorageIndex].MetadataMap.Find(InMetadataName))
+			{
+				return (*MetadataPtrPtr)->GetType();
+			}
 		}
 	}
+	
 	return ERigMetadataType::Invalid;
 }
 
 bool URigHierarchy::RemoveMetadata(FRigElementKey InItem, FName InMetadataName)
 {
-	return RemoveMetadataForElement(InItem, InMetadataName);
+	return RemoveMetadataForElement(Find(InItem), InMetadataName);
 }
 
 bool URigHierarchy::RemoveAllMetadata(FRigElementKey InItem)
 {
-	return RemoveAllMetadataForElement(InItem);
+	return RemoveAllMetadataForElement(Find(InItem));
 }
 
 TArray<const FRigBaseElement*> URigHierarchy::GetSelectedElements(ERigElementType InTypeFilter) const
@@ -5248,16 +5279,18 @@ void URigHierarchy::CleanupInvalidCaches()
 	EnsureCacheValidity();
 }
 
-void URigHierarchy::FMetaDataStorage::Reset()
+void URigHierarchy::FMetadataStorage::Reset()
 {
 	for (TTuple<FName, FRigBaseMetadata*>& Item: MetadataMap)
 	{
 		FRigBaseMetadata::DestroyMetadata(&Item.Value);
 	}
 	MetadataMap.Reset();
+	LastAccessName = NAME_None;
+	LastAccessMetadata = nullptr;
 }
 
-void URigHierarchy::FMetaDataStorage::Serialize(FArchive& Ar)
+void URigHierarchy::FMetadataStorage::Serialize(FArchive& Ar)
 {
 	static const UEnum* MetadataTypeEnum = StaticEnum<ERigMetadataType>();
 	
@@ -5316,160 +5349,185 @@ void URigHierarchy::OnMetadataTagChanged(const FRigElementKey& InKey, const FNam
 	}
 }
 
-FRigBaseMetadata* URigHierarchy::GetMetadataForElement(const FRigElementKey& InKey, const FName& InName, ERigMetadataType InType, bool bInNotify)
+FRigBaseMetadata* URigHierarchy::GetMetadataForElement(FRigBaseElement* InElement, const FName& InName, ERigMetadataType InType, bool bInNotify)
 {
-	FMetaDataStorage* Storage = ElementMetadata.Find(InKey);
-	if (!Storage)
+	if (!ElementMetadata.IsValidIndex(InElement->MetadataStorageIndex))
 	{
-		Storage = &ElementMetadata.Add(InKey);
+		// Do we have entries in the freelist we can recycle?
+		if (!ElementMetadataFreeList.IsEmpty())
+		{
+			constexpr bool bAllowShrinking = false;
+			InElement->MetadataStorageIndex = ElementMetadataFreeList.Pop(bAllowShrinking);
+		}
+		else
+		{
+			InElement->MetadataStorageIndex = ElementMetadata.Num();
+			ElementMetadata.AddDefaulted();
+		}
 	}
 	
-	if (FRigBaseMetadata** MetadataPtrPtr = Storage->MetadataMap.Find(InName))
+	FMetadataStorage& Storage = ElementMetadata[InElement->MetadataStorageIndex];
+
+	// If repeatedly accessing the same element, store it here for faster access to avoid map lookups.
+	if (Storage.LastAccessName == InName && Storage.LastAccessMetadata->GetType() == InType)
+	{
+		return Storage.LastAccessMetadata;
+	}
+
+	FRigBaseMetadata* Metadata;
+	if (FRigBaseMetadata** MetadataPtrPtr = Storage.MetadataMap.Find(InName))
 	{
 		if ((*MetadataPtrPtr)->GetType() == InType)
 		{
-			return (*MetadataPtrPtr);
+			Metadata = *MetadataPtrPtr;
 		}
-		
-		// The type changed, replace the existing metadata with a new one of the correct type.
-		FRigBaseMetadata::DestroyMetadata(MetadataPtrPtr);
-		*MetadataPtrPtr = FRigBaseMetadata::MakeMetadata(InName, InType);
-		if (bInNotify)
+		else
 		{
-			OnMetadataChanged(InKey, InName);
-		}
+			// The type changed, replace the existing metadata with a new one of the correct type.
+			FRigBaseMetadata::DestroyMetadata(MetadataPtrPtr);
+			Metadata = *MetadataPtrPtr = FRigBaseMetadata::MakeMetadata(InName, InType);
 			
-		return *MetadataPtrPtr;
+			if (bInNotify)
+			{
+				OnMetadataChanged(InElement->Key, InName);
+			}
+		}
 	}
-	
-	// No metadata with that name existed on the element, create one from scratch.
-	FRigBaseMetadata* Metadata = FRigBaseMetadata::MakeMetadata(InName, InType);
-	Storage->MetadataMap.Add(InName, Metadata);
-	if (bInNotify)
+	else
 	{
-		OnMetadataChanged(InKey, InName);
+		// No metadata with that name existed on the element, create one from scratch.
+		Metadata = FRigBaseMetadata::MakeMetadata(InName, InType);
+		Storage.MetadataMap.Add(InName, Metadata);
 	}
+	
+	Storage.LastAccessName = InName;
+	Storage.LastAccessMetadata = Metadata;
 	return Metadata;
 }
 
 
-FRigBaseMetadata* URigHierarchy::FindMetadataForElement(const FRigElementKey& InKey, const FName& InName, ERigMetadataType InType)
+FRigBaseMetadata* URigHierarchy::FindMetadataForElement(const FRigBaseElement* InElement, const FName& InName, ERigMetadataType InType)
 {
-	FMetaDataStorage* Storage = ElementMetadata.Find(InKey);
-	if (!Storage)
+	if (!ElementMetadata.IsValidIndex(InElement->MetadataStorageIndex))
 	{
 		return nullptr;
 	}
+
+	FMetadataStorage& Storage = ElementMetadata[InElement->MetadataStorageIndex];
+	if (InName == Storage.LastAccessName && (InType == ERigMetadataType::Invalid || Storage.LastAccessMetadata->GetType() == InType))
+	{
+		return Storage.LastAccessMetadata;
+	}
 	
-	FRigBaseMetadata** MetadataPtrPtr = Storage->MetadataMap.Find(InName);
+	FRigBaseMetadata** MetadataPtrPtr = Storage.MetadataMap.Find(InName);
 	if (!MetadataPtrPtr)
 	{
+		Storage.LastAccessName = NAME_None;
+		Storage.LastAccessMetadata = nullptr;
 		return nullptr;
 	}
 
 	if (InType != ERigMetadataType::Invalid && (*MetadataPtrPtr)->GetType() != InType)
 	{
+		Storage.LastAccessName = NAME_None;
+		Storage.LastAccessMetadata = nullptr;
 		return nullptr;
 	}
 
+	Storage.LastAccessName = InName;
+	Storage.LastAccessMetadata = *MetadataPtrPtr;
+
 	return *MetadataPtrPtr;
 }
 
-const FRigBaseMetadata* URigHierarchy::FindMetadataForElement(const FRigElementKey& InKey, const FName& InName, ERigMetadataType InType) const
+const FRigBaseMetadata* URigHierarchy::FindMetadataForElement(const FRigBaseElement* InElement, const FName& InName, ERigMetadataType InType) const
 {
-	return const_cast<URigHierarchy*>(this)->FindMetadataForElement(InKey, InName, InType);
+	return const_cast<URigHierarchy*>(this)->FindMetadataForElement(InElement, InName, InType);
 }
 
 
-bool URigHierarchy::RemoveMetadataForElement(const FRigElementKey& InKey, const FName& InName)
+bool URigHierarchy::RemoveMetadataForElement(FRigBaseElement* InElement, const FName& InName)
 {
-	FMetaDataStorage* Storage = ElementMetadata.Find(InKey);
-	if (!Storage)
+	if (!ElementMetadata.IsValidIndex(InElement->MetadataStorageIndex))
 	{
 		return false;
 	}
-	
-	FRigBaseMetadata** MetadataPtrPtr = Storage->MetadataMap.Find(InName);
+
+	FMetadataStorage& Storage = ElementMetadata[InElement->MetadataStorageIndex];
+	FRigBaseMetadata** MetadataPtrPtr = Storage.MetadataMap.Find(InName);
 	if (!MetadataPtrPtr)
 	{
 		return false;
 	}
 	
 	FRigBaseMetadata::DestroyMetadata(MetadataPtrPtr);
-	Storage->MetadataMap.Remove(InName);
-	
+	Storage.MetadataMap.Remove(InName);
+
 	// If the storage is now empty, remove the element's storage, so we're not lugging it around
-	// unnecessarily. 
-	if (Storage->MetadataMap.IsEmpty())
+	// unnecessarily. Add the storage slot to the freelist so that the next element to add a new
+	// metadata storage can just recycle that.
+	if (Storage.MetadataMap.IsEmpty())
+	{
+		ElementMetadataFreeList.Add(InElement->MetadataStorageIndex);
+		InElement->MetadataStorageIndex = INDEX_NONE;
+	}
+	else if (Storage.LastAccessName == InName)
 	{
-		ElementMetadata.Remove(InKey);
+		Storage.LastAccessMetadata = nullptr;
 	}
 	
-	OnMetadataChanged(InKey, InName);
+	OnMetadataChanged(InElement->Key, InName);
 	return true;
 }
 
 
-bool URigHierarchy::RemoveAllMetadataForElement(const FRigElementKey& InKey)
+bool URigHierarchy::RemoveAllMetadataForElement(FRigBaseElement* InElement)
 {
-	FMetaDataStorage* Storage = ElementMetadata.Find(InKey);
-	if (!Storage || Storage->MetadataMap.IsEmpty())
+	if (!ElementMetadata.IsValidIndex(InElement->MetadataStorageIndex))
 	{
 		return false;
 	}
 
-	for (TTuple<FName, FRigBaseMetadata*>& Item: Storage->MetadataMap)
-	{
-		FRigBaseMetadata::DestroyMetadata(&Item.Value);
-	} 
 	
-	for (TTuple<FName, FRigBaseMetadata*>& Item: Storage->MetadataMap)
+	FMetadataStorage& Storage = ElementMetadata[InElement->MetadataStorageIndex];
+	TArray<FName> Names;
+	Storage.MetadataMap.GetKeys(Names);
+	
+	// Clear the storage for the next user.
+	Storage.Reset();
+	
+	ElementMetadataFreeList.Push(InElement->MetadataStorageIndex);
+	InElement->MetadataStorageIndex = INDEX_NONE;
+	
+	for (FName Name: Names)
 	{
-		OnMetadataChanged(InKey, Item.Key);
+		OnMetadataChanged(InElement->Key, Name);
 	}
 	
-	ElementMetadata.Remove(InKey);
-	
 	return true; 
 }
 
 
-void URigHierarchy::CopyAllMetadataFromElement(const FRigElementKey& InTargetKey, const FRigBaseElement* InSourceElement)
+void URigHierarchy::CopyAllMetadataFromElement(FRigBaseElement* InTargetElement, const FRigBaseElement* InSourceElement)
 {
 	if (!ensure(InSourceElement->Owner))
 	{
 		return;
 	}
-	
-	const FMetaDataStorage* SourceStorage = InSourceElement->Owner->ElementMetadata.Find(InSourceElement->GetKey());
-	if (!SourceStorage)
-	{
-		return;
-	}
 
-	FMetaDataStorage* TargetStorage = ElementMetadata.Find(InTargetKey);
-	if (!TargetStorage)
+	if (!InSourceElement->Owner->ElementMetadata.IsValidIndex(InSourceElement->MetadataStorageIndex))
 	{
-		TargetStorage = &ElementMetadata.Add(InTargetKey);
+		return;
 	}
 	
-	for (const TTuple<FName, FRigBaseMetadata*>& SourceItem: SourceStorage->MetadataMap)
+	const FMetadataStorage& SourceStorage = InSourceElement->Owner->ElementMetadata[InSourceElement->MetadataStorageIndex];
+	
+	for (const TTuple<FName, FRigBaseMetadata*>& SourceItem: SourceStorage.MetadataMap)
 	{
-			
 		const FRigBaseMetadata* SourceMetadata = SourceItem.Value; 	
-		// Does the destination item already exist?
-		FRigBaseMetadata** TargetMetadataPtrPtr = TargetStorage->MetadataMap.Find(SourceItem.Key);
-		FRigBaseMetadata* TargetMetadata;
-		if (TargetMetadataPtrPtr)
-		{
-			FRigBaseMetadata::DestroyMetadata(TargetMetadataPtrPtr);
-			TargetMetadata = *TargetMetadataPtrPtr = FRigBaseMetadata::MakeMetadata(SourceItem.Key, SourceMetadata->GetType());
-		}
-		else
-		{
-			TargetMetadata = FRigBaseMetadata::MakeMetadata(SourceItem.Key, SourceMetadata->GetType());
-			TargetStorage->MetadataMap.Add(SourceItem.Key, TargetMetadata);
-		}
+
+		constexpr bool bNotify = false;
+		FRigBaseMetadata* TargetMetadata = GetMetadataForElement(InTargetElement, SourceItem.Key, SourceMetadata->GetType(), bNotify);
 		TargetMetadata->SetValueData(SourceMetadata->GetValueData(), SourceMetadata->GetValueSize());
 	}
 }
@@ -6916,3 +6974,4 @@ FRigHierarchyRedirectorGuard::FRigHierarchyRedirectorGuard(UControlRig* InContro
 : Guard(InControlRig->GetHierarchy()->ElementKeyRedirector, &InControlRig->GetElementKeyRedirector())
 {
 }
+;
\ No newline at end of file
diff --git a/Engine/Plugins/Animation/ControlRig/Source/ControlRig/Private/Rigs/RigHierarchyController.cpp b/Engine/Plugins/Animation/ControlRig/Source/ControlRig/Private/Rigs/RigHierarchyController.cpp
--- a/Engine/Plugins/Animation/ControlRig/Source/ControlRig/Private/Rigs/RigHierarchyController.cpp
+++ b/Engine/Plugins/Animation/ControlRig/Source/ControlRig/Private/Rigs/RigHierarchyController.cpp
@@ -2332,14 +2332,6 @@ bool URigHierarchyController::RenameElement(FRigBaseElement* InElement, const FN
 		}
 	}
 
-	// Rename metadata
-	URigHierarchy::FMetaDataStorage MetadataStorage;
-	if (Hierarchy->ElementMetadata.RemoveAndCopyValue(OldKey, MetadataStorage))
-	{
-		Hierarchy->ElementMetadata.Add(NewKey, MoveTemp(MetadataStorage));
-	}
-	
-	
 	Hierarchy->PreviousNameMap.FindOrAdd(NewKey) = OldKey;
 	Hierarchy->IncrementTopologyVersion();
 	Notify(ERigHierarchyNotification::ElementRenamed, InElement);
diff --git a/Engine/Plugins/Animation/ControlRig/Source/ControlRig/Private/Rigs/RigHierarchyElements.cpp b/Engine/Plugins/Animation/ControlRig/Source/ControlRig/Private/Rigs/RigHierarchyElements.cpp
--- a/Engine/Plugins/Animation/ControlRig/Source/ControlRig/Private/Rigs/RigHierarchyElements.cpp
+++ b/Engine/Plugins/Animation/ControlRig/Source/ControlRig/Private/Rigs/RigHierarchyElements.cpp
@@ -19,7 +19,7 @@ FRigBaseElement::~FRigBaseElement()
 {
 	if (Owner)
 	{
-		Owner->RemoveAllMetadataForElement(Key);
+		Owner->RemoveAllMetadataForElement(this);
 	}
 }
 
@@ -122,7 +122,7 @@ void FRigBaseElement::Load(FArchive& Ar, ESerializationPhase SerializationPhase)
 
 				const ERigMetadataType MetadataType = static_cast<ERigMetadataType>(MetadataTypeEnum->GetValueByName(MetadataTypeName));
 
-				FRigBaseMetadata* Md = Owner->GetMetadataForElement(Key, MetadataName, MetadataType, false);
+				FRigBaseMetadata* Md = Owner->GetMetadataForElement(this, MetadataName, MetadataType, false);
 				Md->Serialize(Ar);
 			}
 		}
@@ -136,7 +136,7 @@ FRigBaseMetadata* FRigBaseElement::GetMetadata(const FName& InName, ERigMetadata
 	{
 		return nullptr;
 	}
-	return Owner->FindMetadataForElement(Key, InName, InType);
+	return Owner->FindMetadataForElement(this, InName, InType);
 }
 
 
@@ -146,7 +146,7 @@ const FRigBaseMetadata* FRigBaseElement::GetMetadata(const FName& InName, ERigMe
 	{
 		return nullptr;
 	}
-	return Owner->FindMetadataForElement(Key, InName, InType);
+	return Owner->FindMetadataForElement(this, InName, InType);
 }
 
 
@@ -155,7 +155,7 @@ bool FRigBaseElement::SetMetadata(const FName& InName, ERigMetadataType InType,
 	if (Owner)
 	{
 		constexpr bool bNotify = true;
-		if (FRigBaseMetadata* Metadata = Owner->GetMetadataForElement(Key, InName, InType, bNotify))
+		if (FRigBaseMetadata* Metadata = Owner->GetMetadataForElement(this, InName, InType, bNotify))
 		{
 			Metadata->SetValueData(InData, InSize);
 			return true;
@@ -171,7 +171,7 @@ FRigBaseMetadata* FRigBaseElement::SetupValidMetadata(const FName& InName, ERigM
 		return nullptr;
 	}
 	constexpr bool bNotify = true;
-	return Owner->GetMetadataForElement(Key, InName, InType, bNotify);
+	return Owner->GetMetadataForElement(this, InName, InType, bNotify);
 }
 
 
@@ -181,7 +181,7 @@ bool FRigBaseElement::RemoveMetadata(const FName& InName)
 	{
 		return false;
 	}
-	return Owner->RemoveMetadataForElement(Key, InName);
+	return Owner->RemoveMetadataForElement(this, InName);
 }
 
 bool FRigBaseElement::RemoveAllMetadata()
@@ -190,7 +190,7 @@ bool FRigBaseElement::RemoveAllMetadata()
 	{
 		return false;
 	}
-	return Owner->RemoveAllMetadataForElement(Key);
+	return Owner->RemoveAllMetadataForElement(this);
 }
 
 void FRigBaseElement::NotifyMetadataTagChanged(const FName& InTag, bool bAdded)
diff --git a/Engine/Plugins/Animation/ControlRig/Source/ControlRig/Private/Units/Hierarchy/RigUnit_Metadata.cpp b/Engine/Plugins/Animation/ControlRig/Source/ControlRig/Private/Units/Hierarchy/RigUnit_Metadata.cpp
--- a/Engine/Plugins/Animation/ControlRig/Source/ControlRig/Private/Units/Hierarchy/RigUnit_Metadata.cpp
+++ b/Engine/Plugins/Animation/ControlRig/Source/ControlRig/Private/Units/Hierarchy/RigUnit_Metadata.cpp
@@ -194,12 +194,7 @@ FRigBaseMetadata* FRigDispatch_GetMetadata::FindMetadata(const FRigVMExtendedExe
 		{
 			// first try to find the metadata in the namespace
 			const FName Name = ExecuteContext.AdaptMetadataName(bUseNameSpace, InName);
-			if(FRigBaseMetadata* Metadata = Element->GetMetadata(Name, InType))
-			{
-				return Metadata;
-			}
-			
-			return Element->GetMetadata(InName, InType);
+			return ExecuteContext.Hierarchy->FindMetadataForElement(Element, Name, InType);
 		};
 	}
 	return nullptr;
@@ -315,7 +310,7 @@ const TArray<FRigVMExecuteArgument>& FRigDispatch_SetMetadata::GetExecuteArgumen
 	return ExecuteArguments;
 }
 
-FRigBaseMetadata* FRigDispatch_SetMetadata::FindOrAddMetadata(FControlRigExecuteContext& InContext,
+FRigBaseMetadata* FRigDispatch_SetMetadata::FindOrAddMetadata(const FControlRigExecuteContext& InContext,
                                                               const FRigElementKey& InKey, const FName& InName, ERigMetadataType InType,
                                                               bool bUseNameSpace, FCachedRigElement& Cache)
 {
@@ -324,7 +319,8 @@ FRigBaseMetadata* FRigDispatch_SetMetadata::FindOrAddMetadata(FControlRigExecute
 		if(FRigBaseElement* Element = InContext.Hierarchy->Get(Cache.GetIndex()))
 		{
 			const FName Name = InContext.AdaptMetadataName(bUseNameSpace, InName);
-			return Element->SetupValidMetadata(Name, InType);
+			constexpr bool bNotify = true;
+			return InContext.Hierarchy->GetMetadataForElement(Element, Name, InType, bNotify);
 		}
 	}
 	return nullptr;
diff --git a/Engine/Plugins/Animation/ControlRig/Source/ControlRig/Private/Units/Hierarchy/RigUnit_Metadata.h b/Engine/Plugins/Animation/ControlRig/Source/ControlRig/Private/Units/Hierarchy/RigUnit_Metadata.h
--- a/Engine/Plugins/Animation/ControlRig/Source/ControlRig/Private/Units/Hierarchy/RigUnit_Metadata.h
+++ b/Engine/Plugins/Animation/ControlRig/Source/ControlRig/Private/Units/Hierarchy/RigUnit_Metadata.h
@@ -93,13 +93,13 @@ protected:
 #endif
 
 		// unpack the memory
-		const FRigElementKey& Item = *(const FRigElementKey*)Handles[Factory->ItemArgIndex].GetData();
-		const FName& Name = *(const FName*)Handles[Factory->NameArgIndex].GetData();
-		const bool bUseNameSpace = *(const bool*)Handles[Factory->UseNameSpaceArgIndex].GetData();
-		FCachedRigElement& Cache = *(FCachedRigElement*)Handles[Factory->CacheArgIndex].GetData(false, InContext.GetSlice().GetIndex());
-		const ValueType& Default = *(const ValueType*)Handles[Factory->DefaultArgIndex].GetData();
-		ValueType& Value = *(ValueType*)Handles[Factory->ValueArgIndex].GetData();
-		bool& Found = *(bool*)Handles[Factory->FoundArgIndex].GetData();
+		const FRigElementKey& Item = *reinterpret_cast<const FRigElementKey*>(Handles[Factory->ItemArgIndex].GetData());
+		const FName& Name = *reinterpret_cast<const FName*>(Handles[Factory->NameArgIndex].GetData());
+		const bool bUseNameSpace = *reinterpret_cast<const bool*>(Handles[Factory->UseNameSpaceArgIndex].GetData());
+		FCachedRigElement& Cache = *reinterpret_cast<FCachedRigElement*>(Handles[Factory->CacheArgIndex].GetData(false, InContext.GetSlice().GetIndex()));
+		const ValueType& Default = *reinterpret_cast<const ValueType*>(Handles[Factory->DefaultArgIndex].GetData());
+		ValueType& Value = *reinterpret_cast<ValueType*>(Handles[Factory->ValueArgIndex].GetData());
+		bool& Found = *reinterpret_cast<bool*>(Handles[Factory->FoundArgIndex].GetData());
 
 		// extract the metadata
 		if (const MetadataType* Md = Cast<MetadataType>(FindMetadata(InContext, Item, Name, EnumValue, bUseNameSpace, Cache)))
@@ -129,7 +129,7 @@ struct CONTROLRIG_API FRigDispatch_SetMetadata : public FRigDispatch_MetadataBas
 
 protected:
 
-	static FRigBaseMetadata* FindOrAddMetadata(FControlRigExecuteContext& InContext, const FRigElementKey& InKey, const FName& InName, ERigMetadataType InType, bool bUseNameSpace, FCachedRigElement& Cache);
+	static FRigBaseMetadata* FindOrAddMetadata(const FControlRigExecuteContext& InContext, const FRigElementKey& InKey, const FName& InName, ERigMetadataType InType, bool bUseNameSpace, FCachedRigElement& Cache);
 	virtual FRigVMFunctionPtr GetDispatchFunctionImpl(const FRigVMTemplateTypeMap& InTypes) const override;
 
 #if WITH_EDITOR
diff --git a/Engine/Plugins/Animation/ControlRig/Source/ControlRig/Public/Rigs/RigHierarchy.h b/Engine/Plugins/Animation/ControlRig/Source/ControlRig/Public/Rigs/RigHierarchy.h
--- a/Engine/Plugins/Animation/ControlRig/Source/ControlRig/Public/Rigs/RigHierarchy.h
+++ b/Engine/Plugins/Animation/ControlRig/Source/ControlRig/Public/Rigs/RigHierarchy.h
@@ -4109,14 +4109,16 @@ private:
 	mutable TArray<TArray<FRigBaseElement*>> ElementsPerType;
 
 	//
-	struct FMetaDataStorage
+	struct FMetadataStorage
 	{
 		TMap<FName, FRigBaseMetadata*> MetadataMap;
+		FName LastAccessName;
+		FRigBaseMetadata* LastAccessMetadata = nullptr;
 
 		void Reset();
 		void Serialize(FArchive& Ar);
 
-		friend FArchive& operator<<(FArchive& Ar, FMetaDataStorage& Storage)
+		friend FArchive& operator<<(FArchive& Ar, FMetadataStorage& Storage)
 		{
 			Storage.Serialize(Ar);
 			return Ar;
@@ -4132,7 +4134,10 @@ private:
 	// to reduce memory consumption. Only elements created by MakeElement point to
 	// the element storage. Copied elements via the copy constructor or copy operator
 	// do not have URigHierarchy as an owner and therefore do not carry metadata with them.
-	TMap<FRigElementKey, FMetaDataStorage> ElementMetadata;
+	TArray<FMetadataStorage> ElementMetadata;
+
+	// List of metadata storage entries that have been freed and can be recycled.
+	TArray<int32> ElementMetadataFreeList;
 
 	// A quick-lookup cache for elements' children. Each element that has a ChildCacheIndex
 	// not equal to INDEX_NONE is an index into the offset and count cache below, which in
@@ -4524,7 +4529,7 @@ private:
 	{
 		if(InElement)
 		{
-			if(const FRigBaseMetadata* Metadata = FindMetadataForElement(InElement->GetKey(), InMetadataName, InType))
+			if(const FRigBaseMetadata* Metadata = FindMetadataForElement(InElement, InMetadataName, InType))
 			{
 				return *static_cast<const T*>(Metadata->GetValueData());
 			}
@@ -4557,7 +4562,7 @@ private:
 		if(InElement)
 		{
 			constexpr bool bNotify = true;
-			if (FRigBaseMetadata* Metadata = GetMetadataForElement(InElement->GetKey(), InMetadataName, InType, bNotify))
+			if (FRigBaseMetadata* Metadata = GetMetadataForElement(InElement, InMetadataName, InType, bNotify))
 			{
 				return Metadata->SetValueData(&InValue, sizeof(T));
 			}
@@ -4584,21 +4589,21 @@ private:
 	    is created for that element. If the name matches but the type differs, the existing metadata is destroyed and a new one with the
 	    matching type is created instead.
 	    */
-	FRigBaseMetadata* GetMetadataForElement(const FRigElementKey& InKey, const FName& InName, ERigMetadataType InType, bool bInNotify);
+	FRigBaseMetadata* GetMetadataForElement(FRigBaseElement* InElement, const FName& InName, ERigMetadataType InType, bool bInNotify);
 	
 	/** Attempts to find element's metadata of the given name and type. If either the element doesnt exist, the name doesn't exist or the
 	    type doesn't match, then \c nullptr is returned.
 	    */ 
-	FRigBaseMetadata* FindMetadataForElement(const FRigElementKey& InKey, const FName& InName, ERigMetadataType InType);
-	const FRigBaseMetadata* FindMetadataForElement(const FRigElementKey& InKey, const FName& InName, ERigMetadataType InType) const;
+	FRigBaseMetadata* FindMetadataForElement(const FRigBaseElement* InElement, const FName& InName, ERigMetadataType InType);
+	const FRigBaseMetadata* FindMetadataForElement(const FRigBaseElement* InElement, const FName& InName, ERigMetadataType InType) const;
 	
 	/** Removes the named meta data for the given element, regardless of type. If the element doesn't exist, or it doesn't have any
 	    metadata of the given name, this function does nothing and returns \c false.
 		*/
-	bool RemoveMetadataForElement(const FRigElementKey& InKey, const FName& InName);
-	bool RemoveAllMetadataForElement(const FRigElementKey& InKey);
+	bool RemoveMetadataForElement(FRigBaseElement* InElement, const FName& InName);
+	bool RemoveAllMetadataForElement(FRigBaseElement* InElement);
 	
-	void CopyAllMetadataFromElement(const FRigElementKey& InTargetKey, const FRigBaseElement* InSourceElement);
+	void CopyAllMetadataFromElement(FRigBaseElement* InTargetElement, const FRigBaseElement* InSourceElement);
 	
 protected:
 	bool bEnableCacheValidityCheck;
@@ -4687,6 +4692,8 @@ private:
 	friend struct FRigHierarchyEnableControllerBracket;
 	friend struct FRigHierarchyExecuteContextBracket;
 	friend struct FRigHierarchyRedirectorGuard;
+	friend struct FRigDispatch_GetMetadata;
+	friend struct FRigDispatch_SetMetadata;
 };
 
 struct CONTROLRIG_API FRigHierarchyInteractionBracket
diff --git a/Engine/Plugins/Animation/ControlRig/Source/ControlRig/Public/Rigs/RigHierarchyElements.h b/Engine/Plugins/Animation/ControlRig/Source/ControlRig/Public/Rigs/RigHierarchyElements.h
--- a/Engine/Plugins/Animation/ControlRig/Source/ControlRig/Public/Rigs/RigHierarchyElements.h
+++ b/Engine/Plugins/Animation/ControlRig/Source/ControlRig/Public/Rigs/RigHierarchyElements.h
@@ -644,6 +644,9 @@ protected:
 
 	// Index into the child cache offset and count table in URigHierarchy. Set by URigHierarchy::UpdateCachedChildren
 	int32 ChildCacheIndex = INDEX_NONE;
+
+	// Index into the metadata storage for this element.
+	int32 MetadataStorageIndex = INDEX_NONE;
 	
 	mutable FString CachedNameString;
 
