Control Rig: Control rig track order set by python/bp and set default to 100. Also fixes issue where priorities got broken due to using the bp/python function to add the Control rig, we had an unneeded Player->Stop call there that would make everything become unbound which would remove all Proxy::ControlrigNodes

diff --git a/Engine/Plugins/Animation/ControlRig/Source/ControlRig/Private/Sequencer/ControlRigLayerInstanceProxy.cpp b/Engine/Plugins/Animation/ControlRig/Source/ControlRig/Private/Sequencer/ControlRigLayerInstanceProxy.cpp
--- a/Engine/Plugins/Animation/ControlRig/Source/ControlRig/Private/Sequencer/ControlRigLayerInstanceProxy.cpp
+++ b/Engine/Plugins/Animation/ControlRig/Source/ControlRig/Private/Sequencer/ControlRigLayerInstanceProxy.cpp
@@ -176,7 +176,7 @@ void FControlRigLayerInstanceProxy::SetSourceAnimInstance(UAnimInstance* SourceA
 void FControlRigLayerInstanceProxy::AddControlRigTrack(int32 ControlRigID, UControlRig* InControlRig)
 {
 	FAnimNode_ControlRig_ExternalSource* Node = FindControlRigNode(ControlRigID);
-
+	const int32 DefaultPriorityOrder = 100;
 	if(!Node)
 	{
 		UMovieSceneControlRigParameterTrack* Track = InControlRig->GetTypedOuter<UMovieSceneControlRigParameterTrack>();
@@ -197,7 +197,7 @@ void FControlRigLayerInstanceProxy::AddControlRigTrack(int32 ControlRigID, UCont
 						}
 						else
 						{
-							Track->SetPriorityOrder(ControlRigNodes.Num() -1);
+							Track->SetPriorityOrder(DefaultPriorityOrder + ControlRigNodes.Num() -1);
 						}
 					}
 					Node = ControlRigNodes.Add_GetRef(MakeShared<FAnimNode_ControlRig_ExternalSource>()).Get();
@@ -221,7 +221,7 @@ void FControlRigLayerInstanceProxy::AddControlRigTrack(int32 ControlRigID, UCont
 							}
 							else
 							{
-								Track->SetPriorityOrder(Index);
+								Track->SetPriorityOrder(DefaultPriorityOrder + Index);
 							}
 						}
 						Node = ControlRigNodes.Insert_GetRef(MakeShared<FAnimNode_ControlRig_ExternalSource>(), Index).Get();
@@ -290,7 +290,7 @@ void FControlRigLayerInstanceProxy::AddControlRigTrack(int32 ControlRigID, UCont
 				}
 			}
 		}
-		else
+		else //no nodes add first one
 		{
 			Node = ControlRigNodes.Add_GetRef(MakeShared<FAnimNode_ControlRig_ExternalSource>()).Get();
 			if (Track)
@@ -298,7 +298,7 @@ void FControlRigLayerInstanceProxy::AddControlRigTrack(int32 ControlRigID, UCont
 				int32 PriorityOrder = Track->GetPriorityOrder();
 				if (PriorityOrder == INDEX_NONE) //track has no order so just add to end, will happen on creation
 				{
-					Track->SetPriorityOrder(ControlRigNodes.Num() - 1);
+					Track->SetPriorityOrder(DefaultPriorityOrder);
 				}
 			}
 		}
diff --git a/Engine/Plugins/Animation/ControlRig/Source/ControlRigEditor/Private/ControlRigSequencerEditorLibrary.cpp b/Engine/Plugins/Animation/ControlRig/Source/ControlRigEditor/Private/ControlRigSequencerEditorLibrary.cpp
--- a/Engine/Plugins/Animation/ControlRig/Source/ControlRigEditor/Private/ControlRigSequencerEditorLibrary.cpp
+++ b/Engine/Plugins/Animation/ControlRig/Source/ControlRigEditor/Private/ControlRigSequencerEditorLibrary.cpp
@@ -351,11 +351,8 @@ UMovieSceneTrack* UControlRigSequencerEditorLibrary::FindOrCreateControlRigTrack
 				}
 			}
 
-			if (OutPlayer)
-			{
-				OutPlayer->Stop();
-			}
-
+			//no need to stop player was never running, and stop will cause things to get unbound, messing 
+			//with how priorities are set
 			if (OutActor)
 			{
 				World->DestroyActor(OutActor);
@@ -416,11 +413,6 @@ TArray<UMovieSceneTrack*> UControlRigSequencerEditorLibrary::FindOrCreateControl
 		}
 	}
 
-	if (OutPlayer)
-	{
-		OutPlayer->Stop();
-	}
-
 	if (OutActor)
 	{
 		World->DestroyActor(OutActor);
@@ -529,12 +521,15 @@ UTickableConstraint* UControlRigSequencerEditorLibrary::AddConstraint(UWorld* Wo
 		Constraint->MarkAsGarbage();
 		return nullptr;
 	}
-	
-	//add key
-	const TWeakPtr<ISequencer> WeakSequencer = GetSequencerFromAsset();
-	if (WeakSequencer.IsValid())
+
+	if (ULevelSequence* LevelSequence = ULevelSequenceEditorBlueprintLibrary::GetCurrentLevelSequence())
 	{
-		FMovieSceneConstraintChannelHelper::SmartConstraintKey(WeakSequencer.Pin(), Constraint, TOptional<bool>(), TOptional<FFrameNumber>());
+		//add key
+		const TWeakPtr<ISequencer> WeakSequencer = GetSequencerFromAsset();
+		if (WeakSequencer.IsValid())
+		{
+			FMovieSceneConstraintChannelHelper::SmartConstraintKey(WeakSequencer.Pin(), Constraint, TOptional<bool>(), TOptional<FFrameNumber>());
+		}
 	}
 	else
 	{
@@ -2976,6 +2971,30 @@ FRigElementKey UControlRigSequencerEditorLibrary::GetWorldSpaceReferenceKey()
 	return URigHierarchy::GetWorldSpaceReferenceKey();
 }
 
+void UControlRigSequencerEditorLibrary::SetControlRigPriorityOrder(UMovieSceneTrack* InTrack, int32 PriorityOrder)
+{
+	UMovieSceneControlRigParameterTrack* ParameterTrack = Cast<UMovieSceneControlRigParameterTrack>(InTrack);
+	if (!ParameterTrack)
+	{
+		FFrame::KismetExecutionMessage(TEXT("Cannot call SetControlRigPriorityOrder without a UMovieSceneControlRigParameterTrack"), ELogVerbosity::Error);
+		return;
+	}
+
+	ParameterTrack->SetPriorityOrder(PriorityOrder);
+}
+
+int32 UControlRigSequencerEditorLibrary::GetControlRigPriorityOrder(UMovieSceneTrack* InTrack)
+{
+	UMovieSceneControlRigParameterTrack* ParameterTrack = Cast<UMovieSceneControlRigParameterTrack>(InTrack);
+	if (!ParameterTrack)
+	{
+		FFrame::KismetExecutionMessage(TEXT("Cannot call SetControlRigPriorityOrder without a UMovieSceneControlRigParameterTrack"), ELogVerbosity::Error);
+		return INDEX_NONE;
+	}
+	return 	ParameterTrack->GetPriorityOrder();
+
+}
+
 bool UControlRigSequencerEditorLibrary::GetControlsMask(UMovieSceneSection* InSection, FName ControlName)
 {
 	UMovieSceneControlRigParameterSection* ParameterSection = Cast<UMovieSceneControlRigParameterSection>(InSection);
diff --git a/Engine/Plugins/Animation/ControlRig/Source/ControlRigEditor/Public/ControlRigSequencerEditorLibrary.h b/Engine/Plugins/Animation/ControlRig/Source/ControlRigEditor/Public/ControlRigSequencerEditorLibrary.h
--- a/Engine/Plugins/Animation/ControlRig/Source/ControlRigEditor/Public/ControlRigSequencerEditorLibrary.h
+++ b/Engine/Plugins/Animation/ControlRig/Source/ControlRigEditor/Public/ControlRigSequencerEditorLibrary.h
@@ -1043,6 +1043,14 @@ public:
 	UFUNCTION(BlueprintCallable, Category = "Editor Scripting | Sequencer Tools | Control Rig")
 	static void HideAllControls(UMovieSceneSection* InSection);
 
+	/** Set Control Rig priority order */
+	UFUNCTION(BlueprintCallable, Category = "Editor Scripting | Sequencer Tools | Control Rig")
+	static void SetControlRigPriorityOrder(UMovieSceneTrack* InSection,int32 PriorityOrder);
+
+	/** Get Control Rig prirority order */
+	UFUNCTION(BlueprintCallable, Category = "Editor Scripting | Sequencer Tools | Control Rig")
+	static int32 GetControlRigPriorityOrder(UMovieSceneTrack* InSection);
+
 	/**	Whether or not the control rig is an FK Control Rig.
 	@param InControlRig Rig to test to see if FK Control Rig
 	**/
@@ -1069,4 +1077,5 @@ public:
 	**/
 	UFUNCTION(BlueprintCallable, Category = "Editor Scripting | Sequencer Tools | Control Rig")
 	static bool SetControlRigApplyMode(UControlRig* InControlRig, EControlRigFKRigExecuteMode InApplyMode);
+
 };
