Sequencer: Animation Mode: When reopening level sequence restore tween/tool states, recreate control rig as needed, and restore selection.

diff --git a/Engine/Plugins/Animation/ControlRig/Source/ControlRigEditor/Private/EditMode/ControlRigEditModeToolkit.cpp b/Engine/Plugins/Animation/ControlRig/Source/ControlRigEditor/Private/EditMode/ControlRigEditModeToolkit.cpp
--- a/Engine/Plugins/Animation/ControlRig/Source/ControlRigEditor/Private/EditMode/ControlRigEditModeToolkit.cpp
+++ b/Engine/Plugins/Animation/ControlRig/Source/ControlRigEditor/Private/EditMode/ControlRigEditModeToolkit.cpp
@@ -33,6 +33,11 @@ namespace
 	const TArray<FName> AnimationPaletteNames = { AnimationName };
 }
 
+bool FControlRigEditModeToolkit::bMotionTrailsTabOpen = false;
+bool FControlRigEditModeToolkit::bPoseTabOpen = false;
+bool FControlRigEditModeToolkit::bSnapperTabOpen = false;
+bool FControlRigEditModeToolkit::bTweenOpen = false;
+
 const FName FControlRigEditModeToolkit::PoseTabName = FName(TEXT("PoseTab"));
 const FName FControlRigEditModeToolkit::MotionTrailTabName = FName(TEXT("MotionTrailTab"));
 const FName FControlRigEditModeToolkit::SnapperTabName = FName(TEXT("SnapperTab"));
@@ -97,17 +102,17 @@ void FControlRigEditModeToolkit::TryInvokeToolkitUI(const FName InName)
 	if (InName == MotionTrailTabName)
 	{
 		FTabId TabID(MotionTrailTabName);
-		ModeUILayerPtr->GetTabManager()->TryInvokeTab(TabID);
+		ModeUILayerPtr->GetTabManager()->TryInvokeTab(TabID, false /*bIsActive*/);
 	}
 	else if (InName == PoseTabName)
 	{
 		FTabId TabID(PoseTabName);
-		ModeUILayerPtr->GetTabManager()->TryInvokeTab(TabID);
+		ModeUILayerPtr->GetTabManager()->TryInvokeTab(TabID, false /*bIsActive*/);
 	}
 	else if (InName == SnapperTabName)
 	{
 		FTabId TabID(SnapperTabName);
-		ModeUILayerPtr->GetTabManager()->TryInvokeTab(TabID);
+		ModeUILayerPtr->GetTabManager()->TryInvokeTab(TabID, false /*bIsActive*/);
 	}
 	else if (InName == OutlinerTabName)
 	{
@@ -410,6 +415,22 @@ void FControlRigEditModeToolkit::InvokeUI()
 		ModeUILayerPtr->GetTabManager()->TryInvokeTab(UAssetEditorUISubsystem::TopRightTabID);
 		// doesn't work as expected todo ModeUILayerPtr->GetTabManager()->TryInvokeTab(UAssetEditorUISubsystem::TopLeftTabID);
 		ModeUILayerPtr->GetTabManager()->TryInvokeTab(UAssetEditorUISubsystem::BottomRightTabID);
+		if (bTweenOpen)
+		{
+			CreateAndShowTweenOverlay();
+		}
+		if (bMotionTrailsTabOpen)
+		{
+			TryInvokeToolkitUI(MotionTrailTabName);
+		}
+		if (bSnapperTabOpen)
+		{
+			TryInvokeToolkitUI(SnapperTabName);
+		}
+		if (bPoseTabOpen)
+		{
+			TryInvokeToolkitUI(PoseTabName);
+		}
 	}	
 }
 
@@ -417,6 +438,14 @@ void FControlRigEditModeToolkit::UnregisterAndRemoveFloatingTabs()
 {
 	if (FSlateApplication::IsInitialized())
 	{
+		if (TweenWidgetParent)
+		{
+			bTweenOpen = true;
+		}
+		else
+		{
+			bTweenOpen = false;
+		}
 		RemoveAndDestroyTweenOverlay();
 		if (ModeUILayer.IsValid())
 		{
@@ -425,22 +454,37 @@ void FControlRigEditModeToolkit::UnregisterAndRemoveFloatingTabs()
 			TSharedPtr<SDockTab> MotionTrailTab = ModeUILayerPtr->GetTabManager()->FindExistingLiveTab(FTabId(MotionTrailTabName));
 			if (MotionTrailTab)
 			{
+				bMotionTrailsTabOpen = true;
 				MotionTrailTab->RequestCloseTab();
 			}
+			else
+			{
+				bMotionTrailsTabOpen = false;
+			}
 			ModeUILayerPtr->GetTabManager()->UnregisterTabSpawner(MotionTrailTabName);
 
 			TSharedPtr<SDockTab> SnapperTab = ModeUILayerPtr->GetTabManager()->FindExistingLiveTab(FTabId(SnapperTabName));
 			if (SnapperTab)
 			{
+				bSnapperTabOpen = true;
 				SnapperTab->RequestCloseTab();
 			}
+			else
+			{
+				bSnapperTabOpen = false;
+			}
 			ModeUILayerPtr->GetTabManager()->UnregisterTabSpawner(SnapperTabName);
 		
 			TSharedPtr<SDockTab> PoseTab = ModeUILayerPtr->GetTabManager()->FindExistingLiveTab(FTabId(PoseTabName));
 			if (PoseTab)
 			{
+				bPoseTabOpen = true;
 				PoseTab->RequestCloseTab();
 			}
+			else
+			{
+				bPoseTabOpen = false;
+			}
 			ModeUILayerPtr->GetTabManager()->UnregisterTabSpawner(PoseTabName);
 		}
 	}
diff --git a/Engine/Plugins/Animation/ControlRig/Source/ControlRigEditor/Private/EditMode/ControlRigEditModeToolkit.h b/Engine/Plugins/Animation/ControlRig/Source/ControlRigEditor/Private/EditMode/ControlRigEditModeToolkit.h
--- a/Engine/Plugins/Animation/ControlRig/Source/ControlRigEditor/Private/EditMode/ControlRigEditModeToolkit.h
+++ b/Engine/Plugins/Animation/ControlRig/Source/ControlRigEditor/Private/EditMode/ControlRigEditModeToolkit.h
@@ -82,9 +82,13 @@ protected:
 	virtual void RequestModeUITabs() override;
 	virtual void InvokeUI() override;
 
+	//this also saves the layout
 	void UnregisterAndRemoveFloatingTabs();
 
-private:
+	static bool bMotionTrailsTabOpen;
+	static bool bPoseTabOpen;
+	static bool bSnapperTabOpen;
+	static bool bTweenOpen;
 
 
 private:
diff --git a/Engine/Plugins/Animation/ControlRig/Source/ControlRigEditor/Private/Sequencer/ControlRigParameterTrackEditor.cpp b/Engine/Plugins/Animation/ControlRig/Source/ControlRigEditor/Private/Sequencer/ControlRigParameterTrackEditor.cpp
--- a/Engine/Plugins/Animation/ControlRig/Source/ControlRigEditor/Private/Sequencer/ControlRigParameterTrackEditor.cpp
+++ b/Engine/Plugins/Animation/ControlRig/Source/ControlRigEditor/Private/Sequencer/ControlRigParameterTrackEditor.cpp
@@ -96,6 +96,9 @@
 
 #define LOCTEXT_NAMESPACE "FControlRigParameterTrackEditor"
 
+bool FControlRigParameterTrackEditor::bControlRigEditModeWasOpen = false;
+TArray <TPair<UClass*, TArray<FName>>> FControlRigParameterTrackEditor::PreviousSelectedControlRigs;
+
 TAutoConsoleVariable<bool> CVarAutoGenerateControlRigTrack(TEXT("ControlRig.Sequencer.AutoGenerateTrack"), true, TEXT("When true automatically create control rig tracks in Sequencer when a control rig is added to a level."));
 
 TAutoConsoleVariable<bool> CVarSelectedKeysSelectControls(TEXT("ControlRig.Sequencer.SelectedKeysSelectControls"), false, TEXT("When true when we select a key in Sequencer it will select the Control, by default false."));
@@ -539,6 +542,33 @@ void FControlRigParameterTrackEditor::OnRelease()
 	}
 	ConstraintHandlesToClear.Reset();
 
+	FControlRigEditMode* ControlRigEditMode = GetEditMode();
+	PreviousSelectedControlRigs.Reset();
+	if (ControlRigEditMode)
+	{
+		bControlRigEditModeWasOpen = true;
+		for (TWeakObjectPtr<UControlRig>& ControlRig : BoundControlRigs)
+		{
+			if (ControlRig.IsValid())
+			{
+				TPair<UClass*, TArray<FName>> ClassAndName;
+				ClassAndName.Key = ControlRig->GetClass();
+				ClassAndName.Value = ControlRig->CurrentControlSelection();
+				PreviousSelectedControlRigs.Add(ClassAndName);
+			}
+		}
+		ControlRigEditMode->Exit(); //deactive mode below doesn't exit for some reason so need to make sure things are cleaned up
+		if (FEditorModeTools* Tools = GetEditorModeTools())
+		{
+			Tools->DeactivateMode(FControlRigEditMode::ModeName);
+		}
+
+		ControlRigEditMode->SetObjects(nullptr, nullptr, GetSequencer());
+	}
+	else
+	{
+		bControlRigEditModeWasOpen = false;
+	}
 	UnbindAllControlRigs();
 	if (GetSequencer().IsValid())
 	{
@@ -576,18 +606,6 @@ void FControlRigParameterTrackEditor::OnRelease()
 			}
 		}
 	}
-	FControlRigEditMode* ControlRigEditMode = GetEditMode();
-	if (ControlRigEditMode)
-	{
-		ControlRigEditMode->Exit(); //deactive mode below doesn't exit for some reason so need to make sure things are cleaned up
-		if (FEditorModeTools* Tools = GetEditorModeTools())
-		{
-			Tools->DeactivateMode(FControlRigEditMode::ModeName);
-		}
-
-		ControlRigEditMode->SetObjects(nullptr, nullptr, GetSequencer());
-	}
-
 	AcquiredResources.Release();
 
 }
@@ -1904,6 +1922,40 @@ void FControlRigParameterTrackEditor::OnActivateSequenceChanged(FMovieSceneSeque
 		
 		return false;
 	});
+	if (bControlRigEditModeWasOpen)
+	{
+		GEditor->GetTimerManager()->SetTimerForNextTick([this]()
+		{
+			//true here will turn it on
+			FControlRigEditMode* ControlRigEditMode = GetEditMode(true);
+			if (ControlRigEditMode)
+			{
+				GEditor->GetTimerManager()->SetTimerForNextTick([this, ControlRigEditMode]()
+				{
+					for(TWeakObjectPtr<UControlRig> &ControlRig: BoundControlRigs)
+					{ 
+						if (ControlRig.IsValid())
+						{
+							ControlRigEditMode->AddControlRigObject(ControlRig.Get(), GetSequencer());
+							for (int32 Index = 0; Index < PreviousSelectedControlRigs.Num(); ++Index)
+							{
+								if (ControlRig.Get()->GetClass() == PreviousSelectedControlRigs[Index].Key)
+								{
+									for (const FName& ControlName : PreviousSelectedControlRigs[Index].Value)
+									{
+										ControlRig.Get()->SelectControl(ControlName, true);
+									}
+									PreviousSelectedControlRigs.RemoveAt(Index);
+									break;
+								}
+							}	
+						}
+					}
+					PreviousSelectedControlRigs.Reset();
+				});
+			}
+		});
+	}
 }
 
 
diff --git a/Engine/Plugins/Animation/ControlRig/Source/ControlRigEditor/Private/Sequencer/ControlRigParameterTrackEditor.h b/Engine/Plugins/Animation/ControlRig/Source/ControlRigEditor/Private/Sequencer/ControlRigParameterTrackEditor.h
--- a/Engine/Plugins/Animation/ControlRig/Source/ControlRigEditor/Private/Sequencer/ControlRigParameterTrackEditor.h
+++ b/Engine/Plugins/Animation/ControlRig/Source/ControlRigEditor/Private/Sequencer/ControlRigParameterTrackEditor.h
@@ -301,6 +301,11 @@ private:
 	TSet<FDelegateHandle> ConstraintHandlesToClear;
 
 	friend class FControlRigBlueprintActions;
+
+	/** Whether or not animation/control rig edit mode was open when we closed*/
+	static bool bControlRigEditModeWasOpen;
+	/** Previous selection if open*/
+	static TArray <TPair<UClass*, TArray<FName>>> PreviousSelectedControlRigs;
 };
 
 
