Control Rig: Rename python exposed functions referencing Additive control rigs to Layered control rigs
#rnx

diff --git a/Engine/Plugins/Animation/ControlRig/Source/ControlRigEditor/Private/ControlRigSequencerEditorLibrary.cpp b/Engine/Plugins/Animation/ControlRig/Source/ControlRigEditor/Private/ControlRigSequencerEditorLibrary.cpp
--- a/Engine/Plugins/Animation/ControlRig/Source/ControlRigEditor/Private/ControlRigSequencerEditorLibrary.cpp
+++ b/Engine/Plugins/Animation/ControlRig/Source/ControlRigEditor/Private/ControlRigSequencerEditorLibrary.cpp
@@ -309,7 +309,7 @@ static UMovieSceneControlRigParameterTrack* AddControlRig(ULevelSequence* LevelS
 	return nullptr;
 }
 
-UMovieSceneTrack* UControlRigSequencerEditorLibrary::FindOrCreateControlRigTrack(UWorld* World, ULevelSequence* LevelSequence, const UClass* ControlRigClass, const FMovieSceneBindingProxy& InBinding, bool bIsAdditiveControlRig)
+UMovieSceneTrack* UControlRigSequencerEditorLibrary::FindOrCreateControlRigTrack(UWorld* World, ULevelSequence* LevelSequence, const UClass* ControlRigClass, const FMovieSceneBindingProxy& InBinding, bool bIsLayeredControlRig)
 {
 	UMovieScene* MovieScene = InBinding.Sequence ? InBinding.Sequence->GetMovieScene() : nullptr;
 	UMovieSceneTrack* BaseTrack = nullptr;
@@ -342,7 +342,7 @@ UMovieSceneTrack* UControlRigSequencerEditorLibrary::FindOrCreateControlRigTrack
 				UControlRig* ControlRig = nullptr;
 				if (SkeletalMeshComponent && SkeletalMeshComponent->GetSkeletalMeshAsset() && SkeletalMeshComponent->GetSkeletalMeshAsset()->GetSkeleton())
 				{
-					UMovieSceneControlRigParameterTrack* Track = AddControlRig(LevelSequence, ControlRigClass, SkeletalMeshComponent, InBinding.BindingID, nullptr, bIsAdditiveControlRig);
+					UMovieSceneControlRigParameterTrack* Track = AddControlRig(LevelSequence, ControlRigClass, SkeletalMeshComponent, InBinding.BindingID, nullptr, bIsLayeredControlRig);
 
 					if (Track)
 					{
@@ -3089,7 +3089,7 @@ bool UControlRigSequencerEditorLibrary::IsFKControlRig(UControlRig* InControlRig
 	return (InControlRig && InControlRig->IsA<UFKControlRig>());
 }
 
-bool UControlRigSequencerEditorLibrary::IsAdditiveControlRig(UControlRig* InControlRig)
+bool UControlRigSequencerEditorLibrary::IsLayeredControlRig(UControlRig* InControlRig)
 {
 	return (InControlRig && InControlRig->IsAdditive());
 }
diff --git a/Engine/Plugins/Animation/ControlRig/Source/ControlRigEditor/Public/ControlRigSequencerEditorLibrary.h b/Engine/Plugins/Animation/ControlRig/Source/ControlRigEditor/Public/ControlRigSequencerEditorLibrary.h
--- a/Engine/Plugins/Animation/ControlRig/Source/ControlRigEditor/Public/ControlRigSequencerEditorLibrary.h
+++ b/Engine/Plugins/Animation/ControlRig/Source/ControlRigEditor/Public/ControlRigSequencerEditorLibrary.h
@@ -99,7 +99,7 @@ public:
 	* @return returns Return the found or created track
 	*/
 	UFUNCTION(BlueprintCallable, Category = "Editor Scripting | Sequencer Tools | Control Rig")
-	static UMovieSceneTrack* FindOrCreateControlRigTrack(UWorld* World, ULevelSequence* LevelSequence, const UClass* ControlRigClass, const FMovieSceneBindingProxy& InBinding, bool bIsAdditiveControlRig = false);
+	static UMovieSceneTrack* FindOrCreateControlRigTrack(UWorld* World, ULevelSequence* LevelSequence, const UClass* ControlRigClass, const FMovieSceneBindingProxy& InBinding, bool bIsLayeredControlRig = false);
 
 	/**
 	* Find or create a Control Rig Component
@@ -1057,11 +1057,11 @@ public:
 	UFUNCTION(BlueprintCallable, Category = "Editor Scripting | Sequencer Tools | Control Rig")
 	static bool IsFKControlRig(UControlRig* InControlRig);
 
-	/**	Whether or not the control rig is an Additive Control Rig.
-	@param InControlRig Rig to test to see if Additive Control Rig
+	/**	Whether or not the control rig is an Layered Control Rig.
+	@param InControlRig Rig to test to see if Layered Control Rig
 	**/
 	UFUNCTION(BlueprintCallable, Category = "Editor Scripting | Sequencer Tools | Control Rig")
-	static bool IsAdditiveControlRig(UControlRig* InControlRig);
+	static bool IsLayeredControlRig(UControlRig* InControlRig);
 
 	/**	Get FKControlRig Apply Mode.
 	@param InControlRig Rig to test
