Control RIg: Sort available classes offering FKControlRig as first option
#rnx

diff --git a/Engine/Plugins/Animation/ControlRig/Source/ControlRigEditor/Private/Sequencer/ControlRigParameterTrackEditor.cpp b/Engine/Plugins/Animation/ControlRig/Source/ControlRigEditor/Private/Sequencer/ControlRigParameterTrackEditor.cpp
--- a/Engine/Plugins/Animation/ControlRig/Source/ControlRigEditor/Private/Sequencer/ControlRigParameterTrackEditor.cpp
+++ b/Engine/Plugins/Animation/ControlRig/Source/ControlRigEditor/Private/Sequencer/ControlRigParameterTrackEditor.cpp
@@ -672,6 +672,19 @@ void FControlRigParameterTrackEditor::BuildObjectBindingContextMenu(FMenuBuilder
 	}
 }
 
+static bool ClassViewerSortPredicate(const FClassViewerSortElementInfo& A, const  FClassViewerSortElementInfo& B)
+{
+	if ((A.Class == UFKControlRig::StaticClass() && B.Class == UFKControlRig::StaticClass()) ||
+				(A.Class != UFKControlRig::StaticClass() && B.Class != UFKControlRig::StaticClass()))
+	{
+		return  (*A.DisplayName).Compare(*B.DisplayName, ESearchCase::IgnoreCase) < 0;
+	}
+	else
+	{
+		return A.Class == UFKControlRig::StaticClass();
+	}
+}
+
 void FControlRigParameterTrackEditor::BakeToControlRigSubMenu(FMenuBuilder& MenuBuilder, FGuid ObjectBinding, UObject* BoundObject, USkeletalMeshComponent* SkelMeshComp, USkeleton* Skeleton)
 {
 	const TSharedPtr<ISequencer> ParentSequencer = GetSequencer();
@@ -685,6 +698,7 @@ void FControlRigParameterTrackEditor::BakeToControlRigSubMenu(FMenuBuilder& Menu
 		Options.ClassFilters.Add(ClassFilter.ToSharedRef());
 		Options.bShowNoneOption = false;
 		Options.ExtraPickerCommonClasses.Add(UFKControlRig::StaticClass());
+		Options.ClassViewerSortPredicate = ClassViewerSortPredicate;
 
 		FClassViewerModule& ClassViewerModule = FModuleManager::LoadModuleChecked<FClassViewerModule>("ClassViewer");
 
@@ -1372,6 +1386,7 @@ void FControlRigParameterTrackEditor::HandleAddControlRigSubMenu(FMenuBuilder& M
 		Options.ClassFilters.Add(ClassFilter.ToSharedRef());
 		Options.bShowNoneOption = false;
 		Options.ExtraPickerCommonClasses.Add(UFKControlRig::StaticClass());
+		Options.ClassViewerSortPredicate = ClassViewerSortPredicate;
 
 		UMovieSceneSequence* Sequence = GetSequencer() ? GetSequencer()->GetFocusedMovieSceneSequence() : nullptr;
 		Options.AdditionalReferencingAssets.Add(FAssetData(Sequence));
diff --git a/Engine/Source/Editor/ClassViewer/Private/SClassViewer.cpp b/Engine/Source/Editor/ClassViewer/Private/SClassViewer.cpp
--- a/Engine/Source/Editor/ClassViewer/Private/SClassViewer.cpp
+++ b/Engine/Source/Editor/ClassViewer/Private/SClassViewer.cpp
@@ -2564,8 +2564,20 @@ void SClassViewer::Populate()
 		// Get the class list, passing in certain filter options.
 		ClassViewer::Helpers::GetClassList(RootTreeItems, ClassFilter, InitOptions);
 
-		// Sort the list alphabetically.
-		RootTreeItems.Sort(FClassViewerNodeNameLess(InitOptions.NameTypeToDisplay));
+		if (InitOptions.ClassViewerSortPredicate)
+		{
+			RootTreeItems.Sort([this](const TSharedPtr<FClassViewerNode>& A, const TSharedPtr<FClassViewerNode>& B)
+			{
+				FClassViewerSortElementInfo InfoA (A->Class, A->GetClassName(false), A->GetClassName(true));
+				FClassViewerSortElementInfo InfoB (B->Class, B->GetClassName(false), B->GetClassName(true));
+				return InitOptions.ClassViewerSortPredicate(InfoA, InfoB);
+			});
+		}
+		else
+		{
+			// Sort the list alphabetically.
+			RootTreeItems.Sort(FClassViewerNodeNameLess(InitOptions.NameTypeToDisplay));
+		}
 
 		// Only display this option if the user wants it and in Picker Mode.
 		if(InitOptions.bShowNoneOption && InitOptions.Mode == EClassViewerMode::ClassPicker)
diff --git a/Engine/Source/Editor/ClassViewer/Public/ClassViewerModule.h b/Engine/Source/Editor/ClassViewer/Public/ClassViewerModule.h
--- a/Engine/Source/Editor/ClassViewer/Public/ClassViewerModule.h
+++ b/Engine/Source/Editor/ClassViewer/Public/ClassViewerModule.h
@@ -63,6 +63,19 @@ enum class EClassViewerNameTypeToDisplay : uint8
 	ClassName,
 };
 
+struct FClassViewerSortElementInfo
+{
+	FClassViewerSortElementInfo(TWeakObjectPtr<UClass> InClass, TSharedPtr<FString> InName, TSharedPtr<FString> InDisplayName)
+		: Class(InClass), Name(InName), DisplayName(InDisplayName)
+	{}
+	
+	TWeakObjectPtr<UClass> Class;
+
+	TSharedPtr<FString> Name;
+
+	TSharedPtr<FString> DisplayName;
+};
+
 /**
  * Settings for the Class Viewer set by the programmer before spawning an instance of the widget.  This
  * is used to modify the class viewer's behavior in various ways, such as filtering in or out specific classes.
@@ -79,6 +92,9 @@ public:
 	/** The filter(s) to use on classes in this instance. */
 	TArray<TSharedRef<IClassViewerFilter>> ClassFilters;
 
+	/** Predicate used to sort the class list */
+	TFunction<bool(const FClassViewerSortElementInfo&, const FClassViewerSortElementInfo&)> ClassViewerSortPredicate;
+
 	/** Mode to operate in */
 	EClassViewerMode::Type Mode;
 
