Control Rig: Remove need for private includes

diff --git a/Engine/Plugins/Animation/ControlRig/Source/ControlRig/Private/Units/Deprecated/Math/RigUnit_Converter.cpp b/Engine/Plugins/Animation/ControlRig/Source/ControlRig/Private/Units/Deprecated/Math/RigUnit_Converter.cpp
--- a/Engine/Plugins/Animation/ControlRig/Source/ControlRig/Private/Units/Deprecated/Math/RigUnit_Converter.cpp
+++ b/Engine/Plugins/Animation/ControlRig/Source/ControlRig/Private/Units/Deprecated/Math/RigUnit_Converter.cpp
@@ -1,6 +1,6 @@
 // Copyright Epic Games, Inc. All Rights Reserved.
 
-#include "RigUnit_Converter.h"
+#include "Units/Deprecated/Math/RigUnit_Converter.h"
 #include "RigVMFunctions/Math/RigVMFunction_MathTransform.h"
 #include "RigVMFunctions/Math/RigVMFunction_MathQuaternion.h"
 
diff --git a/Engine/Plugins/Animation/ControlRig/Source/ControlRig/Private/Units/Deprecated/Math/RigUnit_Float.cpp b/Engine/Plugins/Animation/ControlRig/Source/ControlRig/Private/Units/Deprecated/Math/RigUnit_Float.cpp
--- a/Engine/Plugins/Animation/ControlRig/Source/ControlRig/Private/Units/Deprecated/Math/RigUnit_Float.cpp
+++ b/Engine/Plugins/Animation/ControlRig/Source/ControlRig/Private/Units/Deprecated/Math/RigUnit_Float.cpp
@@ -1,6 +1,6 @@
 // Copyright Epic Games, Inc. All Rights Reserved.
 
-#include "RigUnit_Float.h"
+#include "Units/Deprecated/Math/RigUnit_Float.h"
 #include "RigVMFunctions/Math/RigVMFunction_MathFloat.h"
 
 #include UE_INLINE_GENERATED_CPP_BY_NAME(RigUnit_Float)
diff --git a/Engine/Plugins/Animation/ControlRig/Source/ControlRig/Private/Units/Deprecated/Math/RigUnit_Quaternion.cpp b/Engine/Plugins/Animation/ControlRig/Source/ControlRig/Private/Units/Deprecated/Math/RigUnit_Quaternion.cpp
--- a/Engine/Plugins/Animation/ControlRig/Source/ControlRig/Private/Units/Deprecated/Math/RigUnit_Quaternion.cpp
+++ b/Engine/Plugins/Animation/ControlRig/Source/ControlRig/Private/Units/Deprecated/Math/RigUnit_Quaternion.cpp
@@ -1,6 +1,6 @@
 // Copyright Epic Games, Inc. All Rights Reserved.
 
-#include "RigUnit_Quaternion.h"
+#include "Units/Deprecated/Math/RigUnit_Quaternion.h"
 #include "RigVMFunctions/Math/RigVMFunction_MathQuaternion.h"
 
 #include UE_INLINE_GENERATED_CPP_BY_NAME(RigUnit_Quaternion)
diff --git a/Engine/Plugins/Animation/ControlRig/Source/ControlRig/Private/Units/Deprecated/Math/RigUnit_Transform.cpp b/Engine/Plugins/Animation/ControlRig/Source/ControlRig/Private/Units/Deprecated/Math/RigUnit_Transform.cpp
--- a/Engine/Plugins/Animation/ControlRig/Source/ControlRig/Private/Units/Deprecated/Math/RigUnit_Transform.cpp
+++ b/Engine/Plugins/Animation/ControlRig/Source/ControlRig/Private/Units/Deprecated/Math/RigUnit_Transform.cpp
@@ -1,6 +1,6 @@
 // Copyright Epic Games, Inc. All Rights Reserved.
 
-#include "RigUnit_Transform.h"
+#include "Units/Deprecated/Math/RigUnit_Transform.h"
 #include "RigVMFunctions/Math/RigVMFunction_MathTransform.h"
 
 #include UE_INLINE_GENERATED_CPP_BY_NAME(RigUnit_Transform)
diff --git a/Engine/Plugins/Animation/ControlRig/Source/ControlRig/Private/Units/Deprecated/Math/RigUnit_Vector.cpp b/Engine/Plugins/Animation/ControlRig/Source/ControlRig/Private/Units/Deprecated/Math/RigUnit_Vector.cpp
--- a/Engine/Plugins/Animation/ControlRig/Source/ControlRig/Private/Units/Deprecated/Math/RigUnit_Vector.cpp
+++ b/Engine/Plugins/Animation/ControlRig/Source/ControlRig/Private/Units/Deprecated/Math/RigUnit_Vector.cpp
@@ -1,6 +1,6 @@
 // Copyright Epic Games, Inc. All Rights Reserved.
 
-#include "RigUnit_Vector.h"
+#include "Units/Deprecated/Math/RigUnit_Vector.h"
 #include "RigVMFunctions/Math/RigVMFunction_MathVector.h"
 #include "Math/ControlRigMathLibrary.h"
 
diff --git a/Engine/Plugins/Animation/ControlRig/Source/ControlRig/Private/Units/Deprecated/RigUnit_AimConstraint.cpp b/Engine/Plugins/Animation/ControlRig/Source/ControlRig/Private/Units/Deprecated/RigUnit_AimConstraint.cpp
--- a/Engine/Plugins/Animation/ControlRig/Source/ControlRig/Private/Units/Deprecated/RigUnit_AimConstraint.cpp
+++ b/Engine/Plugins/Animation/ControlRig/Source/ControlRig/Private/Units/Deprecated/RigUnit_AimConstraint.cpp
@@ -1,6 +1,6 @@
 // Copyright Epic Games, Inc. All Rights Reserved.
 
-#include "RigUnit_AimConstraint.h"
+#include "Units/Deprecated/RigUnit_AimConstraint.h"
 #include "Units/RigUnitContext.h"
 #include "HelperUtil.h"
 #include "AnimationCoreLibrary.h"
diff --git a/Engine/Plugins/Animation/ControlRig/Source/ControlRig/Private/Units/Deprecated/RigUnit_ApplyFK.cpp b/Engine/Plugins/Animation/ControlRig/Source/ControlRig/Private/Units/Deprecated/RigUnit_ApplyFK.cpp
--- a/Engine/Plugins/Animation/ControlRig/Source/ControlRig/Private/Units/Deprecated/RigUnit_ApplyFK.cpp
+++ b/Engine/Plugins/Animation/ControlRig/Source/ControlRig/Private/Units/Deprecated/RigUnit_ApplyFK.cpp
@@ -1,6 +1,6 @@
 // Copyright Epic Games, Inc. All Rights Reserved.
 
-#include "RigUnit_ApplyFK.h"
+#include "Units/Deprecated/RigUnit_ApplyFK.h"
 #include "Units/RigUnitContext.h"
 #include "HelperUtil.h"
 
diff --git a/Engine/Plugins/Animation/ControlRig/Source/ControlRig/Private/Units/Deprecated/RigUnit_BlendTransform.cpp b/Engine/Plugins/Animation/ControlRig/Source/ControlRig/Private/Units/Deprecated/RigUnit_BlendTransform.cpp
--- a/Engine/Plugins/Animation/ControlRig/Source/ControlRig/Private/Units/Deprecated/RigUnit_BlendTransform.cpp
+++ b/Engine/Plugins/Animation/ControlRig/Source/ControlRig/Private/Units/Deprecated/RigUnit_BlendTransform.cpp
@@ -1,6 +1,6 @@
 // Copyright Epic Games, Inc. All Rights Reserved.
 
-#include "RigUnit_BlendTransform.h"
+#include "Units/Deprecated/RigUnit_BlendTransform.h"
 #include "Units/RigUnitContext.h"
 #include "AnimationRuntime.h"
 
diff --git a/Engine/Plugins/Animation/ControlRig/Source/ControlRig/Private/Units/Deprecated/RigUnit_GetJointTransform.cpp b/Engine/Plugins/Animation/ControlRig/Source/ControlRig/Private/Units/Deprecated/RigUnit_GetJointTransform.cpp
--- a/Engine/Plugins/Animation/ControlRig/Source/ControlRig/Private/Units/Deprecated/RigUnit_GetJointTransform.cpp
+++ b/Engine/Plugins/Animation/ControlRig/Source/ControlRig/Private/Units/Deprecated/RigUnit_GetJointTransform.cpp
@@ -1,6 +1,6 @@
 // Copyright Epic Games, Inc. All Rights Reserved.
 
-#include "RigUnit_GetJointTransform.h"
+#include "Units/Deprecated/RigUnit_GetJointTransform.h"
 #include "Units/RigUnitContext.h"
 #include "HelperUtil.h"
 #include "Units/Hierarchy/RigUnit_GetTransform.h"
diff --git a/Engine/Plugins/Animation/ControlRig/Source/ControlRig/Public/Units/Deprecated/Math/RigUnit_Converter.h b/Engine/Plugins/Animation/ControlRig/Source/ControlRig/Public/Units/Deprecated/Math/RigUnit_Converter.h
new file mode 100644
--- /dev/null
+++ b/Engine/Plugins/Animation/ControlRig/Source/ControlRig/Public/Units/Deprecated/Math/RigUnit_Converter.h
@@ -0,0 +1,189 @@
+// Copyright Epic Games, Inc. All Rights Reserved.
+
+#pragma once
+#include "CoreMinimal.h"
+#include "Units/RigUnit.h"
+#include "EulerTransform.h"
+#include "RigUnit_Converter.generated.h"
+
+USTRUCT(meta = (DisplayName = "ConvertToEulerTransform", Category = "Math|Convert", NodeColor = "0.1 0.1 0.7", Deprecated="4.23.0"))
+struct CONTROLRIG_API FRigUnit_ConvertTransform : public FRigUnit
+{
+	GENERATED_BODY()
+
+	UPROPERTY(meta=(Input))
+	FTransform Input;
+
+	UPROPERTY(meta=(Output))
+	FEulerTransform Result;
+
+	RIGVM_METHOD()
+	virtual void Execute() override;
+
+	RIGVM_METHOD()
+	virtual FRigVMStructUpgradeInfo GetUpgradeInfo() const override;
+};
+
+USTRUCT(meta = (DisplayName = "ConvertToTransform", Category = "Math|Convert", NodeColor = "0.1 0.1 0.7", Deprecated="4.23.0"))
+struct CONTROLRIG_API FRigUnit_ConvertEulerTransform : public FRigUnit
+{
+	GENERATED_BODY()
+
+	UPROPERTY(meta=(Input))
+	FEulerTransform Input;
+
+	UPROPERTY(meta=(Output))
+	FTransform Result;
+
+	RIGVM_METHOD()
+	virtual void Execute() override;
+
+	RIGVM_METHOD()
+	virtual FRigVMStructUpgradeInfo GetUpgradeInfo() const override;
+};
+
+USTRUCT(meta = (DisplayName = "ConvertToQuaternion", Category = "Math|Convert", NodeColor = "0.1 0.1 0.7", Deprecated="4.23.0"))
+struct CONTROLRIG_API FRigUnit_ConvertRotation : public FRigUnit
+{
+	GENERATED_BODY()
+
+	UPROPERTY(meta=(Input))
+	FRotator Input = FRotator(0.f);
+
+	UPROPERTY(meta = (Output))
+	FQuat	Result = FQuat::Identity;
+
+	RIGVM_METHOD()
+	virtual void Execute() override;
+
+	RIGVM_METHOD()
+	virtual FRigVMStructUpgradeInfo GetUpgradeInfo() const override;
+};
+
+USTRUCT(meta = (DisplayName = "ConvertToQuaternionDeprecated", Category = "Math|Convert", NodeColor = "0.1 0.1 0.7", Deprecated="4.23.0"))
+struct CONTROLRIG_API FRigUnit_ConvertVectorRotation : public FRigUnit_ConvertRotation
+{
+	GENERATED_BODY()
+};
+
+USTRUCT(meta = (DisplayName = "ConvertToRotation", Category = "Math|Convert", NodeColor = "0.1 0.1 0.7", Deprecated="4.23.0"))
+struct CONTROLRIG_API FRigUnit_ConvertQuaternion: public FRigUnit
+{
+	GENERATED_BODY()
+
+	UPROPERTY(meta=(Input))
+	FQuat	Input = FQuat::Identity;
+
+	UPROPERTY(meta = (Output))
+	FRotator	Result = FRotator(0.f);
+
+	RIGVM_METHOD()
+	virtual void Execute() override;
+
+	RIGVM_METHOD()
+	virtual FRigVMStructUpgradeInfo GetUpgradeInfo() const override;
+};
+
+USTRUCT(meta = (DisplayName = "ConvertVectorToRotation", Category = "Math|Convert", NodeColor = "0.1 0.1 0.7", Deprecated="4.23.0"))
+struct CONTROLRIG_API FRigUnit_ConvertVectorToRotation: public FRigUnit
+{
+	GENERATED_BODY()
+
+	UPROPERTY(meta = (Input))
+	FVector Input = FVector(0.f);
+
+	UPROPERTY(meta = (Output))
+	FRotator	Result = FRotator(0.f);
+
+	RIGVM_METHOD()
+	virtual void Execute() override;
+
+	RIGVM_METHOD()
+	virtual FRigVMStructUpgradeInfo GetUpgradeInfo() const override;
+};
+
+USTRUCT(meta = (DisplayName = "ConvertVectorToQuaternion", Category = "Math|Convert", NodeColor = "0.1 0.1 0.7", Deprecated="4.23.0"))
+struct CONTROLRIG_API FRigUnit_ConvertVectorToQuaternion: public FRigUnit
+{
+	GENERATED_BODY()
+
+	UPROPERTY(meta = (Input))
+	FVector Input = FVector(0.f);
+
+	UPROPERTY(meta = (Output))
+	FQuat	Result = FQuat::Identity;
+
+	RIGVM_METHOD()
+	virtual void Execute() override;
+
+	RIGVM_METHOD()
+	virtual FRigVMStructUpgradeInfo GetUpgradeInfo() const override;
+};
+
+USTRUCT(meta = (DisplayName = "ConvertRotationToVector", Category = "Math|Convert", NodeColor = "0.1 0.1 0.7", Deprecated="4.23.0"))
+struct CONTROLRIG_API FRigUnit_ConvertRotationToVector: public FRigUnit
+{
+	GENERATED_BODY()
+
+	UPROPERTY(meta = (Input))
+	FRotator Input = FRotator(0.f);
+
+	UPROPERTY(meta = (Output))
+	FVector Result = FVector(0.f);
+
+	RIGVM_METHOD()
+	virtual void Execute() override;
+
+	RIGVM_METHOD()
+	virtual FRigVMStructUpgradeInfo GetUpgradeInfo() const override;
+};
+
+USTRUCT(meta = (DisplayName = "ConvertQuaternionToVector", Category = "Math|Convert", NodeColor = "0.1 0.1 0.7", Deprecated="4.23.0"))
+struct CONTROLRIG_API FRigUnit_ConvertQuaternionToVector: public FRigUnit
+{
+	GENERATED_BODY()
+
+	UPROPERTY(meta = (Input))
+	FQuat	Input = FQuat::Identity;
+
+	UPROPERTY(meta = (Output))
+	FVector Result = FVector(0.f);
+
+	RIGVM_METHOD()
+	virtual void Execute() override;
+
+	RIGVM_METHOD()
+	virtual FRigVMStructUpgradeInfo GetUpgradeInfo() const override;
+};
+
+USTRUCT(meta = (DisplayName = "ToSwingAndTwist", Category = "Math|Transform", NodeColor = "0.1 0.1 0.7", Deprecated="4.23.0"))
+struct CONTROLRIG_API FRigUnit_ToSwingAndTwist : public FRigUnit
+{
+	GENERATED_BODY()
+
+	UPROPERTY(meta=(Input))
+	FQuat	Input;
+
+	UPROPERTY(meta = (Input))
+	FVector TwistAxis;
+
+	UPROPERTY(meta = (Output))
+	FQuat	Swing;
+
+	UPROPERTY(meta = (Output))
+	FQuat	Twist;
+
+	FRigUnit_ToSwingAndTwist()
+		: Input(ForceInitToZero)
+		, TwistAxis(FVector(1.f, 0.f, 0.f))
+		, Swing(ForceInitToZero)
+		, Twist(ForceInitToZero)
+	{
+	}
+
+	RIGVM_METHOD()
+	virtual void Execute() override;
+
+	RIGVM_METHOD()
+	virtual FRigVMStructUpgradeInfo GetUpgradeInfo() const override;
+};
diff --git a/Engine/Plugins/Animation/ControlRig/Source/ControlRig/Public/Units/Deprecated/Math/RigUnit_Float.h b/Engine/Plugins/Animation/ControlRig/Source/ControlRig/Public/Units/Deprecated/Math/RigUnit_Float.h
new file mode 100644
--- /dev/null
+++ b/Engine/Plugins/Animation/ControlRig/Source/ControlRig/Public/Units/Deprecated/Math/RigUnit_Float.h
@@ -0,0 +1,127 @@
+// Copyright Epic Games, Inc. All Rights Reserved.
+
+#pragma once
+
+#include "Units/RigUnit.h"
+#include "MathLibrary.h"
+#include "RigUnit_Float.generated.h"
+
+/** Two args and a result of float type */
+USTRUCT(meta=(Abstract, NodeColor = "0.1 0.7 0.1"))
+struct CONTROLRIG_API FRigUnit_BinaryFloatOp : public FRigUnit
+{
+	GENERATED_BODY()
+
+	UPROPERTY(meta=(Input))
+	float Argument0 = 0.f;
+
+	UPROPERTY(meta=(Input))
+	float Argument1 = 0.f;
+
+	UPROPERTY(meta=(Output))
+	float Result = 0.f;
+};
+
+USTRUCT(meta=(DisplayName="Multiply", Category="Math|Float", Keywords="*", Deprecated="4.23.0"))
+struct CONTROLRIG_API FRigUnit_Multiply_FloatFloat : public FRigUnit_BinaryFloatOp
+{
+	GENERATED_BODY()
+
+	RIGVM_METHOD()
+	virtual void Execute() override;
+
+	RIGVM_METHOD()
+	virtual FRigVMStructUpgradeInfo GetUpgradeInfo() const override;
+};
+
+USTRUCT(meta=(DisplayName="Add", Category="Math|Float", Keywords = "+,Sum", Deprecated="4.23.0"))
+struct CONTROLRIG_API FRigUnit_Add_FloatFloat : public FRigUnit_BinaryFloatOp
+{
+	GENERATED_BODY()
+
+	RIGVM_METHOD()
+	virtual void Execute() override;
+
+	RIGVM_METHOD()
+	virtual FRigVMStructUpgradeInfo GetUpgradeInfo() const override;
+};
+
+USTRUCT(meta=(DisplayName="Subtract", Category="Math|Float", Keywords = "-", Deprecated="4.23.0"))
+struct CONTROLRIG_API FRigUnit_Subtract_FloatFloat : public FRigUnit_BinaryFloatOp
+{
+	GENERATED_BODY()
+
+	RIGVM_METHOD()
+	virtual void Execute() override;
+
+	RIGVM_METHOD()
+	virtual FRigVMStructUpgradeInfo GetUpgradeInfo() const override;
+};
+
+USTRUCT(meta=(DisplayName="Divide", Category="Math|Float", Keywords = "/", Deprecated="4.23.0"))
+struct CONTROLRIG_API FRigUnit_Divide_FloatFloat : public FRigUnit_BinaryFloatOp
+{
+	GENERATED_BODY()
+
+	RIGVM_METHOD()
+	virtual void Execute() override;
+
+	RIGVM_METHOD()
+	virtual FRigVMStructUpgradeInfo GetUpgradeInfo() const override;
+};
+
+/** Two args and a result of float type */
+USTRUCT(meta = (DisplayName = "Clamp", Category = "Math|Float", NodeColor = "0.1 0.7 0.1", Deprecated="4.23.0"))
+struct CONTROLRIG_API FRigUnit_Clamp_Float: public FRigUnit
+{
+	GENERATED_BODY()
+
+	UPROPERTY(meta = (Input))
+	float Value = 0.f;
+
+	UPROPERTY(meta = (Input))
+	float Min = 0.f;
+
+	UPROPERTY(meta = (Input))
+	float Max = 0.f;
+
+	UPROPERTY(meta = (Output))
+	float Result = 0.f;
+
+	RIGVM_METHOD()
+	virtual void Execute() override;
+
+	RIGVM_METHOD()
+	virtual FRigVMStructUpgradeInfo GetUpgradeInfo() const override;
+};
+
+/** Two args and a result of float type */
+USTRUCT(meta = (DisplayName = "MapRange", Category = "Math|Float", Deprecated="4.23.0"))
+struct CONTROLRIG_API FRigUnit_MapRange_Float: public FRigUnit
+{
+	GENERATED_BODY()
+
+	UPROPERTY(meta = (Input))
+	float Value = 0.f;
+
+	UPROPERTY(meta = (Input))
+	float MinIn = 0.f;
+
+	UPROPERTY(meta = (Input))
+	float MaxIn = 0.f;
+
+	UPROPERTY(meta = (Input))
+	float MinOut = 0.f;
+
+	UPROPERTY(meta = (Input))
+	float MaxOut = 0.f;
+
+	UPROPERTY(meta = (Output))
+	float Result = 0.f;
+
+	RIGVM_METHOD()
+	virtual void Execute() override;
+
+	RIGVM_METHOD()
+	virtual FRigVMStructUpgradeInfo GetUpgradeInfo() const override;
+};
\ No newline at end of file
diff --git a/Engine/Plugins/Animation/ControlRig/Source/ControlRig/Public/Units/Deprecated/Math/RigUnit_Quaternion.h b/Engine/Plugins/Animation/ControlRig/Source/ControlRig/Public/Units/Deprecated/Math/RigUnit_Quaternion.h
new file mode 100644
--- /dev/null
+++ b/Engine/Plugins/Animation/ControlRig/Source/ControlRig/Public/Units/Deprecated/Math/RigUnit_Quaternion.h
@@ -0,0 +1,142 @@
+// Copyright Epic Games, Inc. All Rights Reserved.
+
+#pragma once
+
+#include "Units/RigUnit.h"
+#include "MathLibrary.h"
+#include "RigUnit_Quaternion.generated.h"
+
+/** Two args and a result of Quaternion type */
+USTRUCT(meta=(Abstract, NodeColor = "0.1 0.7 0.1"))
+struct CONTROLRIG_API FRigUnit_BinaryQuaternionOp : public FRigUnit
+{
+	GENERATED_BODY()
+
+		FRigUnit_BinaryQuaternionOp()
+		: Argument0(FQuat::Identity)
+		, Argument1(FQuat::Identity)
+		, Result(FQuat::Identity)
+	{}
+
+	UPROPERTY(meta=(Input))
+	FQuat Argument0;
+
+	UPROPERTY(meta=(Input))
+	FQuat Argument1;
+
+	UPROPERTY(meta=(Output))
+	FQuat Result;
+};
+
+USTRUCT(meta=(DisplayName="Multiply(Quaternion)", Category="Math|Quaternion", Deprecated="4.23.0"))
+struct CONTROLRIG_API FRigUnit_MultiplyQuaternion : public FRigUnit_BinaryQuaternionOp
+{
+	GENERATED_BODY()
+
+	RIGVM_METHOD()
+	virtual void Execute() override;
+
+	RIGVM_METHOD()
+	virtual FRigVMStructUpgradeInfo GetUpgradeInfo() const override;
+};
+
+/** Two args and a result of Quaternion type */
+USTRUCT(meta=(Abstract, NodeColor = "0.1 0.7 0.1"))
+struct CONTROLRIG_API FRigUnit_UnaryQuaternionOp : public FRigUnit
+{
+	GENERATED_BODY()
+
+	UPROPERTY(meta=(Input))
+	FQuat Argument = FQuat::Identity;
+
+	UPROPERTY(meta=(Output))
+	FQuat Result = FQuat::Identity;
+};
+
+USTRUCT(meta = (DisplayName = "Inverse(Quaternion)", Category = "Math|Quaternion", Deprecated="4.23.0"))
+struct CONTROLRIG_API FRigUnit_InverseQuaterion: public FRigUnit_UnaryQuaternionOp
+{
+	GENERATED_BODY()
+
+	RIGVM_METHOD()
+	virtual void Execute() override;
+
+	RIGVM_METHOD()
+	virtual FRigVMStructUpgradeInfo GetUpgradeInfo() const override;
+};
+
+USTRUCT(meta = (DisplayName = "To Axis And Angle(Quaternion)", Category = "Math|Quaternion", NodeColor = "0.1 0.7 0.1", Deprecated="4.23.0"))
+struct CONTROLRIG_API FRigUnit_QuaternionToAxisAndAngle : public FRigUnit
+{
+	GENERATED_BODY()
+
+	UPROPERTY(meta = (Input))
+	FQuat Argument = FQuat::Identity;
+
+	UPROPERTY(meta = (Output))
+	FVector Axis = FVector(0.f);
+
+	UPROPERTY(meta = (Output))
+	float Angle = 0.f;
+
+	RIGVM_METHOD()
+	virtual void Execute() override;
+
+	RIGVM_METHOD()
+	virtual FRigVMStructUpgradeInfo GetUpgradeInfo() const override;
+};
+
+USTRUCT(meta = (DisplayName = "From Axis And Angle(Quaternion)", Category = "Math|Quaternion", NodeColor = "0.1 0.7 0.1", Deprecated="4.23.0"))
+struct CONTROLRIG_API FRigUnit_QuaternionFromAxisAndAngle : public FRigUnit
+{
+	GENERATED_BODY()
+
+	FRigUnit_QuaternionFromAxisAndAngle()
+		: Axis(1.f, 0.f, 0.f)
+		, Angle(0.f)
+		, Result(ForceInitToZero)
+	{}
+
+	UPROPERTY(meta = (Input))
+	FVector Axis;
+
+	UPROPERTY(meta = (Input))
+	float Angle;
+
+	UPROPERTY(meta = (Output))
+	FQuat Result;
+
+	RIGVM_METHOD()
+	virtual void Execute() override;
+
+	RIGVM_METHOD()
+	virtual FRigVMStructUpgradeInfo GetUpgradeInfo() const override;
+};
+
+USTRUCT(meta = (DisplayName = "Get Angle Around Axis", Category = "Math|Quaternion", NodeColor = "0.1 0.7 0.1", Deprecated="4.23.0"))
+struct CONTROLRIG_API FRigUnit_QuaternionToAngle : public FRigUnit
+{
+	GENERATED_BODY()
+
+	FRigUnit_QuaternionToAngle()
+		: Axis(1.f, 0.f, 0.f)
+		, Argument(ForceInitToZero)
+		, Angle(0.f)
+	{}
+
+	UPROPERTY(meta = (Input))
+	FVector Axis;
+
+	UPROPERTY(meta = (Input))
+	FQuat Argument;
+
+	UPROPERTY(meta = (Output))
+	float Angle;
+
+	RIGVM_METHOD()
+	virtual void Execute() override;
+	
+	RIGVM_METHOD()
+	virtual FRigVMStructUpgradeInfo GetUpgradeInfo() const override;
+};
+
diff --git a/Engine/Plugins/Animation/ControlRig/Source/ControlRig/Public/Units/Deprecated/Math/RigUnit_Transform.h b/Engine/Plugins/Animation/ControlRig/Source/ControlRig/Public/Units/Deprecated/Math/RigUnit_Transform.h
new file mode 100644
--- /dev/null
+++ b/Engine/Plugins/Animation/ControlRig/Source/ControlRig/Public/Units/Deprecated/Math/RigUnit_Transform.h
@@ -0,0 +1,48 @@
+// Copyright Epic Games, Inc. All Rights Reserved.
+
+#pragma once
+
+#include "Units/RigUnit.h"
+#include "MathLibrary.h"
+#include "RigUnit_Transform.generated.h"
+
+/** Two args and a result of Transform type */
+USTRUCT(meta=(Abstract, NodeColor = "0.1 0.7 0.1"))
+struct CONTROLRIG_API FRigUnit_BinaryTransformOp : public FRigUnit
+{
+	GENERATED_BODY()
+
+	UPROPERTY(meta=(Input))
+	FTransform Argument0;
+
+	UPROPERTY(meta=(Input))
+	FTransform Argument1;
+
+	UPROPERTY(meta=(Output))
+	FTransform Result;
+};
+
+USTRUCT(meta=(DisplayName="Multiply(Transform)", Category="Math|Transform", Deprecated="4.23.0"))
+struct CONTROLRIG_API FRigUnit_MultiplyTransform : public FRigUnit_BinaryTransformOp
+{
+	GENERATED_BODY()
+
+	RIGVM_METHOD()
+	virtual void Execute() override;
+
+	RIGVM_METHOD()
+	virtual FRigVMStructUpgradeInfo GetUpgradeInfo() const override;
+};
+
+USTRUCT(meta = (DisplayName = "GetRelativeTransform", Category = "Math|Transform", Deprecated="4.23.0"))
+struct CONTROLRIG_API FRigUnit_GetRelativeTransform : public FRigUnit_BinaryTransformOp
+{
+	GENERATED_BODY()
+
+	RIGVM_METHOD()
+	virtual void Execute() override;
+
+	RIGVM_METHOD()
+	virtual FRigVMStructUpgradeInfo GetUpgradeInfo() const override;
+};
+
diff --git a/Engine/Plugins/Animation/ControlRig/Source/ControlRig/Public/Units/Deprecated/Math/RigUnit_Vector.h b/Engine/Plugins/Animation/ControlRig/Source/ControlRig/Public/Units/Deprecated/Math/RigUnit_Vector.h
new file mode 100644
--- /dev/null
+++ b/Engine/Plugins/Animation/ControlRig/Source/ControlRig/Public/Units/Deprecated/Math/RigUnit_Vector.h
@@ -0,0 +1,95 @@
+// Copyright Epic Games, Inc. All Rights Reserved.
+
+#pragma once
+
+#include "Units/RigUnit.h"
+#include "MathLibrary.h"
+#include "Math/ControlRigMathLibrary.h"
+#include "RigVMFunctions/Math/RigVMFunction_MathVector.h"
+#include "RigUnit_Vector.generated.h"
+
+/** Two args and a result of Vector type */
+USTRUCT(meta=(Abstract, NodeColor = "0.1 0.7 0.1"))
+struct CONTROLRIG_API FRigUnit_BinaryVectorOp : public FRigUnit
+{
+	GENERATED_BODY()
+
+	UPROPERTY(meta=(Input))
+	FVector Argument0 = FVector(0.f);
+
+	UPROPERTY(meta=(Input))
+	FVector Argument1 = FVector(0.f);
+
+	UPROPERTY(meta=(Output))
+	FVector Result = FVector(0.f);
+};
+
+USTRUCT(meta=(DisplayName="Multiply(Vector)", Category="Math|Vector", Deprecated="4.23.0"))
+struct CONTROLRIG_API FRigUnit_Multiply_VectorVector : public FRigUnit_BinaryVectorOp
+{
+	GENERATED_BODY()
+
+	RIGVM_METHOD()
+	virtual void Execute() override;
+
+	RIGVM_METHOD()
+	virtual FRigVMStructUpgradeInfo GetUpgradeInfo() const override;
+};
+
+USTRUCT(meta=(DisplayName="Add(Vector)", Category="Math|Vector", Deprecated="4.23.0"))
+struct CONTROLRIG_API FRigUnit_Add_VectorVector : public FRigUnit_BinaryVectorOp
+{
+	GENERATED_BODY()
+
+	RIGVM_METHOD()
+	virtual void Execute() override;
+
+	RIGVM_METHOD()
+	virtual FRigVMStructUpgradeInfo GetUpgradeInfo() const override;
+};
+
+USTRUCT(meta=(DisplayName="Subtract(Vector)", Category="Math|Vector", Deprecated="4.23.0"))
+struct CONTROLRIG_API FRigUnit_Subtract_VectorVector : public FRigUnit_BinaryVectorOp
+{
+	GENERATED_BODY()
+
+	RIGVM_METHOD()
+	virtual void Execute() override;
+
+	RIGVM_METHOD()
+	virtual FRigVMStructUpgradeInfo GetUpgradeInfo() const override;
+};
+
+USTRUCT(meta=(DisplayName="Divide(Vector)", Category="Math|Vector", Deprecated="4.23.0"))
+struct CONTROLRIG_API FRigUnit_Divide_VectorVector : public FRigUnit_BinaryVectorOp
+{
+	GENERATED_BODY()
+
+	RIGVM_METHOD()
+	virtual void Execute() override;
+
+	RIGVM_METHOD()
+	virtual FRigVMStructUpgradeInfo GetUpgradeInfo() const override;
+};
+
+USTRUCT(meta = (DisplayName = "Distance", Category = "Math|Vector", Deprecated="4.23.0", NodeColor = "0.1 0.7 0.1"))
+struct CONTROLRIG_API FRigUnit_Distance_VectorVector : public FRigUnit
+{
+	GENERATED_BODY()
+
+	UPROPERTY(meta=(Input))
+	FVector Argument0 = FVector(0.f);
+
+	UPROPERTY(meta=(Input))
+	FVector Argument1 = FVector(0.f);
+
+	UPROPERTY(meta=(Output))
+	float Result = 0.f;
+
+	RIGVM_METHOD()
+	virtual void Execute() override;
+
+	RIGVM_METHOD()
+	virtual FRigVMStructUpgradeInfo GetUpgradeInfo() const override;
+};
+
diff --git a/Engine/Plugins/Animation/ControlRig/Source/ControlRig/Public/Units/Hierarchy/RigUnit_BoneName.h b/Engine/Plugins/Animation/ControlRig/Source/ControlRig/Public/Units/Hierarchy/RigUnit_BoneName.h
new file mode 100644
--- /dev/null
+++ b/Engine/Plugins/Animation/ControlRig/Source/ControlRig/Public/Units/Hierarchy/RigUnit_BoneName.h
@@ -0,0 +1,120 @@
+// Copyright Epic Games, Inc. All Rights Reserved.
+
+#pragma once
+
+#include "Units/RigUnit.h"
+#include "RigUnit_BoneName.generated.h"
+
+/**
+ * The Item node is used to share a specific item across the graph
+ */
+USTRUCT(meta=(DisplayName="Item", Category="Hierarchy", NodeColor = "0.4627450108528137 1.0 0.3294120132923126", DocumentationPolicy = "Strict", Constant))
+struct CONTROLRIG_API FRigUnit_Item : public FRigUnit
+{
+	GENERATED_BODY()
+
+	FRigUnit_Item()
+	{}
+
+	RIGVM_METHOD()
+	virtual void Execute() override;
+
+	/**
+	 * The item
+	 */
+	UPROPERTY(meta = (Input, Output, ExpandByDefault))
+	FRigElementKey Item;
+};
+
+/**
+ * The Item Array node is used to share an array of items across the graph
+ */
+USTRUCT(meta=(DisplayName="Item Array", Category="Hierarchy", NodeColor = "0.4627450108528137 1.0 0.3294120132923126", DocumentationPolicy = "Strict", Constant))
+struct CONTROLRIG_API FRigUnit_ItemArray : public FRigUnit
+{
+	GENERATED_BODY()
+
+	FRigUnit_ItemArray()
+	{}
+
+	RIGVM_METHOD()
+	virtual void Execute() override;
+
+	/**
+	 * The items
+	 */
+	UPROPERTY(meta = (Input, Output, ExpandByDefault))
+	TArray<FRigElementKey> Items;
+};
+
+/**
+ * BoneName is used to represent a bone name in the graph
+ */
+USTRUCT(meta=(DisplayName="Bone Name", Category="Hierarchy", DocumentationPolicy = "Strict", Constant, Deprecated = "4.25"))
+struct CONTROLRIG_API FRigUnit_BoneName : public FRigUnit
+{
+	GENERATED_BODY()
+
+	FRigUnit_BoneName()
+	{}
+
+	RIGVM_METHOD()
+	virtual void Execute() override;
+
+	/**
+	 * The name of the Bone
+	 */
+	UPROPERTY(meta = (Input, Output))
+	FName Bone;
+
+	RIGVM_METHOD()
+	virtual FRigVMStructUpgradeInfo GetUpgradeInfo() const override;
+};
+
+/**
+ * SpaceName is used to represent a Space name in the graph
+ */
+USTRUCT(meta=(DisplayName="Space Name", Category="Hierarchy", DocumentationPolicy = "Strict", Deprecated = "4.25"))
+struct CONTROLRIG_API FRigUnit_SpaceName : public FRigUnit
+{
+	GENERATED_BODY()
+
+	FRigUnit_SpaceName()
+	{}
+
+	RIGVM_METHOD()
+	virtual void Execute() override;
+
+	/**
+	 * The name of the Space
+	 */
+	UPROPERTY(meta = (Input, Output))
+	FName Space;
+
+	RIGVM_METHOD()
+	virtual FRigVMStructUpgradeInfo GetUpgradeInfo() const override;
+};
+
+/**
+ * ControlName is used to represent a Control name in the graph
+ */
+USTRUCT(meta=(DisplayName="Control Name", Category="Hierarchy", DocumentationPolicy = "Strict", Deprecated = "4.25"))
+struct CONTROLRIG_API FRigUnit_ControlName : public FRigUnit
+{
+	GENERATED_BODY()
+
+	FRigUnit_ControlName()
+	{}
+
+	RIGVM_METHOD()
+	virtual void Execute() override;
+
+	/**
+	 * The name of the Control
+	 */
+	UPROPERTY(meta = (Input, Output))
+	FName Control;
+
+	RIGVM_METHOD()
+	virtual FRigVMStructUpgradeInfo GetUpgradeInfo() const override;
+};
diff --git a/Engine/Plugins/Animation/ControlRig/Source/ControlRig/Public/Units/Hierarchy/RigUnit_OffsetTransform.h b/Engine/Plugins/Animation/ControlRig/Source/ControlRig/Public/Units/Hierarchy/RigUnit_OffsetTransform.h
new file mode 100644
--- /dev/null
+++ b/Engine/Plugins/Animation/ControlRig/Source/ControlRig/Public/Units/Hierarchy/RigUnit_OffsetTransform.h
@@ -0,0 +1,54 @@
+// Copyright Epic Games, Inc. All Rights Reserved.
+
+#pragma once
+
+#include "Units/RigUnit.h"
+#include "RigUnit_OffsetTransform.generated.h"
+
+/**
+ * Offset Transform is used to add an offset to an existing transform in the hierarchy. The offset is post multiplied.
+ */
+USTRUCT(meta=(DisplayName="Offset Transform", Category="Transforms", DocumentationPolicy = "Strict", Keywords = "Offset,Relative,AddBoneTransform", NodeColor="0, 0.364706, 1.0", Varying))
+struct CONTROLRIG_API FRigUnit_OffsetTransformForItem : public FRigUnitMutable
+{
+	GENERATED_BODY()
+
+	FRigUnit_OffsetTransformForItem()
+		: Item(NAME_None, ERigElementType::Bone)
+		, OffsetTransform(FTransform::Identity)
+		, Weight(1.f)
+		, bPropagateToChildren(true)
+		, CachedIndex()
+	{}
+
+#if WITH_EDITOR
+	virtual bool UpdateHierarchyForDirectManipulation(const URigVMUnitNode* InNode, TSharedPtr<FStructOnScope> InInstance, FControlRigExecuteContext& InContext, TSharedPtr<FRigDirectManipulationInfo> InInfo) override;
+	virtual bool UpdateDirectManipulationFromHierarchy(const URigVMUnitNode* InNode, TSharedPtr<FStructOnScope> InInstance, FControlRigExecuteContext& InContext, TSharedPtr<FRigDirectManipulationInfo> InInfo) override;
+#endif
+	
+	RIGVM_METHOD()
+	virtual void Execute() override;
+
+	/**
+	 * The item to offset the transform for
+	 */
+	UPROPERTY(meta = (Input, ExpandByDefault))
+	FRigElementKey Item;
+
+	// The transform of the item relative to its previous transform
+	UPROPERTY(meta=(Input))
+	FTransform OffsetTransform;
+
+	// Defines how much the change will be applied
+	UPROPERTY(meta=(Input, UIMin = "0.0", UIMax = "1.0"))
+	float Weight;
+
+	// If set to true children of affected items in the hierarchy
+	// will follow the transform change - otherwise only the parent will move.
+	UPROPERTY(meta=(Input))
+	bool bPropagateToChildren;
+
+	// Used to cache the item internally
+	UPROPERTY()
+	FCachedRigElement CachedIndex;
+};
diff --git a/Engine/Plugins/Animation/ControlRig/Source/ControlRig/Public/Units/Hierarchy/RigUnit_SendEvent.h b/Engine/Plugins/Animation/ControlRig/Source/ControlRig/Public/Units/Hierarchy/RigUnit_SendEvent.h
new file mode 100644
--- /dev/null
+++ b/Engine/Plugins/Animation/ControlRig/Source/ControlRig/Public/Units/Hierarchy/RigUnit_SendEvent.h
@@ -0,0 +1,56 @@
+// Copyright Epic Games, Inc. All Rights Reserved.
+
+#pragma once
+
+#include "Units/RigUnit.h"
+#include "RigUnit_SendEvent.generated.h"
+
+/**
+ * SendEvent is used to notify the engine / editor of a change that happend within the Control Rig.
+ */
+USTRUCT(meta=(DisplayName="Send Event", Category="Hierarchy", DocumentationPolicy="Strict", Keywords = "SendEvent", TemplateName="Event,Notify,Notification", NodeColor="1, 0, 0"))
+struct CONTROLRIG_API FRigUnit_SendEvent : public FRigUnitMutable
+{
+	GENERATED_BODY()
+
+	FRigUnit_SendEvent()
+		: Event(ERigEvent::RequestAutoKey)
+		, Item()
+		, OffsetInSeconds(0.f)
+		, bEnable(true)
+		, bOnlyDuringInteraction(true)
+	{}
+
+	RIGVM_METHOD()
+	virtual void Execute() override;
+
+	/**
+	 * The event to send to the engine
+	 */
+	UPROPERTY(meta = (Input))
+	ERigEvent Event;
+
+	/**
+	 * The item to send the event for
+	 */
+	UPROPERTY(meta = (Input, ExpandByDefault))
+	FRigElementKey Item;
+
+	/**
+	 * The time offset to use for the send event
+	 */
+	UPROPERTY(meta = (Input, UIMin = "-1000", UIMax = "1000"))
+	float OffsetInSeconds;
+
+	/**
+	 * The event will be sent if this is checked
+	 */
+	UPROPERTY(meta = (Input))
+	bool bEnable;
+
+	/**
+	 * The event will be sent if this only during an interaction
+	 */
+	UPROPERTY(meta = (Input))
+	bool bOnlyDuringInteraction;
+};
diff --git a/Engine/Plugins/Animation/ControlRig/Source/ControlRig/Public/Units/Highlevel/Hierarchy/RigUnit_ModifyTransforms.h b/Engine/Plugins/Animation/ControlRig/Source/ControlRig/Public/Units/Highlevel/Hierarchy/RigUnit_ModifyTransforms.h
new file mode 100644
--- /dev/null
+++ b/Engine/Plugins/Animation/ControlRig/Source/ControlRig/Public/Units/Highlevel/Hierarchy/RigUnit_ModifyTransforms.h
@@ -0,0 +1,130 @@
+// Copyright Epic Games, Inc. All Rights Reserved.
+
+#pragma once
+
+#include "Units/RigUnit.h"
+#include "Units/Highlevel/RigUnit_HighlevelBase.h"
+#include "RigUnit_ModifyTransforms.generated.h"
+
+UENUM()
+enum class EControlRigModifyBoneMode : uint8
+{
+	/** Override existing local transform */
+	OverrideLocal,
+
+	/** Override existing global transform */
+	OverrideGlobal,
+
+	/** 
+	 * Additive to existing local transform.
+	 * Input transform is added within the bone's space.
+	 */
+	AdditiveLocal,
+
+	/**
+     * Additive to existing global transform.
+     * Input transform is added as a global offset in the root of the hierarchy.
+	 */
+	AdditiveGlobal,
+
+	/** MAX - invalid */
+	Max UMETA(Hidden),
+};
+
+USTRUCT()
+struct CONTROLRIG_API FRigUnit_ModifyTransforms_PerItem
+{
+	GENERATED_BODY()
+
+	FRigUnit_ModifyTransforms_PerItem()
+		: Item(FRigElementKey(NAME_None, ERigElementType::Bone))
+	{
+	}
+
+	/**
+	 * The item to set the transform for.
+	 */
+	UPROPERTY(EditAnywhere, meta = (Input, ExpandByDefault), Category = FRigUnit_ModifyTransforms_PerItem)
+	FRigElementKey Item;
+
+	/**
+	 * The transform value to set for the given Bone.
+	 */
+	UPROPERTY(EditAnywhere, meta = (Input), Category = FRigUnit_ModifyTransforms_PerItem)
+	FTransform Transform;
+};
+
+USTRUCT()
+struct CONTROLRIG_API FRigUnit_ModifyTransforms_WorkData
+{
+	GENERATED_BODY()
+
+	UPROPERTY()
+	TArray<FCachedRigElement> CachedItems;
+};
+
+/**
+ * Modify Transforms is used to perform a change in the hierarchy by setting one or more bones' transforms
+ */
+USTRUCT(meta=(DisplayName="Modify Transforms", Category="Transforms", DocumentationPolicy="Strict", Keywords = "ModifyBone"))
+struct CONTROLRIG_API FRigUnit_ModifyTransforms : public FRigUnit_HighlevelBaseMutable
+{
+	GENERATED_BODY()
+
+	FRigUnit_ModifyTransforms()
+		: Weight(1.f)
+		, WeightMinimum(0.f)
+		, WeightMaximum(1.f)
+		, Mode(EControlRigModifyBoneMode::AdditiveLocal)
+	{
+		ItemToModify.Add(FRigUnit_ModifyTransforms_PerItem());
+		ItemToModify[0].Item = FRigElementKey(NAME_None, ERigElementType::Bone);
+	}
+
+#if WITH_EDITOR
+	virtual bool GetDirectManipulationTargets(const URigVMUnitNode* InNode, TSharedPtr<FStructOnScope> InInstance, URigHierarchy* InHierarchy, TArray<FRigDirectManipulationTarget>& InOutTargets, FString* OutFailureReason) const override;
+	virtual bool UpdateHierarchyForDirectManipulation(const URigVMUnitNode* InNode, TSharedPtr<FStructOnScope> InInstance, FControlRigExecuteContext& InContext, TSharedPtr<FRigDirectManipulationInfo> InInfo) override;
+	virtual bool UpdateDirectManipulationFromHierarchy(const URigVMUnitNode* InNode, TSharedPtr<FStructOnScope> InInstance, FControlRigExecuteContext& InContext, TSharedPtr<FRigDirectManipulationInfo> InInfo) override;
+	virtual TArray<const URigVMPin*> GetPinsForDirectManipulation(const URigVMUnitNode* InNode, const FRigDirectManipulationTarget& InTarget) const override;
+	int32 GetIndexFromTarget(const FString& InTarget) const;
+#endif
+	
+	RIGVM_METHOD()
+	virtual void Execute() override;
+
+	/**
+	 * The items to modify.
+	 */
+	UPROPERTY(meta = (Input, ExpandByDefault, DefaultArraySize=1))
+	TArray<FRigUnit_ModifyTransforms_PerItem> ItemToModify;
+
+	/**
+	 * At 1 this sets the transform, between 0 and 1 the transform is blended with previous results.
+	 */
+	UPROPERTY(meta = (Input, ClampMin=0.f, ClampMax=1.f, UIMin = 0.f, UIMax = 1.f))
+	float Weight;
+
+	/**
+	 * The minimum of the weight - defaults to 0.0
+	 */
+	UPROPERTY(meta = (Input, Constant, ClampMin = 0.f, ClampMax = 1.f, UIMin = 0.f, UIMax = 1.f))
+	float WeightMinimum;
+
+	/**
+	 * The maximum of the weight - defaults to 1.0
+	 */
+	UPROPERTY(meta = (Input, Constant, ClampMin = 0.f, ClampMax = 1.f, UIMin = 0.f, UIMax = 1.f))
+	float WeightMaximum;
+
+	/**
+	 * Defines if the bone's transform should be set
+	 * in local or global space, additive or override.
+	 */
+	UPROPERTY(meta = (Input, Constant))
+	EControlRigModifyBoneMode Mode;
+
+	// Used to cache the internally used bone index
+	UPROPERTY(transient)
+	FRigUnit_ModifyTransforms_WorkData WorkData;
+};
+
diff --git a/Engine/Plugins/Animation/ControlRig/Source/ControlRigDeveloper/ControlRigDeveloper.Build.cs b/Engine/Plugins/Animation/ControlRig/Source/ControlRigDeveloper/ControlRigDeveloper.Build.cs
--- a/Engine/Plugins/Animation/ControlRig/Source/ControlRigDeveloper/ControlRigDeveloper.Build.cs
+++ b/Engine/Plugins/Animation/ControlRig/Source/ControlRigDeveloper/ControlRigDeveloper.Build.cs
@@ -6,9 +6,6 @@ namespace UnrealBuildTool.Rules
     {
         public ControlRigDeveloper(ReadOnlyTargetRules Target) : base(Target)
         {
-            PrivateIncludePaths.Add("ControlRig/Private");
-			PrivateIncludePaths.Add("ControlRigEditor/Private");
-
 			// Copying some these from ControlRig.build.cs, our deps are likely leaner
 			// and therefore these could be pruned if needed:
 			PrivateDependencyModuleNames.AddRange(
diff --git a/Engine/Plugins/Animation/ControlRig/Source/ControlRigEditor/ControlRigEditor.Build.cs b/Engine/Plugins/Animation/ControlRig/Source/ControlRigEditor/ControlRigEditor.Build.cs
--- a/Engine/Plugins/Animation/ControlRig/Source/ControlRigEditor/ControlRigEditor.Build.cs
+++ b/Engine/Plugins/Animation/ControlRig/Source/ControlRigEditor/ControlRigEditor.Build.cs
@@ -6,12 +6,6 @@ namespace UnrealBuildTool.Rules
     {
         public ControlRigEditor(ReadOnlyTargetRules Target) : base(Target)
         {
-			PrivateIncludePaths.AddRange(
-                new string[] {
-					System.IO.Path.Combine(GetModuleDirectory("ControlRig"), "Private"),
-				}
-			);
-
             PublicDependencyModuleNames.AddRange(
 				new string[]
 				{
diff --git a/Engine/Plugins/Animation/ControlRig/Source/ControlRigEditor/Public/ControlRigBlueprintFactory.h b/Engine/Plugins/Animation/ControlRig/Source/ControlRigEditor/Public/ControlRigBlueprintFactory.h
new file mode 100644
--- /dev/null
+++ b/Engine/Plugins/Animation/ControlRig/Source/ControlRigEditor/Public/ControlRigBlueprintFactory.h
@@ -0,0 +1,46 @@
+// Copyright Epic Games, Inc. All Rights Reserved.
+
+#pragma once
+
+#include "CoreMinimal.h"
+#include "UObject/ObjectMacros.h"
+#include "Templates/SubclassOf.h"
+#include "Engine/Blueprint.h"
+#include "Factories/Factory.h"
+#include "ControlRigBlueprint.h"
+#include "ControlRig.h"
+#include "ControlRigBlueprintFactory.generated.h"
+
+UCLASS(HideCategories=Object)
+class CONTROLRIGEDITOR_API UControlRigBlueprintFactory : public UFactory
+{
+	GENERATED_BODY()
+
+public:
+	UControlRigBlueprintFactory();
+
+	// The parent class of the created blueprint
+	UPROPERTY(EditAnywhere, Category="Control Rig Factory", meta=(AllowAbstract = ""))
+	TSubclassOf<UControlRig> ParentClass;
+
+	// UFactory Interface
+	virtual bool ConfigureProperties() override;
+	virtual UObject* FactoryCreateNew(UClass* Class, UObject* InParent, FName Name, EObjectFlags Flags, UObject* Context, FFeedbackContext* Warn, FName CallingContext) override;
+	virtual UObject* FactoryCreateNew(UClass* Class, UObject* InParent, FName Name, EObjectFlags Flags, UObject* Context, FFeedbackContext* Warn) override;
+
+	/**
+	 * Create a new control rig asset within the contents space of the project.
+	 * @param InDesiredPackagePath The package path to use for the control rig asset
+	 */
+	UFUNCTION(BlueprintCallable, Category = "Control Rig")
+	static UControlRigBlueprint* CreateNewControlRigAsset(const FString& InDesiredPackagePath);
+
+	/**
+	 * Create a new control rig asset within the contents space of the project
+	 * based on a skeletal mesh or skeleton object.
+	 * @param InSelectedObject The SkeletalMesh / Skeleton object to base the control rig asset on
+	 */
+	UFUNCTION(BlueprintCallable, Category = "Control Rig")
+	static UControlRigBlueprint* CreateControlRigFromSkeletalMeshOrSkeleton(UObject* InSelectedObject);
+};
+
