Control Rig: Add ability to convert an absolute control rig to layered control rig in sequencer
#rnx

diff --git a/Engine/Plugins/Animation/ControlRig/Source/ControlRigEditor/Private/Sequencer/ControlRigParameterTrackEditor.cpp b/Engine/Plugins/Animation/ControlRig/Source/ControlRigEditor/Private/Sequencer/ControlRigParameterTrackEditor.cpp
--- a/Engine/Plugins/Animation/ControlRig/Source/ControlRigEditor/Private/Sequencer/ControlRigParameterTrackEditor.cpp
+++ b/Engine/Plugins/Animation/ControlRig/Source/ControlRigEditor/Private/Sequencer/ControlRigParameterTrackEditor.cpp
@@ -1214,6 +1214,64 @@ void FControlRigParameterTrackEditor::BakeInvertedPose(UControlRig* InControlRig
 	ParentSequencer->NotifyMovieSceneDataChanged(EMovieSceneDataChangeType::TrackValueChanged);
 }
 
+void FControlRigParameterTrackEditor::ConvertIsLayered(UMovieSceneControlRigParameterTrack* Track)
+{
+	UControlRig* ControlRig = Track->GetControlRig();
+	if (!ControlRig)
+	{
+		return;
+	}
+
+	const FScopedTransaction Transaction(LOCTEXT("ConvertToLayeredControlRig_Transaction", "Convert to Layered Control Rig"));
+	Track->Modify();
+	ControlRig->Modify();
+
+	const bool bSetAdditive = !ControlRig->IsAdditive();
+
+	ControlRig->ClearPoseBeforeBackwardsSolve();
+	ControlRig->ResetControlValues();
+	ControlRig->SetIsAdditive(bSetAdditive);
+
+	ControlRig->Evaluate_AnyThread();
+
+	FString ObjectName = ControlRig->GetClass()->GetName(); //GetDisplayNameText().ToString();
+	ObjectName.RemoveFromEnd(TEXT("_C"));
+	
+	if (bSetAdditive)
+	{
+		const FString AdditiveObjectName = ObjectName + TEXT(" (Layered)");
+		Track->SetTrackName(FName(*ObjectName));
+		Track->SetDisplayName(FText::FromString(AdditiveObjectName));
+		Track->SetColorTint(UMovieSceneControlRigParameterTrack::LayeredRigTrackColor);
+	}
+	else
+	{
+		Track->SetTrackName(FName(*ObjectName));
+		Track->SetDisplayName(FText::FromString(ObjectName));
+		Track->SetColorTint(UMovieSceneControlRigParameterTrack::AbsoluteRigTrackColor);
+	}
+
+	if (FControlRigEditMode* EditMode = GetEditMode())
+	{
+		EditMode->ZeroTransforms(false);
+	}
+
+	for (UMovieSceneSection* Section : Track->GetAllSections())
+	{
+		if (Section)
+		{
+			UMovieSceneControlRigParameterSection* CRSection = Cast<UMovieSceneControlRigParameterSection>(Section);
+			if (CRSection)
+			{
+				Section->Modify();
+				CRSection->ClearAllParameters();
+				CRSection->RecreateWithThisControlRig(CRSection->GetControlRig(), true);
+			}
+		}
+	}
+	GetSequencer()->NotifyMovieSceneDataChanged(EMovieSceneDataChangeType::MovieSceneStructureItemsChanged);
+}
+
 void FControlRigParameterTrackEditor::BuildObjectBindingTrackMenu(FMenuBuilder& MenuBuilder, const TArray<FGuid>& ObjectBindings, const UClass* ObjectClass)
 {
 	if(!ObjectClass)
@@ -3848,6 +3906,17 @@ void FControlRigParameterTrackEditor::BuildTrackContextMenu(FMenuBuilder& MenuBu
 			LOCTEXT("Order", "Order")
 		);
 
+		MenuBuilder.AddMenuEntry(
+				LOCTEXT("ConvertIsLayeredControlRig", "Convert To Layered"),
+				LOCTEXT("ConvertIsLayeredControlRigToolTip", "Converts the Control Rig from an Absolute rig to a Layered rig"),
+				FSlateIcon(),
+				FUIAction(
+					FExecuteAction::CreateRaw(this, &FControlRigParameterTrackEditor::ConvertIsLayered, Track),
+					FCanExecuteAction(),
+					FIsActionChecked::CreateUObject(Track->GetControlRig(), &UControlRig::IsAdditive)
+				),
+				NAME_None,
+				EUserInterfaceActionType::ToggleButton);
 	}
 	MenuBuilder.EndSection();
 
diff --git a/Engine/Plugins/Animation/ControlRig/Source/ControlRigEditor/Private/Sequencer/ControlRigParameterTrackEditor.h b/Engine/Plugins/Animation/ControlRig/Source/ControlRigEditor/Private/Sequencer/ControlRigParameterTrackEditor.h
--- a/Engine/Plugins/Animation/ControlRig/Source/ControlRigEditor/Private/Sequencer/ControlRigParameterTrackEditor.h
+++ b/Engine/Plugins/Animation/ControlRig/Source/ControlRigEditor/Private/Sequencer/ControlRigParameterTrackEditor.h
@@ -199,6 +199,9 @@ private:
 	/** Bake Inersion of Additive Control Rig to Rest Pose*/
 	void BakeInvertedPose(UControlRig* InControlRig, UMovieSceneControlRigParameterTrack* Track);
 
+	/** Convert an absolute control rit to a layered control rig */
+	void ConvertIsLayered(UMovieSceneControlRigParameterTrack* Track);
+
 	/** Set Up EditMode for Specified Control Rig*/
 	void SetUpEditModeIfNeeded(UControlRig* ControlRig);
 
