Control Rig: Add ability to convert an absolute control rig to layered control rig in sequencer through python
#rnx

diff --git a/Engine/Plugins/Animation/ControlRig/Source/ControlRigEditor/Private/ControlRigSequencerEditorLibrary.cpp b/Engine/Plugins/Animation/ControlRig/Source/ControlRigEditor/Private/ControlRigSequencerEditorLibrary.cpp
--- a/Engine/Plugins/Animation/ControlRig/Source/ControlRigEditor/Private/ControlRigSequencerEditorLibrary.cpp
+++ b/Engine/Plugins/Animation/ControlRig/Source/ControlRigEditor/Private/ControlRigSequencerEditorLibrary.cpp
@@ -3097,6 +3097,90 @@ bool UControlRigSequencerEditorLibrary::IsLayeredControlRig(UControlRig* InContr
 	return (InControlRig && InControlRig->IsAdditive());
 }
 
+bool UControlRigSequencerEditorLibrary::SetControlRigLayeredMode(UMovieSceneControlRigParameterTrack* InTrack, bool bSetIsLayered)
+{
+	if (!InTrack)
+	{
+		FFrame::KismetExecutionMessage(TEXT("Invalid track"), ELogVerbosity::Error);
+		return false;
+	}
+
+	UControlRig* ControlRig = InTrack->GetControlRig();
+	if (!ControlRig)
+	{
+		FFrame::KismetExecutionMessage(TEXT("Track does not have a control rig"), ELogVerbosity::Error);
+		return false;
+	}
+
+	if (ControlRig->IsAdditive() == bSetIsLayered)
+	{
+		if(bSetIsLayered)
+		{
+			FFrame::KismetExecutionMessage(TEXT("Control rig is already in layered mode"), ELogVerbosity::Error);
+		}
+		else
+		{
+			FFrame::KismetExecutionMessage(TEXT("Control rig is already in absolute mode"), ELogVerbosity::Error);
+		}
+		return false;
+	}
+
+	const FScopedTransaction Transaction(LOCTEXT("ConvertToLayeredControlRig_Transaction", "Convert to Layered Control Rig"));
+	InTrack->Modify();
+	ControlRig->Modify();
+
+	ControlRig->ClearPoseBeforeBackwardsSolve();
+	ControlRig->ResetControlValues();
+	ControlRig->SetIsAdditive(bSetIsLayered);
+
+	ControlRig->Evaluate_AnyThread();
+
+	FString ObjectName = ControlRig->GetClass()->GetName(); //GetDisplayNameText().ToString();
+	ObjectName.RemoveFromEnd(TEXT("_C"));
+	
+	if (bSetIsLayered)
+	{
+		const FString AdditiveObjectName = ObjectName + TEXT(" (Layered)");
+		InTrack->SetTrackName(FName(*ObjectName));
+		InTrack->SetDisplayName(FText::FromString(AdditiveObjectName));
+		InTrack->SetColorTint(UMovieSceneControlRigParameterTrack::LayeredRigTrackColor);
+	}
+	else
+	{
+		InTrack->SetTrackName(FName(*ObjectName));
+		InTrack->SetDisplayName(FText::FromString(ObjectName));
+		InTrack->SetColorTint(UMovieSceneControlRigParameterTrack::AbsoluteRigTrackColor);
+	}
+
+	FControlRigEditMode* ControlRigEditMode = static_cast<FControlRigEditMode*>(GLevelEditorModeTools().GetActiveMode(FControlRigEditMode::ModeName));
+	if (ControlRigEditMode)
+	{
+		ControlRigEditMode->ZeroTransforms(false);
+	}
+
+	for (UMovieSceneSection* Section : InTrack->GetAllSections())
+	{
+		if (Section)
+		{
+			UMovieSceneControlRigParameterSection* CRSection = Cast<UMovieSceneControlRigParameterSection>(Section);
+			if (CRSection)
+			{
+				Section->Modify();
+				CRSection->ClearAllParameters();
+				CRSection->RecreateWithThisControlRig(CRSection->GetControlRig(), true);
+			}
+		}
+	}
+
+	TWeakPtr<ISequencer> WeakSequencer = GetSequencerFromAsset();
+	if (WeakSequencer.IsValid())
+	{
+		WeakSequencer.Pin()->NotifyMovieSceneDataChanged(EMovieSceneDataChangeType::MovieSceneStructureItemsChanged);
+	}
+
+	return true;
+}
+
 EControlRigFKRigExecuteMode UControlRigSequencerEditorLibrary::GetFKControlRigApplyMode(UControlRig* InControlRig)
 {
 	EControlRigFKRigExecuteMode ApplyMode = EControlRigFKRigExecuteMode::Direct;
diff --git a/Engine/Plugins/Animation/ControlRig/Source/ControlRigEditor/Private/Sequencer/ControlRigParameterTrackEditor.cpp b/Engine/Plugins/Animation/ControlRig/Source/ControlRigEditor/Private/Sequencer/ControlRigParameterTrackEditor.cpp
--- a/Engine/Plugins/Animation/ControlRig/Source/ControlRigEditor/Private/Sequencer/ControlRigParameterTrackEditor.cpp
+++ b/Engine/Plugins/Animation/ControlRig/Source/ControlRigEditor/Private/Sequencer/ControlRigParameterTrackEditor.cpp
@@ -1222,54 +1222,8 @@ void FControlRigParameterTrackEditor::ConvertIsLayered(UMovieSceneControlRigPara
 		return;
 	}
 
-	const FScopedTransaction Transaction(LOCTEXT("ConvertToLayeredControlRig_Transaction", "Convert to Layered Control Rig"));
-	Track->Modify();
-	ControlRig->Modify();
-
 	const bool bSetAdditive = !ControlRig->IsAdditive();
-
-	ControlRig->ClearPoseBeforeBackwardsSolve();
-	ControlRig->ResetControlValues();
-	ControlRig->SetIsAdditive(bSetAdditive);
-
-	ControlRig->Evaluate_AnyThread();
-
-	FString ObjectName = ControlRig->GetClass()->GetName(); //GetDisplayNameText().ToString();
-	ObjectName.RemoveFromEnd(TEXT("_C"));
-	
-	if (bSetAdditive)
-	{
-		const FString AdditiveObjectName = ObjectName + TEXT(" (Layered)");
-		Track->SetTrackName(FName(*ObjectName));
-		Track->SetDisplayName(FText::FromString(AdditiveObjectName));
-		Track->SetColorTint(UMovieSceneControlRigParameterTrack::LayeredRigTrackColor);
-	}
-	else
-	{
-		Track->SetTrackName(FName(*ObjectName));
-		Track->SetDisplayName(FText::FromString(ObjectName));
-		Track->SetColorTint(UMovieSceneControlRigParameterTrack::AbsoluteRigTrackColor);
-	}
-
-	if (FControlRigEditMode* EditMode = GetEditMode())
-	{
-		EditMode->ZeroTransforms(false);
-	}
-
-	for (UMovieSceneSection* Section : Track->GetAllSections())
-	{
-		if (Section)
-		{
-			UMovieSceneControlRigParameterSection* CRSection = Cast<UMovieSceneControlRigParameterSection>(Section);
-			if (CRSection)
-			{
-				Section->Modify();
-				CRSection->ClearAllParameters();
-				CRSection->RecreateWithThisControlRig(CRSection->GetControlRig(), true);
-			}
-		}
-	}
-	GetSequencer()->NotifyMovieSceneDataChanged(EMovieSceneDataChangeType::MovieSceneStructureItemsChanged);
+	UControlRigSequencerEditorLibrary::SetControlRigLayeredMode(Track, bSetAdditive);
 }
 
 void FControlRigParameterTrackEditor::BuildObjectBindingTrackMenu(FMenuBuilder& MenuBuilder, const TArray<FGuid>& ObjectBindings, const UClass* ObjectClass)
diff --git a/Engine/Plugins/Animation/ControlRig/Source/ControlRigEditor/Public/ControlRigSequencerEditorLibrary.h b/Engine/Plugins/Animation/ControlRig/Source/ControlRigEditor/Public/ControlRigSequencerEditorLibrary.h
--- a/Engine/Plugins/Animation/ControlRig/Source/ControlRigEditor/Public/ControlRigSequencerEditorLibrary.h
+++ b/Engine/Plugins/Animation/ControlRig/Source/ControlRigEditor/Public/ControlRigSequencerEditorLibrary.h
@@ -1063,6 +1063,15 @@ public:
 	UFUNCTION(BlueprintCallable, Category = "Editor Scripting | Sequencer Tools | Control Rig")
 	static bool IsLayeredControlRig(UControlRig* InControlRig);
 
+	/*
+	 * Convert the control rig track into absolute or layered rig
+	 *
+	 * @param InTrack Control rig track to convert 
+	 * @param bSetIsLayered Convert to layered rig if true, or absolute if false
+	 */
+	UFUNCTION(BlueprintCallable, Category = "Editor Scripting | Sequencer Tools | Control Rig")
+	static bool SetControlRigLayeredMode(UMovieSceneControlRigParameterTrack* InTrack, bool bSetIsLayered);
+
 	/**	Get FKControlRig Apply Mode.
 	@param InControlRig Rig to test
 	@return The EControlRigFKRigExecuteMode mode it is in, either Replace,Additive or Direct
