Control Rig: Hide already added control rig classes
#rnx

diff --git a/Engine/Plugins/Animation/ControlRig/Source/ControlRigEditor/Private/Sequencer/ControlRigParameterTrackEditor.cpp b/Engine/Plugins/Animation/ControlRig/Source/ControlRigEditor/Private/Sequencer/ControlRigParameterTrackEditor.cpp
--- a/Engine/Plugins/Animation/ControlRig/Source/ControlRigEditor/Private/Sequencer/ControlRigParameterTrackEditor.cpp
+++ b/Engine/Plugins/Animation/ControlRig/Source/ControlRigEditor/Private/Sequencer/ControlRigParameterTrackEditor.cpp
@@ -685,6 +685,34 @@ static bool ClassViewerSortPredicate(const FClassViewerSortElementInfo& A, const
 	}
 }
 
+/** Filter class does not allow classes that already exist in a skeletal mesh component. */
+class FClassViewerHideAlreadyAddedRigsFilter : public IClassViewerFilter
+{
+public:
+	FClassViewerHideAlreadyAddedRigsFilter(const TArray<UClass*> InExistingClasses)
+		: AlreadyAddedRigs(InExistingClasses)
+	{}
+
+	virtual bool IsClassAllowed(const FClassViewerInitializationOptions& InInitOptions, const UClass* InClass, TSharedRef< class FClassViewerFilterFuncs > InFilterFuncs ) override
+	{
+		return !AlreadyAddedRigs.Contains(InClass);
+	}
+	
+	virtual bool IsUnloadedClassAllowed(const FClassViewerInitializationOptions& InInitOptions,
+										const TSharedRef< const class IUnloadedBlueprintData > InUnloadedClassData,
+										TSharedRef< class FClassViewerFilterFuncs > InFilterFuncs) override
+	{
+		const FTopLevelAssetPath ClassPath = InUnloadedClassData.Get().GetClassPathName();
+		return !AlreadyAddedRigs.ContainsByPredicate([ClassPath](const UClass* Class)
+		{
+			return ClassPath == Class->GetClassPathName();
+		});
+	}
+	
+private:
+	TArray<UClass*> AlreadyAddedRigs;
+};
+
 void FControlRigParameterTrackEditor::BakeToControlRigSubMenu(FMenuBuilder& MenuBuilder, FGuid ObjectBinding, UObject* BoundObject, USkeletalMeshComponent* SkelMeshComp, USkeleton* Skeleton)
 {
 	const TSharedPtr<ISequencer> ParentSequencer = GetSequencer();
@@ -1129,11 +1157,12 @@ void FControlRigParameterTrackEditor::IterateTracksInMovieScene(UMovieScene& Mov
 	const TArray<FMovieSceneBinding>& Bindings = MovieScene.GetBindings();
 	for (const FMovieSceneBinding& Binding : Bindings)
 	{
-		if (UMovieSceneControlRigParameterTrack* Track = Cast<UMovieSceneControlRigParameterTrack>(MovieScene.FindTrack(UMovieSceneControlRigParameterTrack::StaticClass(), Binding.GetObjectGuid(), NAME_None)))
+		TArray<UMovieSceneTrack*> FoundTracks = MovieScene.FindTracks(UMovieSceneControlRigParameterTrack::StaticClass(), Binding.GetObjectGuid(), NAME_None);
+		for(UMovieSceneTrack* Track : FoundTracks)
 		{
-			if (Callback(Track))
+			if (UMovieSceneControlRigParameterTrack* CRTrack = Cast<UMovieSceneControlRigParameterTrack>(Track))
 			{
-				return;
+				Callback(CRTrack);
 			}
 		}
 	}
@@ -1141,13 +1170,10 @@ void FControlRigParameterTrackEditor::IterateTracksInMovieScene(UMovieScene& Mov
 	for (UMovieSceneTrack* Track : MovieScene.GetTracks())
 	{
 		if (UMovieSceneControlRigParameterTrack* CRTrack = Cast<UMovieSceneControlRigParameterTrack>(Track))
-		{			
-			if (Callback(CRTrack))
-			{
-				return;
-			}
+		{
+			Callback(CRTrack);
 		}
-	}	
+	}
 }
 
 void FControlRigParameterTrackEditor::BakeInvertedPose(UControlRig* InControlRig, UMovieSceneControlRigParameterTrack* Track)
@@ -1397,8 +1423,26 @@ void FControlRigParameterTrackEditor::HandleAddControlRigSubMenu(FMenuBuilder& M
 		TSharedPtr<FControlRigClassFilter> ClassFilter = MakeShareable(new FControlRigClassFilter(bFilterAssetBySkeleton, bFilterAssetByAnimatableControls, bCheckInversion, Skeleton));
 		Options.ClassFilters.Add(ClassFilter.ToSharedRef());
 		Options.bShowNoneOption = false;
-		Options.ExtraPickerCommonClasses.Add(UFKControlRig::StaticClass());
 		Options.ClassViewerSortPredicate = ClassViewerSortPredicate;
+		TArray<UClass*> ExistingRigs;
+		USkeletalMeshComponent* SkeletalMeshComponent = AcquireSkeletalMeshFromObject(BoundObject, ParentSequencer);
+		IterateTracks([&ExistingRigs, SkeletalMeshComponent](UMovieSceneControlRigParameterTrack* Track) -> bool
+		{
+			if (UControlRig* ControlRig = Track->GetControlRig())
+			{
+				if (ControlRig->GetObjectBinding()->GetBoundObject() == SkeletalMeshComponent)
+				{
+					ExistingRigs.Add(ControlRig->GetClass());
+				}
+			}
+			return true;
+		});
+		TSharedPtr<FClassViewerHideAlreadyAddedRigsFilter> ExistingClassesFilter = MakeShareable(new FClassViewerHideAlreadyAddedRigsFilter(ExistingRigs));
+		Options.ClassFilters.Add(ExistingClassesFilter.ToSharedRef());
+		if (!ExistingRigs.Contains(UFKControlRig::StaticClass()))
+		{
+			Options.ExtraPickerCommonClasses.Add(UFKControlRig::StaticClass());
+		}
 
 		UMovieSceneSequence* Sequence = GetSequencer() ? GetSequencer()->GetFocusedMovieSceneSequence() : nullptr;
 		Options.AdditionalReferencingAssets.Add(FAssetData(Sequence));
