Control Rig: Fix bug when adding modules as child of other modules
#rnx

diff --git a/Engine/Plugins/Animation/ControlRig/Source/ControlRigDeveloper/Private/ModularRigController.cpp b/Engine/Plugins/Animation/ControlRig/Source/ControlRigDeveloper/Private/ModularRigController.cpp
--- a/Engine/Plugins/Animation/ControlRig/Source/ControlRigDeveloper/Private/ModularRigController.cpp
+++ b/Engine/Plugins/Animation/ControlRig/Source/ControlRigDeveloper/Private/ModularRigController.cpp
@@ -101,48 +101,31 @@ bool UModularRigController::AddModule(const FName& InModuleName, TSubclassOf<UCo
 
 FRigModuleReference* UModularRigController::FindModule(const FString& InPath)
 {
+	TArray<FRigModuleReference*>* Children = &Model->RootModules;
 	FString Left = InPath, Right;
-	TArray<FString> Tokens;
 	while (Left.Split(UModularRig::NamespaceSeparator, &Left, &Right))
 	{
-		Tokens.Add(Left);
+		FRigModuleReference** Cur = Children->FindByPredicate([Left](FRigModuleReference* Module)
+		{
+			return Module->Name == Left;
+		});
+		if (!Cur)
+		{
+			return nullptr;
+		}
+		Children = &(*Cur)->CachedChildren;
 		Left = Right;
 	}
-	Tokens.Add(Left);
 
-	FRigModuleReference* CurModule = nullptr;
-	for (FRigModuleReference* Module : Model->RootModules)
-	{
-		if (Module->Name.ToString() == Tokens[0])
+	FRigModuleReference** Cur = Children->FindByPredicate([Left](FRigModuleReference* Module)
 		{
-			CurModule = Module;
-			break;
-		}
-	}
-	if (!CurModule)
+			return Module->Name == Left;
+		});
+	if (!Cur)
 	{
 		return nullptr;
 	}
-
-	for (int32 Index=1; Index<Tokens.Num(); ++Index)
-	{
-		bool bFound = false;
-		for (FRigModuleReference* Child : CurModule->CachedChildren)
-		{
-			if (Child->Name.ToString() == Tokens[Index])
-			{
-				CurModule = Child;
-				bFound = true;
-				break;
-			}
-		}
-		if (!bFound)
-		{
-			return nullptr;
-		}
-	}
-
-	return CurModule;
+	return *Cur;
 }
 
 bool UModularRigController::ConnectModuleToElement(const FRigElementKey& InConnectorKey, const FRigElementKey& InTargetKey)
diff --git a/Engine/Plugins/Animation/ControlRig/Source/ControlRigDeveloper/Private/ModularRigModel.cpp b/Engine/Plugins/Animation/ControlRig/Source/ControlRigDeveloper/Private/ModularRigModel.cpp
--- a/Engine/Plugins/Animation/ControlRig/Source/ControlRigDeveloper/Private/ModularRigModel.cpp
+++ b/Engine/Plugins/Animation/ControlRig/Source/ControlRigDeveloper/Private/ModularRigModel.cpp
@@ -11,7 +11,7 @@ FString FRigModuleReference::GetNamespace() const
 	{
 		return Name.ToString() + UModularRig::NamespaceSeparator;
 	}
-	return ParentNamespace + UModularRig::NamespaceSeparator + Name.ToString() + UModularRig::NamespaceSeparator;
+	return ParentNamespace + Name.ToString() + UModularRig::NamespaceSeparator;
 }
 
 UModularRigController* FModularRigModel::GetController(bool bCreateIfNeeded)
