Control Rig: Fix crash when spawned control is selected in details panel and we hit compile
#rnx

diff --git a/Engine/Plugins/Animation/ControlRig/Source/ControlRig/Private/Rigs/RigHierarchy.cpp b/Engine/Plugins/Animation/ControlRig/Source/ControlRig/Private/Rigs/RigHierarchy.cpp
--- a/Engine/Plugins/Animation/ControlRig/Source/ControlRig/Private/Rigs/RigHierarchy.cpp
+++ b/Engine/Plugins/Animation/ControlRig/Source/ControlRig/Private/Rigs/RigHierarchy.cpp
@@ -107,6 +107,7 @@ URigHierarchy::URigHierarchy()
 , bIsInteracting(false)
 , LastInteractedKey()
 , bSuspendNotifications(false)
+, bSuspendMetadataNotifications(false)
 , HierarchyController(nullptr)
 , bIsControllerAvailable(true)
 , ResetPoseHash(INDEX_NONE)
@@ -498,125 +499,133 @@ void URigHierarchy::CopyHierarchy(URigHierarchy* InHierarchy)
 		}
 	}
 
-	Reset_Impl(bReallocateElements);
+	{
+		TGuardValue<bool> SuspendMetadataNotifications(bSuspendMetadataNotifications, true);
+		Reset_Impl(bReallocateElements);
 
-	static const TArray<int32> StructureSizePerType = {
-		sizeof(FRigBoneElement),
-		sizeof(FRigNullElement),
-		sizeof(FRigControlElement),
-		sizeof(FRigCurveElement),
-		sizeof(FRigRigidBodyElement),
-		sizeof(FRigReferenceElement),
-		sizeof(FRigConnectorElement),
-		sizeof(FRigSocketElement),
-	}; 
+		static const TArray<int32> StructureSizePerType = {
+			sizeof(FRigBoneElement),
+			sizeof(FRigNullElement),
+			sizeof(FRigControlElement),
+			sizeof(FRigCurveElement),
+			sizeof(FRigRigidBodyElement),
+			sizeof(FRigReferenceElement),
+			sizeof(FRigConnectorElement),
+			sizeof(FRigSocketElement),
+		}; 
 
-	if(bReallocateElements)
-	{
-		// Allocate the elements in batches to improve performance
-		TArray<uint8*> NewElementsPerType;
-		for(int32 ElementTypeIndex = 0; ElementTypeIndex < InHierarchy->ElementsPerType.Num(); ElementTypeIndex++)
+		if(bReallocateElements)
 		{
-			const ERigElementType ElementType = FlatIndexToRigElementType(ElementTypeIndex);
-			int32 StructureSize = 0;
-
-			const int32 Count = InHierarchy->ElementsPerType[ElementTypeIndex].Num();
-			if(Count)
-			{
-				FRigBaseElement* ElementMemory = MakeElement(ElementType, Count, &StructureSize);
-				verify(StructureSize == StructureSizePerType[ElementTypeIndex]);
-				NewElementsPerType.Add(reinterpret_cast<uint8*>(ElementMemory));
-			}
-			else
+			// Allocate the elements in batches to improve performance
+			TArray<uint8*> NewElementsPerType;
+			for(int32 ElementTypeIndex = 0; ElementTypeIndex < InHierarchy->ElementsPerType.Num(); ElementTypeIndex++)
 			{
-				NewElementsPerType.Add(nullptr);
-			}
+				const ERigElementType ElementType = FlatIndexToRigElementType(ElementTypeIndex);
+				int32 StructureSize = 0;
+
+				const int32 Count = InHierarchy->ElementsPerType[ElementTypeIndex].Num();
+				if(Count)
+				{
+					FRigBaseElement* ElementMemory = MakeElement(ElementType, Count, &StructureSize);
+					verify(StructureSize == StructureSizePerType[ElementTypeIndex]);
+					NewElementsPerType.Add(reinterpret_cast<uint8*>(ElementMemory));
+				}
+				else
+				{
+					NewElementsPerType.Add(nullptr);
+				}
 			
-			ElementsPerType[ElementTypeIndex].Reserve(Count);
-		}
+				ElementsPerType[ElementTypeIndex].Reserve(Count);
+			}
 
-		Elements.Reserve(InHierarchy->Elements.Num());
-		IndexLookup.Reserve(InHierarchy->IndexLookup.Num());
+			Elements.Reserve(InHierarchy->Elements.Num());
+			IndexLookup.Reserve(InHierarchy->IndexLookup.Num());
 
-		for(int32 Index = 0; Index < InHierarchy->Num(); Index++)
-		{
-			const FRigBaseElement* Source = InHierarchy->Get(Index);
-			const FRigElementKey& Key = Source->Key;
+			for(int32 Index = 0; Index < InHierarchy->Num(); Index++)
+			{
+				const FRigBaseElement* Source = InHierarchy->Get(Index);
+				const FRigElementKey& Key = Source->Key;
 
-			const int32 ElementTypeIndex = RigElementTypeToFlatIndex(Key.Type);
+				const int32 ElementTypeIndex = RigElementTypeToFlatIndex(Key.Type);
 		
-			const int32 SubIndex = Num(Key.Type);
+				const int32 SubIndex = Num(Key.Type);
 
-			const int32 StructureSize = StructureSizePerType[ElementTypeIndex];
-			check(NewElementsPerType[ElementTypeIndex] != nullptr);
-			FRigBaseElement* Target = reinterpret_cast<FRigBaseElement*>(&NewElementsPerType[ElementTypeIndex][StructureSize * SubIndex]);
+				const int32 StructureSize = StructureSizePerType[ElementTypeIndex];
+				check(NewElementsPerType[ElementTypeIndex] != nullptr);
+				FRigBaseElement* Target = reinterpret_cast<FRigBaseElement*>(&NewElementsPerType[ElementTypeIndex][StructureSize * SubIndex]);
 
-			Target->InitializeFrom(Source);
+				Target->InitializeFrom(Source);
 			
-			Target->SubIndex = SubIndex;
-			Target->Index = Elements.Add(Target);
+				Target->SubIndex = SubIndex;
+				Target->Index = Elements.Add(Target);
 
-			ElementsPerType[ElementTypeIndex].Add(Target);
-			IndexLookup.Add(Key, Target->Index);
+				ElementsPerType[ElementTypeIndex].Add(Target);
+				IndexLookup.Add(Key, Target->Index);
 			
-			IncrementPoseVersion(Index);
+				IncrementPoseVersion(Index);
 
-			check(Source->Index == Index);
-			check(Target->Index == Index);
+				check(Source->Index == Index);
+				check(Target->Index == Index);
+			}
 		}
-	}
-	else
-	{
-		// remove the superfluous elements
-		for(int32 ElementIndex = Elements.Num() - 1; ElementIndex >= InHierarchy->Elements.Num(); ElementIndex--)
+		else
 		{
-			DestroyElement(Elements[ElementIndex]);
-		}
+			// remove the superfluous elements
+			for(int32 ElementIndex = Elements.Num() - 1; ElementIndex >= InHierarchy->Elements.Num(); ElementIndex--)
+			{
+				DestroyElement(Elements[ElementIndex]);
+			}
 
-		// shrink the containers accordingly
-		Elements.SetNum(InHierarchy->Elements.Num());
-		const UEnum* ElementTypeEnum = StaticEnum<ERigElementType>();
-		for(int32 ElementTypeIndex = 0; ; ElementTypeIndex++)
-		{
-			if ((ERigElementType)ElementTypeEnum->GetValueByIndex(ElementTypeIndex) == ERigElementType::All)
+			// shrink the containers accordingly
+			Elements.SetNum(InHierarchy->Elements.Num());
+			const UEnum* ElementTypeEnum = StaticEnum<ERigElementType>();
+			for(int32 ElementTypeIndex = 0; ; ElementTypeIndex++)
 			{
-				break;
+				if ((ERigElementType)ElementTypeEnum->GetValueByIndex(ElementTypeIndex) == ERigElementType::All)
+				{
+					break;
+				}
+				ElementsPerType[ElementTypeIndex].SetNum(InHierarchy->ElementsPerType[ElementTypeIndex].Num());
+			}
+
+			for(int32 Index = 0; Index < InHierarchy->Num(); Index++)
+			{
+				const FRigBaseElement* Source = InHierarchy->Get(Index);
+				FRigBaseElement* Target = Elements[Index];
+
+				check(Target->Key.Type == Source->Key.Type);
+				Target->InitializeFrom(Source);
+
+				IncrementPoseVersion(Index);
 			}
-			ElementsPerType[ElementTypeIndex].SetNum(InHierarchy->ElementsPerType[ElementTypeIndex].Num());
+
+			IndexLookup = InHierarchy->IndexLookup;
 		}
 
+		// Copy all the element subclass data and all elements' metadata over.
 		for(int32 Index = 0; Index < InHierarchy->Num(); Index++)
 		{
 			const FRigBaseElement* Source = InHierarchy->Get(Index);
 			FRigBaseElement* Target = Elements[Index];
 
-			check(Target->Key.Type == Source->Key.Type);
-			Target->InitializeFrom(Source);
+			Target->CopyFrom(Source);
 
-			IncrementPoseVersion(Index);
+			CopyAllMetadataFromElement(Target, Source);
 		}
 
-		IndexLookup = InHierarchy->IndexLookup;
+		PreviousNameMap.Append(InHierarchy->PreviousNameMap);
+
+		// Increment the topology version to invalidate our cached children.
+		IncrementTopologyVersion();
+	
+		MetadataVersion = InHierarchy->GetMetadataVersion();
 	}
 
-	// Copy all the element subclass data and all elements' metadata over.
-	for(int32 Index = 0; Index < InHierarchy->Num(); Index++)
+	if (MetadataChangedDelegate.IsBound())
 	{
-		const FRigBaseElement* Source = InHierarchy->Get(Index);
-		FRigBaseElement* Target = Elements[Index];
-
-		Target->CopyFrom(Source);
-
-		CopyAllMetadataFromElement(Target, Source);
+		MetadataChangedDelegate.Broadcast(FRigElementKey(ERigElementType::All), NAME_None);
 	}
 
-	PreviousNameMap.Append(InHierarchy->PreviousNameMap);
-
-	// Increment the topology version to invalidate our cached children.
-	IncrementTopologyVersion();
-	
-	MetadataVersion = InHierarchy->GetMetadataVersion();
-
 	EnsureCacheValidity();
 }
 
@@ -5332,20 +5341,26 @@ void URigHierarchy::FMetadataStorage::Serialize(FArchive& Ar)
 void URigHierarchy::OnMetadataChanged(const FRigElementKey& InKey, const FName& InName)
 {
 	MetadataVersion++;
-	
-	if(MetadataChangedDelegate.IsBound())
+
+	if (!bSuspendMetadataNotifications)
 	{
-		MetadataChangedDelegate.Broadcast(InKey, InName);
+		if(MetadataChangedDelegate.IsBound())
+		{
+			MetadataChangedDelegate.Broadcast(InKey, InName);
+		}
 	}
 }
 
 void URigHierarchy::OnMetadataTagChanged(const FRigElementKey& InKey, const FName& InTag, bool bAdded)
 {
 	MetadataTagVersion++;
-	
-	if(MetadataTagChangedDelegate.IsBound())
+
+	if (!bSuspendMetadataNotifications)
 	{
-		MetadataTagChangedDelegate.Broadcast(InKey, InTag, bAdded);
+		if(MetadataTagChangedDelegate.IsBound())
+		{
+			MetadataTagChangedDelegate.Broadcast(InKey, InTag, bAdded);
+		}
 	}
 }
 
diff --git a/Engine/Plugins/Animation/ControlRig/Source/ControlRig/Public/Rigs/RigHierarchy.h b/Engine/Plugins/Animation/ControlRig/Source/ControlRig/Public/Rigs/RigHierarchy.h
--- a/Engine/Plugins/Animation/ControlRig/Source/ControlRig/Public/Rigs/RigHierarchy.h
+++ b/Engine/Plugins/Animation/ControlRig/Source/ControlRig/Public/Rigs/RigHierarchy.h
@@ -4237,6 +4237,11 @@ private:
 	 */
 	bool bSuspendNotifications;
 
+	/** 
+	 * If set to true all metadata changes notifs coming from this hierarchy will be suspended
+	 */
+	bool bSuspendMetadataNotifications = false;
+
 	/**
 	 * The event fired during undo / redo
 	 */
diff --git a/Engine/Plugins/Animation/ControlRig/Source/ControlRigEditor/Private/ControlRigElementDetails.cpp b/Engine/Plugins/Animation/ControlRig/Source/ControlRigEditor/Private/ControlRigElementDetails.cpp
--- a/Engine/Plugins/Animation/ControlRig/Source/ControlRigEditor/Private/ControlRigElementDetails.cpp
+++ b/Engine/Plugins/Animation/ControlRig/Source/ControlRigEditor/Private/ControlRigElementDetails.cpp
@@ -1236,7 +1236,7 @@ void FRigBaseElementDetails::CustomizeMetadata(IDetailLayoutBuilder& DetailBuild
 		MetadataHandle = Hierarchy->OnMetadataChanged().AddLambda([this, PropertyUtilities](const FRigElementKey& InKey, const FName&)
 		{
 			const FRigBaseElement* Element = PerElementInfos.Num() == 1 ? PerElementInfos[0].GetElement() : nullptr;
-			if (Element && Element->GetKey() == InKey)
+			if (InKey.Type == ERigElementType::All || (Element && Element->GetKey() == InKey))
 			{
 				PropertyUtilities->ForceRefresh();
 			}
