Control RIg: Hide the Convert To Layered option in sequencer behind the ControlRig.Sequencer.EnableAdditiveControlRigs CVar
#rnx

diff --git a/Engine/Plugins/Animation/ControlRig/Source/ControlRigEditor/Private/Sequencer/ControlRigParameterTrackEditor.cpp b/Engine/Plugins/Animation/ControlRig/Source/ControlRigEditor/Private/Sequencer/ControlRigParameterTrackEditor.cpp
--- a/Engine/Plugins/Animation/ControlRig/Source/ControlRigEditor/Private/Sequencer/ControlRigParameterTrackEditor.cpp
+++ b/Engine/Plugins/Animation/ControlRig/Source/ControlRigEditor/Private/Sequencer/ControlRigParameterTrackEditor.cpp
@@ -3904,17 +3904,20 @@ void FControlRigParameterTrackEditor::BuildTrackContextMenu(FMenuBuilder& MenuBu
 			LOCTEXT("Order", "Order")
 		);
 
-		MenuBuilder.AddMenuEntry(
-				LOCTEXT("ConvertIsLayeredControlRig", "Convert To Layered"),
-				LOCTEXT("ConvertIsLayeredControlRigToolTip", "Converts the Control Rig from an Absolute rig to a Layered rig"),
-				FSlateIcon(),
-				FUIAction(
-					FExecuteAction::CreateRaw(this, &FControlRigParameterTrackEditor::ConvertIsLayered, Track),
-					FCanExecuteAction(),
-					FIsActionChecked::CreateUObject(Track->GetControlRig(), &UControlRig::IsAdditive)
-				),
-				NAME_None,
-				EUserInterfaceActionType::ToggleButton);
+		if (CVarEnableAdditiveControlRigs->GetBool())
+		{
+			MenuBuilder.AddMenuEntry(
+				   LOCTEXT("ConvertIsLayeredControlRig", "Convert To Layered"),
+				   LOCTEXT("ConvertIsLayeredControlRigToolTip", "Converts the Control Rig from an Absolute rig to a Layered rig"),
+				   FSlateIcon(),
+				   FUIAction(
+					   FExecuteAction::CreateRaw(this, &FControlRigParameterTrackEditor::ConvertIsLayered, Track),
+					   FCanExecuteAction(),
+					   FIsActionChecked::CreateUObject(Track->GetControlRig(), &UControlRig::IsAdditive)
+				   ),
+				   NAME_None,
+				   EUserInterfaceActionType::ToggleButton);
+		}
 	}
 	MenuBuilder.EndSection();
 
