Sequencer:  Tween tool not working on objects since Sequencer isn't present.

diff --git a/Engine/Plugins/Animation/ControlRig/Source/ControlRigEditor/Private/EditMode/SControlRigTweenWidget.cpp b/Engine/Plugins/Animation/ControlRig/Source/ControlRigEditor/Private/EditMode/SControlRigTweenWidget.cpp
--- a/Engine/Plugins/Animation/ControlRig/Source/ControlRigEditor/Private/EditMode/SControlRigTweenWidget.cpp
+++ b/Engine/Plugins/Animation/ControlRig/Source/ControlRigEditor/Private/EditMode/SControlRigTweenWidget.cpp
@@ -29,9 +29,13 @@ void SControlRigTweenSlider::Construct(const FArguments& InArgs)
 	bIsBlending = false;
 	bSliderStartedTransaction = false;
 	AnimSlider = InArgs._InAnimSlider;
-	WeakSequencer = InArgs._InSequencer;
 	WeakEditMode = InArgs._InWeakEditMode;
 	
+	ULevelSequence* LevelSequence = ULevelSequenceEditorBlueprintLibrary::GetCurrentLevelSequence();
+	IAssetEditorInstance* AssetEditor = GEditor->GetEditorSubsystem<UAssetEditorSubsystem>()->FindEditorForAsset(LevelSequence, false);
+	ILevelSequenceEditorToolkit* LevelSequenceEditor = static_cast<ILevelSequenceEditorToolkit*>(AssetEditor);
+	WeakSequencer = LevelSequenceEditor ? LevelSequenceEditor->GetSequencer() : nullptr;
+	
 	ChildSlot
 	[
 		SNew(SVerticalBox)
@@ -330,7 +334,6 @@ void SControlRigTweenWidget::Construct(const FArguments& InArgs)
 		[
 			SAssignNew(SliderWidget,SControlRigTweenSlider)
 			.InAnimSlider(SliderPtr)
-			.InSequencer(OwningEditMode.Pin()->GetWeakSequencer().Pin())
 			.InWeakEditMode(OwningEditMode)
 		];
 	
diff --git a/Engine/Plugins/Animation/ControlRig/Source/ControlRigEditor/Private/EditMode/SControlRigTweenWidget.h b/Engine/Plugins/Animation/ControlRig/Source/ControlRigEditor/Private/EditMode/SControlRigTweenWidget.h
--- a/Engine/Plugins/Animation/ControlRig/Source/ControlRigEditor/Private/EditMode/SControlRigTweenWidget.h
+++ b/Engine/Plugins/Animation/ControlRig/Source/ControlRigEditor/Private/EditMode/SControlRigTweenWidget.h
@@ -25,7 +25,6 @@ class SControlRigTweenSlider : public SCompoundWidget
 {
 	SLATE_BEGIN_ARGS(SControlRigTweenSlider) {}
 	SLATE_ARGUMENT(TSharedPtr<FBaseAnimSlider>, InAnimSlider)
-	SLATE_ARGUMENT(TSharedPtr<ISequencer>, InSequencer)
 	SLATE_ARGUMENT(TWeakPtr<FControlRigEditMode>, InWeakEditMode)
 	SLATE_END_ARGS()
 	~SControlRigTweenSlider()
