Fix BP compile errors from parameter rename

diff --git a/Engine/Plugins/Animation/ControlRig/Config/BaseControlRig.ini b/Engine/Plugins/Animation/ControlRig/Config/BaseControlRig.ini
--- a/Engine/Plugins/Animation/ControlRig/Config/BaseControlRig.ini
+++ b/Engine/Plugins/Animation/ControlRig/Config/BaseControlRig.ini
@@ -439,3 +439,5 @@
 +StructRedirects=(OldName="/Script/ControlRig.ControlRigPublicFunctionData",NewName="/Script/RigVMDeveloper.RigVMOldPublicFunctionData")
 +ClassRedirects=(OldName="/Script/ControlRig.NameSpacedUserData", NewName="/Script/RigVM.NameSpacedUserData")
 +ClassRedirects=(OldName="/Script/ControlRig.DataAssetLink", NewName="/Script/RigVM.DataAssetLink")
+
++K2ParamRedirects=(NodeName="/Script/BlueprintGraph.K2Node_CallFunction", OldParamName="ControlRigSequencerEditorLibrary.FindOrCreateControlRigTrack.bIsAdditiveControlRig", NewParamName="bIsLayeredControlRig")
\ No newline at end of file
