Second attempt to redirect function param

diff --git a/Engine/Plugins/Animation/ControlRig/Config/BaseControlRig.ini b/Engine/Plugins/Animation/ControlRig/Config/BaseControlRig.ini
--- a/Engine/Plugins/Animation/ControlRig/Config/BaseControlRig.ini
+++ b/Engine/Plugins/Animation/ControlRig/Config/BaseControlRig.ini
@@ -440,4 +440,5 @@
 +ClassRedirects=(OldName="/Script/ControlRig.NameSpacedUserData", NewName="/Script/RigVM.NameSpacedUserData")
 +ClassRedirects=(OldName="/Script/ControlRig.DataAssetLink", NewName="/Script/RigVM.DataAssetLink")
 
-+K2ParamRedirects=(NodeName="/Script/BlueprintGraph.K2Node_CallFunction", OldParamName="ControlRigSequencerEditorLibrary.FindOrCreateControlRigTrack.bIsAdditiveControlRig", NewParamName="bIsLayeredControlRig")
\ No newline at end of file
+[/Script/Engine.Engine]
++K2ParamRedirects=(NodeName="/Script/BlueprintGraph.K2Node_CallFunction", OldParamName="ControlRigSequencerEditorLibrary.FindOrCreateControlRigTrack.bIsAdditiveControlRig", NewParamName="bIsLayeredControlRig")
