Control Rig: Allow undo/redo of adding and connecting modules
#rnx

diff --git a/Engine/Plugins/Animation/ControlRig/Source/ControlRigDeveloper/Private/ControlRigBlueprint.cpp b/Engine/Plugins/Animation/ControlRig/Source/ControlRigDeveloper/Private/ControlRigBlueprint.cpp
--- a/Engine/Plugins/Animation/ControlRig/Source/ControlRigDeveloper/Private/ControlRigBlueprint.cpp
+++ b/Engine/Plugins/Animation/ControlRig/Source/ControlRigDeveloper/Private/ControlRigBlueprint.cpp
@@ -972,7 +972,7 @@ bool UControlRigBlueprint::ResolveConnector(const FRigElementKey& DraggedKey, co
 		// Add connection to the model
 		if (UModularRigController* Controller = GetModularRigController())
 		{
-			Controller->ConnectModuleToElement(DraggedKey, TargetKey);
+			Controller->ConnectModuleToElement(DraggedKey, TargetKey, bSetupUndoRedo);
 		}
 	}
 	else
@@ -1292,6 +1292,15 @@ void UControlRigBlueprint::PostTransacted(const FTransactionObjectEvent& Transac
 			(void)MarkPackageDirty();
 		}
 
+		if (PropertiesChanged.Contains(GET_MEMBER_NAME_CHECKED(UControlRigBlueprint, ModularRigModel)))
+		{
+			if (IsModularRig())
+			{
+				ModularRigModel.UpdateCachedChildren();
+				RecompileModularRig();
+			}
+		}
+		
 		if (PropertiesChanged.Contains(GET_MEMBER_NAME_CHECKED(UControlRigBlueprint, DrawContainer)))
 		{
 			PropagateDrawInstructionsFromBPToInstances();
diff --git a/Engine/Plugins/Animation/ControlRig/Source/ControlRigDeveloper/Private/ModularRigController.cpp b/Engine/Plugins/Animation/ControlRig/Source/ControlRigDeveloper/Private/ModularRigController.cpp
--- a/Engine/Plugins/Animation/ControlRig/Source/ControlRigDeveloper/Private/ModularRigController.cpp
+++ b/Engine/Plugins/Animation/ControlRig/Source/ControlRigDeveloper/Private/ModularRigController.cpp
@@ -10,18 +10,35 @@
 
 #include UE_INLINE_GENERATED_CPP_BY_NAME(ModularRigController)
 
+#if WITH_EDITOR
+#include "ScopedTransaction.h"
+#endif
+
 UModularRigController::UModularRigController(const FObjectInitializer& ObjectInitializer)
 	: Super(ObjectInitializer)
 {
 }
 
-bool UModularRigController::AddModule(const FName& InModuleName, TSubclassOf<UControlRig> InClass, const FString& InParentModulePath)
+bool UModularRigController::AddModule(const FName& InModuleName, TSubclassOf<UControlRig> InClass, const FString& InParentModulePath, bool bSetupUndo)
 {
 	if (!InClass->GetDefaultObject<UControlRig>()->IsRigModule())
 	{
 		return false;
 	}
 
+#if WITH_EDITOR
+	TSharedPtr<FScopedTransaction> TransactionPtr;
+	if (bSetupUndo)
+	{
+		TransactionPtr = MakeShared<FScopedTransaction>(NSLOCTEXT("ModularRigController", "AddModuleTransaction", "Add Module"));
+		if(UControlRigBlueprint* Blueprint = Cast<UControlRigBlueprint>(GetOuter()))
+		{
+			Blueprint->Modify();
+			Blueprint->Hierarchy->Modify();
+		}
+	}
+#endif 
+
 	FRigModuleReference* NewModule = nullptr;
 	if (InParentModulePath.IsEmpty())
 	{
@@ -95,6 +112,10 @@ bool UModularRigController::AddModule(const FName& InModuleName, TSubclassOf<UCo
 	}
 
 	Notify(EModularRigNotification::ModuleAdded, NewModule);
+
+#if WITH_EDITOR
+	TransactionPtr.Reset();
+#endif
 	
 	return true;
 }
@@ -128,7 +149,7 @@ FRigModuleReference* UModularRigController::FindModule(const FString& InPath)
 	return *Cur;
 }
 
-bool UModularRigController::ConnectModuleToElement(const FRigElementKey& InConnectorKey, const FRigElementKey& InTargetKey)
+bool UModularRigController::ConnectModuleToElement(const FRigElementKey& InConnectorKey, const FRigElementKey& InTargetKey, bool bSetupUndo)
 {
 	FString ConnectorNameSpace, ConnectorName;
 	if (!InConnectorKey.Name.ToString().Split(UModularRig::NamespaceSeparator, &ConnectorNameSpace, &ConnectorName, ESearchCase::CaseSensitive, ESearchDir::FromEnd))
@@ -154,28 +175,42 @@ bool UModularRigController::ConnectModuleToElement(const FRigElementKey& InConne
 		return false;
 	}
 
+#if WITH_EDITOR
+	TSharedPtr<FScopedTransaction> TransactionPtr;
+	if (bSetupUndo)
+	{
+		TransactionPtr = MakeShared<FScopedTransaction>(NSLOCTEXT("ModularRigController", "ConnectModuleToElementTransaction", "Connect to Element"));
+		Blueprint->Modify();
+		Blueprint->Hierarchy->Modify();
+	}
+#endif 
+
 	const FRigElementKey ConnectorKey(*ConnectorName, ERigElementType::Connector);
 	FRigElementKey& TargetKey = Module->Connections.FindOrAdd(ConnectorKey);
 	TargetKey = InTargetKey;
 
 	Notify(EModularRigNotification::ConnectionChanged, Module);
+
+#if WITH_EDITOR
+	TransactionPtr.Reset();
+#endif
 	
 	return true;
 }
 
-bool UModularRigController::RemoveModule(const FString& InModulesPath)
+bool UModularRigController::RemoveModule(const FString& InModulesPath, bool bSetupUndo)
 {
 	// todo: UE-199050
 	return false;
 }
 
-bool UModularRigController::RenameModule(const FString& InModulesPath, const FName& InNewName)
+bool UModularRigController::RenameModule(const FString& InModulesPath, const FName& InNewName, bool bSetupUndo)
 {
 	// todo: UE-199050
 	return false;
 }
 
-bool UModularRigController::ReparentModule(const FString& InModulesPath, const FString& InNewParentModulePath)
+bool UModularRigController::ReparentModule(const FString& InModulesPath, const FString& InNewParentModulePath, bool bSetupUndo)
 {
 	// todo: UE-199050
 	return false;
diff --git a/Engine/Plugins/Animation/ControlRig/Source/ControlRigDeveloper/Private/ModularRigModel.cpp b/Engine/Plugins/Animation/ControlRig/Source/ControlRigDeveloper/Private/ModularRigModel.cpp
--- a/Engine/Plugins/Animation/ControlRig/Source/ControlRigDeveloper/Private/ModularRigModel.cpp
+++ b/Engine/Plugins/Animation/ControlRig/Source/ControlRigDeveloper/Private/ModularRigModel.cpp
@@ -43,6 +43,8 @@ void FModularRigModel::UpdateCachedChildren()
 		Module.CachedChildren.Reset();
 		PathToModule.Add(Module.GetNamespace(), &Module);
 	}
+	
+	RootModules.Reset();
 	for (FRigModuleReference& Module : Modules)
 	{
 		if (Module.ParentNamespace.IsEmpty())
diff --git a/Engine/Plugins/Animation/ControlRig/Source/ControlRigDeveloper/Public/ModularRigController.h b/Engine/Plugins/Animation/ControlRig/Source/ControlRigDeveloper/Public/ModularRigController.h
--- a/Engine/Plugins/Animation/ControlRig/Source/ControlRigDeveloper/Public/ModularRigController.h
+++ b/Engine/Plugins/Animation/ControlRig/Source/ControlRigDeveloper/Public/ModularRigController.h
@@ -41,19 +41,19 @@ class CONTROLRIGDEVELOPER_API UModularRigController : public UObject
 
 
 	UFUNCTION(BlueprintCallable, Category = "Control Rig | Modules")
-	bool AddModule(const FName& InModuleName, TSubclassOf<UControlRig> InClass, const FString& InParentModulePath);
+	bool AddModule(const FName& InModuleName, TSubclassOf<UControlRig> InClass, const FString& InParentModulePath, bool bSetupUndo = true);
 
 	UFUNCTION(BlueprintCallable, Category = "Control Rig | Modules")
-	bool ConnectModuleToElement(const FRigElementKey& InConnectorKey, const FRigElementKey& InTargetKey);
+	bool ConnectModuleToElement(const FRigElementKey& InConnectorKey, const FRigElementKey& InTargetKey, bool bSetupUndo = true);
 
 	UFUNCTION(BlueprintCallable, Category = "ControlRig | Modules")
-	bool RemoveModule(const FString& InModulesPath);
+	bool RemoveModule(const FString& InModulesPath, bool bSetupUndo = true);
 
 	UFUNCTION(BlueprintCallable, Category = "ControlRig | Modules")
-	bool RenameModule(const FString& InModulesPath, const FName& InNewName);
+	bool RenameModule(const FString& InModulesPath, const FName& InNewName, bool bSetupUndo = true);
 
 	UFUNCTION(BlueprintCallable, Category = "ControlRig | Modules")
-	bool ReparentModule(const FString& InModulesPath, const FString& InNewParentModulePath);
+	bool ReparentModule(const FString& InModulesPath, const FString& InNewParentModulePath, bool bSetupUndo = true);
 
 
 	void SetModel(FModularRigModel* InModel) { Model = InModel; }
diff --git a/Engine/Plugins/Animation/ControlRig/Source/ControlRigEditor/Private/Editor/ControlRigEditor.cpp b/Engine/Plugins/Animation/ControlRig/Source/ControlRigEditor/Private/Editor/ControlRigEditor.cpp
--- a/Engine/Plugins/Animation/ControlRig/Source/ControlRigEditor/Private/Editor/ControlRigEditor.cpp
+++ b/Engine/Plugins/Animation/ControlRig/Source/ControlRigEditor/Private/Editor/ControlRigEditor.cpp
@@ -3104,7 +3104,7 @@ void FControlRigEditor::OnHierarchyModified(ERigHierarchyNotification InNotif, U
 	{
 		case ERigHierarchyNotification::ElementAdded:
 		{
-			if (GetControlRig()->IsA<UControlRig>())
+			if (!RigBlueprint->IsModularRig())
 			{
 				if(InElement->GetType() == ERigElementType::Connector)
 				{
