Control Rig: Save configuration values of modules
#rnx

diff --git a/Engine/Plugins/Animation/ControlRig/Source/ControlRig/Private/ModularRig.cpp b/Engine/Plugins/Animation/ControlRig/Source/ControlRig/Private/ModularRig.cpp
--- a/Engine/Plugins/Animation/ControlRig/Source/ControlRig/Private/ModularRig.cpp
+++ b/Engine/Plugins/Animation/ControlRig/Source/ControlRig/Private/ModularRig.cpp
@@ -134,13 +134,15 @@ void UModularRig::ResetModules()
 	Modules.Reset();
 }
 
-bool UModularRig::AddModuleInstance(const FName& InModuleName, TSubclassOf<UControlRig> InModuleClass, FString InParentPath, const TMap<FRigElementKey, FRigElementKey>& InConnectionMap)
+bool UModularRig::AddModuleInstance(const FName& InModuleName, TSubclassOf<UControlRig> InModuleClass, FString InParentPath,
+	const TMap<FRigElementKey, FRigElementKey>& InConnectionMap, const TMap<FName, FString>& InVariableDefaultValues )
 {
 	FRigModuleInstance* ParentModule = FindModule(InParentPath);
-	return AddModuleInstance(InModuleName, InModuleClass, ParentModule, InConnectionMap) != nullptr;
+	return AddModuleInstance(InModuleName, InModuleClass, ParentModule, InConnectionMap, InVariableDefaultValues ) != nullptr;
 }
 
-FRigModuleInstance* UModularRig::AddModuleInstance(const FName& InModuleName, TSubclassOf<UControlRig> InModuleClass, FRigModuleInstance* InParent, const TMap<FRigElementKey, FRigElementKey>& InConnectionMap) 
+FRigModuleInstance* UModularRig::AddModuleInstance(const FName& InModuleName, TSubclassOf<UControlRig> InModuleClass, FRigModuleInstance* InParent,
+	const TMap<FRigElementKey, FRigElementKey>& InConnectionMap, const TMap<FName, FString>& InVariableDefaultValues ) 
 {
 	// Make sure there are no name clashes
 	if (InParent)
@@ -196,6 +198,11 @@ FRigModuleInstance* UModularRig::AddModuleInstance(const FName& InModuleName, TS
 		ModulePublicContext.RigModuleNameSpace = NewModule.Rig->GetRigModuleNameSpace();
 		ModulePublicContext.RigModuleNameSpaceHash = GetTypeHash(ModulePublicContext.RigModuleNameSpace);
 		NewModule.Rig->SetElementKeyRedirector(FRigElementKeyRedirector(InConnectionMap, Hierarchy));
+
+		for (TPair<FName, FString> Variable : InVariableDefaultValues )
+		{
+			NewModule.Rig->SetVariableFromString(Variable.Key, Variable.Value);
+		}
 	}
 	
 	return &NewModule;
diff --git a/Engine/Plugins/Animation/ControlRig/Source/ControlRig/Public/ModularRig.h b/Engine/Plugins/Animation/ControlRig/Source/ControlRig/Public/ModularRig.h
--- a/Engine/Plugins/Animation/ControlRig/Source/ControlRig/Public/ModularRig.h
+++ b/Engine/Plugins/Animation/ControlRig/Source/ControlRig/Public/ModularRig.h
@@ -61,8 +61,8 @@ public:
 	void ResetModules();
 
 	/** Adds a module to the rig*/
-	bool AddModuleInstance(const FName& InModuleName, TSubclassOf<UControlRig> InModuleClass, FString InParentPath, const TMap<FRigElementKey, FRigElementKey>& InConnectionMap);
-	FRigModuleInstance* AddModuleInstance(const FName& InModuleName, TSubclassOf<UControlRig> InModuleClass, FRigModuleInstance* InParent, const TMap<FRigElementKey, FRigElementKey>& InConnectionMap);
+	bool AddModuleInstance(const FName& InModuleName, TSubclassOf<UControlRig> InModuleClass, FString InParentPath, const TMap<FRigElementKey, FRigElementKey>& InConnectionMap, const TMap<FName, FString>& InVariables);
+	FRigModuleInstance* AddModuleInstance(const FName& InModuleName, TSubclassOf<UControlRig> InModuleClass, FRigModuleInstance* InParent, const TMap<FRigElementKey, FRigElementKey>& InConnectionMap, const TMap<FName, FString>& InVariables);
 
 	FRigModuleInstance* FindModule(const FString& InPath);
 
diff --git a/Engine/Plugins/Animation/ControlRig/Source/ControlRigDeveloper/Private/ControlRigBlueprint.cpp b/Engine/Plugins/Animation/ControlRig/Source/ControlRigDeveloper/Private/ControlRigBlueprint.cpp
--- a/Engine/Plugins/Animation/ControlRig/Source/ControlRigDeveloper/Private/ControlRigBlueprint.cpp
+++ b/Engine/Plugins/Animation/ControlRig/Source/ControlRigDeveloper/Private/ControlRigBlueprint.cpp
@@ -1549,7 +1549,7 @@ static void AddModulesRecursively(UModularRig* Rig, const FRigModuleReference* I
 		}
 	}
 	
-	FRigModuleInstance* NewModule = Rig->AddModuleInstance(InModule->Name, InModule->Class.Get(), InParent, InModule->Connections);
+	FRigModuleInstance* NewModule = Rig->AddModuleInstance(InModule->Name, InModule->Class.Get(), InParent, InModule->Connections, InModule->ConfigValues);
 	
 	for (const FRigModuleReference* ChildModule : InModule->CachedChildren)
 	{
@@ -2273,6 +2273,19 @@ void UControlRigBlueprint::HandleHierarchyModified(ERigHierarchyNotification InN
 
 void UControlRigBlueprint::HandleRigModulesModified(EModularRigNotification InNotification, const FRigModuleReference* InModule)
 {
+	switch (InNotification)
+	{
+		case EModularRigNotification::ConnectionChanged:
+		{
+			const FString Namespace = InModule->GetNamespace();
+			for (TPair<FRigElementKey, FRigElementKey> Connection : InModule->Connections)
+			{
+				ConnectionMap.FindOrAdd(Connection.Key) = Connection.Value;
+			}
+			HierarchyModifiedEvent.Broadcast(ERigHierarchyNotification::HierarchyReset, Hierarchy, nullptr);
+			break;
+		}
+	}
 	RecompileModularRig();
 }
 
diff --git a/Engine/Plugins/Animation/ControlRig/Source/ControlRigDeveloper/Private/ModularRigController.cpp b/Engine/Plugins/Animation/ControlRig/Source/ControlRigDeveloper/Private/ModularRigController.cpp
--- a/Engine/Plugins/Animation/ControlRig/Source/ControlRigDeveloper/Private/ModularRigController.cpp
+++ b/Engine/Plugins/Animation/ControlRig/Source/ControlRigDeveloper/Private/ModularRigController.cpp
@@ -6,6 +6,7 @@
 #include "ModularRig.h"
 #include "ControlRigBlueprint.h"
 #include "ModularRigModel.h"
+#include "Kismet2/BlueprintEditorUtils.h"
 #include "Rigs/RigHierarchyController.h"
 
 #include UE_INLINE_GENERATED_CPP_BY_NAME(ModularRigController)
@@ -21,8 +22,16 @@ UModularRigController::UModularRigController(const FObjectInitializer& ObjectIni
 
 bool UModularRigController::AddModule(const FName& InModuleName, TSubclassOf<UControlRig> InClass, const FString& InParentModulePath, bool bSetupUndo)
 {
-	if (!InClass->GetDefaultObject<UControlRig>()->IsRigModule())
+	if (!InClass)
 	{
+		UE_LOG(LogControlRig, Error, TEXT("Invalid InClass"));
+		return false;
+	}
+
+	UControlRig* ClassDefaultObject = InClass->GetDefaultObject<UControlRig>();
+	if (!ClassDefaultObject->IsRigModule())
+	{
+		UE_LOG(LogControlRig, Error, TEXT("Class %s is not a rig module"), *InClass->GetClassPathName().ToString());
 		return false;
 	}
 
@@ -71,12 +80,7 @@ bool UModularRigController::AddModule(const FName& InModuleName, TSubclassOf<UCo
 
 	if (!NewModule)
 	{
-		return false;
-	}
-
-	UControlRig* DefaultModule = InClass->GetDefaultObject<UControlRig>();
-	if (!DefaultModule)
-	{
+		UE_LOG(LogControlRig, Error, TEXT("Error while creating module %s"), *InModuleName.ToString());
 		return false;
 	}
 
@@ -96,7 +100,7 @@ bool UModularRigController::AddModule(const FName& InModuleName, TSubclassOf<UCo
 			FControlRigExecuteContextRigModuleGuard RigModuleGuard(PublicContext, NewNamespace);
 
 			
-			const TArray<FRigModuleConnector>& Connectors = DefaultModule->GetRigModuleSettings().ExposedConnectors;
+			const TArray<FRigModuleConnector>& Connectors = ClassDefaultObject->GetRigModuleSettings().ExposedConnectors;
 			for (const FRigModuleConnector& Connector : Connectors)
 			{
 				Controller->AddConnector(*Connector.Name, Connector.Settings);
@@ -154,12 +158,14 @@ bool UModularRigController::ConnectModuleToElement(const FRigElementKey& InConne
 	FString ConnectorNameSpace, ConnectorName;
 	if (!InConnectorKey.Name.ToString().Split(UModularRig::NamespaceSeparator, &ConnectorNameSpace, &ConnectorName, ESearchCase::CaseSensitive, ESearchDir::FromEnd))
 	{
+		UE_LOG(LogControlRig, Error, TEXT("Connector %s does not contain a namespace"), *InConnectorKey.ToString());
 		return false;
 	}
 	
 	FRigModuleReference* Module = FindModule(ConnectorNameSpace);
 	if (!Module)
 	{
+		UE_LOG(LogControlRig, Error, TEXT("Could not find module %s"), *ConnectorNameSpace);
 		return false;
 	}
 
@@ -167,11 +173,13 @@ bool UModularRigController::ConnectModuleToElement(const FRigElementKey& InConne
 	FRigConnectorElement* Connector = Cast<FRigConnectorElement>(Blueprint->Hierarchy->Find(InConnectorKey));
 	if (!Connector)
 	{
+		UE_LOG(LogControlRig, Error, TEXT("Could not find connector %s"), *InConnectorKey.ToString());
 		return false;
 	}
 
 	if (!Blueprint->Hierarchy->Contains(InTargetKey))
 	{
+		UE_LOG(LogControlRig, Error, TEXT("Could not find target %s"), *InTargetKey.ToString());
 		return false;
 	}
 
@@ -198,6 +206,59 @@ bool UModularRigController::ConnectModuleToElement(const FRigElementKey& InConne
 	return true;
 }
 
+bool UModularRigController::SetConfigValueInModule(const FString& InModulePath, const FName& InVariableName, const FString& InValue, bool bSetupUndo)
+{
+	FRigModuleReference* Module = FindModule(InModulePath);
+	if (!Module)
+	{
+		UE_LOG(LogControlRig, Error, TEXT("Could not find module %s"), *InModulePath);
+		return false;
+	}
+
+	if (!Module->Class.IsValid())
+	{
+		UE_LOG(LogControlRig, Error, TEXT("Class defined in module %s is not valid"), *InModulePath);
+		return false;
+	}
+
+	const FProperty* Property = Module->Class->FindPropertyByName(InVariableName);
+	if (!Property)
+	{
+		UE_LOG(LogControlRig, Error, TEXT("Could not find variable %s in module %s"), *InVariableName.ToString(), *InModulePath);
+		return false;
+	}
+
+	TArray<uint8, TAlignedHeapAllocator<16>> TempStorage;
+	TempStorage.AddZeroed(Property->GetSize());
+	if (!FBlueprintEditorUtils::PropertyValueFromString(Property, InValue, TempStorage.GetData()))
+	{
+		UE_LOG(LogControlRig, Error, TEXT("Value %s for variable %s in module %s is not valid"), *InValue, *InVariableName.ToString(), *InModulePath);
+		return false;
+	}
+
+#if WITH_EDITOR
+	TSharedPtr<FScopedTransaction> TransactionPtr;
+	if (bSetupUndo)
+	{
+		TransactionPtr = MakeShared<FScopedTransaction>(NSLOCTEXT("ModularRigController", "ConfigureModuleValueTransaction", "Configure Module Value"));
+		if(UControlRigBlueprint* Blueprint = Cast<UControlRigBlueprint>(GetOuter()))
+		{
+			Blueprint->Modify();
+		}
+	}
+#endif 
+
+	Module->ConfigValues.FindOrAdd(InVariableName) = InValue;
+
+	Notify(EModularRigNotification::ModuleConfigValueChanged, Module);
+
+#if WITH_EDITOR
+	TransactionPtr.Reset();
+#endif
+
+	return true;
+}
+
 bool UModularRigController::RemoveModule(const FString& InModulesPath, bool bSetupUndo)
 {
 	// todo: UE-199050
diff --git a/Engine/Plugins/Animation/ControlRig/Source/ControlRigDeveloper/Public/ModularRigController.h b/Engine/Plugins/Animation/ControlRig/Source/ControlRigDeveloper/Public/ModularRigController.h
--- a/Engine/Plugins/Animation/ControlRig/Source/ControlRigDeveloper/Public/ModularRigController.h
+++ b/Engine/Plugins/Animation/ControlRig/Source/ControlRigDeveloper/Public/ModularRigController.h
@@ -18,6 +18,8 @@ enum class EModularRigNotification : uint8
 
 	ConnectionChanged,
 
+	ModuleConfigValueChanged,
+
 	/** MAX - invalid */
 	Max UMETA(Hidden),
 };
@@ -46,6 +48,9 @@ class CONTROLRIGDEVELOPER_API UModularRigController : public UObject
 	UFUNCTION(BlueprintCallable, Category = "Control Rig | Modules")
 	bool ConnectModuleToElement(const FRigElementKey& InConnectorKey, const FRigElementKey& InTargetKey, bool bSetupUndo = true);
 
+	UFUNCTION(BlueprintCallable, Category = "Control Rig | Modules")
+	bool SetConfigValueInModule(const FString& InModulePath, const FName& InVariableName, const FString& InValue, bool bSetupUndo = true);
+
 	UFUNCTION(BlueprintCallable, Category = "ControlRig | Modules")
 	bool RemoveModule(const FString& InModulesPath, bool bSetupUndo = true);
 
diff --git a/Engine/Plugins/Animation/ControlRig/Source/ControlRigDeveloper/Public/ModularRigModel.h b/Engine/Plugins/Animation/ControlRig/Source/ControlRigDeveloper/Public/ModularRigModel.h
--- a/Engine/Plugins/Animation/ControlRig/Source/ControlRigDeveloper/Public/ModularRigModel.h
+++ b/Engine/Plugins/Animation/ControlRig/Source/ControlRigDeveloper/Public/ModularRigModel.h
@@ -35,6 +35,9 @@ struct CONTROLRIGDEVELOPER_API FRigModuleReference
 	UPROPERTY()
 	TMap<FRigElementKey, FRigElementKey> Connections; // Connectors to Connection element
 
+	UPROPERTY()
+	TMap<FName, FString> ConfigValues;
+
 	TArray<FRigModuleReference*> CachedChildren;
 
 	FString GetNamespace() const;
