Control Rig: Refresh the SHieararchy when a connector is resolved
#rnx

diff --git a/Engine/Plugins/Animation/ControlRig/Source/ControlRig/Private/ModularRig.cpp b/Engine/Plugins/Animation/ControlRig/Source/ControlRig/Private/ModularRig.cpp
--- a/Engine/Plugins/Animation/ControlRig/Source/ControlRig/Private/ModularRig.cpp
+++ b/Engine/Plugins/Animation/ControlRig/Source/ControlRig/Private/ModularRig.cpp
@@ -18,8 +18,14 @@ FRigModuleInstance::~FRigModuleInstance()
 	CachedChildren.Reset();
 	if (Rig && Rig.IsValid())
 	{
-		Rig->Rename(nullptr, GetTransientPackage(), REN_ForceNoResetLoaders | REN_DoNotDirty | REN_DontCreateRedirectors | REN_NonTransactional);
-		Rig->MarkAsGarbage();
+		static int32 ObjectIndexToBeDestroyed = 0;
+		static constexpr TCHAR ObjectNameFormat[] = TEXT("FRigModuleInstance_ObjectToBeDestroyed_%d");
+		const FString NewObjectName = FString::Printf(ObjectNameFormat, ObjectIndexToBeDestroyed++);
+		Rig->Rename(*NewObjectName, GetTransientPackage(), REN_ForceNoResetLoaders | REN_DoNotDirty | REN_DontCreateRedirectors | REN_NonTransactional);
+		if(!Rig->IsRooted())
+		{
+			Rig->MarkAsGarbage();
+		}
 		Rig = nullptr;
 	}
 }
@@ -174,7 +180,7 @@ FRigModuleInstance* UModularRig::AddModuleInstance(const FName& InModuleName, TS
 
 	FString Name = (InParent) ? InParent->Name.ToString() + NamespaceSeparator + InModuleName.ToString() : InModuleName.ToString();
 	FRigModuleInstance& NewModule = Modules.Add_GetRef(FRigModuleInstance());
-	NewModule.Rig = NewObject<UControlRig>(this, InModuleClass, *Name, RF_Transient | RF_Transactional);
+	NewModule.Rig = NewObject<UControlRig>(this, InModuleClass, *Name);
 	NewModule.Name = InModuleName;
 
 	if (InParent)
diff --git a/Engine/Plugins/Animation/ControlRig/Source/ControlRigDeveloper/Private/ControlRigBlueprint.cpp b/Engine/Plugins/Animation/ControlRig/Source/ControlRigDeveloper/Private/ControlRigBlueprint.cpp
--- a/Engine/Plugins/Animation/ControlRig/Source/ControlRigDeveloper/Private/ControlRigBlueprint.cpp
+++ b/Engine/Plugins/Animation/ControlRig/Source/ControlRigDeveloper/Private/ControlRigBlueprint.cpp
@@ -2280,7 +2280,8 @@ void UControlRigBlueprint::HandleRigModulesModified(EModularRigNotification InNo
 			const FString Namespace = InModule->GetNamespace();
 			for (TPair<FRigElementKey, FRigElementKey> Connection : InModule->Connections)
 			{
-				ConnectionMap.FindOrAdd(Connection.Key) = Connection.Value;
+				const FString NamespacedName = FString::Printf(TEXT("%s%s"), *Namespace, *Connection.Key.Name.ToString());
+				ConnectionMap.FindOrAdd(FRigElementKey(*NamespacedName, ERigElementType::Connector)) = Connection.Value;
 			}
 			HierarchyModifiedEvent.Broadcast(ERigHierarchyNotification::HierarchyReset, Hierarchy, nullptr);
 			break;
diff --git a/Engine/Plugins/Animation/ControlRig/Source/ControlRigEditor/Private/Editor/SRigHierarchy.cpp b/Engine/Plugins/Animation/ControlRig/Source/ControlRigEditor/Private/Editor/SRigHierarchy.cpp
--- a/Engine/Plugins/Animation/ControlRig/Source/ControlRigEditor/Private/Editor/SRigHierarchy.cpp
+++ b/Engine/Plugins/Animation/ControlRig/Source/ControlRigEditor/Private/Editor/SRigHierarchy.cpp
@@ -655,11 +655,6 @@ void SRigHierarchy::OnSelectionChanged(TSharedPtr<FRigTreeElement> Selection, ES
 
 void SRigHierarchy::OnHierarchyModified(ERigHierarchyNotification InNotif, URigHierarchy* InHierarchy, const FRigBaseElement* InElement)
 {
-	if (!InElement)
-	{
-		return;
-	}
-
 	if(!ControlRigBlueprint.IsValid())
 	{
 		return;
