Removing redundant scoped rig hierarchy validity bracket

diff --git a/Engine/Plugins/Animation/ControlRig/Source/ControlRig/Private/AnimNode_ControlRigBase.cpp b/Engine/Plugins/Animation/ControlRig/Source/ControlRig/Private/AnimNode_ControlRigBase.cpp
--- a/Engine/Plugins/Animation/ControlRig/Source/ControlRig/Private/AnimNode_ControlRigBase.cpp
+++ b/Engine/Plugins/Animation/ControlRig/Source/ControlRig/Private/AnimNode_ControlRigBase.cpp
@@ -153,9 +153,6 @@ void FAnimNode_ControlRigBase::UpdateInput(UControlRig* ControlRig, const FPoseC
 		// so the reset here ensures excluded bones are also reset
 		if(!ControlRigBoneInputMappingByName.IsEmpty() || bResetInputPoseToInitial)
 		{
-			{
-				FRigHierarchyValidityBracket ValidityBracket(Hierarchy);
-			}
 			FRigHierarchyValidityBracket ValidityBracket(Hierarchy);
 
 			{
