Obscure warning V547: Expression 'bWasReplaced == false' is always true.
Centralize logic and fix bool case to Modify and set bWasReplaced true
#rnx

diff --git a/Engine/Plugins/Animation/ControlRig/Source/ControlRig/Private/Sequencer/MovieSceneControlRigParameterSection.cpp b/Engine/Plugins/Animation/ControlRig/Source/ControlRig/Private/Sequencer/MovieSceneControlRigParameterSection.cpp
--- a/Engine/Plugins/Animation/ControlRig/Source/ControlRig/Private/Sequencer/MovieSceneControlRigParameterSection.cpp
+++ b/Engine/Plugins/Animation/ControlRig/Source/ControlRig/Private/Sequencer/MovieSceneControlRigParameterSection.cpp
@@ -822,113 +822,33 @@ void UMovieSceneControlRigParameterSection::PreSave(FObjectPreSaveContext SaveCo
 bool UMovieSceneControlRigParameterSection::RenameParameterName(const FName& OldParameterName, const FName& NewParameterName)
 {
 	bool bWasReplaced = false;
-	for (FScalarParameterNameAndCurve& ScalarParameterNameAndCurve : ScalarParameterNamesAndCurves)
-	{
-		if (ScalarParameterNameAndCurve.ParameterName == OldParameterName)
-		{
-			if (bWasReplaced == false)
-			{
-				Modify();
-				bWasReplaced = true;
-			}
-			ScalarParameterNameAndCurve.ParameterName = NewParameterName;
-			break;
-		}
-	}
-
-	for (FBoolParameterNameAndCurve& BoolParameterNameAndCurve : BoolParameterNamesAndCurves)
-	{
-		if (BoolParameterNameAndCurve.ParameterName == OldParameterName)
-		{
-			BoolParameterNameAndCurve.ParameterName = NewParameterName;
-			break;
-		}
-	}
-
-	for (FEnumParameterNameAndCurve& EnumParameterNameAndCurve : EnumParameterNamesAndCurves)
-	{
-		if (EnumParameterNameAndCurve.ParameterName == OldParameterName)
-		{
-			if (bWasReplaced == false)
-			{
-				Modify();
-				bWasReplaced = true;
-			}
-			EnumParameterNameAndCurve.ParameterName = NewParameterName;
-			break;
-		}
-	}
 
-	for (FIntegerParameterNameAndCurve& IntegerParameterNameAndCurve : IntegerParameterNamesAndCurves)
+	auto RenameParameterNameInner = [this, &bWasReplaced, OldParameterName, NewParameterName](auto& ParameterNamesAndCurves)
 	{
-		if (IntegerParameterNameAndCurve.ParameterName == OldParameterName)
+		for (auto& ParameterNameAndCurve : ParameterNamesAndCurves)
 		{
-			if (bWasReplaced == false)
+			if (ParameterNameAndCurve.ParameterName == OldParameterName)
 			{
-				Modify();
-				bWasReplaced = true;
-			}
-			IntegerParameterNameAndCurve.ParameterName = NewParameterName;
-			break;
-		}
-	}
-
-	for (FVector2DParameterNameAndCurves& Vector2DParameterNameAndCurve : Vector2DParameterNamesAndCurves)
-	{
-		if (Vector2DParameterNameAndCurve.ParameterName == OldParameterName)
-		{
-			if (bWasReplaced == false)
-			{
-				Modify();
-				bWasReplaced = true;
+				if (!bWasReplaced)
+				{
+					Modify();
+					bWasReplaced = true;
+				}
+				ParameterNameAndCurve.ParameterName = NewParameterName;
+				break;
 			}
-			Vector2DParameterNameAndCurve.ParameterName = NewParameterName;
-			break;
 		}
-	}
+	};
 
-	for (FVectorParameterNameAndCurves& VectorParameterNameAndCurve : VectorParameterNamesAndCurves)
-	{
-		if (VectorParameterNameAndCurve.ParameterName == OldParameterName)
-		{
-			if (bWasReplaced == false)
-			{
-				Modify();
-				bWasReplaced = true;
-			}
-			VectorParameterNameAndCurve.ParameterName = NewParameterName;
-			break;
-		}
-	}
+	RenameParameterNameInner(ScalarParameterNamesAndCurves);
+	RenameParameterNameInner(BoolParameterNamesAndCurves);
+	RenameParameterNameInner(EnumParameterNamesAndCurves);
+	RenameParameterNameInner(IntegerParameterNamesAndCurves);
+	RenameParameterNameInner(Vector2DParameterNamesAndCurves);
+	RenameParameterNameInner(VectorParameterNamesAndCurves);
+	RenameParameterNameInner(ColorParameterNamesAndCurves);
+	RenameParameterNameInner(TransformParameterNamesAndCurves);
 
-	for (FColorParameterNameAndCurves& ColorParameterNameAndCurve : ColorParameterNamesAndCurves)
-	{
-		if (ColorParameterNameAndCurve.ParameterName == OldParameterName)
-		{
-			if (bWasReplaced == false)
-			{
-				Modify();
-				bWasReplaced = true;
-			}
-			ColorParameterNameAndCurve.ParameterName = NewParameterName;
-			break;
-		}
-				
-	}
-
-	for (FTransformParameterNameAndCurves& TransformParameterNamesAndCurve : TransformParameterNamesAndCurves)
-	{
-		if (TransformParameterNamesAndCurve.ParameterName == OldParameterName)
-		{
-			if (bWasReplaced == false)
-			{
-				Modify();
-				bWasReplaced = true;
-			}
-			TransformParameterNamesAndCurve.ParameterName = NewParameterName;
-			break;
-		}
-	}
 	if (bWasReplaced)
 	{
 		ReconstructChannelProxy();
