Sequencer: Control Rig: Masked controls may not key correctly since wrong control/mask index was being used. Also fix old CIS warning from Marc.

diff --git a/Engine/Plugins/Animation/ControlRig/Source/ControlRigEditor/Private/Sequencer/ControlRigParameterTrackEditor.cpp b/Engine/Plugins/Animation/ControlRig/Source/ControlRigEditor/Private/Sequencer/ControlRigParameterTrackEditor.cpp
--- a/Engine/Plugins/Animation/ControlRig/Source/ControlRigEditor/Private/Sequencer/ControlRigParameterTrackEditor.cpp
+++ b/Engine/Plugins/Animation/ControlRig/Source/ControlRigEditor/Private/Sequencer/ControlRigParameterTrackEditor.cpp
@@ -3293,9 +3293,9 @@ void FControlRigParameterTrackEditor::GetControlRigKeys(
 
 	//Need seperate index fo bools,ints and enums and floats since there are seperate entries for each later when they are accessed by the set key stuff.
 	int32 SpaceChannelIndex = 0;
-	for (int32 ControlIndex = 0; ControlIndex < Controls.Num(); ++ControlIndex)
+	for (int32 LocalControlIndex = 0; LocalControlIndex < Controls.Num(); ++LocalControlIndex)
 	{
-		FRigControlElement* ControlElement = Controls[ControlIndex];
+		FRigControlElement* ControlElement = Controls[LocalControlIndex];
 		check(ControlElement);
 
 		if (!Hierarchy->IsAnimatable(ControlElement))
@@ -3306,9 +3306,9 @@ void FControlRigParameterTrackEditor::GetControlRigKeys(
 		if (FChannelMapInfo* pChannelIndex = SectionToKey->ControlChannelMap.Find(ControlElement->GetFName()))
 		{
 			int32 ChannelIndex = pChannelIndex->ChannelIndex;
+			const int32 MaskIndex = pChannelIndex->MaskIndex;
 
-
-			bool bMaskKeyOut = (ControlIndex >= ControlsMask.Num() || ControlsMask[ControlIndex] == false);
+			bool bMaskKeyOut = (MaskIndex >= ControlsMask.Num() || ControlsMask[MaskIndex] == false);
 			bool bSetKey = ParameterName.IsNone() || (ControlElement->GetFName() == ParameterName && !bMaskKeyOut);
 
 			FRigControlValue ControlValue = InControlRig->GetControlValue(ControlElement, ERigControlValueType::Current);
diff --git a/Engine/Plugins/Animation/ControlRig/Source/ControlRigEditor/Private/Tools/ConstraintBaker.cpp b/Engine/Plugins/Animation/ControlRig/Source/ControlRigEditor/Private/Tools/ConstraintBaker.cpp
--- a/Engine/Plugins/Animation/ControlRig/Source/ControlRigEditor/Private/Tools/ConstraintBaker.cpp
+++ b/Engine/Plugins/Animation/ControlRig/Source/ControlRigEditor/Private/Tools/ConstraintBaker.cpp
@@ -231,7 +231,7 @@ void FConstraintBaker::Bake(UWorld* InWorld,
 	FConstraintSections ConstraintSections = FMovieSceneConstraintChannelHelper::GetConstraintSectionAndChannel(
 		InConstraint, InSequencer);
 	IMovieSceneConstrainedSection* ConstrainedSection = Cast<IMovieSceneConstrainedSection>(ConstraintSections.ConstraintSection);
-	if (ConstrainedSection == nullptr) 
+	if (ConstrainedSection == nullptr || ConstraintSections.ActiveChannel == nullptr)
 	{
 		return;
 	}
