CIS Fix: Nullptr check for viewports in control rig edit mode.
#trivial

diff --git a/Engine/Plugins/Animation/ControlRig/Source/ControlRigEditor/Private/EditMode/ControlRigEditMode.cpp b/Engine/Plugins/Animation/ControlRig/Source/ControlRigEditor/Private/EditMode/ControlRigEditMode.cpp
--- a/Engine/Plugins/Animation/ControlRig/Source/ControlRigEditor/Private/EditMode/ControlRigEditMode.cpp
+++ b/Engine/Plugins/Animation/ControlRig/Source/ControlRigEditor/Private/EditMode/ControlRigEditMode.cpp
@@ -500,6 +500,11 @@ void FControlRigEditMode::Tick(FEditorViewportClient* ViewportClient, float Delt
 {
 	FEdMode::Tick(ViewportClient, DeltaTime);
 	
+	//if we have don't have a viewport client or viewport, bail we can be in UMG for example
+	if (ViewportClient == nullptr || ViewportClient->Viewport == nullptr)
+	{
+		return;
+	}
 	CheckMovieSceneSig();
 
 	if (bool* GameView = ViewportToGameView.Find(ViewportClient->Viewport))
