Sequencer: Control Rig:  Procedural control rigs would no longer have channels show up, need to referesh tree on PostConstruct for all CR's, and need to BroadcastChange on the Section so the UI picks up this change and recreates it.

diff --git a/Engine/Plugins/Animation/ControlRig/Source/ControlRig/Private/Sequencer/MovieSceneControlRigParameterSection.cpp b/Engine/Plugins/Animation/ControlRig/Source/ControlRig/Private/Sequencer/MovieSceneControlRigParameterSection.cpp
--- a/Engine/Plugins/Animation/ControlRig/Source/ControlRig/Private/Sequencer/MovieSceneControlRigParameterSection.cpp
+++ b/Engine/Plugins/Animation/ControlRig/Source/ControlRig/Private/Sequencer/MovieSceneControlRigParameterSection.cpp
@@ -2508,7 +2508,8 @@ void UMovieSceneControlRigParameterSection::RecreateWithThisControlRig(UControlR
 			break;
 		}
 	}
-	ReconstructChannelProxy();
+	CacheChannelProxy();
+	BroadcastChanged();
 }
 
 void UMovieSceneControlRigParameterSection::SetControlRig(UControlRig* InControlRig)
diff --git a/Engine/Plugins/Animation/ControlRig/Source/ControlRigEditor/Private/Sequencer/ControlRigParameterTrackEditor.cpp b/Engine/Plugins/Animation/ControlRig/Source/ControlRigEditor/Private/Sequencer/ControlRigParameterTrackEditor.cpp
--- a/Engine/Plugins/Animation/ControlRig/Source/ControlRigEditor/Private/Sequencer/ControlRigParameterTrackEditor.cpp
+++ b/Engine/Plugins/Animation/ControlRig/Source/ControlRigEditor/Private/Sequencer/ControlRigParameterTrackEditor.cpp
@@ -3101,11 +3101,8 @@ void FControlRigParameterTrackEditor::HandleOnPostConstructed(UControlRig* Subje
 
 		if (GetSequencer().IsValid())
 		{
-			//If FK control rig on next tick we refresh the tree
-			if (ControlRig->IsA<UFKControlRig>())
-			{
-				GetSequencer()->NotifyMovieSceneDataChanged(EMovieSceneDataChangeType::RefreshTree);
-			}
+			//refresh tree for ANY control rig may be FK or procedural
+			GetSequencer()->NotifyMovieSceneDataChanged(EMovieSceneDataChangeType::RefreshTree);
 		}
 	}
 }
