Control RIg: Modular RIg Hierarchy Tab

#virtualized

diff --git a/Engine/Plugins/Animation/ControlRig/Source/ControlRig/Private/ModularRig.cpp b/Engine/Plugins/Animation/ControlRig/Source/ControlRig/Private/ModularRig.cpp
--- a/Engine/Plugins/Animation/ControlRig/Source/ControlRig/Private/ModularRig.cpp
+++ b/Engine/Plugins/Animation/ControlRig/Source/ControlRig/Private/ModularRig.cpp
@@ -13,23 +13,44 @@
 
 const FString UModularRig::NamespaceSeparator = TEXT(":");
 
-FRigModuleInstance::~FRigModuleInstance()
+////////////////////////////////////////////////////////////////////////////////
+// FModuleInstanceHandle
+////////////////////////////////////////////////////////////////////////////////
+
+FModuleInstanceHandle::FModuleInstanceHandle(UModularRig* InRig, const FString& InPath)
+: Rig(InRig)
+, Path(InPath)
+{
+}
+
+FModuleInstanceHandle::FModuleInstanceHandle(UModularRig* InRig, const FRigModuleInstance* InModule)
+: Rig(InRig)
+, Path(InModule->GetPath())
 {
-	CachedChildren.Reset();
-	if (Rig && Rig.IsValid())
+}
+
+const FRigModuleInstance* FModuleInstanceHandle::Get() const
+{
+	if(Rig.IsValid())
 	{
-		static int32 ObjectIndexToBeDestroyed = 0;
-		static constexpr TCHAR ObjectNameFormat[] = TEXT("FRigModuleInstance_ObjectToBeDestroyed_%d");
-		const FString NewObjectName = FString::Printf(ObjectNameFormat, ObjectIndexToBeDestroyed++);
-		Rig->Rename(*NewObjectName, GetTransientPackage(), REN_ForceNoResetLoaders | REN_DoNotDirty | REN_DontCreateRedirectors | REN_NonTransactional);
-		if(!Rig->IsRooted())
-		{
-			Rig->MarkAsGarbage();
-		}
-		Rig = nullptr;
+		return Rig->FindModule(Path);
+	}
+	return nullptr;
+}
+
+FRigModuleInstance* FModuleInstanceHandle::Get()
+{
+	if(Rig.IsValid())
+	{
+		return Rig->FindModule(Path);
 	}
+	return nullptr;
 }
 
+////////////////////////////////////////////////////////////////////////////////
+// UModularRig
+////////////////////////////////////////////////////////////////////////////////
+
 UModularRig::UModularRig(const FObjectInitializer& ObjectInitializer)
 	: Super(ObjectInitializer)
 {
@@ -40,6 +61,8 @@ UModularRig::UModularRig(const FObjectInitializer& ObjectInitializer)
 
 void UModularRig::BeginDestroy()
 {
+	ResetModules();
+	
 	Super::BeginDestroy();
 	
 #if WITH_EDITOR
@@ -47,6 +70,15 @@ void UModularRig::BeginDestroy()
 #endif
 }
 
+FString FRigModuleInstance::GetPath() const
+{
+	if (!ParentPath.IsEmpty())
+	{
+		return FString::Printf(TEXT("%s:%s"), *ParentPath, *Name.ToString()); 
+	}
+	return Name.ToString();
+}
+
 void UModularRig::InitializeVMs(bool bRequestInit)
 {
 	URigVMHost::Initialize(bRequestInit);
@@ -136,10 +168,54 @@ void UModularRig::OnObjectsReplaced(const TMap<UObject*, UObject*>& OldToNewInst
 
 void UModularRig::ResetModules()
 {
+	for (FRigModuleInstance& Module : Modules)
+	{
+		Module.CachedChildren.Reset();
+		TSoftObjectPtr<UControlRig> Rig = Module.Rig;
+		if (Rig && Rig.IsValid())
+		{
+			static int32 ObjectIndexToBeDestroyed = 0;
+			static constexpr TCHAR ObjectNameFormat[] = TEXT("FRigModuleInstance_ObjectToBeDestroyed_%d");
+			const FString NewObjectName = FString::Printf(ObjectNameFormat, ObjectIndexToBeDestroyed++);
+			Rig->Rename(*NewObjectName, GetTransientPackage(), REN_ForceNoResetLoaders | REN_DoNotDirty | REN_DontCreateRedirectors | REN_NonTransactional);
+			if(!Rig->IsRooted())
+			{
+				Rig->MarkAsGarbage();
+			}
+			Rig = nullptr;
+		}
+	}
+	
 	RootModules.Reset();
 	Modules.Reset();
 }
 
+void UModularRig::UpdateCachedChildren()
+{
+	TMap<FString, FRigModuleInstance*> PathToModule;
+	for (FRigModuleInstance& Module : Modules)
+	{
+		Module.CachedChildren.Reset();
+		PathToModule.Add(Module.GetPath(), &Module);
+	}
+	
+	RootModules.Reset();
+	for (FRigModuleInstance& Module : Modules)
+	{
+		if (Module.ParentPath.IsEmpty())
+		{
+			RootModules.Add(&Module);
+		}
+		else
+		{
+			if (FRigModuleInstance** ParentModule = PathToModule.Find(Module.ParentPath))
+			{
+				(*ParentModule)->CachedChildren.Add(&Module);
+			}
+		}
+	}
+}
+
 bool UModularRig::AddModuleInstance(const FName& InModuleName, TSubclassOf<UControlRig> InModuleClass, FString InParentPath,
 	const TMap<FRigElementKey, FRigElementKey>& InConnectionMap, const TMap<FName, FString>& InVariableDefaultValues )
 {
@@ -185,13 +261,9 @@ FRigModuleInstance* UModularRig::AddModuleInstance(const FName& InModuleName, TS
 
 	if (InParent)
 	{
-		InParent->CachedChildren.Add(&NewModule);
-		NewModule.ParentPath = InParent->ParentPath.IsEmpty() ? InParent->Name.ToString() : InParent->ParentPath + NamespaceSeparator + InParent->Name.ToString();
-	}
-	else
-	{
-		RootModules.Add(&NewModule);
+		NewModule.ParentPath = InParent->GetPath();
 	}
+	UpdateCachedChildren();
 
 	// Configure module
 	{
@@ -214,13 +286,13 @@ FRigModuleInstance* UModularRig::AddModuleInstance(const FName& InModuleName, TS
 	return &NewModule;
 }
 
-FRigModuleInstance* UModularRig::FindModule(const FString& InPath)
+FRigModuleInstance* UModularRig::FindModule(const FString& InPath) const
 {
-	TArray<FRigModuleInstance*>* Children = &RootModules;
+	const TArray<FRigModuleInstance*>* Children = &RootModules;
 	FString Left = InPath, Right;
 	while (Left.Split(NamespaceSeparator, &Left, &Right))
 	{
-		FRigModuleInstance** Cur = Children->FindByPredicate([Left](FRigModuleInstance* Module)
+		FRigModuleInstance* const * Cur = Children->FindByPredicate([Left](FRigModuleInstance* Module)
 		{
 			return Module->Name == Left;
 		});
@@ -232,7 +304,7 @@ FRigModuleInstance* UModularRig::FindModule(const FString& InPath)
 		Left = Right;
 	}
 
-	FRigModuleInstance** Cur = Children->FindByPredicate([Left](FRigModuleInstance* Module)
+	FRigModuleInstance* const * Cur = Children->FindByPredicate([Left](FRigModuleInstance* Module)
 		{
 			return Module->Name == Left;
 		});
@@ -243,6 +315,15 @@ FRigModuleInstance* UModularRig::FindModule(const FString& InPath)
 	return *Cur;
 }
 
+FString UModularRig::GetParentPath(const FString& InPath) const
+{
+	if (FRigModuleInstance* Element = FindModule(InPath))
+	{
+		return Element->ParentPath;
+	}
+	return FString();
+}
+
 void UModularRig::ForEachModule(TFunctionRef<bool(FRigModuleInstance*)> PerModuleFunction)
 {
 	TArray<FRigModuleInstance*> ModuleInstances = RootModules;
@@ -256,4 +337,17 @@ void UModularRig::ForEachModule(TFunctionRef<bool(FRigModuleInstance*)> PerModul
 	}
 }
 
+void UModularRig::ForEachModule(TFunctionRef<bool(const FRigModuleInstance*)> PerModuleFunction) const
+{
+	TArray<FRigModuleInstance*> ModuleInstances = RootModules;
+	for (int32 ModuleIndex = 0; ModuleIndex < ModuleInstances.Num(); ++ModuleIndex)
+	{
+		if (!PerModuleFunction(ModuleInstances[ModuleIndex]))
+		{
+			break;
+		}
+		ModuleInstances.Append(ModuleInstances[ModuleIndex]->CachedChildren);
+	}
+}
+
 #undef LOCTEXT_NAMESPACE
diff --git a/Engine/Plugins/Animation/ControlRig/Source/ControlRig/Public/ModularRig.h b/Engine/Plugins/Animation/ControlRig/Source/ControlRig/Public/ModularRig.h
--- a/Engine/Plugins/Animation/ControlRig/Source/ControlRig/Public/ModularRig.h
+++ b/Engine/Plugins/Animation/ControlRig/Source/ControlRig/Public/ModularRig.h
@@ -5,6 +5,36 @@
 
 #include "ModularRig.generated.h"
 
+struct FRigModuleInstance;
+
+struct CONTROLRIG_API FModuleInstanceHandle
+{
+public:
+
+	FModuleInstanceHandle()
+		: Rig(nullptr)
+		, Path()
+	{}
+
+	FModuleInstanceHandle(UModularRig* InRig, const FString& InPath);
+	FModuleInstanceHandle(UModularRig* InRig, const FRigModuleInstance* InElement);
+
+	bool IsValid() const { return Get() != nullptr; }
+	operator bool() const { return IsValid(); }
+	
+	const UModularRig* GetRig() const { return Rig.Get(); }
+	UModularRig* GetHierarchy() { return Rig.Get(); }
+	const FString& GetPath() const { return Path; }
+
+	const FRigModuleInstance* Get() const;
+	FRigModuleInstance* Get();
+
+private:
+
+	TWeakObjectPtr<UModularRig> Rig;
+	FString Path;
+};
+
 USTRUCT(BlueprintType)
 struct CONTROLRIG_API FRigModuleInstance 
 {
@@ -17,8 +47,6 @@ struct CONTROLRIG_API FRigModuleInstance
 	{
 	}
 
-	~FRigModuleInstance();
-
 	UPROPERTY()
 	FName Name;
 
@@ -29,6 +57,8 @@ struct CONTROLRIG_API FRigModuleInstance
 	FString ParentPath;
 
 	TArray<FRigModuleInstance*> CachedChildren;
+
+	FString GetPath() const;
 };
 
 /** Runs logic for mapping input data to transforms (the "Rig") */
@@ -60,13 +90,31 @@ public:
 
 	void ResetModules();
 
+	void UpdateCachedChildren();
+
 	/** Adds a module to the rig*/
 	bool AddModuleInstance(const FName& InModuleName, TSubclassOf<UControlRig> InModuleClass, FString InParentPath, const TMap<FRigElementKey, FRigElementKey>& InConnectionMap, const TMap<FName, FString>& InVariables);
 	FRigModuleInstance* AddModuleInstance(const FName& InModuleName, TSubclassOf<UControlRig> InModuleClass, FRigModuleInstance* InParent, const TMap<FRigElementKey, FRigElementKey>& InConnectionMap, const TMap<FName, FString>& InVariables);
 
-	FRigModuleInstance* FindModule(const FString& InPath);
+	FRigModuleInstance* FindModule(const FString& InPath) const;
+	FString GetParentPath(const FString& InPath) const;
 
 	void ForEachModule(TFunctionRef<bool(FRigModuleInstance*)> PerModuleFunction);
+	void ForEachModule(TFunctionRef<bool(const FRigModuleInstance*)> PerModuleFunction) const;
+
+	/**
+	 * Returns a handle to an existing element
+	 * @param InKey The key of the handle to retrieve.
+	 * @return The retrieved handle (may be invalid)
+	 */
+	FModuleInstanceHandle GetHandle(const FString& InPath) const
+	{
+		if(FRigModuleInstance* Module = FindModule(InPath))
+		{
+			return FModuleInstanceHandle((UModularRig*)this, InPath);
+		}
+		return FModuleInstanceHandle();
+	}
 
 	static const FString NamespaceSeparator;
 };
diff --git a/Engine/Plugins/Animation/ControlRig/Source/ControlRig/Public/Rigs/RigHierarchyCache.h b/Engine/Plugins/Animation/ControlRig/Source/ControlRig/Public/Rigs/RigHierarchyCache.h
--- a/Engine/Plugins/Animation/ControlRig/Source/ControlRig/Public/Rigs/RigHierarchyCache.h
+++ b/Engine/Plugins/Animation/ControlRig/Source/ControlRig/Public/Rigs/RigHierarchyCache.h
@@ -185,4 +185,5 @@ private:
 	friend class URigHierarchy;
 	friend class URigHierarchyController;
 	friend struct FRigConnectionInfo;
+	friend class FRigModuleInstanceDetails;
 };
\ No newline at end of file
diff --git a/Engine/Plugins/Animation/ControlRig/Source/ControlRigDeveloper/Private/ControlRigBlueprint.cpp b/Engine/Plugins/Animation/ControlRig/Source/ControlRigDeveloper/Private/ControlRigBlueprint.cpp
--- a/Engine/Plugins/Animation/ControlRig/Source/ControlRigDeveloper/Private/ControlRigBlueprint.cpp
+++ b/Engine/Plugins/Animation/ControlRig/Source/ControlRigDeveloper/Private/ControlRigBlueprint.cpp
@@ -1587,6 +1587,8 @@ void UControlRigBlueprint::RecompileModularRig()
 			}
 		}
 	}
+
+	OnModularRigCompiled().Broadcast(this);
 }
 
 #endif
diff --git a/Engine/Plugins/Animation/ControlRig/Source/ControlRigDeveloper/Private/ModularRigController.cpp b/Engine/Plugins/Animation/ControlRig/Source/ControlRigDeveloper/Private/ModularRigController.cpp
--- a/Engine/Plugins/Animation/ControlRig/Source/ControlRigDeveloper/Private/ModularRigController.cpp
+++ b/Engine/Plugins/Animation/ControlRig/Source/ControlRigDeveloper/Private/ModularRigController.cpp
@@ -7,6 +7,7 @@
 #include "ControlRigBlueprint.h"
 #include "ModularRigModel.h"
 #include "Kismet2/BlueprintEditorUtils.h"
+#include "Misc/DefaultValueHelper.h"
 #include "Rigs/RigHierarchyController.h"
 
 #include UE_INLINE_GENERATED_CPP_BY_NAME(ModularRigController)
@@ -61,7 +62,6 @@ bool UModularRigController::AddModule(const FName& InModuleName, TSubclassOf<UCo
 
 		Model->Modules.Add(FRigModuleReference(InModuleName, InClass, FString()));
 		NewModule = &Model->Modules.Last();
-		Model->RootModules.Add(NewModule);
 	}
 	else if (FRigModuleReference* ParentModule = FindModule(InParentModulePath))
 	{
@@ -75,15 +75,18 @@ bool UModularRigController::AddModule(const FName& InModuleName, TSubclassOf<UCo
 
 		Model->Modules.Add(FRigModuleReference(InModuleName, InClass, ParentModule->GetNamespace()));
 		NewModule = &Model->Modules.Last();
-		ParentModule->CachedChildren.Add(&Model->Modules.Last());
 	}
 
+	Model->UpdateCachedChildren();
+
 	if (!NewModule)
 	{
 		UE_LOG(LogControlRig, Error, TEXT("Error while creating module %s"), *InModuleName.ToString());
 		return false;
 	}
 
+	Notify(EModularRigNotification::ModuleAdded, NewModule);
+
 	const FString NewNamespace = NewModule->GetNamespace();
 
 	// Add Connectors/Sockets
@@ -115,8 +118,6 @@ bool UModularRigController::AddModule(const FName& InModuleName, TSubclassOf<UCo
 		}
 	}
 
-	Notify(EModularRigNotification::ModuleAdded, NewModule);
-
 #if WITH_EDITOR
 	TransactionPtr.Reset();
 #endif
@@ -230,7 +231,10 @@ bool UModularRigController::SetConfigValueInModule(const FString& InModulePath,
 
 	TArray<uint8, TAlignedHeapAllocator<16>> TempStorage;
 	TempStorage.AddZeroed(Property->GetSize());
-	if (!FBlueprintEditorUtils::PropertyValueFromString(Property, InValue, TempStorage.GetData()))
+	uint8* TempMemory = TempStorage.GetData();
+	Property->InitializeValue(TempMemory);
+
+	if (!FBlueprintEditorUtils::PropertyValueFromString_Direct(Property, InValue, TempMemory))
 	{
 		UE_LOG(LogControlRig, Error, TEXT("Value %s for variable %s in module %s is not valid"), *InValue, *InVariableName.ToString(), *InModulePath);
 		return false;
@@ -277,6 +281,39 @@ bool UModularRigController::ReparentModule(const FString& InModulesPath, const F
 	return false;
 }
 
+FName UModularRigController::GetSafeNewName(const FString& InModuleDesiredPath)
+{
+	FString ParentPath, DesiredName = InModuleDesiredPath;
+	InModuleDesiredPath.Split(UModularRig::NamespaceSeparator, &ParentPath, &DesiredName);
+
+	TArray<FRigModuleReference*>* Children = &Model->RootModules;
+	if (!ParentPath.IsEmpty())
+	{
+		if (FRigModuleReference* Parent = FindModule(ParentPath))
+		{
+			Children = &Parent->CachedChildren;
+		}
+	}
+
+	bool bSafeToUse = false;
+	FString NewName = DesiredName;
+	int32 Index = 0;
+	while (!bSafeToUse)
+	{
+		bSafeToUse = true;
+		for (FRigModuleReference* Child : *Children)
+		{
+			if (Child->Name == *NewName)
+			{
+				bSafeToUse = false;
+				NewName = FString::Printf(TEXT("%s%d"), *DesiredName, ++Index);
+				break;
+			}
+		}
+	}
+	return *NewName;
+}
+
 void UModularRigController::Notify(const EModularRigNotification& InNotification, const FRigModuleReference* InElement)
 {
 	if(!bSuspendNotifications)
diff --git a/Engine/Plugins/Animation/ControlRig/Source/ControlRigDeveloper/Private/ModularRigModel.cpp b/Engine/Plugins/Animation/ControlRig/Source/ControlRigDeveloper/Private/ModularRigModel.cpp
--- a/Engine/Plugins/Animation/ControlRig/Source/ControlRigDeveloper/Private/ModularRigModel.cpp
+++ b/Engine/Plugins/Animation/ControlRig/Source/ControlRigDeveloper/Private/ModularRigModel.cpp
@@ -61,14 +61,14 @@ void FModularRigModel::UpdateCachedChildren()
 	}
 }
 
-FRigModuleReference* FModularRigModel::FindModule(const FString InNameSpace)
+FRigModuleReference* FModularRigModel::FindModule(const FString InNameSpace) const
 {
-	TArray<FRigModuleReference*>* Children = &RootModules;
+	const TArray<FRigModuleReference*>* Children = &RootModules;
 
 	FString Left = InNameSpace, Right;
 	while (Left.Split(UModularRig::NamespaceSeparator, &Left, &Right))
 	{
-		FRigModuleReference** Child = Children->FindByPredicate([Left](FRigModuleReference* Module)
+		FRigModuleReference* const * Child = Children->FindByPredicate([Left](FRigModuleReference* Module)
 		{
 			return Module->Name.ToString() == Left;
 		});
@@ -82,7 +82,7 @@ FRigModuleReference* FModularRigModel::FindModule(const FString InNameSpace)
 		Left = Right;
 	}
 
-	FRigModuleReference** Child = Children->FindByPredicate([Left](FRigModuleReference* Module)
+	FRigModuleReference* const * Child = Children->FindByPredicate([Left](FRigModuleReference* Module)
 		{
 			return Module->Name.ToString() == Left;
 		});
@@ -93,3 +93,12 @@ FRigModuleReference* FModularRigModel::FindModule(const FString InNameSpace)
 	}
 	return *Child;
 }
+
+FString FModularRigModel::FindParentNamespace(const FString InNameSpace) const
+{
+	if (FRigModuleReference* Element = FindModule(InNameSpace))
+	{
+		return Element->ParentNamespace;
+	}
+	return FString();
+}
diff --git a/Engine/Plugins/Animation/ControlRig/Source/ControlRigDeveloper/Public/ControlRigBlueprint.h b/Engine/Plugins/Animation/ControlRig/Source/ControlRigDeveloper/Public/ControlRigBlueprint.h
--- a/Engine/Plugins/Animation/ControlRig/Source/ControlRigDeveloper/Public/ControlRigBlueprint.h
+++ b/Engine/Plugins/Animation/ControlRig/Source/ControlRigDeveloper/Public/ControlRigBlueprint.h
@@ -298,12 +298,15 @@ public:
 	*/
 	FRigHierarchyModifiedEvent& OnHierarchyModified() { return HierarchyModifiedEvent; }
 
+	FOnRigVMRefreshEditorEvent& OnModularRigCompiled() { return ModularRigCompiled; }
+
 private:
 
 	UPROPERTY()
 	TObjectPtr<UControlRigValidator> Validator;
 
 	FRigHierarchyModifiedEvent	HierarchyModifiedEvent;
+	FOnRigVMRefreshEditorEvent ModularRigCompiled;
 
 	void HandleHierarchyModified(ERigHierarchyNotification InNotification, URigHierarchy* InHierarchy, const FRigBaseElement* InElement);
 
diff --git a/Engine/Plugins/Animation/ControlRig/Source/ControlRigDeveloper/Public/ModularRigController.h b/Engine/Plugins/Animation/ControlRig/Source/ControlRigDeveloper/Public/ModularRigController.h
--- a/Engine/Plugins/Animation/ControlRig/Source/ControlRigDeveloper/Public/ModularRigController.h
+++ b/Engine/Plugins/Animation/ControlRig/Source/ControlRigDeveloper/Public/ModularRigController.h
@@ -52,13 +52,16 @@ class CONTROLRIGDEVELOPER_API UModularRigController : public UObject
 	bool SetConfigValueInModule(const FString& InModulePath, const FName& InVariableName, const FString& InValue, bool bSetupUndo = true);
 
 	UFUNCTION(BlueprintCallable, Category = "ControlRig | Modules")
-	bool RemoveModule(const FString& InModulesPath, bool bSetupUndo = true);
+	bool RemoveModule(const FString& InModulePath, bool bSetupUndo = true);
 
 	UFUNCTION(BlueprintCallable, Category = "ControlRig | Modules")
-	bool RenameModule(const FString& InModulesPath, const FName& InNewName, bool bSetupUndo = true);
+	bool RenameModule(const FString& InModulePath, const FName& InNewName, bool bSetupUndo = true);
 
 	UFUNCTION(BlueprintCallable, Category = "ControlRig | Modules")
-	bool ReparentModule(const FString& InModulesPath, const FString& InNewParentModulePath, bool bSetupUndo = true);
+	bool ReparentModule(const FString& InModulePath, const FString& InNewParentModulePath, bool bSetupUndo = true);
+
+
+	FName GetSafeNewName(const FString& InModuleDesiredPath);
 
 
 	void SetModel(FModularRigModel* InModel) { Model = InModel; }
diff --git a/Engine/Plugins/Animation/ControlRig/Source/ControlRigDeveloper/Public/ModularRigModel.h b/Engine/Plugins/Animation/ControlRig/Source/ControlRigDeveloper/Public/ModularRigModel.h
--- a/Engine/Plugins/Animation/ControlRig/Source/ControlRigDeveloper/Public/ModularRigModel.h
+++ b/Engine/Plugins/Animation/ControlRig/Source/ControlRigDeveloper/Public/ModularRigModel.h
@@ -68,7 +68,9 @@ public:
 
 	void UpdateCachedChildren();
 
-	FRigModuleReference* FindModule(const FString InNameSpace);
+	FRigModuleReference* FindModule(const FString InNameSpace) const;
+
+	FString FindParentNamespace(const FString InNameSpace) const;
 	
 private:
 	TWeakObjectPtr<UObject> OuterClientHost;
diff --git a/Engine/Plugins/Animation/ControlRig/Source/ControlRigEditor/Private/ControlRigBlueprintEditorLibrary.cpp b/Engine/Plugins/Animation/ControlRig/Source/ControlRigEditor/Private/ControlRigBlueprintEditorLibrary.cpp
--- a/Engine/Plugins/Animation/ControlRig/Source/ControlRigEditor/Private/ControlRigBlueprintEditorLibrary.cpp
+++ b/Engine/Plugins/Animation/ControlRig/Source/ControlRigEditor/Private/ControlRigBlueprintEditorLibrary.cpp
@@ -2,6 +2,7 @@
 
 #include "ControlRigBlueprintEditorLibrary.h"
 #include "Editor/SRigHierarchy.h"
+#include "Editor/SModularRigHierarchy.h"
 #include "AssetRegistry/AssetRegistryModule.h"
 #include "RigVMModel/RigVMPin.h"
 
@@ -88,6 +89,7 @@ void UControlRigBlueprintEditorLibrary::SetupAllEditorMenus()
 {
 	SRigHierarchy::CreateContextMenu();
 	SRigHierarchy::CreateDragDropMenu();
+	SModularRigHierarchy::CreateContextMenu();
 }
 
 TArray<FRigModuleDescription> UControlRigBlueprintEditorLibrary::GetAvailableRigModules()
diff --git a/Engine/Plugins/Animation/ControlRig/Source/ControlRigEditor/Private/ControlRigEditorModule.cpp b/Engine/Plugins/Animation/ControlRig/Source/ControlRigEditor/Private/ControlRigEditorModule.cpp
--- a/Engine/Plugins/Animation/ControlRig/Source/ControlRigEditor/Private/ControlRigEditorModule.cpp
+++ b/Engine/Plugins/Animation/ControlRig/Source/ControlRigEditor/Private/ControlRigEditorModule.cpp
@@ -38,9 +38,11 @@
 #include "Graph/ControlRigGraphPanelPinFactory.h"
 #include <Editor/ControlRigEditorCommands.h>
 #include "ControlRigHierarchyCommands.h"
+#include "ControlRigModularRigHierarchyCommands.h"
 #include "Animation/AnimSequence.h"
 #include "Editor/ControlRigEditorEditMode.h"
 #include "ControlRigElementDetails.h"
+#include "ControlRigModuleDetails.h"
 #include "ControlRigCompilerDetails.h"
 #include "ControlRigDrawingDetails.h"
 #include "ControlRigAnimGraphDetails.h"
@@ -104,6 +106,7 @@ void FControlRigEditorModule::StartupModule()
 	FControlRigEditModeCommands::Register();
 	FControlRigEditorCommands::Register();
 	FControlRigHierarchyCommands::Register();
+	FControlRigModularHierarchyCommands::Register();
 	FControlRigEditorStyle::Get();
 
 	EdGraphPanelPinFactory = MakeShared<FControlRigGraphPanelPinFactory>();
diff --git a/Engine/Plugins/Animation/ControlRig/Source/ControlRigEditor/Private/ControlRigModularRigHierarchyCommands.cpp b/Engine/Plugins/Animation/ControlRig/Source/ControlRigEditor/Private/ControlRigModularRigHierarchyCommands.cpp
new file mode 100644
--- /dev/null
+++ b/Engine/Plugins/Animation/ControlRig/Source/ControlRigEditor/Private/ControlRigModularRigHierarchyCommands.cpp
@@ -0,0 +1,12 @@
+// Copyright Epic Games, Inc. All Rights Reserved.
+
+#include "ControlRigModularRigHierarchyCommands.h"
+
+#define LOCTEXT_NAMESPACE "ControlRigModularRigHierarchyCommands"
+
+void FControlRigModularHierarchyCommands::RegisterCommands()
+{
+	UI_COMMAND(AddModuleItem, "New Module", "Add new module to the rig.", EUserInterfaceActionType::Button, FInputChord());
+}
+
+#undef LOCTEXT_NAMESPACE
diff --git a/Engine/Plugins/Animation/ControlRig/Source/ControlRigEditor/Private/ControlRigModularRigHierarchyCommands.h b/Engine/Plugins/Animation/ControlRig/Source/ControlRigEditor/Private/ControlRigModularRigHierarchyCommands.h
new file mode 100644
--- /dev/null
+++ b/Engine/Plugins/Animation/ControlRig/Source/ControlRigEditor/Private/ControlRigModularRigHierarchyCommands.h
@@ -0,0 +1,28 @@
+// Copyright Epic Games, Inc. All Rights Reserved.
+
+#pragma once
+
+#include "CoreMinimal.h"
+#include "Framework/Commands/Commands.h"
+#include "ControlRigEditorStyle.h"
+
+class FControlRigModularHierarchyCommands : public TCommands<FControlRigModularHierarchyCommands>
+{
+public:
+	FControlRigModularHierarchyCommands() : TCommands<FControlRigModularHierarchyCommands>
+	(
+		"ControlRigModuleRigHierarchy",
+		NSLOCTEXT("Contexts", "ModuleRigHierarchy", "Module Rig Hierarchy"),
+		NAME_None, // "MainFrame" // @todo Fix this crash
+		FControlRigEditorStyle::Get().GetStyleSetName() // Icon Style Set
+	)
+	{}
+	
+	/** Add Module at root */
+	TSharedPtr< FUICommandInfo > AddModuleItem;
+
+	/**
+	 * Initialize commands
+	 */
+	virtual void RegisterCommands() override;
+};
diff --git a/Engine/Plugins/Animation/ControlRig/Source/ControlRigEditor/Private/ControlRigModuleDetails.cpp b/Engine/Plugins/Animation/ControlRig/Source/ControlRigEditor/Private/ControlRigModuleDetails.cpp
new file mode 100644
--- /dev/null
+++ b/Engine/Plugins/Animation/ControlRig/Source/ControlRigEditor/Private/ControlRigModuleDetails.cpp
@@ -0,0 +1,422 @@
+// Copyright Epic Games, Inc. All Rights Reserved.
+
+#include "ControlRigModuleDetails.h"
+#include "Widgets/SWidget.h"
+#include "IDetailChildrenBuilder.h"
+#include "Widgets/Text/STextBlock.h"
+#include "Widgets/Input/SEditableTextBox.h"
+#include "Widgets/Text/SInlineEditableTextBlock.h"
+#include "Widgets/Input/SVectorInputBox.h"
+#include "Widgets/Input/SCheckBox.h"
+#include "Widgets/Input/SButton.h"
+#include "ControlRigBlueprint.h"
+#include "ControlRigElementDetails.h"
+#include "Graph/ControlRigGraph.h"
+#include "PropertyCustomizationHelpers.h"
+#include "PropertyEditorModule.h"
+#include "SEnumCombo.h"
+#include "HAL/PlatformApplicationMisc.h"
+#include "Styling/AppStyle.h"
+#include "Editor/SModularRigHierarchyTreeView.h"
+#include "StructViewerFilter.h"
+#include "StructViewerModule.h"
+
+#define LOCTEXT_NAMESPACE "ControlRigModuleDetails"
+
+static const FText ControlRigModuleDetailsMultipleValues = LOCTEXT("MultipleValues", "Multiple Values");
+
+static void RigModuleDetails_GetCustomizedInfo(TSharedRef<IPropertyHandle> InStructPropertyHandle, UControlRigBlueprint*& OutBlueprint)
+{
+	TArray<UObject*> Objects;
+	InStructPropertyHandle->GetOuterObjects(Objects);
+	for (UObject* Object : Objects)
+	{
+		if (Object->IsA<UControlRigBlueprint>())
+		{
+			OutBlueprint = CastChecked<UControlRigBlueprint>(Object);
+			break;
+		}
+
+		OutBlueprint = Object->GetTypedOuter<UControlRigBlueprint>();
+		if(OutBlueprint)
+		{
+			break;
+		}
+
+		if(const UControlRig* ControlRig = Object->GetTypedOuter<UControlRig>())
+		{
+			OutBlueprint = Cast<UControlRigBlueprint>(ControlRig->GetClass()->ClassGeneratedBy);
+			if(OutBlueprint)
+			{
+				break;
+			}
+		}
+	}
+
+	if (OutBlueprint == nullptr)
+	{
+		TArray<UPackage*> Packages;
+		InStructPropertyHandle->GetOuterPackages(Packages);
+		for (UPackage* Package : Packages)
+		{
+			if (Package == nullptr)
+			{
+				continue;
+			}
+
+			TArray<UObject*> SubObjects;
+			Package->GetDefaultSubobjects(SubObjects);
+			for (UObject* SubObject : SubObjects)
+			{
+				if (UControlRig* Rig = Cast<UControlRig>(SubObject))
+				{
+					UControlRigBlueprint* Blueprint = Cast<UControlRigBlueprint>(Rig->GetClass()->ClassGeneratedBy);
+					if (Blueprint)
+					{
+						if(Blueprint->GetOutermost() == Package)
+						{
+							OutBlueprint = Blueprint;
+							break;
+						}
+					}
+				}
+			}
+
+			if (OutBlueprint)
+			{
+				break;
+			}
+		}
+	}
+}
+
+static UControlRigBlueprint* RigModuleDetails_GetBlueprintFromRig(UModularRig* InRig)
+{
+	if(InRig == nullptr)
+	{
+		return nullptr;
+	}
+
+	UControlRigBlueprint* Blueprint = InRig->GetTypedOuter<UControlRigBlueprint>();
+	if(Blueprint == nullptr)
+	{
+		Blueprint = Cast<UControlRigBlueprint>(InRig->GetClass()->ClassGeneratedBy);
+	}
+	return Blueprint;
+}
+
+void FRigModuleInstanceDetails::CustomizeDetails(IDetailLayoutBuilder& DetailBuilder)
+{
+	PerModuleInfos.Reset();
+
+	TArray<TWeakObjectPtr<UObject>> DetailObjects;
+	DetailBuilder.GetObjectsBeingCustomized(DetailObjects);
+	for(TWeakObjectPtr<UObject> DetailObject : DetailObjects)
+	{
+		URigVMDetailsViewWrapperObject* WrapperObject = CastChecked<URigVMDetailsViewWrapperObject>(DetailObject.Get());
+
+		const FString Path = WrapperObject->GetContent<FRigModuleInstance>().GetPath();
+
+		FPerModuleInfo Info;
+		Info.WrapperObject = WrapperObject;
+		Info.Module = Cast<UModularRig>(WrapperObject->GetSubject())->GetHandle(Path);
+
+		if(!Info.Module.IsValid())
+		{
+			return;
+		}
+		if(const UControlRigBlueprint* Blueprint = Info.GetBlueprint())
+		{
+			if (UModularRig* ModularRig = Cast<UModularRig>(Blueprint->GetObjectBeingDebugged()))
+			{
+				Info.DefaultModule = ModularRig->GetHandle(Path);
+			}
+		}
+
+		PerModuleInfos.Add(Info);
+	}
+
+	DetailBuilder.HideCategory(TEXT("RigModuleInstance"));
+
+	IDetailCategoryBuilder& GeneralCategory = DetailBuilder.EditCategory(TEXT("General"), LOCTEXT("General", "General"));
+
+	{
+		GeneralCategory.AddCustomRow(FText::FromString(TEXT("Name")))
+		.NameContent()
+		[
+			SNew(STextBlock)
+			.Text(FText::FromString(TEXT("Name")))
+			.Font(IDetailLayoutBuilder::GetDetailFont())
+			.IsEnabled(true)
+		]
+		.ValueContent()
+		[
+			SNew(SInlineEditableTextBlock)
+			.Font(IDetailLayoutBuilder::GetDetailFont())
+			.Text(this, &FRigModuleInstanceDetails::GetName)
+			.IsEnabled(true)
+		];
+
+		GeneralCategory.AddCustomRow(FText::FromString(TEXT("RigClass")))
+		.NameContent()
+		[
+			SNew(STextBlock)
+			.Text(FText::FromString(TEXT("RigClass")))
+			.Font(IDetailLayoutBuilder::GetDetailFont())
+			.IsEnabled(true)
+		]
+		.ValueContent()
+		[
+			SNew(SInlineEditableTextBlock)
+			.Font(IDetailLayoutBuilder::GetDetailFont())
+			.Text(this, &FRigModuleInstanceDetails::GetRigClassPath)
+			.IsEnabled(true)
+		];
+	}
+
+	IDetailCategoryBuilder& ConnectionsCategory = DetailBuilder.EditCategory(TEXT("Connections"), LOCTEXT("Connections", "Connections"));
+
+	{
+		FRigElementKeyRedirector Redirector = GetConnections();
+		for(TPair<FRigElementKey, FCachedRigElement>& Connection : Redirector.InternalKeyToExternalKey)
+		{
+			if (!Connection.Key.IsTypeOf(ERigElementType::Connector))
+			{
+				continue;
+			}
+			const FText Label = FText::FromName(Connection.Key.Name);
+			const FRigElementKey& Target = Connection.Value.GetKey();
+			ConnectionsCategory.AddCustomRow(Label)
+				.NameContent()
+				[
+					SNew(STextBlock)
+					.Text(Label)
+					.Font(IDetailLayoutBuilder::GetDetailFont())
+					.IsEnabled(true)
+				]
+				.ValueContent()
+				[
+					SNew(STextBlock)
+					.Text(FText::FromName(Target.Name))
+					.Font(IDetailLayoutBuilder::GetDetailFont())
+					.IsEnabled(true)
+				];
+		}
+	}
+
+	IDetailCategoryBuilder& ConfigValuesCategory = DetailBuilder.EditCategory(TEXT("Config Values"), LOCTEXT("ConfigValues", "Config Values"));
+	{
+		TArray<UObject*> DetailObjectsRaw;
+		DetailObjectsRaw.Reserve(DetailObjects.Num());
+		for (TWeakObjectPtr<UObject> Obj : DetailObjects)
+		{
+			if (Obj.IsValid())
+			{
+				DetailObjectsRaw.Add(Obj.Get());
+			}
+		}
+		
+		TArray<FRigVMExternalVariable> Variables = GetConfigValues();
+		for(FRigVMExternalVariable& Variable : Variables)
+		{
+			if (!Variable.bIsPublic)
+			{
+				continue;
+			}
+			const FText Label = FText::FromName(Variable.Name);
+
+			if (FRigModuleInstance* ModuleInstance = PerModuleInfos[0].GetModule())
+			{
+				TSoftObjectPtr<UControlRig> Rig = ModuleInstance->Rig;
+				if (Rig.IsValid())
+				{
+					if (FProperty* Property = Rig.Get()->GetClass()->FindPropertyByName(Variable.Name))
+					{
+						uint8* Container = (uint8*)Rig.Get();
+
+						TSharedPtr<FStructOnScope> StructOnScope = MakeShareable(new FStructOnScope(Rig.Get()->GetClass(), Container));
+						IDetailPropertyRow* Row = ConfigValuesCategory.AddExternalObjectProperty({Rig.Get()}, Variable.Name);
+						
+
+						Row->DisplayName(FText::FromName(Variable.Name));
+
+						const FSimpleDelegate OnValueChangedDelegate = FSimpleDelegate::CreateSP(this, &FRigModuleInstanceDetails::OnConfigValueChanged, Variable.Name);
+
+						TSharedPtr<IPropertyHandle> Handle = Row->GetPropertyHandle();
+						Handle->SetOnPropertyValueChanged(OnValueChangedDelegate);
+						Handle->SetOnChildPropertyValueChanged(OnValueChangedDelegate);
+						
+					}
+				}
+			}
+		}
+	}
+}
+
+FString FRigModuleInstanceDetails::GetModulePath() const
+{
+	check(PerModuleInfos.Num() == 1);
+	if (FRigModuleInstance* Module = PerModuleInfos[0].GetModule())
+	{
+		return Module->GetPath();
+	}
+	return FString();
+}
+
+FText FRigModuleInstanceDetails::GetName() const
+{
+	if(PerModuleInfos.Num() > 1)
+	{
+		bool bSame = true;
+		for (int32 i=1; i<PerModuleInfos.Num(); ++i)
+		{
+			if (PerModuleInfos[i].GetModule()->Name !=  PerModuleInfos[0].GetModule()->Name)
+			{
+				bSame = false;
+				break;
+			}
+		}
+		if (!bSame)
+		{
+			return ControlRigModuleDetailsMultipleValues;
+		}
+	}
+	return FText::FromName(PerModuleInfos[0].GetModule()->Name);
+}
+
+FText FRigModuleInstanceDetails::GetRigClassPath() const
+{
+	if(PerModuleInfos.Num() > 1)
+	{
+		bool bSame = true;
+		for (int32 i=1; i<PerModuleInfos.Num(); ++i)
+		{
+			if (PerModuleInfos[i].GetModule()->Rig->GetClass() !=  PerModuleInfos[0].GetModule()->Rig->GetClass())
+			{
+				bSame = false;
+				break;
+			}
+		}
+		if (!bSame)
+		{
+			return ControlRigModuleDetailsMultipleValues;
+		}
+	}
+
+	if (FRigModuleInstance* Module = PerModuleInfos[0].GetModule())
+	{
+		if (TSoftObjectPtr<UControlRig> Rig = Module->Rig)
+		{
+			if (Rig.IsValid())
+			{
+				return FText::FromString(Module->Rig->GetClass()->GetClassPathName().ToString());
+			}
+		}
+	}
+
+	return FText();
+}
+
+FRigElementKeyRedirector FRigModuleInstanceDetails::GetConnections() const
+{
+	if(PerModuleInfos.Num() > 1)
+	{
+		return FRigElementKeyRedirector();
+	}
+
+	if (FRigModuleInstance* Module = PerModuleInfos[0].GetModule())
+	{
+		if (TSoftObjectPtr<UControlRig> Rig = Module->Rig)
+		{
+			if (Rig.IsValid())
+			{
+				return Rig->GetElementKeyRedirector();
+			}
+		}
+	}
+
+	return FRigElementKeyRedirector();
+}
+
+TArray<FRigVMExternalVariable> FRigModuleInstanceDetails::GetConfigValues() const
+{
+	if(PerModuleInfos.Num() > 1)
+	{
+		return TArray<FRigVMExternalVariable>();
+	}
+
+	if (FRigModuleInstance* Module = PerModuleInfos[0].GetModule())
+	{
+		if (TSoftObjectPtr<UControlRig> Rig = Module->Rig)
+		{
+			if (Rig.IsValid())
+			{
+				return Rig->GetExternalVariables();
+			}
+		}
+	}
+
+	return TArray<FRigVMExternalVariable>();
+}
+
+void FRigModuleInstanceDetails::OnConfigValueChanged(const FName InVariableName)
+{
+	if (FRigModuleInstance* ModuleInstance = PerModuleInfos[0].GetModule())
+	{
+		TSoftObjectPtr<UControlRig> Rig = ModuleInstance->Rig;
+		if (Rig.IsValid())
+		{
+			FString ValueStr = Rig->GetVariableAsString(InVariableName);
+			if (UControlRigBlueprint* Blueprint = RigModuleDetails_GetBlueprintFromRig(PerModuleInfos[0].GetRig()))
+			{
+				UModularRigController* Controller = Blueprint->GetModularRigController();
+				Controller->SetConfigValueInModule(ModuleInstance->GetPath(), InVariableName, ValueStr);
+			}
+		}
+	}
+}
+
+TArray<FString> FRigModuleInstanceDetails::GetModulePaths() const
+{
+	TArray<FString> Paths;
+	Algo::Transform(PerModuleInfos, Paths, [](const FPerModuleInfo& Info)
+	{
+		return Info.GetModule()->GetPath();
+	});
+	return Paths;
+}
+
+const FRigModuleInstanceDetails::FPerModuleInfo& FRigModuleInstanceDetails::FindModule(const FString& InPath) const
+{
+	const FPerModuleInfo* Info = FindModuleByPredicate([InPath](const FPerModuleInfo& Info)
+	{
+		return Info.GetModule()->GetPath() == InPath;
+	});
+
+	if(Info)
+	{
+		return *Info;
+	}
+
+	static const FPerModuleInfo EmptyInfo;
+	return EmptyInfo;
+}
+
+const FRigModuleInstanceDetails::FPerModuleInfo* FRigModuleInstanceDetails::FindModuleByPredicate(const TFunction<bool(const FPerModuleInfo&)>& InPredicate) const
+{
+	return PerModuleInfos.FindByPredicate(InPredicate);
+}
+
+bool FRigModuleInstanceDetails::ContainsModuleByPredicate(const TFunction<bool(const FPerModuleInfo&)>& InPredicate) const
+{
+	return PerModuleInfos.ContainsByPredicate(InPredicate);
+}
+
+void FRigModuleInstanceDetails::RegisterSectionMappings(FPropertyEditorModule& PropertyEditorModule, UClass* InClass)
+{
+	TSharedRef<FPropertySection> MetadataSection = PropertyEditorModule.FindOrCreateSection(InClass->GetFName(), "Metadata", LOCTEXT("Metadata", "Metadata"));
+	MetadataSection->AddCategory("Metadata");
+}
+
+
+#undef LOCTEXT_NAMESPACE
diff --git a/Engine/Plugins/Animation/ControlRig/Source/ControlRigEditor/Private/ControlRigModuleDetails.h b/Engine/Plugins/Animation/ControlRig/Source/ControlRigEditor/Private/ControlRigModuleDetails.h
new file mode 100644
--- /dev/null
+++ b/Engine/Plugins/Animation/ControlRig/Source/ControlRigEditor/Private/ControlRigModuleDetails.h
@@ -0,0 +1,124 @@
+// Copyright Epic Games, Inc. All Rights Reserved.
+
+#pragma once
+
+#include "CoreMinimal.h"
+#include "IDetailCustomization.h"
+#include "IPropertyTypeCustomization.h"
+#include "ControlRig.h"
+#include "ModularRig.h"
+#include "ControlRigBlueprint.h"
+#include "Editor/ControlRigWrapperObject.h"
+#include "Styling/SlateTypes.h"
+#include "IPropertyUtilities.h"
+#include "SSearchableComboBox.h"
+#include "Widgets/Input/SSegmentedControl.h"
+#include "DetailCategoryBuilder.h"
+#include "DetailLayoutBuilder.h"
+#include "DetailWidgetRow.h"
+#include "HAL/PlatformApplicationMisc.h"
+#include "Internationalization/FastDecimalFormat.h"
+#include "ScopedTransaction.h"
+#include "Styling/AppStyle.h"
+#include "Algo/Transform.h"
+
+class IPropertyHandle;
+
+class FRigModuleInstanceDetails : public IDetailCustomization
+{
+public:
+
+	/** IDetailCustomization interface */
+	virtual void CustomizeDetails(IDetailLayoutBuilder& DetailBuilder) override;
+	virtual void BeginDestroy() {};
+
+	// Makes a new instance of this detail layout class for a specific detail view requesting it
+	static TSharedRef<IDetailCustomization> MakeInstance()
+	{
+		return MakeShareable(new FRigModuleInstanceDetails);
+	}
+
+	FString GetModulePath() const;
+	FText GetName() const;
+	FText GetRigClassPath() const;
+	FRigElementKeyRedirector GetConnections() const;
+	TArray<FRigVMExternalVariable> GetConfigValues() const;
+
+	void OnConfigValueChanged(const FName InVariableName);
+	
+	TArray<FString> GetModulePaths() const;
+
+	TArray<FRigModuleInstance> GetModulesInDetailsView(const TArray<FString>& InFilter = TArray<FString>()) const
+	{
+		TArray<FRigModuleInstance> Elements;
+		for(const FPerModuleInfo& Info : PerModuleInfos)
+		{
+			FRigModuleInstance Content = Info.WrapperObject->GetContent<FRigModuleInstance>();
+			if(!InFilter.IsEmpty() && !InFilter.Contains(Content.GetPath()))
+			{
+				continue;
+			}
+			Elements.Add(Content);
+		}
+		return Elements;
+	}
+
+	struct FPerModuleInfo
+	{
+		FPerModuleInfo()
+			: WrapperObject()
+			, Module()
+			, DefaultModule()
+		{}
+
+		bool IsValid() const { return Module.IsValid(); }
+		operator bool() const { return IsValid(); }
+
+		UModularRig* GetRig() const { return (UModularRig*)Module.GetRig(); }
+		UModularRig* GetDefaultRig() const
+		{
+			if(DefaultModule.IsValid())
+			{
+				return (UModularRig*)DefaultModule.GetRig();
+			}
+			return GetRig();
+		}
+
+		UControlRigBlueprint* GetBlueprint() const
+		{
+			if(const UModularRig* ControlRig = GetRig()->GetTypedOuter<UModularRig>())
+			{
+				return Cast<UControlRigBlueprint>(ControlRig->GetClass()->ClassGeneratedBy);
+			}
+			return GetDefaultRig()->GetTypedOuter<UControlRigBlueprint>();
+		}
+
+		FRigModuleInstance* GetModule() const
+		{
+			return (FRigModuleInstance*)Module.Get();
+		}
+
+		FRigModuleInstance* GetDefaultElement() const
+		{
+			if(DefaultModule)
+			{
+				return (FRigModuleInstance*)DefaultModule.Get();
+			}
+			return GetModule();
+		}
+
+		TWeakObjectPtr<URigVMDetailsViewWrapperObject> WrapperObject;
+		FModuleInstanceHandle Module;
+		FModuleInstanceHandle DefaultModule;
+	};
+
+	const FPerModuleInfo& FindModule(const FString& InKey) const;
+	const FPerModuleInfo* FindModuleByPredicate(const TFunction<bool(const FPerModuleInfo&)>& InPredicate) const;
+	bool ContainsModuleByPredicate(const TFunction<bool(const FPerModuleInfo&)>& InPredicate) const;
+
+	virtual void RegisterSectionMappings(FPropertyEditorModule& PropertyEditorModule, UClass* InClass);
+
+protected:
+
+	TArray<FPerModuleInfo> PerModuleInfos;
+};
diff --git a/Engine/Plugins/Animation/ControlRig/Source/ControlRigEditor/Private/Editor/ControlRigContextMenuContext.cpp b/Engine/Plugins/Animation/ControlRig/Source/ControlRigEditor/Private/Editor/ControlRigContextMenuContext.cpp
--- a/Engine/Plugins/Animation/ControlRig/Source/ControlRigEditor/Private/Editor/ControlRigContextMenuContext.cpp
+++ b/Engine/Plugins/Animation/ControlRig/Source/ControlRigEditor/Private/Editor/ControlRigContextMenuContext.cpp
@@ -80,6 +80,15 @@ SRigHierarchy* UControlRigContextMenuContext::GetRigHierarchyPanel() const
 	return nullptr;
 }
 
+SModularRigHierarchy* UControlRigContextMenuContext::GetModularRigHierarchyPanel() const
+{
+	if (const TSharedPtr<SModularRigHierarchy> ModularRigHierarchyPanel = MenuSpecificContext.ModularRigHierarchyPanel.Pin())
+	{
+		return ModularRigHierarchyPanel.Get();
+	}
+	return nullptr;
+}
+
 FControlRigEditor* UControlRigContextMenuContext::GetControlRigEditor() const
 {
 	if (const TSharedPtr<FControlRigEditor> Editor = WeakControlRigEditor.Pin())
diff --git a/Engine/Plugins/Animation/ControlRig/Source/ControlRigEditor/Private/Editor/ControlRigContextMenuContext.h b/Engine/Plugins/Animation/ControlRig/Source/ControlRigEditor/Private/Editor/ControlRigContextMenuContext.h
--- a/Engine/Plugins/Animation/ControlRig/Source/ControlRigEditor/Private/Editor/ControlRigContextMenuContext.h
+++ b/Engine/Plugins/Animation/ControlRig/Source/ControlRigEditor/Private/Editor/ControlRigContextMenuContext.h
@@ -15,6 +15,7 @@ class URigVMGraph;
 class URigVMNode;
 class URigVMPin;
 class SRigHierarchy;
+class SModularRigHierarchy;
 
 USTRUCT(BlueprintType)
 struct FControlRigRigHierarchyDragAndDropContext
@@ -96,6 +97,8 @@ struct FControlRigMenuSpecificContext
 	TWeakPtr<SRigHierarchy> RigHierarchyPanel;
 	
 	FControlRigRigHierarchyDragAndDropContext RigHierarchyDragAndDropContext;
+
+	TWeakPtr<SModularRigHierarchy> ModularRigHierarchyPanel;
 	
 	FControlRigGraphNodeContextMenuContext GraphNodeContextMenuContext;
 
@@ -140,6 +143,8 @@ public:
     FControlRigRigHierarchyToGraphDragAndDropContext GetRigHierarchyToGraphDragAndDropContext();
 
 	SRigHierarchy* GetRigHierarchyPanel() const;
+
+	SModularRigHierarchy* GetModularRigHierarchyPanel() const;
 	
 	FControlRigEditor* GetControlRigEditor() const;
 
diff --git a/Engine/Plugins/Animation/ControlRig/Source/ControlRigEditor/Private/Editor/ControlRigEditor.cpp b/Engine/Plugins/Animation/ControlRig/Source/ControlRigEditor/Private/Editor/ControlRigEditor.cpp
--- a/Engine/Plugins/Animation/ControlRig/Source/ControlRigEditor/Private/Editor/ControlRigEditor.cpp
+++ b/Engine/Plugins/Animation/ControlRig/Source/ControlRigEditor/Private/Editor/ControlRigEditor.cpp
@@ -125,6 +125,7 @@ FControlRigEditor::FControlRigEditor()
 	, ActiveController(nullptr)
 	, bExecutionControlRig(true)
 	, RigHierarchyTabCount(0)
+	, ModularRigHierarchyTabCount(0)
 	, bIsConstructionEventRunning(false)
 	, LastHierarchyHash(INDEX_NONE)
 	, bRefreshDirectionManipulationTargetsRequired(false)
@@ -1136,6 +1137,67 @@ bool FControlRigEditor::DetailViewShowsRigElement(FRigElementKey InKey) const
 	return false;
 }
 
+void FControlRigEditor::SetDetailViewForRigModules(const TArray<FString>& InKeys)
+{
+	if(IsDetailsPanelRefreshSuspended())
+	{
+		return;
+	}
+
+	ClearDetailObject();
+
+	UControlRigBlueprint* RigBlueprint = Cast<UControlRigBlueprint>(GetBlueprintObj());
+	UModularRig* RigBeingDebugged = Cast<UModularRig>(RigBlueprint->GetDebuggedControlRig());
+	TArray<UObject*> Objects;
+
+	for(const FString& Key : InKeys)
+	{
+		FRigModuleInstance* Element = RigBeingDebugged->FindModule(Key);
+		if (Element == nullptr)
+		{
+			continue;
+		}
+
+		URigVMDetailsViewWrapperObject* WrapperObject = URigVMDetailsViewWrapperObject::MakeInstance(GetDetailWrapperClass(), GetBlueprintObj(), FRigModuleInstance::StaticStruct(), (uint8*)Element, RigBeingDebugged);
+		WrapperObject->GetWrappedPropertyChangedChainEvent().AddSP(this, &FControlRigEditor::OnWrappedPropertyChangedChainEvent);
+		WrapperObject->AddToRoot();
+
+		Objects.Add(WrapperObject);
+	}
+	
+	SetDetailObjects(Objects);
+}
+
+bool FControlRigEditor::DetailViewShowsAnyRigModule() const
+{
+	return DetailViewShowsStruct(FRigModuleInstance::StaticStruct());
+}
+
+bool FControlRigEditor::DetailViewShowsRigModule(FString InKey) const
+{
+	TArray< TWeakObjectPtr<UObject> > SelectedObjects = Inspector->GetSelectedObjects();
+	for (TWeakObjectPtr<UObject> SelectedObject : SelectedObjects)
+	{
+		if (SelectedObject.IsValid())
+		{
+			if(URigVMDetailsViewWrapperObject* WrapperObject = Cast<URigVMDetailsViewWrapperObject>(SelectedObject.Get()))
+			{
+				if (const UScriptStruct* WrappedStruct = WrapperObject->GetWrappedStruct())
+				{
+					if (WrappedStruct->IsChildOf(FRigModuleInstance::StaticStruct()))
+					{
+						if(WrapperObject->GetContent<FRigModuleInstance>().GetPath() == InKey)
+						{
+							return true;
+						}
+					}
+				}
+			}
+		}
+	}
+	return false;
+}
+
 void FControlRigEditor::Compile()
 {
 	{
diff --git a/Engine/Plugins/Animation/ControlRig/Source/ControlRigEditor/Private/Editor/ControlRigEditor.h b/Engine/Plugins/Animation/ControlRig/Source/ControlRigEditor/Private/Editor/ControlRigEditor.h
--- a/Engine/Plugins/Animation/ControlRig/Source/ControlRigEditor/Private/Editor/ControlRigEditor.h
+++ b/Engine/Plugins/Animation/ControlRig/Source/ControlRigEditor/Private/Editor/ControlRigEditor.h
@@ -86,6 +86,7 @@ public:
 	virtual void CreateEmptyGraphContent(URigVMController* InController) override;
 
 	int32 GetRigHierarchyTabCount() const { return RigHierarchyTabCount; }
+	int32 GetModularRigHierarchyTabCount() const { return ModularRigHierarchyTabCount; }
 
 	bool IsModularRig() const;
 
@@ -125,6 +126,12 @@ public:
 	void SetDetailViewForRigElements(const TArray<FRigElementKey>& InKeys);
 	bool DetailViewShowsAnyRigElement() const;
 	bool DetailViewShowsRigElement(FRigElementKey InKey) const;
+
+	//void SetDetailViewForRigModules();
+	void SetDetailViewForRigModules(const TArray<FString>& InKeys);
+	bool DetailViewShowsAnyRigModule() const;
+	bool DetailViewShowsRigModule(FString InKey) const;
+
 	virtual void RefreshDetailView() override;
 
 	void CreatePersonaToolKitIfRequired();
@@ -306,6 +313,7 @@ protected:
 	bool IsConstructionModeEnabled() const;
 
 	int32 RigHierarchyTabCount;
+	int32 ModularRigHierarchyTabCount;
 	TWeakObjectPtr<AStaticMeshActor> WeakGroundActorPtr;
 
 	void OnPreForwardsSolve_AnyThread(UControlRig* InRig, const FName& InEventName);
@@ -334,5 +342,7 @@ protected:
 	friend class FControlRigEditorMode;
 	friend class SControlRigStackView;
 	friend class SRigHierarchy;
+	friend class SModularRigHierarchy;
 	friend struct FRigHierarchyTabSummoner;
+	friend struct FModularRigHierarchyTabSummoner;
 };
diff --git a/Engine/Plugins/Animation/ControlRig/Source/ControlRigEditor/Private/Editor/ControlRigEditorMode.cpp b/Engine/Plugins/Animation/ControlRig/Source/ControlRigEditor/Private/Editor/ControlRigEditorMode.cpp
--- a/Engine/Plugins/Animation/ControlRig/Source/ControlRigEditor/Private/Editor/ControlRigEditorMode.cpp
+++ b/Engine/Plugins/Animation/ControlRig/Source/ControlRigEditor/Private/Editor/ControlRigEditorMode.cpp
@@ -7,6 +7,7 @@
 #include "IPersonaToolkit.h"
 #include "PersonaTabs.h"
 #include "Editor/RigHierarchyTabSummoner.h"
+#include "Editor/ModularRigHierarchyTabSummoner.h"
 #include "Editor/RigVMExecutionStackTabSummoner.h"
 #include "Editor/RigCurveContainerTabSummoner.h"
 #include "Editor/RigValidationTabSummoner.h"
@@ -20,6 +21,10 @@ FControlRigEditorMode::FControlRigEditorMode(const TSharedRef<FControlRigEditor>
 	ControlRigBlueprintPtr = CastChecked<UControlRigBlueprint>(InControlRigEditor->GetBlueprintObj());
 
 	TabFactories.RegisterFactory(MakeShared<FRigHierarchyTabSummoner>(InControlRigEditor));
+	if (ControlRigBlueprintPtr->IsModularRig())
+	{
+		TabFactories.RegisterFactory(MakeShared<FModularRigHierarchyTabSummoner>(InControlRigEditor));
+	}
 	TabFactories.RegisterFactory(MakeShared<FRigVMExecutionStackTabSummoner>(InControlRigEditor));
 	TabFactories.RegisterFactory(MakeShared<FRigCurveContainerTabSummoner>(InControlRigEditor));
 	TabFactories.RegisterFactory(MakeShared<FRigValidationTabSummoner>(InControlRigEditor));
@@ -41,7 +46,7 @@ FControlRigEditorMode::FControlRigEditorMode(const TSharedRef<FControlRigEditor>
 
 	if(bCreateDefaultLayout)
 	{
-		TabLayout = FTabManager::NewLayout("Standalone_ControlRigEditMode_Layout_v1.5")
+		TabLayout = FTabManager::NewLayout("Standalone_ControlRigEditMode_Layout_v1.6")
 			->AddArea
 			(
 				// Main application area
@@ -68,12 +73,21 @@ FControlRigEditorMode::FControlRigEditorMode(const TSharedRef<FControlRigEditor>
 						->Split
 						(
 							//	Left bottom - rig/hierarchy
-							FTabManager::NewStack()
-							->SetSizeCoefficient(0.5f)
-							->AddTab(FRigHierarchyTabSummoner::TabID, ETabState::OpenedTab)
-							->AddTab(FRigVMExecutionStackTabSummoner::TabID, ETabState::OpenedTab)
-							->AddTab(FRigCurveContainerTabSummoner::TabID, ETabState::OpenedTab)
-							->AddTab(FBlueprintEditorTabs::MyBlueprintID, ETabState::OpenedTab)
+							ControlRigBlueprintPtr->IsModularRig() ?
+								FTabManager::NewStack()
+								->SetSizeCoefficient(0.5f)
+								->AddTab(FRigHierarchyTabSummoner::TabID, ETabState::OpenedTab)
+								->AddTab(FModularRigHierarchyTabSummoner::TabID, ETabState::OpenedTab)
+								->AddTab(FRigVMExecutionStackTabSummoner::TabID, ETabState::OpenedTab)
+								->AddTab(FRigCurveContainerTabSummoner::TabID, ETabState::OpenedTab)
+								->AddTab(FBlueprintEditorTabs::MyBlueprintID, ETabState::OpenedTab)
+									:
+								FTabManager::NewStack()
+								->SetSizeCoefficient(0.5f)
+								->AddTab(FRigHierarchyTabSummoner::TabID, ETabState::OpenedTab)
+								->AddTab(FRigVMExecutionStackTabSummoner::TabID, ETabState::OpenedTab)
+								->AddTab(FRigCurveContainerTabSummoner::TabID, ETabState::OpenedTab)
+								->AddTab(FBlueprintEditorTabs::MyBlueprintID, ETabState::OpenedTab)
 						)
 					)
 					->Split
@@ -170,12 +184,21 @@ FModularRigEditorMode::FModularRigEditorMode(const TSharedRef<FControlRigEditor>
 					->Split
 					(
 						//	Left bottom - rig/hierarchy
-						FTabManager::NewStack()
-						->SetSizeCoefficient(0.5f)
-						->AddTab(FBlueprintEditorTabs::CompilerResultsID, ETabState::ClosedTab)
-						->AddTab(FRigHierarchyTabSummoner::TabID, ETabState::OpenedTab)
-						->AddTab(FRigCurveContainerTabSummoner::TabID, ETabState::OpenedTab)
-						->AddTab(FBlueprintEditorTabs::MyBlueprintID, ETabState::OpenedTab)
+						InControlRigEditor->GetControlRigBlueprint()->IsModularRig() ?
+							FTabManager::NewStack()
+							->SetSizeCoefficient(0.5f)
+							->AddTab(FBlueprintEditorTabs::CompilerResultsID, ETabState::ClosedTab)
+							->AddTab(FRigHierarchyTabSummoner::TabID, ETabState::OpenedTab)
+							->AddTab(FModularRigHierarchyTabSummoner::TabID, ETabState::OpenedTab)
+							->AddTab(FRigCurveContainerTabSummoner::TabID, ETabState::OpenedTab)
+							->AddTab(FBlueprintEditorTabs::MyBlueprintID, ETabState::OpenedTab)
+								:
+							FTabManager::NewStack()
+							->SetSizeCoefficient(0.5f)
+							->AddTab(FBlueprintEditorTabs::CompilerResultsID, ETabState::ClosedTab)
+							->AddTab(FRigHierarchyTabSummoner::TabID, ETabState::OpenedTab)
+							->AddTab(FRigCurveContainerTabSummoner::TabID, ETabState::OpenedTab)
+							->AddTab(FBlueprintEditorTabs::MyBlueprintID, ETabState::OpenedTab)
 					)
 				)
 				->Split
diff --git a/Engine/Plugins/Animation/ControlRig/Source/ControlRigEditor/Private/Editor/ControlRigWrapperObject.cpp b/Engine/Plugins/Animation/ControlRig/Source/ControlRigEditor/Private/Editor/ControlRigWrapperObject.cpp
--- a/Engine/Plugins/Animation/ControlRig/Source/ControlRigEditor/Private/Editor/ControlRigWrapperObject.cpp
+++ b/Engine/Plugins/Animation/ControlRig/Source/ControlRigEditor/Private/Editor/ControlRigWrapperObject.cpp
@@ -4,6 +4,7 @@
 
 #if WITH_EDITOR
 #include "ControlRigElementDetails.h"
+#include "ControlRigModuleDetails.h"
 #include "PropertyEditorModule.h"
 #include "Modules/ModuleManager.h"
 #endif
@@ -46,6 +47,14 @@ UClass* UControlRigWrapperObject::GetClassForStruct(UScriptStruct* InStruct, boo
 			}
 		}
 	}
+	else if (InStruct->IsChildOf(FRigModuleInstance::StaticStruct()))
+	{
+		FPropertyEditorModule& PropertyEditorModule = FModuleManager::LoadModuleChecked<FPropertyEditorModule>("PropertyEditor");
+		if (!PropertyEditorModule.GetClassNameToDetailLayoutNameMap().Contains(WrapperClassName))
+		{
+			PropertyEditorModule.RegisterCustomClassLayout(WrapperClassName, FOnGetDetailCustomizationInstance::CreateStatic(&FRigModuleInstanceDetails::MakeInstance));
+		}
+	}
 #endif
 
 	return Class;
diff --git a/Engine/Plugins/Animation/ControlRig/Source/ControlRigEditor/Private/Editor/ModularRigHierarchyTabSummoner.cpp b/Engine/Plugins/Animation/ControlRig/Source/ControlRigEditor/Private/Editor/ModularRigHierarchyTabSummoner.cpp
new file mode 100644
--- /dev/null
+++ b/Engine/Plugins/Animation/ControlRig/Source/ControlRigEditor/Private/Editor/ModularRigHierarchyTabSummoner.cpp
@@ -0,0 +1,77 @@
+// Copyright Epic Games, Inc. All Rights Reserved.
+
+#include "Editor/ModularRigHierarchyTabSummoner.h"
+#include "Editor/SModularRigHierarchy.h"
+#include "ControlRigEditorStyle.h"
+#include "Editor/ControlRigEditor.h"
+#include "Widgets/Docking/SDockTab.h"
+
+#define LOCTEXT_NAMESPACE "ModularRigHierarchyTabSummoner"
+
+const FName FModularRigHierarchyTabSummoner::TabID(TEXT("ModularRigHierarchy"));
+
+FModularRigHierarchyTabSummoner::FModularRigHierarchyTabSummoner(const TSharedRef<FControlRigEditor>& InControlRigEditor)
+	: FWorkflowTabFactory(TabID, InControlRigEditor)
+	, ControlRigEditor(InControlRigEditor)
+{
+	TabLabel = LOCTEXT("ModularRigHierarchyTabLabel", "Modular Rig Hierarchy");
+	TabIcon = FSlateIcon(FControlRigEditorStyle::Get().GetStyleSetName(), "ModularRigHierarchy.TabIcon");
+
+	ViewMenuDescription = LOCTEXT("ModularRigHierarchy_ViewMenu_Desc", "Modular Rig Hierarchy");
+	ViewMenuTooltip = LOCTEXT("ModularRigHierarchy_ViewMenu_ToolTip", "Show the Modular Rig Hierarchy tab");
+}
+
+FTabSpawnerEntry& FModularRigHierarchyTabSummoner::RegisterTabSpawner(TSharedRef<FTabManager> InTabManager, const FApplicationMode* CurrentApplicationMode) const
+{
+	FTabSpawnerEntry& SpawnerEntry = FWorkflowTabFactory::RegisterTabSpawner(InTabManager, CurrentApplicationMode);
+
+	SpawnerEntry.SetReuseTabMethod(FOnFindTabToReuse::CreateLambda([](const FTabId& InTabId) ->TSharedPtr<SDockTab> {
+	
+		return TSharedPtr<SDockTab>();
+
+	}));
+
+	return SpawnerEntry;
+}
+
+TSharedRef<SWidget> FModularRigHierarchyTabSummoner::CreateTabBody(const FWorkflowTabSpawnInfo& Info) const
+{
+	ControlRigEditor.Pin()->ModularRigHierarchyTabCount++;
+	return SNew(SModularRigHierarchy, ControlRigEditor.Pin().ToSharedRef());
+}
+
+TSharedRef<SDockTab> FModularRigHierarchyTabSummoner::SpawnTab(const FWorkflowTabSpawnInfo& Info) const
+{
+	TSharedRef<SDockTab>  DockTab = FWorkflowTabFactory::SpawnTab(Info);
+	TWeakPtr<SDockTab> WeakDockTab = DockTab;
+	DockTab->SetCanCloseTab(SDockTab::FCanCloseTab::CreateLambda([WeakDockTab]()
+    {
+		int32 HierarchyTabCount = 0;
+		if (TSharedPtr<SDockTab> SharedDocTab = WeakDockTab.Pin())
+		{
+			if(SWidget* Content = &SharedDocTab->GetContent().Get())
+			{
+				SModularRigHierarchy* RigHierarchy = (SModularRigHierarchy*)Content;
+				if(FControlRigEditor* ControlRigEditorForTab = RigHierarchy->GetControlRigEditor())
+				{
+					HierarchyTabCount = ControlRigEditorForTab->GetModularRigHierarchyTabCount();
+				}
+			}
+		}
+		return HierarchyTabCount > 0;
+    }));
+	DockTab->SetOnTabClosed( SDockTab::FOnTabClosedCallback::CreateLambda([](TSharedRef<SDockTab> DockTab)
+	{
+		if(SWidget* Content = &DockTab->GetContent().Get())
+		{
+			SModularRigHierarchy* RigHierarchy = (SModularRigHierarchy*)Content;
+			if(FControlRigEditor* ControlRigEditorForTab = RigHierarchy->GetControlRigEditor())
+			{
+				ControlRigEditorForTab->ModularRigHierarchyTabCount--;
+			}
+		}
+	}));
+	return DockTab;
+}
+
+#undef LOCTEXT_NAMESPACE 
diff --git a/Engine/Plugins/Animation/ControlRig/Source/ControlRigEditor/Private/Editor/ModularRigHierarchyTabSummoner.h b/Engine/Plugins/Animation/ControlRig/Source/ControlRigEditor/Private/Editor/ModularRigHierarchyTabSummoner.h
new file mode 100644
--- /dev/null
+++ b/Engine/Plugins/Animation/ControlRig/Source/ControlRigEditor/Private/Editor/ModularRigHierarchyTabSummoner.h
@@ -0,0 +1,24 @@
+// Copyright Epic Games, Inc. All Rights Reserved.
+
+#pragma once
+
+#include "CoreMinimal.h"
+#include "WorkflowOrientedApp/WorkflowTabFactory.h"
+
+class FControlRigEditor;
+
+struct FModularRigHierarchyTabSummoner : public FWorkflowTabFactory
+{
+public:
+	static const FName TabID;
+	
+public:
+	FModularRigHierarchyTabSummoner(const TSharedRef<FControlRigEditor>& InControlRigEditor);
+	
+	virtual FTabSpawnerEntry& RegisterTabSpawner(TSharedRef<FTabManager> TabManager, const FApplicationMode* CurrentApplicationMode) const;
+	virtual TSharedRef<SWidget> CreateTabBody(const FWorkflowTabSpawnInfo& Info) const override;
+	virtual TSharedRef<SDockTab> SpawnTab(const FWorkflowTabSpawnInfo& Info) const override;
+
+protected:
+	TWeakPtr<FControlRigEditor> ControlRigEditor;
+};
diff --git a/Engine/Plugins/Animation/ControlRig/Source/ControlRigEditor/Private/Editor/SModularRigHierarchy.cpp b/Engine/Plugins/Animation/ControlRig/Source/ControlRigEditor/Private/Editor/SModularRigHierarchy.cpp
new file mode 100644
--- /dev/null
+++ b/Engine/Plugins/Animation/ControlRig/Source/ControlRigEditor/Private/Editor/SModularRigHierarchy.cpp
@@ -0,0 +1,612 @@
+// Copyright Epic Games, Inc. All Rights Reserved.
+
+#include "Editor/SModularRigHierarchy.h"
+#include "Widgets/Input/SComboButton.h"
+#include "Styling/AppStyle.h"
+#include "Widgets/Input/SButton.h"
+#include "Widgets/Input/SSearchBox.h"
+#include "Kismet2/BlueprintEditorUtils.h"
+#include "ScopedTransaction.h"
+#include "Editor/ControlRigEditor.h"
+#include "BlueprintActionDatabase.h"
+#include "BlueprintVariableNodeSpawner.h"
+#include "Widgets/Layout/SSpacer.h"
+#include "Widgets/Images/SImage.h"
+#include "Widgets/Text/STextBlock.h"
+#include "K2Node_VariableGet.h"
+#include "RigVMBlueprintUtils.h"
+#include "ControlRigModularRigHierarchyCommands.h"
+#include "ControlRigBlueprint.h"
+#include "Graph/ControlRigGraph.h"
+#include "Graph/ControlRigGraphNode.h"
+#include "Graph/ControlRigGraphSchema.h"
+#include "GraphEditorModule.h"
+#include "Framework/MultiBox/MultiBoxBuilder.h"
+#include "AnimationRuntime.h"
+#include "ClassViewerFilter.h"
+#include "PropertyCustomizationHelpers.h"
+#include "Framework/Application/SlateApplication.h"
+#include "Editor/EditorEngine.h"
+#include "HelperUtil.h"
+#include "Widgets/Text/SInlineEditableTextBlock.h"
+#include "ControlRig.h"
+#include "HAL/PlatformApplicationMisc.h"
+#include "HAL/PlatformTime.h"
+#include "Dialogs/Dialogs.h"
+#include "IPersonaToolkit.h"
+#include "SKismetInspector.h"
+#include "Types/WidgetActiveTimerDelegate.h"
+#include "Dialog/SCustomDialog.h"
+#include "EditMode/ControlRigEditMode.h"
+#include "ToolMenus.h"
+#include "Editor/ControlRigContextMenuContext.h"
+#include "Editor/SRigSpacePickerWidget.h"
+#include "Settings/ControlRigSettings.h"
+#include "Framework/Notifications/NotificationManager.h"
+#include "Widgets/Notifications/SNotificationList.h"
+#include "Styling/AppStyle.h"
+#include "ControlRigSkeletalMeshComponent.h"
+#include "Sequencer/ControlRigLayerInstance.h"
+#include "Algo/MinElement.h"
+#include "Algo/MaxElement.h"
+#include "AssetRegistry/AssetRegistryModule.h"
+#include "DragAndDrop/AssetDragDropOp.h"
+#include "Kismet2/SClassPickerDialog.h"
+#include "RigVMFunctions/Math/RigVMMathLibrary.h"
+#include "Preferences/PersonaOptions.h"
+
+#define LOCTEXT_NAMESPACE "SModularRigHierarchy"
+
+
+///////////////////////////////////////////////////////////
+
+const FName SModularRigHierarchy::ContextMenuName = TEXT("ControlRigEditor.ModularRigHierarchy.ContextMenu");
+
+SModularRigHierarchy::~SModularRigHierarchy()
+{
+	const FControlRigEditor* Editor = ControlRigEditor.IsValid() ? ControlRigEditor.Pin().Get() : nullptr;
+	OnEditorClose(Editor, ControlRigBlueprint.Get());
+}
+
+void SModularRigHierarchy::Construct(const FArguments& InArgs, TSharedRef<FControlRigEditor> InControlRigEditor)
+{
+	ControlRigEditor = InControlRigEditor;
+
+	ControlRigBlueprint = ControlRigEditor.Pin()->GetControlRigBlueprint();
+
+	ControlRigBlueprint->OnRefreshEditor().AddRaw(this, &SModularRigHierarchy::HandleRefreshEditorFromBlueprint);
+	ControlRigBlueprint->OnSetObjectBeingDebugged().AddRaw(this, &SModularRigHierarchy::HandleSetObjectBeingDebugged);
+	ControlRigBlueprint->OnModularRigCompiled().AddRaw(this, &SModularRigHierarchy::HandleRefreshEditorFromBlueprint);
+
+	// for deleting, renaming, dragging
+	CommandList = MakeShared<FUICommandList>();
+
+	UEditorEngine* Editor = Cast<UEditorEngine>(GEngine);
+	if (Editor != nullptr)
+	{
+		Editor->RegisterForUndo(this);
+	}
+
+	BindCommands();
+
+	// setup all delegates for the rig hierarchy widget
+	FModularRigTreeDelegates Delegates;
+	Delegates.OnGetHierarchy = FOnGetModularRigTreeHierarchy::CreateSP(this, &SModularRigHierarchy::GetHierarchyForTreeView);
+	Delegates.OnContextMenuOpening = FOnContextMenuOpening::CreateSP(this, &SModularRigHierarchy::CreateContextMenuWidget);
+	Delegates.OnCanAcceptDrop = FOnModularRigTreeCanAcceptDrop::CreateSP(this, &SModularRigHierarchy::OnCanAcceptDrop);
+	Delegates.OnAcceptDrop = FOnModularRigTreeAcceptDrop::CreateSP(this, &SModularRigHierarchy::OnAcceptDrop);
+	Delegates.OnMouseButtonClick = FOnModularRigTreeMouseButtonClick::CreateSP(this, &SModularRigHierarchy::OnItemClicked);
+	Delegates.OnRequestDetailsInspection = FOnModularRigTreeRequestDetailsInspection::CreateSP(this, &SModularRigHierarchy::OnRequestDetailsInspection);
+	
+	ChildSlot
+	[
+		SNew(SVerticalBox)
+		+SVerticalBox::Slot()
+		.Padding(0.0f, 0.0f)
+		[
+			SNew(SBorder)
+			.Padding(0.0f)
+			.ShowEffectWhenDisabled(false)
+			[
+				SNew(SBorder)
+				.Padding(2.0f)
+				.BorderImage(FAppStyle::GetBrush("SCSEditor.TreePanel"))
+				[
+					SAssignNew(TreeView, SModularRigHierarchyTreeView)
+					.RigTreeDelegates(Delegates)
+					.AutoScrollEnabled(true)
+				]
+			]
+		]
+	];
+
+	RefreshTreeView();
+
+	if (ControlRigEditor.IsValid())
+	{
+		ControlRigEditor.Pin()->GetKeyDownDelegate().BindLambda([&](const FGeometry& MyGeometry, const FKeyEvent& InKeyEvent)->FReply {
+			return OnKeyDown(MyGeometry, InKeyEvent);
+		});
+		ControlRigEditor.Pin()->OnGetViewportContextMenu().BindSP(this, &SModularRigHierarchy::GetContextMenu);
+		ControlRigEditor.Pin()->OnViewportContextMenuCommands().BindSP(this, &SModularRigHierarchy::GetContextMenuCommands);
+		ControlRigEditor.Pin()->OnEditorClosed().AddSP(this, &SModularRigHierarchy::OnEditorClose);
+	}
+	
+	CreateContextMenu();
+}
+
+void SModularRigHierarchy::OnEditorClose(const FRigVMEditor* InEditor, URigVMBlueprint* InBlueprint)
+{
+	if (InEditor)
+	{
+		FControlRigEditor* Editor = (FControlRigEditor*)InEditor;  
+		Editor->OnGetViewportContextMenu().Unbind();
+		Editor->OnViewportContextMenuCommands().Unbind();
+	}
+
+	if (UControlRigBlueprint* BP = Cast<UControlRigBlueprint>(InBlueprint))
+	{
+		InBlueprint->OnRefreshEditor().RemoveAll(this);
+		InBlueprint->OnSetObjectBeingDebugged().RemoveAll(this);
+	}
+	
+	ControlRigEditor.Reset();
+	ControlRigBlueprint.Reset();
+}
+
+void SModularRigHierarchy::BindCommands()
+{
+	// create new command
+	const FControlRigModularHierarchyCommands& Commands = FControlRigModularHierarchyCommands::Get();
+
+	CommandList->MapAction(Commands.AddModuleItem,
+		FExecuteAction::CreateSP(this, &SModularRigHierarchy::HandleNewItem),
+		FCanExecuteAction());
+}
+
+void SModularRigHierarchy::RefreshTreeView(bool bRebuildContent)
+{
+	const UModularRig* Hierarchy = GetHierarchy();
+	bool bDummySuspensionFlag = false;
+	bool* SuspensionFlagPtr = &bDummySuspensionFlag;
+	if (ControlRigEditor.IsValid())
+	{
+		SuspensionFlagPtr = &ControlRigEditor.Pin()->GetSuspendDetailsPanelRefreshFlag();
+	}
+	TGuardValue<bool> SuspendDetailsPanelRefreshGuard(*SuspensionFlagPtr, true);
+
+	TreeView->RefreshTreeView(bRebuildContent);
+}
+
+TArray<FString> SModularRigHierarchy::GetSelectedKeys() const
+{
+	TArray<TSharedPtr<FModularRigTreeElement>> SelectedItems = TreeView->GetSelectedItems();
+	
+	TArray<FString> SelectedKeys;
+	for (const TSharedPtr<FModularRigTreeElement>& SelectedItem : SelectedItems)
+	{
+		if(!SelectedItem->Key.IsEmpty())
+		{
+			SelectedKeys.AddUnique(SelectedItem->Key);
+		}
+	}
+
+	return SelectedKeys;
+}
+
+void SModularRigHierarchy::HandleRefreshEditorFromBlueprint(URigVMBlueprint* InBlueprint)
+{
+	RefreshTreeView();
+}
+
+void SModularRigHierarchy::HandleSetObjectBeingDebugged(UObject* InObject)
+{
+	if(ControlRigBeingDebuggedPtr.Get() == InObject)
+	{
+		return;
+	}
+
+	ControlRigBeingDebuggedPtr.Reset();
+	
+	if(UModularRig* ControlRig = Cast<UModularRig>(InObject))
+	{
+		ControlRigBeingDebuggedPtr = ControlRig;
+	}
+
+	RefreshTreeView();
+}
+
+TSharedPtr< SWidget > SModularRigHierarchy::CreateContextMenuWidget()
+{
+	UToolMenus* ToolMenus = UToolMenus::Get();
+
+	if (UToolMenu* Menu = GetContextMenu())
+	{
+		return ToolMenus->GenerateWidget(Menu);
+	}
+	
+	return SNullWidget::NullWidget;
+}
+
+void SModularRigHierarchy::OnItemClicked(TSharedPtr<FModularRigTreeElement> InItem)
+{
+	UModularRig* Rig = GetHierarchy();
+	check(Rig);
+
+	if (ControlRigEditor.IsValid() && InItem.IsValid())
+	{
+		ControlRigEditor.Pin()->SetDetailViewForRigModules({InItem->Key});
+	}
+
+	uint32 CurrentCycles = FPlatformTime::Cycles();
+	// double SecondsPassed = double(CurrentCycles - TreeView->LastClickCycles) * FPlatformTime::GetSecondsPerCycle();
+	// if (SecondsPassed > 0.5f)
+	// {
+	// 	RegisterActiveTimer(0.f, FWidgetActiveTimerDelegate::CreateLambda([this](double, float) {
+	// 		HandleRenameItem();
+	// 		return EActiveTimerReturnType::Stop;
+	// 	}));
+	// }
+
+	TreeView->LastClickCycles = CurrentCycles;
+}
+
+void SModularRigHierarchy::CreateContextMenu()
+{
+	static bool bCreatedMenu = false;
+	if(bCreatedMenu)
+	{
+		return;
+	}
+	bCreatedMenu = true;
+	
+	const FName MenuName = ContextMenuName;
+
+	UToolMenus* ToolMenus = UToolMenus::Get();
+	
+	if (!ensure(ToolMenus))
+	{
+		return;
+	}
+
+	if (UToolMenu* Menu = ToolMenus->ExtendMenu(MenuName))
+	{
+		Menu->AddDynamicSection(NAME_None, FNewToolMenuDelegate::CreateLambda([](UToolMenu* InMenu)
+			{
+				UControlRigContextMenuContext* MainContext = InMenu->FindContext<UControlRigContextMenuContext>();
+				
+				if (SModularRigHierarchy* RigHierarchyPanel = MainContext->GetModularRigHierarchyPanel())
+				{
+					const FControlRigModularHierarchyCommands& Commands = FControlRigModularHierarchyCommands::Get(); 
+				
+					FToolMenuSection& ElementsSection = InMenu->AddSection(TEXT("Elements"), LOCTEXT("ElementsHeader", "Elements"));
+					ElementsSection.AddSubMenu(TEXT("New"), LOCTEXT("New", "New"), LOCTEXT("New_ToolTip", "Create New Elements"),
+						FNewToolMenuDelegate::CreateLambda([Commands, RigHierarchyPanel](UToolMenu* InSubMenu)
+						{
+							FToolMenuSection& DefaultSection = InSubMenu->AddSection(NAME_None);
+							FString SelectedKey;
+							TArray<TSharedPtr<FModularRigTreeElement>> SelectedItems = RigHierarchyPanel->TreeView->GetSelectedItems();
+							if (SelectedItems.Num() > 0)
+							{
+								SelectedKey = SelectedItems[0]->Key;
+							}
+							
+							DefaultSection.AddMenuEntry(Commands.AddModuleItem);
+						})
+					);
+				}
+			})
+		);
+	}
+}
+
+UToolMenu* SModularRigHierarchy::GetContextMenu()
+{
+	const FName MenuName = ContextMenuName;
+	UToolMenus* ToolMenus = UToolMenus::Get();
+
+	if(!ensure(ToolMenus))
+	{
+		return nullptr;
+	}
+
+	// individual entries in this menu can access members of this context, particularly useful for editor scripting
+	UControlRigContextMenuContext* ContextMenuContext = NewObject<UControlRigContextMenuContext>();
+	FControlRigMenuSpecificContext MenuSpecificContext;
+	MenuSpecificContext.ModularRigHierarchyPanel = SharedThis(this);
+	ContextMenuContext->Init(ControlRigEditor, MenuSpecificContext);
+
+	FToolMenuContext MenuContext(CommandList);
+	MenuContext.AddObject(ContextMenuContext);
+
+	UToolMenu* Menu = ToolMenus->GenerateMenu(MenuName, MenuContext);
+
+	return Menu;
+}
+
+TSharedPtr<FUICommandList> SModularRigHierarchy::GetContextMenuCommands() const
+{
+	return CommandList;
+}
+
+bool SModularRigHierarchy::IsSingleSelected() const
+{
+	if(GetSelectedKeys().Num() == 1)
+	{
+		return true;
+	}
+	return false;
+}
+
+/** Filter class to show only RigModules. */
+class FClassViewerRigModulesFilter : public IClassViewerFilter
+{
+public:
+	FClassViewerRigModulesFilter()
+		: AssetRegistry(FModuleManager::GetModuleChecked<FAssetRegistryModule>(TEXT("AssetRegistry")).Get())
+	{}
+	
+	virtual bool IsClassAllowed(const FClassViewerInitializationOptions& InInitOptions, const UClass* InClass, TSharedRef< FClassViewerFilterFuncs > InFilterFuncs) override
+	{
+		if(InClass)
+		{
+			const bool bChildOfObjectClass = InClass->IsChildOf(UControlRig::StaticClass());
+			const bool bMatchesFlags = !InClass->HasAnyClassFlags(CLASS_Hidden | CLASS_HideDropDown | CLASS_Deprecated | CLASS_Abstract);
+			const bool bNotNative = !InClass->IsNative();
+
+			// Allow any class contained in the extra picker common classes array
+			if (InInitOptions.ExtraPickerCommonClasses.Contains(InClass))
+			{
+				return true;
+			}
+			
+			if (bChildOfObjectClass && bMatchesFlags && bNotNative)
+			{
+				const FAssetData AssetData(InClass);
+				return MatchesFilter(AssetData);
+			}
+		}
+		return false;
+	}
+
+	virtual bool IsUnloadedClassAllowed(const FClassViewerInitializationOptions& InInitOptions, const TSharedRef< const IUnloadedBlueprintData > InUnloadedClassData, TSharedRef< FClassViewerFilterFuncs > InFilterFuncs) override
+	{
+		const bool bChildOfObjectClass = InUnloadedClassData->IsChildOf(UControlRig::StaticClass());
+		const bool bMatchesFlags = !InUnloadedClassData->HasAnyClassFlags(CLASS_Hidden | CLASS_HideDropDown | CLASS_Deprecated | CLASS_Abstract);
+		if (bChildOfObjectClass && bMatchesFlags)
+		{
+			const FString GeneratedClassPathString = InUnloadedClassData->GetClassPathName().ToString();
+			const FString BlueprintPath = GeneratedClassPathString.LeftChop(2); // Chop off _C
+			const FAssetData AssetData = AssetRegistry.GetAssetByObjectPath(FSoftObjectPath(BlueprintPath));
+			return MatchesFilter(AssetData);
+
+		}
+		return false;
+	}
+
+private:
+	bool MatchesFilter(const FAssetData& AssetData)
+	{
+		static const UEnum* ControlTypeEnum = StaticEnum<EControlRigType>();
+		const FString ControlRigTypeStr = AssetData.GetTagValueRef<FString>(TEXT("ControlRigType"));
+		if (ControlRigTypeStr.IsEmpty())
+		{
+			return false;
+		}
+
+		const EControlRigType ControlRigType = (EControlRigType)(ControlTypeEnum->GetValueByName(*ControlRigTypeStr));
+		return ControlRigType == EControlRigType::RigModule;
+	}
+
+	const IAssetRegistry& AssetRegistry;
+};
+
+/** Create Item */
+void SModularRigHierarchy::HandleNewItem()
+{
+	if(!ControlRigEditor.IsValid())
+	{
+		return;
+	}
+
+	FString ParentPath;
+	if (IsSingleSelected())
+	{
+		ParentPath = GetSelectedKeys()[0];
+	}
+	
+	FClassViewerInitializationOptions Options;
+	Options.bShowUnloadedBlueprints = true;
+	Options.NameTypeToDisplay = EClassViewerNameTypeToDisplay::DisplayName;
+
+	TSharedPtr<FClassViewerRigModulesFilter> ClassFilter = MakeShareable(new FClassViewerRigModulesFilter());
+	Options.ClassFilters.Add(ClassFilter.ToSharedRef());
+	Options.bShowNoneOption = false;
+	
+	UClass* ChosenClass;
+	const FText TitleText = LOCTEXT("ModularRigHierarchy", "Pick Rig Module Class");
+	const bool bPressedOk = SClassPickerDialog::PickClass(TitleText, Options, ChosenClass, UControlRig::StaticClass());
+	if (bPressedOk)
+	{
+		HandleNewItem(ChosenClass, ParentPath);
+	}
+}
+
+void SModularRigHierarchy::HandleNewItem(UClass* InClass, const FString &InParentPath)
+{
+	UControlRig* ControlRig = InClass->GetDefaultObject<UControlRig>();
+	if (!ControlRig)
+	{
+		return;
+	}
+	
+	if (ControlRigBlueprint.IsValid())
+	{
+		FString ClassName = InClass->GetName();
+		ClassName.RemoveFromEnd(TEXT("_C"));
+		FString PathName = FString::Printf(TEXT("%s:%s"), *InParentPath, *ClassName);
+		const FName Name = CreateUniqueName(*PathName);
+		ControlRigBlueprint->GetModularRigController()->AddModule(Name, InClass, InParentPath);
+	}
+	
+	FSlateApplication::Get().DismissAllMenus();
+}
+
+class SModularRigHierarchyPasteTransformsErrorPipe : public FOutputDevice
+{
+public:
+
+	int32 NumErrors;
+
+	SModularRigHierarchyPasteTransformsErrorPipe()
+		: FOutputDevice()
+		, NumErrors(0)
+	{
+	}
+
+	virtual void Serialize(const TCHAR* V, ELogVerbosity::Type Verbosity, const class FName& Category) override
+	{
+		UE_LOG(LogControlRig, Error, TEXT("Error importing transforms to Hierarchy: %s"), V);
+		NumErrors++;
+	}
+};
+
+UModularRig* SModularRigHierarchy::GetHierarchy() const
+{
+	if (ControlRigBlueprint.IsValid())
+	{
+		if (UControlRig* DebuggedRig = ControlRigBeingDebuggedPtr.Get())
+		{
+			return Cast<UModularRig>(DebuggedRig);
+		}
+	}
+	if (ControlRigEditor.IsValid())
+	{
+		if (UControlRig* CurrentRig = ControlRigEditor.Pin()->GetControlRig())
+		{
+			return Cast<UModularRig>(CurrentRig);
+		}
+	}
+	return nullptr;
+}
+
+UModularRig* SModularRigHierarchy::GetDefaultHierarchy() const
+{
+	if (ControlRigBlueprint.IsValid())
+	{
+		if (UControlRig* DebuggedRig = ControlRigBeingDebuggedPtr.Get())
+		{
+			return Cast<UModularRig>(DebuggedRig);
+		}
+	}
+	return nullptr;
+}
+
+
+FName SModularRigHierarchy::CreateUniqueName(const FName& InBasePath) const
+{
+	return ControlRigBlueprint->GetModularRigController()->GetSafeNewName(InBasePath.ToString());
+}
+
+void SModularRigHierarchy::OnRequestDetailsInspection(const FString& InKey)
+{
+	if(!ControlRigEditor.IsValid())
+	{
+		return;
+	}
+	ControlRigEditor.Pin()->SetDetailViewForRigModules({InKey});
+}
+
+void SModularRigHierarchy::PostRedo(bool bSuccess) 
+{
+	if (bSuccess)
+	{
+		RefreshTreeView();
+	}
+}
+
+void SModularRigHierarchy::PostUndo(bool bSuccess) 
+{
+	if (bSuccess)
+	{
+		RefreshTreeView();
+	}
+}
+
+TOptional<EItemDropZone> SModularRigHierarchy::OnCanAcceptDrop(const FDragDropEvent& DragDropEvent, EItemDropZone DropZone, TSharedPtr<FModularRigTreeElement> TargetItem)
+{
+	const TOptional<EItemDropZone> InvalidDropZone;
+	TOptional<EItemDropZone> ReturnDropZone;
+
+	ReturnDropZone = DropZone;
+	return ReturnDropZone;
+}
+
+FReply SModularRigHierarchy::OnAcceptDrop(const FDragDropEvent& DragDropEvent, EItemDropZone DropZone, TSharedPtr<FModularRigTreeElement> TargetItem)
+{
+	const TSharedPtr<FModularRigTreeElement>* ItemAtMouse = TreeView->FindItemAtPosition(DragDropEvent.GetScreenSpacePosition());
+	FString ParentPath;
+	if (ItemAtMouse && ItemAtMouse->IsValid())
+	{
+		ParentPath = ItemAtMouse->Get()->Key;
+	}
+
+	TSharedPtr<FAssetDragDropOp> AssetDragDropOperation = DragDropEvent.GetOperationAs<FAssetDragDropOp>();
+	if (AssetDragDropOperation)
+	{
+		for (const FAssetData& AssetData : AssetDragDropOperation->GetAssets())
+		{
+			static const UEnum* ControlTypeEnum = StaticEnum<EControlRigType>();
+			const FString ControlRigTypeStr = AssetData.GetTagValueRef<FString>(TEXT("ControlRigType"));
+			if (ControlRigTypeStr.IsEmpty())
+			{
+				continue;
+			}
+
+			const EControlRigType ControlRigType = (EControlRigType)(ControlTypeEnum->GetValueByName(*ControlRigTypeStr));
+			if (ControlRigType != EControlRigType::RigModule)
+			{
+				continue;
+			}
+
+			UClass* AssetClass = AssetData.GetClass();
+			if (!AssetClass->IsChildOf(UControlRigBlueprint::StaticClass()))
+			{
+				continue;
+			}
+
+			if(UControlRigBlueprint* AssetBlueprint = Cast<UControlRigBlueprint>(AssetData.GetAsset()))
+			{
+				HandleNewItem(AssetBlueprint->GetControlRigClass(), ParentPath);
+			}
+		}
+
+		FReply::Handled();
+	}
+	
+	return FReply::Unhandled();
+}
+
+FReply SModularRigHierarchy::OnDrop(const FGeometry& MyGeometry, const FDragDropEvent& DragDropEvent)
+{
+	// only allow drops onto empty space of the widget (when there's no target item under the mouse)
+	// when dropped onto an item SModularRigHierarchy::OnAcceptDrop will deal with the event
+	const TSharedPtr<FModularRigTreeElement>* ItemAtMouse = TreeView->FindItemAtPosition(DragDropEvent.GetScreenSpacePosition());
+	FString ParentPath;
+	if (ItemAtMouse && ItemAtMouse->IsValid())
+	{
+		return SCompoundWidget::OnDrop(MyGeometry, DragDropEvent);
+	}
+	
+	if (OnCanAcceptDrop(DragDropEvent, EItemDropZone::BelowItem, nullptr))
+	{
+		if (OnAcceptDrop(DragDropEvent, EItemDropZone::BelowItem, nullptr).IsEventHandled())
+		{
+			return FReply::Handled();
+		}
+	}
+	return SCompoundWidget::OnDrop(MyGeometry, DragDropEvent);
+}
+
+#undef LOCTEXT_NAMESPACE
+
diff --git a/Engine/Plugins/Animation/ControlRig/Source/ControlRigEditor/Private/Editor/SModularRigHierarchy.h b/Engine/Plugins/Animation/ControlRig/Source/ControlRigEditor/Private/Editor/SModularRigHierarchy.h
new file mode 100644
--- /dev/null
+++ b/Engine/Plugins/Animation/ControlRig/Source/ControlRigEditor/Private/Editor/SModularRigHierarchy.h
@@ -0,0 +1,113 @@
+// Copyright Epic Games, Inc. All Rights Reserved.
+
+#pragma once
+
+#include "CoreMinimal.h"
+#include "EditorUndoClient.h"
+#include "Editor/SModularRigHierarchyTreeView.h"
+#include "ControlRigBlueprint.h"
+#include "Editor/RigVMEditor.h"
+//#include "SModularRigHierarchy.generated.h"
+
+class SModularRigHierarchy;
+class FControlRigEditor;
+class SSearchBox;
+class FUICommandList;
+class URigVMBlueprint;
+class UControlRig;
+struct FAssetData;
+class FMenuBuilder;
+class UToolMenu;
+struct FToolMenuContext;
+
+
+/** Widget allowing editing of a control rig's structure */
+class SModularRigHierarchy : public SCompoundWidget, public FEditorUndoClient
+{
+public:
+	SLATE_BEGIN_ARGS(SModularRigHierarchy) {}
+	SLATE_END_ARGS()
+
+	~SModularRigHierarchy();
+
+	void Construct(const FArguments& InArgs, TSharedRef<FControlRigEditor> InControlRigEditor);
+
+	FControlRigEditor* GetControlRigEditor() const
+	{
+		if(ControlRigEditor.IsValid())
+		{
+			return ControlRigEditor.Pin().Get();
+		}
+		return nullptr;
+	}
+
+private:
+
+	void OnEditorClose(const FRigVMEditor* InEditor, URigVMBlueprint* InBlueprint);
+
+	/** Bind commands that this widget handles */
+	void BindCommands();
+
+	/** Rebuild the tree view */
+	void RefreshTreeView(bool bRebuildContent = true);
+
+	/** Return all selected keys */
+	TArray<FString> GetSelectedKeys() const;
+
+	/** Create a new item */
+	void HandleNewItem();
+
+	void HandleNewItem(UClass* InClass, const FString& InParentPath);
+
+	/** Set Selection Changed */
+	void OnSelectionChanged(TSharedPtr<FModularRigTreeElement> Selection, ESelectInfo::Type SelectInfo);
+
+	TSharedPtr< SWidget > CreateContextMenuWidget();
+	void OnItemClicked(TSharedPtr<FModularRigTreeElement> InItem);
+	
+	// FEditorUndoClient
+	virtual void PostUndo(bool bSuccess) override;
+	virtual void PostRedo(bool bSuccess) override;
+
+	// reply to a drop operation on item
+	TOptional<EItemDropZone> OnCanAcceptDrop(const FDragDropEvent& DragDropEvent, EItemDropZone DropZone, TSharedPtr<FModularRigTreeElement> TargetItem);
+	FReply OnAcceptDrop(const FDragDropEvent& DragDropEvent, EItemDropZone DropZone, TSharedPtr<FModularRigTreeElement> TargetItem);
+
+	// SWidget Overrides
+	virtual FReply OnDrop(const FGeometry& MyGeometry, const FDragDropEvent& DragDropEvent) override;
+
+	static const FName ContextMenuName;
+	static void CreateContextMenu();
+	UToolMenu* GetContextMenu();
+	TSharedPtr<FUICommandList> GetContextMenuCommands() const;
+	
+	/** Our owning control rig editor */
+	TWeakPtr<FControlRigEditor> ControlRigEditor;
+
+	/** Tree view widget */
+	TSharedPtr<SModularRigHierarchyTreeView> TreeView;
+
+	TWeakObjectPtr<UControlRigBlueprint> ControlRigBlueprint;
+	TWeakObjectPtr<UModularRig> ControlRigBeingDebuggedPtr;
+	
+	/** Command list we bind to */
+	TSharedPtr<FUICommandList> CommandList;
+
+	bool IsSingleSelected() const;
+	
+	UModularRig* GetHierarchy() const;
+	UModularRig* GetDefaultHierarchy() const;
+	const UModularRig* GetHierarchyForTreeView() const { return GetHierarchy(); }
+	FName CreateUniqueName(const FName& InBasePath) const;
+	void OnRequestDetailsInspection(const FString& InKey);
+
+	void HandleRefreshEditorFromBlueprint(URigVMBlueprint* InBlueprint);
+	void HandleSetObjectBeingDebugged(UObject* InObject);
+
+public:
+
+	friend class FModularRigTreeElement;
+	friend class SModularRigHierarchyItem;
+	friend class UControlRigBlueprintEditorLibrary;
+};
+
diff --git a/Engine/Plugins/Animation/ControlRig/Source/ControlRigEditor/Private/Editor/SModularRigHierarchyTreeView.cpp b/Engine/Plugins/Animation/ControlRig/Source/ControlRigEditor/Private/Editor/SModularRigHierarchyTreeView.cpp
new file mode 100644
--- /dev/null
+++ b/Engine/Plugins/Animation/ControlRig/Source/ControlRigEditor/Private/Editor/SModularRigHierarchyTreeView.cpp
@@ -0,0 +1,657 @@
+// Copyright Epic Games, Inc. All Rights Reserved.
+
+#include "Editor/SModularRigHierarchyTreeView.h"
+#include "Styling/AppStyle.h"
+#include "Widgets/Layout/SSpacer.h"
+#include "Widgets/Layout/SScrollBox.h"
+#include "Widgets/Images/SImage.h"
+#include "Widgets/Input/SSearchBox.h"
+#include "Widgets/Text/STextBlock.h"
+#include "PropertyCustomizationHelpers.h"
+#include "Framework/Application/SlateApplication.h"
+#include "Editor/EditorEngine.h"
+#include "HelperUtil.h"
+#include "Widgets/Text/SInlineEditableTextBlock.h"
+#include "ControlRig.h"
+#include "ControlRigEditorStyle.h"
+#include "ModularRig.h"
+#include "Blueprint/WidgetBlueprintLibrary.h"
+#include "Editor/SModularRigHierarchy.h"
+#include "Settings/ControlRigSettings.h"
+#include "Graph/ControlRigGraphSchema.h"
+#include "Rigs/AdditiveControlRig.h"
+#include "Rigs/RigHierarchyController.h"
+#include "Styling/AppStyle.h"
+
+#define LOCTEXT_NAMESPACE "SModularRigHierarchyTreeView"
+
+TMap<FSoftObjectPath, FSlateBrush> SModularRigHierarchyItem::IconPathToBrush;
+
+//////////////////////////////////////////////////////////////
+/// FModularRigTreeElement
+///////////////////////////////////////////////////////////
+FModularRigTreeElement::FModularRigTreeElement(const FString& InKey, TWeakPtr<SModularRigHierarchyTreeView> InTreeView, bool InSupportsRename)
+{
+	Key = InKey;
+	FString ShortNameStr = Key;
+	FString ParentPath;
+	ShortNameStr.Split(UModularRig::NamespaceSeparator, &ParentPath, &ShortNameStr, ESearchCase::Type::CaseSensitive, ESearchDir::FromEnd);
+	ShortName = *ShortNameStr;
+	
+	if(InTreeView.IsValid())
+	{
+		if(const UModularRig* Hierarchy = InTreeView.Pin()->GetRigTreeDelegates().GetHierarchy())
+		{
+			RefreshDisplaySettings(Hierarchy);
+		}
+	}
+}
+
+void FModularRigTreeElement::RefreshDisplaySettings(const UModularRig* InHierarchy)
+{
+	const TPair<const FSlateBrush*, FSlateColor> Result = SModularRigHierarchyItem::GetBrushForElementType(InHierarchy, Key);
+
+	IconBrush = Result.Key;
+	IconColor = Result.Value;
+	TextColor = FSlateColor::UseForeground();
+}
+
+TSharedRef<ITableRow> FModularRigTreeElement::MakeTreeRowWidget(const TSharedRef<STableViewBase>& InOwnerTable, TSharedRef<FModularRigTreeElement> InRigTreeElement, TSharedPtr<SModularRigHierarchyTreeView> InTreeView, bool bPinned)
+{
+	return SNew(SModularRigHierarchyItem, InOwnerTable, InRigTreeElement, InTreeView, bPinned);
+}
+
+//////////////////////////////////////////////////////////////
+/// SModularRigHierarchyItem
+///////////////////////////////////////////////////////////
+void SModularRigHierarchyItem::Construct(const FArguments& InArgs, const TSharedRef<STableViewBase>& OwnerTable, TSharedRef<FModularRigTreeElement> InRigTreeElement, TSharedPtr<SModularRigHierarchyTreeView> InTreeView, bool bPinned)
+{
+	WeakRigTreeElement = InRigTreeElement;
+	Delegates = InTreeView->GetRigTreeDelegates();
+
+	if (InRigTreeElement->Key.IsEmpty())
+	{
+		STableRow<TSharedPtr<FModularRigTreeElement>>::Construct(
+			STableRow<TSharedPtr<FModularRigTreeElement>>::FArguments()
+			.ShowSelection(false)
+			.OnCanAcceptDrop(Delegates.OnCanAcceptDrop)
+			.OnAcceptDrop(Delegates.OnAcceptDrop)
+			.Content()
+			[
+				SNew(SVerticalBox)
+				+ SVerticalBox::Slot()
+				.FillHeight(200.f)
+				[
+					SNew(SSpacer)
+				]
+			], OwnerTable);
+		return;
+	}
+
+	TSharedPtr< SInlineEditableTextBlock > InlineWidget;
+	TSharedPtr< SHorizontalBox > HorizontalBox;
+
+	STableRow<TSharedPtr<FModularRigTreeElement>>::Construct(
+		STableRow<TSharedPtr<FModularRigTreeElement>>::FArguments()
+		.OnCanAcceptDrop(Delegates.OnCanAcceptDrop)
+		.OnAcceptDrop(Delegates.OnAcceptDrop)
+		.ShowWires(true)
+		.Content()
+		[
+			SAssignNew(HorizontalBox, SHorizontalBox)
+			+ SHorizontalBox::Slot()
+			.MaxWidth(18)
+			.FillWidth(1.0)
+			.HAlign(HAlign_Left)
+			.VAlign(VAlign_Center)
+			.Padding(FMargin(0.f, 0.f, 3.f, 0.f))
+			[
+				SNew(SImage)
+				.Image_Lambda([this]() -> const FSlateBrush*
+				{
+					if(WeakRigTreeElement.IsValid())
+					{
+						return WeakRigTreeElement.Pin()->IconBrush;
+					}
+					return nullptr;
+				})
+				.ColorAndOpacity_Lambda([this]()
+				{
+					if(WeakRigTreeElement.IsValid())
+					{
+						return WeakRigTreeElement.Pin()->IconColor;
+					}
+					return FSlateColor::UseForeground();
+				})
+			]
+			+ SHorizontalBox::Slot()
+			.AutoWidth()
+			.VAlign(VAlign_Center)
+			[
+				SAssignNew(InlineWidget, SInlineEditableTextBlock)
+				.Text(this, &SModularRigHierarchyItem::GetName, true)
+				.ToolTipText(this, &SModularRigHierarchyItem::GetItemTooltip)
+				.MultiLine(false)
+				.ColorAndOpacity_Lambda([this]()
+				{
+					if(WeakRigTreeElement.IsValid())
+					{
+						return WeakRigTreeElement.Pin()->TextColor;
+					}
+					return FSlateColor::UseForeground();
+				})
+			]
+		], OwnerTable);
+}
+
+FText SModularRigHierarchyItem::GetName(bool bUseShortName) const
+{
+	if(bUseShortName)
+	{
+		return (FText::FromName(WeakRigTreeElement.Pin()->ShortName));
+	}
+	return (FText::FromString(WeakRigTreeElement.Pin()->Key));
+}
+
+FText SModularRigHierarchyItem::GetItemTooltip() const
+{
+	const FText FullName = GetName(false);
+	const FText ShortName = GetName(true);
+	if(FullName.EqualTo(ShortName))
+	{
+		return FText();
+	}
+	return FullName;
+}
+
+//////////////////////////////////////////////////////////////
+/// SModularRigHierarchyTreeView
+///////////////////////////////////////////////////////////
+
+void SModularRigHierarchyTreeView::Construct(const FArguments& InArgs)
+{
+	Delegates = InArgs._RigTreeDelegates;
+	bAutoScrollEnabled = InArgs._AutoScrollEnabled;
+
+	STreeView<TSharedPtr<FModularRigTreeElement>>::FArguments SuperArgs;
+	SuperArgs.TreeItemsSource(&RootElements);
+	SuperArgs.SelectionMode(ESelectionMode::Multi);
+	SuperArgs.OnGenerateRow(this, &SModularRigHierarchyTreeView::MakeTableRowWidget, false);
+	SuperArgs.OnGetChildren(this, &SModularRigHierarchyTreeView::HandleGetChildrenForTree);
+	SuperArgs.OnContextMenuOpening(Delegates.OnContextMenuOpening);
+	SuperArgs.HighlightParentNodesForSelection(true);
+	SuperArgs.ItemHeight(24);
+	SuperArgs.AllowInvisibleItemSelection(true);  //without this we deselect everything when we filter or we collapse
+	SuperArgs.OnMouseButtonClick(Delegates.OnMouseButtonClick);
+	
+	SuperArgs.ShouldStackHierarchyHeaders_Lambda([]() -> bool {
+		return UControlRigEditorSettings::Get()->bShowStackedHierarchy;
+	});
+	SuperArgs.OnGeneratePinnedRow(this, &SModularRigHierarchyTreeView::MakeTableRowWidget, true);
+	SuperArgs.MaxPinnedItems_Lambda([]() -> int32
+	{
+		return FMath::Max<int32>(1, UControlRigEditorSettings::Get()->MaxStackSize);
+	});
+
+	STreeView<TSharedPtr<FModularRigTreeElement>>::Construct(SuperArgs);
+
+	LastMousePosition = FVector2D::ZeroVector;
+	TimeAtMousePosition = 0.0;
+}
+
+void SModularRigHierarchyTreeView::Tick(const FGeometry& AllottedGeometry, const double InCurrentTime, const float InDeltaTime)
+{
+	STreeView<TSharedPtr<FModularRigTreeElement, ESPMode::ThreadSafe>>::Tick(AllottedGeometry, InCurrentTime, InDeltaTime);
+
+	const FGeometry PaintGeometry = GetPaintSpaceGeometry();
+	const FVector2D MousePosition = FSlateApplication::Get().GetCursorPos();
+
+	if(PaintGeometry.IsUnderLocation(MousePosition))
+	{
+		const FVector2D WidgetPosition = PaintGeometry.AbsoluteToLocal(MousePosition);
+
+		static constexpr float SteadyMousePositionTolerance = 5.f;
+
+		if(LastMousePosition.Equals(MousePosition, SteadyMousePositionTolerance))
+		{
+			TimeAtMousePosition += InDeltaTime;
+		}
+		else
+		{
+			LastMousePosition = MousePosition;
+			TimeAtMousePosition = 0.0;
+		}
+
+		static constexpr float AutoScrollStartDuration = 0.5f; // in seconds
+		static constexpr float AutoScrollDistance = 24.f; // in pixels
+		static constexpr float AutoScrollSpeed = 150.f;
+
+		if(TimeAtMousePosition > AutoScrollStartDuration && FSlateApplication::Get().IsDragDropping())
+		{
+			if((WidgetPosition.Y < AutoScrollDistance) || (WidgetPosition.Y > PaintGeometry.Size.Y - AutoScrollDistance))
+			{
+				if(bAutoScrollEnabled)
+				{
+					const bool bScrollUp = (WidgetPosition.Y < AutoScrollDistance);
+
+					const float DeltaInSlateUnits = (bScrollUp ? -InDeltaTime : InDeltaTime) * AutoScrollSpeed; 
+					ScrollBy(GetCachedGeometry(), DeltaInSlateUnits, EAllowOverscroll::No);
+				}
+			}
+			else
+			{
+				const TSharedPtr<FModularRigTreeElement>* Item = FindItemAtPosition(MousePosition);
+				if(Item && Item->IsValid())
+				{
+					if(!IsItemExpanded(*Item))
+					{
+						SetItemExpansion(*Item, true);
+					}
+				}
+			}
+		}
+	}
+}
+
+TSharedPtr<FModularRigTreeElement> SModularRigHierarchyTreeView::FindElement(const FString& InElementKey, TSharedPtr<FModularRigTreeElement> CurrentItem)
+{
+	if (CurrentItem->Key == InElementKey)
+	{
+		return CurrentItem;
+	}
+
+	for (int32 ChildIndex = 0; ChildIndex < CurrentItem->Children.Num(); ++ChildIndex)
+	{
+		TSharedPtr<FModularRigTreeElement> Found = FindElement(InElementKey, CurrentItem->Children[ChildIndex]);
+		if (Found.IsValid())
+		{
+			return Found;
+		}
+	}
+
+	return TSharedPtr<FModularRigTreeElement>();
+}
+
+bool SModularRigHierarchyTreeView::AddElement(FString InKey, FString InParentKey)
+{
+	if(ElementMap.Contains(InKey))
+	{
+		return false;
+	}
+
+	if (!InKey.IsEmpty())
+	{
+		TSharedPtr<FModularRigTreeElement> NewItem = MakeShared<FModularRigTreeElement>(InKey, SharedThis(this), false);
+
+		if (!InKey.IsEmpty())
+		{
+			ElementMap.Add(InKey, NewItem);
+			if (!InParentKey.IsEmpty())
+			{
+				ParentMap.Add(InKey, InParentKey);
+			}
+
+			if (!InParentKey.IsEmpty())
+			{
+				TSharedPtr<FModularRigTreeElement>* FoundItem = ElementMap.Find(InParentKey);
+				check(FoundItem);
+				FoundItem->Get()->Children.Add(NewItem);
+			}
+			else
+			{
+				RootElements.Add(NewItem);
+			}
+		}
+		else
+		{
+			RootElements.Add(NewItem);
+		}
+	}
+
+	return true;
+}
+
+bool SModularRigHierarchyTreeView::AddElement(const FRigModuleInstance* InElement)
+{
+	check(InElement);
+	
+	if (ElementMap.Contains(InElement->GetPath()))
+	{
+		return false;
+	}
+
+	const UModularRig* Hierarchy = Delegates.GetHierarchy();
+
+	if(!AddElement(InElement->GetPath(), FString()))
+	{
+		return false;
+	}
+
+	if (ElementMap.Contains(InElement->GetPath()))
+	{
+		if(Hierarchy)
+		{
+			FString ParentKey = Hierarchy->GetParentPath(InElement->GetPath());
+			if (!ParentKey.IsEmpty())
+			{
+				if(const FRigModuleInstance* ParentElement = Hierarchy->FindModule(ParentKey))
+				{
+					AddElement(ParentElement);
+
+					if(ElementMap.Contains(ParentKey))
+					{
+						ReparentElement(InElement->GetPath(), ParentKey);
+					}
+				}
+			}
+		}
+	}
+
+	return true;
+}
+
+void SModularRigHierarchyTreeView::AddSpacerElement()
+{
+	AddElement(FString(), FString());
+}
+
+bool SModularRigHierarchyTreeView::ReparentElement(const FString InKey, const FString InParentKey)
+{
+	if (InKey.IsEmpty() || InKey == InParentKey)
+	{
+		return false;
+	}
+
+	TSharedPtr<FModularRigTreeElement>* FoundItem = ElementMap.Find(InKey);
+	if (FoundItem == nullptr)
+	{
+		return false;
+	}
+
+	if (const FString* ExistingParentKey = ParentMap.Find(InKey))
+	{
+		if (*ExistingParentKey == InParentKey)
+		{
+			return false;
+		}
+
+		if (TSharedPtr<FModularRigTreeElement>* ExistingParent = ElementMap.Find(*ExistingParentKey))
+		{
+			(*ExistingParent)->Children.Remove(*FoundItem);
+		}
+
+		ParentMap.Remove(InKey);
+	}
+	else
+	{
+		if (InParentKey.IsEmpty())
+		{
+			return false;
+		}
+
+		RootElements.Remove(*FoundItem);
+	}
+
+	if (!InParentKey.IsEmpty())
+	{
+		ParentMap.Add(InKey, InParentKey);
+
+		TSharedPtr<FModularRigTreeElement>* FoundParent = ElementMap.Find(InParentKey);
+		if(FoundParent)
+		{
+			FoundParent->Get()->Children.Add(*FoundItem);
+		}
+	}
+	else
+	{
+		RootElements.Add(*FoundItem);
+	}
+
+	return true;
+}
+
+void SModularRigHierarchyTreeView::RefreshTreeView(bool bRebuildContent)
+{
+	TMap<FString, bool> ExpansionState;
+
+	if(bRebuildContent)
+	{
+		for (TPair<FString, TSharedPtr<FModularRigTreeElement>> Pair : ElementMap)
+		{
+			ExpansionState.FindOrAdd(Pair.Key) = IsItemExpanded(Pair.Value);
+		}
+
+		// internally save expansion states before rebuilding the tree, so the states can be restored later
+		SaveAndClearSparseItemInfos();
+
+		RootElements.Reset();
+		ElementMap.Reset();
+		ParentMap.Reset();
+	}
+
+	if(bRebuildContent)
+	{
+		const UModularRig* Hierarchy = Delegates.GetHierarchy();
+		if(Hierarchy)
+		{
+			Hierarchy->ForEachModule([&](const FRigModuleInstance* Element)
+			{
+				AddElement(Element);
+				return true;
+			});
+
+			// expand all elements upon the initial construction of the tree
+			if (ExpansionState.Num() < ElementMap.Num())
+			{
+				for (const TPair<FString, TSharedPtr<FModularRigTreeElement>>& Element : ElementMap)
+				{
+					if (!ExpansionState.Contains(Element.Key))
+					{
+						SetItemExpansion(Element.Value, true);
+					}
+				}
+			}
+
+			for (const auto& Pair : ElementMap)
+			{
+				RestoreSparseItemInfos(Pair.Value);
+			}
+
+			if (RootElements.Num() > 0)
+			{
+				AddSpacerElement();
+			}
+		}
+	}
+	else
+	{
+		if (RootElements.Num()> 0)
+		{
+			// elements may be added at the end of the list after a spacer element
+			// we need to remove the spacer element and re-add it at the end
+			RootElements.RemoveAll([](TSharedPtr<FModularRigTreeElement> InElement)
+			{
+				return InElement.Get()->Key == FString();
+			});
+			AddSpacerElement();
+		}
+	}
+
+	RequestTreeRefresh();
+	{
+		ClearSelection();
+	}
+}
+
+TSharedRef<ITableRow> SModularRigHierarchyTreeView::MakeTableRowWidget(TSharedPtr<FModularRigTreeElement> InItem,
+	const TSharedRef<STableViewBase>& OwnerTable, bool bPinned)
+{
+	return InItem->MakeTreeRowWidget(OwnerTable, InItem.ToSharedRef(), SharedThis(this), bPinned);
+}
+
+void SModularRigHierarchyTreeView::HandleGetChildrenForTree(TSharedPtr<FModularRigTreeElement> InItem,
+	TArray<TSharedPtr<FModularRigTreeElement>>& OutChildren)
+{
+	OutChildren = InItem.Get()->Children;
+}
+
+TArray<FString> SModularRigHierarchyTreeView::GetSelectedKeys() const
+{
+	TArray<FString> Keys;
+	TArray<TSharedPtr<FModularRigTreeElement>> SelectedElements = GetSelectedItems();
+	for(const TSharedPtr<FModularRigTreeElement>& SelectedElement : SelectedElements)
+	{
+		Keys.Add(SelectedElement->Key);
+	}
+	return Keys;
+}
+
+const TSharedPtr<FModularRigTreeElement>* SModularRigHierarchyTreeView::FindItemAtPosition(FVector2D InScreenSpacePosition) const
+{
+	if (ItemsPanel.IsValid() && SListView<TSharedPtr<FModularRigTreeElement>>::HasValidItemsSource())
+	{
+		FArrangedChildren ArrangedChildren(EVisibility::Visible);
+		const int32 Index = FindChildUnderPosition(ArrangedChildren, InScreenSpacePosition);
+		if (ArrangedChildren.IsValidIndex(Index))
+		{
+			TSharedRef<SModularRigHierarchyItem> ItemWidget = StaticCastSharedRef<SModularRigHierarchyItem>(ArrangedChildren[Index].Widget);
+			if (ItemWidget->WeakRigTreeElement.IsValid())
+			{
+				const FString Key = ItemWidget->WeakRigTreeElement.Pin()->Key;
+				const TSharedPtr<FModularRigTreeElement>* ResultPtr = SListView<TSharedPtr<FModularRigTreeElement>>::GetItems().FindByPredicate([Key](const TSharedPtr<FModularRigTreeElement>& Item) -> bool
+					{
+						return Item->Key == Key;
+					});
+
+				if (ResultPtr)
+				{
+					return ResultPtr;
+				}
+			}
+		}
+	}
+	return nullptr;
+}
+
+TPair<const FSlateBrush*, FSlateColor> SModularRigHierarchyItem::GetBrushForElementType(const UModularRig* InHierarchy, const FString& InKey)
+{
+	const FSlateBrush* Brush = nullptr;
+	FSlateColor Color = FSlateColor::UseForeground();
+
+	if (FRigModuleInstance* Module = InHierarchy->FindModule(InKey))
+	{
+		if (Module->Rig.IsValid())
+		{
+			FSoftObjectPath IconPath = Module->Rig->GetRigModuleSettings().Icon;
+			Brush = IconPathToBrush.Find(IconPath);
+			if (!Brush)
+			{
+				if(UTexture2D* Icon = Cast<UTexture2D>(IconPath.TryLoad()))
+				{
+					Brush = &IconPathToBrush.Add(IconPath, UWidgetBlueprintLibrary::MakeBrushFromTexture(Icon, 16.0f, 16.0f));
+				}
+			}
+		}
+	}
+	if (!Brush)
+	{
+		Brush = FControlRigEditorStyle::Get().GetBrush("ControlRig.Tree.RigidBody");
+	}
+	return TPair<const FSlateBrush*, FSlateColor>(Brush, Color);
+}
+
+FLinearColor SModularRigHierarchyItem::GetColorForControlType(ERigControlType InControlType, UEnum* InControlEnum)
+{
+	FEdGraphPinType PinType;
+	switch(InControlType)
+	{
+		case ERigControlType::Bool:
+		{
+			PinType = RigVMTypeUtils::PinTypeFromCPPType(RigVMTypeUtils::BoolTypeName, nullptr);
+			break;
+		}
+		case ERigControlType::Float:
+		case ERigControlType::ScaleFloat:
+		{
+			PinType = RigVMTypeUtils::PinTypeFromCPPType(RigVMTypeUtils::FloatTypeName, nullptr);
+			break;
+		}
+		case ERigControlType::Integer:
+		{
+			if(InControlEnum)
+			{
+				PinType = RigVMTypeUtils::PinTypeFromCPPType(NAME_None, InControlEnum);
+			}
+			else
+			{
+				PinType = RigVMTypeUtils::PinTypeFromCPPType(RigVMTypeUtils::Int32TypeName, nullptr);
+			}
+			break;
+		}
+		case ERigControlType::Vector2D:
+		{
+			UScriptStruct* Struct = TBaseStructure<FVector2D>::Get(); 
+			PinType = RigVMTypeUtils::PinTypeFromCPPType(*RigVMTypeUtils::GetUniqueStructTypeName(Struct), Struct);
+			break;
+		}
+		case ERigControlType::Position:
+		case ERigControlType::Scale:
+		{
+			UScriptStruct* Struct = TBaseStructure<FVector>::Get(); 
+			PinType = RigVMTypeUtils::PinTypeFromCPPType(*RigVMTypeUtils::GetUniqueStructTypeName(Struct), Struct);
+			break;
+		}
+		case ERigControlType::Rotator:
+		{
+			UScriptStruct* Struct = TBaseStructure<FRotator>::Get(); 
+			PinType = RigVMTypeUtils::PinTypeFromCPPType(*RigVMTypeUtils::GetUniqueStructTypeName(Struct), Struct);
+			break;
+		}
+		case ERigControlType::Transform:
+		case ERigControlType::TransformNoScale:
+		case ERigControlType::EulerTransform:
+		default:
+		{
+			UScriptStruct* Struct = TBaseStructure<FTransform>::Get(); 
+			PinType = RigVMTypeUtils::PinTypeFromCPPType(*RigVMTypeUtils::GetUniqueStructTypeName(Struct), Struct);
+			break;
+		}
+	}
+	const UControlRigGraphSchema* Schema = GetDefault<UControlRigGraphSchema>();
+	return Schema->GetPinTypeColor(PinType);
+}
+
+//////////////////////////////////////////////////////////////
+/// SSearchableRigHierarchyTreeView
+///////////////////////////////////////////////////////////
+
+void SSearchableModularRigHierarchyTreeView::Construct(const FArguments& InArgs)
+{
+	FModularRigTreeDelegates TreeDelegates = InArgs._RigTreeDelegates;
+	
+	ChildSlot
+	[
+		SNew(SVerticalBox)
+		+SVerticalBox::Slot()
+		.AutoHeight()
+		.VAlign(VAlign_Top)
+		.HAlign(HAlign_Fill)
+		.Padding(0.0f, 0.0f)
+		[
+			SNew(SScrollBox)
+			+ SScrollBox::Slot()
+			[
+				SNew(SBorder)
+				.Padding(2.0f)
+				.BorderImage(FAppStyle::GetBrush("SCSEditor.TreePanel"))
+				[
+					SAssignNew(TreeView, SModularRigHierarchyTreeView)
+					.RigTreeDelegates(TreeDelegates)
+				]
+			]
+		]
+	];
+}
+
+
+#undef LOCTEXT_NAMESPACE
diff --git a/Engine/Plugins/Animation/ControlRig/Source/ControlRigEditor/Private/Editor/SModularRigHierarchyTreeView.h b/Engine/Plugins/Animation/ControlRig/Source/ControlRigEditor/Private/Editor/SModularRigHierarchyTreeView.h
new file mode 100644
--- /dev/null
+++ b/Engine/Plugins/Animation/ControlRig/Source/ControlRigEditor/Private/Editor/SModularRigHierarchyTreeView.h
@@ -0,0 +1,201 @@
+// Copyright Epic Games, Inc. All Rights Reserved.
+
+#pragma once
+
+#include "CoreMinimal.h"
+#include "Widgets/Views/STreeView.h"
+#include "ModularRig.h"
+
+class SSearchBox;
+class SModularRigHierarchyTreeView;
+class SModularRigHierarchyItem;
+class FModularRigTreeElement;
+
+
+DECLARE_DELEGATE_RetVal(const UModularRig*, FOnGetModularRigTreeHierarchy);
+DECLARE_DELEGATE_OneParam(FOnModularRigTreeRequestDetailsInspection, const FString&);
+
+typedef STreeView<TSharedPtr<FModularRigTreeElement>>::FOnMouseButtonClick FOnModularRigTreeMouseButtonClick;
+typedef STableRow<TSharedPtr<FModularRigTreeElement>>::FOnCanAcceptDrop FOnModularRigTreeCanAcceptDrop;
+typedef STableRow<TSharedPtr<FModularRigTreeElement>>::FOnAcceptDrop FOnModularRigTreeAcceptDrop;
+
+struct CONTROLRIGEDITOR_API FModularRigTreeDelegates
+{
+	FOnGetModularRigTreeHierarchy OnGetHierarchy;
+	FOnModularRigTreeMouseButtonClick OnMouseButtonClick;
+	FOnModularRigTreeCanAcceptDrop OnCanAcceptDrop;
+	FOnModularRigTreeAcceptDrop OnAcceptDrop;
+	FOnContextMenuOpening OnContextMenuOpening;
+	FOnModularRigTreeRequestDetailsInspection OnRequestDetailsInspection;
+	
+	FModularRigTreeDelegates()
+	{
+	}
+
+	const UModularRig* GetHierarchy() const
+	{
+		if(OnGetHierarchy.IsBound())
+		{
+			return OnGetHierarchy.Execute();
+		}
+		return nullptr;
+	}
+};
+
+
+/** An item in the tree */
+class FModularRigTreeElement : public TSharedFromThis<FModularRigTreeElement>
+{
+public:
+	FModularRigTreeElement(const FString& InKey, TWeakPtr<SModularRigHierarchyTreeView> InTreeView, bool InSupportsRename);
+
+public:
+	/** Element Data to display */
+	FString Key;
+	FName ShortName;
+	TArray<TSharedPtr<FModularRigTreeElement>> Children;
+
+	TSharedRef<ITableRow> MakeTreeRowWidget(const TSharedRef<STableViewBase>& InOwnerTable, TSharedRef<FModularRigTreeElement> InRigTreeElement, TSharedPtr<SModularRigHierarchyTreeView> InTreeView, bool bPinned);
+
+	void RefreshDisplaySettings(const UModularRig* InHierarchy);
+
+	/** The brush to use when rendering an icon */
+	const FSlateBrush* IconBrush;
+
+	/** The color to use when rendering an icon */
+	FSlateColor IconColor;
+
+	/** The color to use when rendering the label text */
+	FSlateColor TextColor;
+};
+
+class SModularRigHierarchyItem : public STableRow<TSharedPtr<FModularRigTreeElement>>
+{
+public:
+	
+	void Construct(const FArguments& InArgs, const TSharedRef<STableViewBase>& OwnerTable, TSharedRef<FModularRigTreeElement> InRigTreeElement, TSharedPtr<SModularRigHierarchyTreeView> InTreeView, bool bPinned);
+ 	static TPair<const FSlateBrush*, FSlateColor> GetBrushForElementType(const UModularRig* InHierarchy, const FString& InKey);
+	static FLinearColor GetColorForControlType(ERigControlType InControlType, UEnum* InControlEnum);
+
+private:
+	TWeakPtr<FModularRigTreeElement> WeakRigTreeElement;
+ 	FModularRigTreeDelegates Delegates;
+
+	FText GetName(bool bUseShortName) const;
+	FText GetItemTooltip() const;
+
+	static TMap<FSoftObjectPath, FSlateBrush> IconPathToBrush;
+
+	friend class SModularRigHierarchyTreeView; 
+};
+
+class SModularRigHierarchyTreeView : public STreeView<TSharedPtr<FModularRigTreeElement>>
+{
+public:
+
+	SLATE_BEGIN_ARGS(SModularRigHierarchyTreeView)
+		: _AutoScrollEnabled(false)
+	{}
+		SLATE_ARGUMENT(FModularRigTreeDelegates, RigTreeDelegates)
+		SLATE_ARGUMENT(bool, AutoScrollEnabled)
+	SLATE_END_ARGS()
+
+	void Construct(const FArguments& InArgs);
+	virtual ~SModularRigHierarchyTreeView() {}
+
+	/** Performs auto scroll */
+	virtual void Tick(const FGeometry& AllottedGeometry, const double InCurrentTime, const float InDeltaTime) override;
+
+	virtual FReply OnFocusReceived(const FGeometry& MyGeometry, const FFocusEvent& InFocusEvent) override
+	{
+		FReply Reply = STreeView<TSharedPtr<FModularRigTreeElement>>::OnFocusReceived(MyGeometry, InFocusEvent);
+
+		LastClickCycles = FPlatformTime::Cycles();
+
+		return Reply;
+	}
+
+	uint32 LastClickCycles = 0;
+
+	/** Save a snapshot of the internal map that tracks item expansion before tree reconstruction */
+	void SaveAndClearSparseItemInfos()
+	{
+		// Only save the info if there is something to save (do not overwrite info with an empty map)
+		if (!SparseItemInfos.IsEmpty())
+		{
+			OldSparseItemInfos = SparseItemInfos;
+		}
+		ClearExpandedItems();
+	}
+
+	/** Restore the expansion infos map from the saved snapshot after tree reconstruction */
+	void RestoreSparseItemInfos(TSharedPtr<FModularRigTreeElement> ItemPtr)
+	{
+		for (const auto& Pair : OldSparseItemInfos)
+		{
+			if (Pair.Key->Key == ItemPtr->Key)
+			{
+				// the SparseItemInfos now reference the new element, but keep the same expansion state
+				SparseItemInfos.Add(ItemPtr, Pair.Value);
+				break;
+			}
+		}
+	}
+
+
+	static TSharedPtr<FModularRigTreeElement> FindElement(const FString& InElementKey, TSharedPtr<FModularRigTreeElement> CurrentItem);
+	bool AddElement(FString InKey, FString InParentKey = FString());
+	bool AddElement(const FRigModuleInstance* InElement);
+	void AddSpacerElement();
+	bool ReparentElement(const FString InKey, const FString InParentKey);
+	void RefreshTreeView(bool bRebuildContent = true);
+	TSharedRef<ITableRow> MakeTableRowWidget(TSharedPtr<FModularRigTreeElement> InItem, const TSharedRef<STableViewBase>& OwnerTable, bool bPinned);
+	void HandleGetChildrenForTree(TSharedPtr<FModularRigTreeElement> InItem, TArray<TSharedPtr<FModularRigTreeElement>>& OutChildren);
+
+	TArray<FString> GetSelectedKeys() const;
+	const TArray<TSharedPtr<FModularRigTreeElement>>& GetRootElements() const { return RootElements; }
+	FModularRigTreeDelegates& GetRigTreeDelegates() { return Delegates; }
+
+	/** Given a position, return the item under that position. If nothing is there, return null. */
+	const TSharedPtr<FModularRigTreeElement>* FindItemAtPosition(FVector2D InScreenSpacePosition) const;
+
+private:
+
+	/** A temporary snapshot of the SparseItemInfos in STreeView, used during RefreshTreeView() */
+	TSparseItemMap OldSparseItemInfos;
+
+	/** Backing array for tree view */
+	TArray<TSharedPtr<FModularRigTreeElement>> RootElements;
+	
+	/** A map for looking up items based on their key */
+	TMap<FString, TSharedPtr<FModularRigTreeElement>> ElementMap;
+
+	/** A map for looking up a parent based on their key */
+	TMap<FString, FString> ParentMap;
+
+	FModularRigTreeDelegates Delegates;
+
+	bool bAutoScrollEnabled;
+	FVector2D LastMousePosition;
+	double TimeAtMousePosition;
+
+	friend class SModularRigHierarchy;
+};
+
+class SSearchableModularRigHierarchyTreeView : public SCompoundWidget
+{
+public:
+
+	SLATE_BEGIN_ARGS(SSearchableModularRigHierarchyTreeView) {}
+		SLATE_ARGUMENT(FModularRigTreeDelegates, RigTreeDelegates)
+		SLATE_ARGUMENT(FText, InitialFilterText)
+	SLATE_END_ARGS()
+
+	void Construct(const FArguments& InArgs);
+	virtual ~SSearchableModularRigHierarchyTreeView() {}
+	TSharedRef<SModularRigHierarchyTreeView> GetTreeView() const { return TreeView.ToSharedRef(); }
+
+private:
+
+	TSharedPtr<SModularRigHierarchyTreeView> TreeView;
+};
