Control Rig: Add Rig module asset browser tab to editor

diff --git a/Engine/Plugins/Animation/ControlRig/Source/ControlRigEditor/Private/Editor/ControlRigEditorMode.cpp b/Engine/Plugins/Animation/ControlRig/Source/ControlRigEditor/Private/Editor/ControlRigEditorMode.cpp
--- a/Engine/Plugins/Animation/ControlRig/Source/ControlRigEditor/Private/Editor/ControlRigEditorMode.cpp
+++ b/Engine/Plugins/Animation/ControlRig/Source/ControlRigEditor/Private/Editor/ControlRigEditorMode.cpp
@@ -8,6 +8,7 @@
 #include "PersonaTabs.h"
 #include "Editor/RigHierarchyTabSummoner.h"
 #include "Editor/ModularRigHierarchyTabSummoner.h"
+#include "Editor/RigModuleAssetBrowserTabSummoner.h"
 #include "Editor/RigVMExecutionStackTabSummoner.h"
 #include "Editor/RigCurveContainerTabSummoner.h"
 #include "Editor/RigValidationTabSummoner.h"
@@ -17,14 +18,9 @@
 FControlRigEditorMode::FControlRigEditorMode(const TSharedRef<FControlRigEditor>& InControlRigEditor, bool bCreateDefaultLayout)
 	: FBlueprintEditorApplicationMode(InControlRigEditor, FControlRigEditorModes::ControlRigEditorMode, FControlRigEditorModes::GetLocalizedMode, false, false)
 {
-
 	ControlRigBlueprintPtr = CastChecked<UControlRigBlueprint>(InControlRigEditor->GetBlueprintObj());
 
 	TabFactories.RegisterFactory(MakeShared<FRigHierarchyTabSummoner>(InControlRigEditor));
-	if (ControlRigBlueprintPtr->IsModularRig())
-	{
-		TabFactories.RegisterFactory(MakeShared<FModularRigHierarchyTabSummoner>(InControlRigEditor));
-	}
 	TabFactories.RegisterFactory(MakeShared<FRigVMExecutionStackTabSummoner>(InControlRigEditor));
 	TabFactories.RegisterFactory(MakeShared<FRigCurveContainerTabSummoner>(InControlRigEditor));
 	TabFactories.RegisterFactory(MakeShared<FRigValidationTabSummoner>(InControlRigEditor));
@@ -46,7 +42,7 @@ FControlRigEditorMode::FControlRigEditorMode(const TSharedRef<FControlRigEditor>
 
 	if(bCreateDefaultLayout)
 	{
-		TabLayout = FTabManager::NewLayout("Standalone_ControlRigEditMode_Layout_v1.6")
+		TabLayout = FTabManager::NewLayout("Standalone_ControlRigEditMode_Layout_v1.8")
 			->AddArea
 			(
 				// Main application area
@@ -73,21 +69,12 @@ FControlRigEditorMode::FControlRigEditorMode(const TSharedRef<FControlRigEditor>
 						->Split
 						(
 							//	Left bottom - rig/hierarchy
-							ControlRigBlueprintPtr->IsModularRig() ?
-								FTabManager::NewStack()
-								->SetSizeCoefficient(0.5f)
-								->AddTab(FRigHierarchyTabSummoner::TabID, ETabState::OpenedTab)
-								->AddTab(FModularRigHierarchyTabSummoner::TabID, ETabState::OpenedTab)
-								->AddTab(FRigVMExecutionStackTabSummoner::TabID, ETabState::OpenedTab)
-								->AddTab(FRigCurveContainerTabSummoner::TabID, ETabState::OpenedTab)
-								->AddTab(FBlueprintEditorTabs::MyBlueprintID, ETabState::OpenedTab)
-									:
-								FTabManager::NewStack()
-								->SetSizeCoefficient(0.5f)
-								->AddTab(FRigHierarchyTabSummoner::TabID, ETabState::OpenedTab)
-								->AddTab(FRigVMExecutionStackTabSummoner::TabID, ETabState::OpenedTab)
-								->AddTab(FRigCurveContainerTabSummoner::TabID, ETabState::OpenedTab)
-								->AddTab(FBlueprintEditorTabs::MyBlueprintID, ETabState::OpenedTab)
+							FTabManager::NewStack()
+							->SetSizeCoefficient(0.5f)
+							->AddTab(FRigHierarchyTabSummoner::TabID, ETabState::OpenedTab)
+							->AddTab(FRigVMExecutionStackTabSummoner::TabID, ETabState::OpenedTab)
+							->AddTab(FRigCurveContainerTabSummoner::TabID, ETabState::OpenedTab)
+							->AddTab(FBlueprintEditorTabs::MyBlueprintID, ETabState::OpenedTab)
 						)
 					)
 					->Split
@@ -157,7 +144,11 @@ void FControlRigEditorMode::RegisterTabFactories(TSharedPtr<FTabManager> InTabMa
 FModularRigEditorMode::FModularRigEditorMode(const TSharedRef<FControlRigEditor>& InControlRigEditor)
 	: FControlRigEditorMode(InControlRigEditor, false)
 {
-	TabLayout = FTabManager::NewLayout("Standalone_ModularRigEditMode_Layout_v1.0")
+
+	TabFactories.RegisterFactory(MakeShared<FModularRigHierarchyTabSummoner>(InControlRigEditor));
+	TabFactories.RegisterFactory(MakeShared<FRigModuleAssetBrowserTabSummoner>(InControlRigEditor));
+	
+	TabLayout = FTabManager::NewLayout("Standalone_ModularRigEditMode_Layout_v1.2")
 		->AddArea
 		(
 			// Main application area
@@ -183,22 +174,27 @@ FModularRigEditorMode::FModularRigEditorMode(const TSharedRef<FControlRigEditor>
 					)
 					->Split
 					(
-						//	Left bottom - rig/hierarchy
-						InControlRigEditor->GetControlRigBlueprint()->IsModularRig() ?
+						//	Left bottom - rig/hierarchy/modules
+						FTabManager::NewSplitter()
+						->SetOrientation(Orient_Horizontal)
+						->SetSizeCoefficient(0.5f)
+						->Split
+						(
+							// Left bottom left
 							FTabManager::NewStack()
 							->SetSizeCoefficient(0.5f)
+							->SetForegroundTab(FRigHierarchyTabSummoner::TabID)
 							->AddTab(FBlueprintEditorTabs::CompilerResultsID, ETabState::ClosedTab)
 							->AddTab(FRigHierarchyTabSummoner::TabID, ETabState::OpenedTab)
-							->AddTab(FModularRigHierarchyTabSummoner::TabID, ETabState::OpenedTab)
 							->AddTab(FRigCurveContainerTabSummoner::TabID, ETabState::OpenedTab)
 							->AddTab(FBlueprintEditorTabs::MyBlueprintID, ETabState::OpenedTab)
-								:
+						)
+						->Split(
+							// Left bottom right
 							FTabManager::NewStack()
 							->SetSizeCoefficient(0.5f)
-							->AddTab(FBlueprintEditorTabs::CompilerResultsID, ETabState::ClosedTab)
-							->AddTab(FRigHierarchyTabSummoner::TabID, ETabState::OpenedTab)
-							->AddTab(FRigCurveContainerTabSummoner::TabID, ETabState::OpenedTab)
-							->AddTab(FBlueprintEditorTabs::MyBlueprintID, ETabState::OpenedTab)
+							->AddTab(FModularRigHierarchyTabSummoner::TabID, ETabState::OpenedTab)
+						)
 					)
 				)
 				->Split
@@ -212,12 +208,20 @@ FModularRigEditorMode::FModularRigEditorMode(const TSharedRef<FControlRigEditor>
 						// Right top
 						FTabManager::NewStack()
 						->SetHideTabWell(false)
-						->SetSizeCoefficient(1.f)
+						->SetSizeCoefficient(0.5f )
 						->AddTab(FBlueprintEditorTabs::DetailsID, ETabState::OpenedTab)
 						->AddTab(FPersonaTabs::AdvancedPreviewSceneSettingsID, ETabState::OpenedTab)
 						->AddTab(FRigAnimAttributeTabSummoner::TabID, ETabState::OpenedTab)
 						->SetForegroundTab(FBlueprintEditorTabs::DetailsID)
 					)
+					->Split(
+						// Right bottom
+						FTabManager::NewStack()
+						->SetHideTabWell(false)
+						->SetSizeCoefficient(0.5f)
+						->AddTab(FRigModuleAssetBrowserTabSummoner::TabID, ETabState::OpenedTab)
+						->SetForegroundTab(FBlueprintEditorTabs::DetailsID)
+					)
 				)
 			)
 		);
diff --git a/Engine/Plugins/Animation/ControlRig/Source/ControlRigEditor/Private/Editor/RigModuleAssetBrowserTabSummoner.cpp b/Engine/Plugins/Animation/ControlRig/Source/ControlRigEditor/Private/Editor/RigModuleAssetBrowserTabSummoner.cpp
new file mode 100644
--- /dev/null
+++ b/Engine/Plugins/Animation/ControlRig/Source/ControlRigEditor/Private/Editor/RigModuleAssetBrowserTabSummoner.cpp
@@ -0,0 +1,41 @@
+// Copyright Epic Games, Inc. All Rights Reserved.
+#include "RigModuleAssetBrowserTabSummoner.h"
+
+#include "IDocumentation.h"
+#include "ControlRigEditorStyle.h"
+#include "ControlRigEditor.h"
+#include "SRigModuleAssetBrowser.h"
+
+#define LOCTEXT_NAMESPACE "RigModuleAssetBrowserTabSummoner"
+
+const FName FRigModuleAssetBrowserTabSummoner::TabID(TEXT("RigModuleAssetBrowser"));
+
+FRigModuleAssetBrowserTabSummoner::FRigModuleAssetBrowserTabSummoner(const TSharedRef<FControlRigEditor>& InEditor)
+	: FWorkflowTabFactory(TabID, InEditor),
+	ControlRigEditor(InEditor)
+{
+	bIsSingleton = true; // only allow a single instance of this tab
+	
+	TabLabel = LOCTEXT("RigModuleAssetBrowserTabLabel", "Rig Module Asset Browser");
+	TabIcon = FSlateIcon(FAppStyle::GetAppStyleSetName(), "ContentBrowser.TabIcon");
+
+	ViewMenuDescription = LOCTEXT("RigModuleAssetBrowser_ViewMenu_Desc", "Rig Module Asset Browser");
+	ViewMenuTooltip = LOCTEXT("RigModuleAssetBrowser_ViewMenu_ToolTip", "Show the Rig Module Asset Browser Tab");
+}
+
+TSharedPtr<SToolTip> FRigModuleAssetBrowserTabSummoner::CreateTabToolTipWidget(const FWorkflowTabSpawnInfo& Info) const
+{
+	return  IDocumentation::Get()->CreateToolTip(LOCTEXT("RigModuleAssetBrowserTooltip", "Drag and drop the rig modules into the Modular Rig Hierarchy tab."), NULL, TEXT("Shared/Editors/Persona"), TEXT("RigModuleAssetBrowser_Window"));
+}
+
+TSharedRef<SWidget> FRigModuleAssetBrowserTabSummoner::CreateTabBody(const FWorkflowTabSpawnInfo& Info) const
+{
+	return SNew(SRigModuleAssetBrowser, ControlRigEditor.Pin().ToSharedRef());
+}
+
+TSharedRef<SDockTab> FRigModuleAssetBrowserTabSummoner::SpawnTab(const FWorkflowTabSpawnInfo& Info) const
+{
+	return FWorkflowTabFactory::SpawnTab(Info);
+}
+
+#undef LOCTEXT_NAMESPACE 
diff --git a/Engine/Plugins/Animation/ControlRig/Source/ControlRigEditor/Private/Editor/RigModuleAssetBrowserTabSummoner.h b/Engine/Plugins/Animation/ControlRig/Source/ControlRigEditor/Private/Editor/RigModuleAssetBrowserTabSummoner.h
new file mode 100644
--- /dev/null
+++ b/Engine/Plugins/Animation/ControlRig/Source/ControlRigEditor/Private/Editor/RigModuleAssetBrowserTabSummoner.h
@@ -0,0 +1,22 @@
+// Copyright Epic Games, Inc. All Rights Reserved.
+
+#pragma once
+
+#include "CoreMinimal.h"
+#include "WorkflowOrientedApp/WorkflowTabFactory.h"
+
+class FControlRigEditor;
+
+struct FRigModuleAssetBrowserTabSummoner : public FWorkflowTabFactory
+{
+public:
+	static const FName TabID;
+	
+	FRigModuleAssetBrowserTabSummoner(const TSharedRef<FControlRigEditor>& InControlRigEditor);
+	virtual TSharedRef<SWidget> CreateTabBody(const FWorkflowTabSpawnInfo& Info) const override;
+	virtual TSharedPtr<SToolTip> CreateTabToolTipWidget(const FWorkflowTabSpawnInfo& Info) const override;
+	virtual TSharedRef<SDockTab> SpawnTab(const FWorkflowTabSpawnInfo& Info) const override;
+
+protected:
+	TWeakPtr<FControlRigEditor> ControlRigEditor;
+};
diff --git a/Engine/Plugins/Animation/ControlRig/Source/ControlRigEditor/Private/Editor/SRigModuleAssetBrowser.cpp b/Engine/Plugins/Animation/ControlRig/Source/ControlRigEditor/Private/Editor/SRigModuleAssetBrowser.cpp
new file mode 100644
--- /dev/null
+++ b/Engine/Plugins/Animation/ControlRig/Source/ControlRigEditor/Private/Editor/SRigModuleAssetBrowser.cpp
@@ -0,0 +1,124 @@
+﻿// Copyright Epic Games, Inc. All Rights Reserved.
+
+#include "SRigModuleAssetBrowser.h"
+
+#include "ContentBrowserModule.h"
+#include "ControlRigBlueprint.h"
+#include "IContentBrowserSingleton.h"
+
+#include "ControlRigEditor.h"
+#include "Widgets/Layout/SUniformGridPanel.h"
+
+#define LOCTEXT_NAMESPACE "RigModuleAssetBrowser"
+
+void SRigModuleAssetBrowser::Construct(
+	const FArguments& InArgs,
+	TSharedRef<FControlRigEditor> InEditor)
+{
+	ControlRigEditor = InEditor;
+
+	ChildSlot
+	[
+		SNew(SVerticalBox)
+		+SVerticalBox::Slot()
+		[
+			SAssignNew(AssetBrowserBox, SBox)
+		]
+	];
+
+	RefreshView();
+}
+
+void SRigModuleAssetBrowser::RefreshView()
+{
+	FAssetPickerConfig AssetPickerConfig;
+	
+	// setup filtering
+	AssetPickerConfig.Filter.ClassPaths.Add(UControlRigBlueprint::StaticClass()->GetClassPathName());
+	AssetPickerConfig.InitialAssetViewType = EAssetViewType::List;
+	AssetPickerConfig.bAddFilterUI = true;
+	AssetPickerConfig.bShowPathInColumnView = true;
+	AssetPickerConfig.bShowTypeInColumnView = true;
+	AssetPickerConfig.OnShouldFilterAsset = FOnShouldFilterAsset::CreateSP(this, &SRigModuleAssetBrowser::OnShouldFilterAsset);
+	AssetPickerConfig.DefaultFilterMenuExpansion = EAssetTypeCategories::Blueprint;
+	AssetPickerConfig.OnGetAssetContextMenu = FOnGetAssetContextMenu::CreateSP(this, &SRigModuleAssetBrowser::OnGetAssetContextMenu);
+	AssetPickerConfig.GetCurrentSelectionDelegates.Add(&GetCurrentSelectionDelegate);
+	AssetPickerConfig.bAllowNullSelection = false;
+	AssetPickerConfig.bFocusSearchBoxWhenOpened = false;
+
+	// hide all asset registry columns by default (we only really want the name and path)
+	TArray<UObject::FAssetRegistryTag> AssetRegistryTags;
+	UControlRigBlueprint::StaticClass()->GetDefaultObject()->GetAssetRegistryTags(AssetRegistryTags);
+	for(UObject::FAssetRegistryTag& AssetRegistryTag : AssetRegistryTags)
+	{
+		AssetPickerConfig.HiddenColumnNames.Add(AssetRegistryTag.Name.ToString());
+	}
+
+	// Also hide the type column by default (but allow users to enable it, so don't use bShowTypeInColumnView)
+	AssetPickerConfig.HiddenColumnNames.Add(TEXT("Class"));
+	AssetPickerConfig.HiddenColumnNames.Add(TEXT("Has Virtualized Data"));
+
+	const FContentBrowserModule& ContentBrowserModule = FModuleManager::Get().LoadModuleChecked<FContentBrowserModule>(TEXT("ContentBrowser"));
+	AssetBrowserBox->SetContent(ContentBrowserModule.Get().CreateAssetPicker(AssetPickerConfig));
+}
+
+TSharedPtr<SWidget> SRigModuleAssetBrowser::OnGetAssetContextMenu(const TArray<FAssetData>& SelectedAssets) const
+{
+	if (SelectedAssets.Num() <= 0)
+	{
+		return nullptr;
+	}
+
+	UObject* SelectedAsset = SelectedAssets[0].GetAsset();
+	if (SelectedAsset == nullptr)
+	{
+		return nullptr;
+	}
+	
+	FMenuBuilder MenuBuilder(true, MakeShared<FUICommandList>());
+
+	MenuBuilder.BeginSection(TEXT("Asset"), LOCTEXT("AssetSectionLabel", "Asset"));
+	{
+		MenuBuilder.AddMenuEntry(
+			LOCTEXT("Browse", "Browse to Asset"),
+			LOCTEXT("BrowseTooltip", "Browses to the associated asset and selects it in the most recently used Content Browser (summoning one if necessary)"),
+			FSlateIcon(FAppStyle::GetAppStyleSetName(), "SystemWideCommands.FindInContentBrowser.Small"),
+			FUIAction(
+				FExecuteAction::CreateLambda([SelectedAsset] ()
+				{
+					if (SelectedAsset)
+					{
+						const TArray<FAssetData>& Assets = { SelectedAsset };
+						const FContentBrowserModule& ContentBrowserModule = FModuleManager::Get().LoadModuleChecked<FContentBrowserModule>("ContentBrowser");
+						ContentBrowserModule.Get().SyncBrowserToAssets(Assets);
+					}
+				}),
+				FCanExecuteAction::CreateLambda([] () { return true; })
+			)
+		);
+	}
+	MenuBuilder.EndSection();
+
+	return MenuBuilder.MakeWidget();
+}
+
+bool SRigModuleAssetBrowser::OnShouldFilterAsset(const struct FAssetData& AssetData)
+{
+	// is this an control rig blueprint asset?
+	if (!AssetData.IsInstanceOf(UControlRigBlueprint::StaticClass()))
+	{
+		return true;
+	}
+	
+	static const UEnum* ControlTypeEnum = StaticEnum<EControlRigType>();
+	const FString ControlRigTypeStr = AssetData.GetTagValueRef<FString>(TEXT("ControlRigType"));
+	if (ControlRigTypeStr.IsEmpty())
+	{
+		return true;
+	}
+
+	const EControlRigType ControlRigType = (EControlRigType)(ControlTypeEnum->GetValueByName(*ControlRigTypeStr));
+	return ControlRigType != EControlRigType::RigModule;
+}
+
+#undef LOCTEXT_NAMESPACE
diff --git a/Engine/Plugins/Animation/ControlRig/Source/ControlRigEditor/Private/Editor/SRigModuleAssetBrowser.h b/Engine/Plugins/Animation/ControlRig/Source/ControlRigEditor/Private/Editor/SRigModuleAssetBrowser.h
new file mode 100644
--- /dev/null
+++ b/Engine/Plugins/Animation/ControlRig/Source/ControlRigEditor/Private/Editor/SRigModuleAssetBrowser.h
@@ -0,0 +1,41 @@
+﻿// Copyright Epic Games, Inc. All Rights Reserved.
+
+#pragma once
+
+#include "CoreMinimal.h"
+
+#include "ContentBrowserDelegates.h"
+#include "Widgets/Layout/SBox.h"
+#include "Widgets/SWindow.h"
+
+
+class FControlRigEditor;
+
+class SRigModuleAssetBrowser : public SBox
+{
+public:
+	SLATE_BEGIN_ARGS(SRigModuleAssetBrowser) {}
+	SLATE_END_ARGS()
+
+	void Construct(const FArguments& InArgs, TSharedRef<FControlRigEditor> InEditor);
+
+	void RefreshView();
+	
+private:
+
+	
+	bool OnShouldFilterAsset(const struct FAssetData& AssetData);
+	TSharedPtr<SWidget> OnGetAssetContextMenu(const TArray<FAssetData>& SelectedAssets) const;
+
+	/** Used to get the currently selected assets */
+	FGetCurrentSelectionDelegate GetCurrentSelectionDelegate;
+		
+	/** editor controller */
+	TWeakPtr<FControlRigEditor> ControlRigEditor;
+
+	/** the animation asset browser */
+	TSharedPtr<SBox> AssetBrowserBox;
+
+
+	friend FControlRigEditor;
+};
