Control Rig: Rename module right after adding it
#rnx

diff --git a/Engine/Plugins/Animation/ControlRig/Source/ControlRig/Public/Rigs/RigHierarchy.h b/Engine/Plugins/Animation/ControlRig/Source/ControlRig/Public/Rigs/RigHierarchy.h
--- a/Engine/Plugins/Animation/ControlRig/Source/ControlRig/Public/Rigs/RigHierarchy.h
+++ b/Engine/Plugins/Animation/ControlRig/Source/ControlRig/Public/Rigs/RigHierarchy.h
@@ -4766,7 +4766,7 @@ private:
 
 	friend class UControlRig;
 	friend class UModularRig;
-	friend class UModularRigController;
+	friend class UControlRigBlueprint;
 };
 
 struct CONTROLRIG_API FRigHierarchyValidityBracket
diff --git a/Engine/Plugins/Animation/ControlRig/Source/ControlRig/Public/Rigs/RigHierarchyController.h b/Engine/Plugins/Animation/ControlRig/Source/ControlRig/Public/Rigs/RigHierarchyController.h
--- a/Engine/Plugins/Animation/ControlRig/Source/ControlRig/Public/Rigs/RigHierarchyController.h
+++ b/Engine/Plugins/Animation/ControlRig/Source/ControlRig/Public/Rigs/RigHierarchyController.h
@@ -767,6 +767,7 @@ private:
 	friend class UControlRig;
 	friend class URigHierarchy;
 	friend class FRigHierarchyControllerInstructionBracket;
+	friend class UControlRigBlueprint;
 };
 
 class CONTROLRIG_API FRigHierarchyControllerInstructionBracket : TGuardValue<int32>
diff --git a/Engine/Plugins/Animation/ControlRig/Source/ControlRigDeveloper/Private/ControlRigBlueprint.cpp b/Engine/Plugins/Animation/ControlRig/Source/ControlRigDeveloper/Private/ControlRigBlueprint.cpp
--- a/Engine/Plugins/Animation/ControlRig/Source/ControlRigDeveloper/Private/ControlRigBlueprint.cpp
+++ b/Engine/Plugins/Animation/ControlRig/Source/ControlRigDeveloper/Private/ControlRigBlueprint.cpp
@@ -2277,6 +2277,107 @@ void UControlRigBlueprint::HandleRigModulesModified(EModularRigNotification InNo
 {
 	switch (InNotification)
 	{
+		case EModularRigNotification::ModuleAdded:
+		{
+			if (InModule)
+			{
+				if (URigHierarchyController* Controller = GetHierarchyController())
+				{
+					if (UControlRig* CDO = GetControlRigClass()->GetDefaultObject<UControlRig>())
+					{
+						FRigVMExtendedExecuteContext& Context = CDO->GetRigVMExtendedExecuteContext();
+						FRigHierarchyExecuteContextBracket HierarchyContextGuard(Controller->GetHierarchy(), &Context);
+			
+						// setup the module information
+						FControlRigExecuteContext& PublicContext = Context.GetPublicDataSafe<FControlRigExecuteContext>();
+						FControlRigExecuteContextRigModuleGuard RigModuleGuard(PublicContext, InModule->GetNamespace());
+			
+						UControlRig* ClassDefaultObject = InModule->Class->GetDefaultObject<UControlRig>();
+						const TArray<FRigModuleConnector>& Connectors = ClassDefaultObject->GetRigModuleSettings().ExposedConnectors;
+						for (const FRigModuleConnector& Connector : Connectors)
+						{
+							Controller->AddConnector(*Connector.Name, Connector.Settings);
+						}
+			
+						// todo: copy the sockets
+						// TArray<FRigSocketElement*> Sockets = DefaultModule->GetHierarchy()->GetElementsOfType<FRigSocketElement>();
+						// for (FRigSocketElement* Socket : Sockets)
+						// {
+						// 	Controller->CopySocket(NewModule->GetPath(), Socket);
+						// }
+
+						PropagateHierarchyFromBPToInstances();
+					}
+				}
+			}
+			break;
+		}
+		case EModularRigNotification::ModuleRenamed:
+		{
+			if (InModule)
+			{
+				if (URigHierarchyController* Controller = GetHierarchyController())
+				{
+					if (UControlRig* CDO = GetControlRigClass()->GetDefaultObject<UControlRig>())
+					{
+						struct ConnectionInfo
+						{
+							FString NewPath;
+							FRigElementKey TargetConnection;
+							FRigConnectorSettings Settings;
+						};
+						FString OldPath = FString::Printf(TEXT("%s%s:"), *InModule->ParentNamespace, *InModule->PreviousName.ToString());
+						FString NewPath = FString::Printf(TEXT("%s%s:"), *InModule->ParentNamespace, *InModule->Name.ToString());
+						TArray<FRigElementKey> Connectors = Controller->GetHierarchy()->GetKeysOfType<FRigConnectorElement>();
+						TMap<FRigElementKey, ConnectionInfo> RenamedConnectors; // old key -> new key
+						for (const FRigElementKey& Connector : Connectors)
+						{
+							FString OldConnectorName = Connector.Name.ToString();
+							if (OldConnectorName.StartsWith(OldPath))
+							{
+								ConnectionInfo& Info = RenamedConnectors.FindOrAdd(Connector);
+								Info.NewPath = OldConnectorName.Replace(*OldPath, *NewPath);
+								Info.Settings = CastChecked<FRigConnectorElement>(Controller->GetHierarchy()->FindChecked(Connector))->Settings;
+								if (FRigElementKey* TargetKey = ConnectionMap.Find(Connector))
+								{
+									Info.TargetConnection = *TargetKey;
+								}
+							}
+						}
+
+						// Remove connectors
+						for (TPair<FRigElementKey, ConnectionInfo>& Pair : RenamedConnectors)
+						{
+							Controller->RemoveElement(Pair.Key);
+						}
+
+						// Add connectors
+						{
+							FRigVMExtendedExecuteContext& Context = CDO->GetRigVMExtendedExecuteContext();
+							FRigHierarchyExecuteContextBracket HierarchyContextGuard(Controller->GetHierarchy(), &Context);
+							FControlRigExecuteContext& PublicContext = Context.GetPublicDataSafe<FControlRigExecuteContext>();
+							for (TPair<FRigElementKey, ConnectionInfo>& Pair : RenamedConnectors)
+							{
+								FString Namespace, ConnectorName;
+								Pair.Value.NewPath.Split(UModularRig::NamespaceSeparator, &Namespace, &ConnectorName, ESearchCase::CaseSensitive, ESearchDir::FromEnd);
+								Namespace.AppendChar(TCHAR(':'));
+								FControlRigExecuteContextRigModuleGuard RigModuleGuard(PublicContext, Namespace);
+								Controller->AddConnector(*ConnectorName, Pair.Value.Settings);
+							}
+						}
+
+						// Resolve Connectors
+						for (TPair<FRigElementKey, ConnectionInfo>& Pair : RenamedConnectors)
+						{
+							ResolveConnector(FRigElementKey(*Pair.Value.NewPath, ERigElementType::Connector), Pair.Value.TargetConnection);
+						}
+
+						PropagateHierarchyFromBPToInstances();
+					}
+				}
+			}
+			break;
+		}
 		case EModularRigNotification::ConnectionChanged:
 		{
 			const FString Namespace = InModule->GetNamespace();
diff --git a/Engine/Plugins/Animation/ControlRig/Source/ControlRigDeveloper/Private/ModularRigController.cpp b/Engine/Plugins/Animation/ControlRig/Source/ControlRigDeveloper/Private/ModularRigController.cpp
--- a/Engine/Plugins/Animation/ControlRig/Source/ControlRigDeveloper/Private/ModularRigController.cpp
+++ b/Engine/Plugins/Animation/ControlRig/Source/ControlRigDeveloper/Private/ModularRigController.cpp
@@ -87,37 +87,6 @@ bool UModularRigController::AddModule(const FName& InModuleName, TSubclassOf<UCo
 
 	Notify(EModularRigNotification::ModuleAdded, NewModule);
 
-	const FString NewNamespace = NewModule->GetNamespace();
-
-	// Add Connectors/Sockets
-	UControlRigBlueprint* Blueprint = Cast<UControlRigBlueprint>(GetOuter());
-	if (URigHierarchyController* Controller = Blueprint->GetHierarchyController())
-	{
-		if (UControlRig* CDO = Blueprint->GetControlRigClass()->GetDefaultObject<UControlRig>())
-		{
-			FRigVMExtendedExecuteContext& Context = CDO->GetRigVMExtendedExecuteContext();
-			FRigHierarchyExecuteContextBracket HierarchyContextGuard(Controller->GetHierarchy(), &Context);
-
-			// setup the module information
-			FControlRigExecuteContext& PublicContext = Context.GetPublicDataSafe<FControlRigExecuteContext>();
-			FControlRigExecuteContextRigModuleGuard RigModuleGuard(PublicContext, NewNamespace);
-
-			
-			const TArray<FRigModuleConnector>& Connectors = ClassDefaultObject->GetRigModuleSettings().ExposedConnectors;
-			for (const FRigModuleConnector& Connector : Connectors)
-			{
-				Controller->AddConnector(*Connector.Name, Connector.Settings);
-			}
-
-			// todo: copy the sockets
-			// TArray<FRigSocketElement*> Sockets = DefaultModule->GetHierarchy()->GetElementsOfType<FRigSocketElement>();
-			// for (FRigSocketElement* Socket : Sockets)
-			// {
-			// 	Controller->CopySocket(NewModule->GetPath(), Socket);
-			// }
-		}
-	}
-
 #if WITH_EDITOR
 	TransactionPtr.Reset();
 #endif
@@ -263,19 +232,105 @@ bool UModularRigController::SetConfigValueInModule(const FString& InModulePath,
 	return true;
 }
 
-bool UModularRigController::RemoveModule(const FString& InModulesPath, bool bSetupUndo)
+bool UModularRigController::RemoveModule(const FString& InModulePath, bool bSetupUndo)
 {
 	// todo: UE-199050
 	return false;
 }
 
-bool UModularRigController::RenameModule(const FString& InModulesPath, const FName& InNewName, bool bSetupUndo)
+bool UModularRigController::RenameModule(const FString& InModulePath, const FName& InNewName, bool bSetupUndo)
 {
-	// todo: UE-199050
-	return false;
+	FRigModuleReference* Module = FindModule(InModulePath);
+	if (!Module)
+	{
+		UE_LOG(LogControlRig, Error, TEXT("Could not find module %s"), *InModulePath);
+		return false;
+	}
+
+	FText ErrorMessage;
+	if (!CanRenameModule(InModulePath, InNewName, ErrorMessage))
+	{
+		UE_LOG(LogControlRig, Error, TEXT("Could not rename module %s: %s"), *InModulePath, *ErrorMessage.ToString());
+		return false;
+	}
+	
+#if WITH_EDITOR
+	TSharedPtr<FScopedTransaction> TransactionPtr;
+	if (bSetupUndo)
+	{
+		TransactionPtr = MakeShared<FScopedTransaction>(NSLOCTEXT("ModularRigController", "RenameModuleTransaction", "Rename Module"));
+		if(UControlRigBlueprint* Blueprint = Cast<UControlRigBlueprint>(GetOuter()))
+		{
+			Blueprint->Modify();
+		}
+	}
+#endif
+	
+	const FString OldPath = FString::Printf(TEXT("%s%s"), *Module->ParentNamespace, *Module->Name.ToString());
+	const FString NewPath = FString::Printf(TEXT("%s%s"), *Module->ParentNamespace, *InNewName.ToString());
+	Module->PreviousName = Module->Name;
+	Module->Name = InNewName;
+	TArray<FRigModuleReference*> Children;
+	Children.Append(Module->CachedChildren);
+	for (int32 i=0; i<Children.Num(); ++i)
+	{
+		FRigModuleReference* Child = Children[i];
+		Child->ParentNamespace.ReplaceInline(*OldPath, *NewPath);
+
+		Children.Append(Child->CachedChildren);
+	}
+
+	Notify(EModularRigNotification::ModuleRenamed, Module);
+
+#if WITH_EDITOR
+	TransactionPtr.Reset();
+#endif
+	
+	return true;
+}
+
+bool UModularRigController::CanRenameModule(const FString& InModulePath, const FName& InNewName, FText& OutErrorMessage)
+{
+	if (InNewName.IsNone() || InNewName.ToString().IsEmpty())
+	{
+		OutErrorMessage = FText::FromString(TEXT("Name is empty."));
+		return false;
+	}
+
+	FRigModuleReference* Module = FindModule(InModulePath);
+	if (!Module)
+	{
+		OutErrorMessage = FText::FromString(FString::Printf(TEXT("Module %s not found."), *InModulePath));
+		return false;
+	}
+
+	const FString NewPath = FString::Printf(TEXT("%s%s"), *Module->ParentNamespace, *InNewName.ToString());
+	if (InModulePath == NewPath)
+	{
+		return true;
+	}
+
+	// Check there are no siblings with the same name
+	{
+		FRigModuleReference* Parent = FindModule(Module->ParentNamespace);
+		TArray<FRigModuleReference*>* Siblings = &Model->RootModules;
+		if (Parent)
+		{
+			Siblings = &Parent->CachedChildren;
+		}
+		for (FRigModuleReference* Sibling : *Siblings)
+		{
+			if (Sibling->Name == InNewName)
+			{
+				OutErrorMessage = FText::FromString(FString::Printf(TEXT("Sibling with path %s already exists"), *Sibling->GetNamespace()));
+				return false;
+			}
+		}
+	}
+	return true;
 }
 
-bool UModularRigController::ReparentModule(const FString& InModulesPath, const FString& InNewParentModulePath, bool bSetupUndo)
+bool UModularRigController::ReparentModule(const FString& InModulePath, const FString& InNewParentModulePath, bool bSetupUndo)
 {
 	// todo: UE-199050
 	return false;
diff --git a/Engine/Plugins/Animation/ControlRig/Source/ControlRigDeveloper/Public/ModularRigController.h b/Engine/Plugins/Animation/ControlRig/Source/ControlRigDeveloper/Public/ModularRigController.h
--- a/Engine/Plugins/Animation/ControlRig/Source/ControlRigDeveloper/Public/ModularRigController.h
+++ b/Engine/Plugins/Animation/ControlRig/Source/ControlRigDeveloper/Public/ModularRigController.h
@@ -16,6 +16,8 @@ enum class EModularRigNotification : uint8
 {
 	ModuleAdded,
 
+	ModuleRenamed,
+
 	ConnectionChanged,
 
 	ModuleConfigValueChanged,
@@ -56,9 +58,11 @@ class CONTROLRIGDEVELOPER_API UModularRigController : public UObject
 
 	UFUNCTION(BlueprintCallable, Category = "ControlRig | Modules")
 	bool RenameModule(const FString& InModulePath, const FName& InNewName, bool bSetupUndo = true);
+	bool CanRenameModule(const FString& InModulePath, const FName& InNewName, FText& OutErrorMessage);
 
 	UFUNCTION(BlueprintCallable, Category = "ControlRig | Modules")
 	bool ReparentModule(const FString& InModulePath, const FString& InNewParentModulePath, bool bSetupUndo = true);
+	
 
 
 	FName GetSafeNewName(const FString& InModuleDesiredPath);
diff --git a/Engine/Plugins/Animation/ControlRig/Source/ControlRigDeveloper/Public/ModularRigModel.h b/Engine/Plugins/Animation/ControlRig/Source/ControlRigDeveloper/Public/ModularRigModel.h
--- a/Engine/Plugins/Animation/ControlRig/Source/ControlRigDeveloper/Public/ModularRigModel.h
+++ b/Engine/Plugins/Animation/ControlRig/Source/ControlRigDeveloper/Public/ModularRigModel.h
@@ -38,6 +38,9 @@ struct CONTROLRIGDEVELOPER_API FRigModuleReference
 	UPROPERTY()
 	TMap<FName, FString> ConfigValues;
 
+	UPROPERTY(transient)
+	FName PreviousName;
+
 	TArray<FRigModuleReference*> CachedChildren;
 
 	FString GetNamespace() const;
diff --git a/Engine/Plugins/Animation/ControlRig/Source/ControlRigEditor/Private/ControlRigModularRigHierarchyCommands.cpp b/Engine/Plugins/Animation/ControlRig/Source/ControlRigEditor/Private/ControlRigModularRigHierarchyCommands.cpp
--- a/Engine/Plugins/Animation/ControlRig/Source/ControlRigEditor/Private/ControlRigModularRigHierarchyCommands.cpp
+++ b/Engine/Plugins/Animation/ControlRig/Source/ControlRigEditor/Private/ControlRigModularRigHierarchyCommands.cpp
@@ -7,6 +7,7 @@
 void FControlRigModularHierarchyCommands::RegisterCommands()
 {
 	UI_COMMAND(AddModuleItem, "New Module", "Add new module to the rig.", EUserInterfaceActionType::Button, FInputChord());
+	UI_COMMAND(RenameModuleItem, "Rename", "Rename module to the rig.", EUserInterfaceActionType::Button, FInputChord(EKeys::F2));
 }
 
 #undef LOCTEXT_NAMESPACE
diff --git a/Engine/Plugins/Animation/ControlRig/Source/ControlRigEditor/Private/ControlRigModularRigHierarchyCommands.h b/Engine/Plugins/Animation/ControlRig/Source/ControlRigEditor/Private/ControlRigModularRigHierarchyCommands.h
--- a/Engine/Plugins/Animation/ControlRig/Source/ControlRigEditor/Private/ControlRigModularRigHierarchyCommands.h
+++ b/Engine/Plugins/Animation/ControlRig/Source/ControlRigEditor/Private/ControlRigModularRigHierarchyCommands.h
@@ -21,6 +21,9 @@ public:
 	/** Add Module at root */
 	TSharedPtr< FUICommandInfo > AddModuleItem;
 
+	/** Rename Module */
+	TSharedPtr< FUICommandInfo > RenameModuleItem;
+
 	/**
 	 * Initialize commands
 	 */
diff --git a/Engine/Plugins/Animation/ControlRig/Source/ControlRigEditor/Private/Editor/SModularRigHierarchy.cpp b/Engine/Plugins/Animation/ControlRig/Source/ControlRigEditor/Private/Editor/SModularRigHierarchy.cpp
--- a/Engine/Plugins/Animation/ControlRig/Source/ControlRigEditor/Private/Editor/SModularRigHierarchy.cpp
+++ b/Engine/Plugins/Animation/ControlRig/Source/ControlRigEditor/Private/Editor/SModularRigHierarchy.cpp
@@ -96,7 +96,10 @@ void SModularRigHierarchy::Construct(const FArguments& InArgs, TSharedRef<FContr
 	Delegates.OnCanAcceptDrop = FOnModularRigTreeCanAcceptDrop::CreateSP(this, &SModularRigHierarchy::OnCanAcceptDrop);
 	Delegates.OnAcceptDrop = FOnModularRigTreeAcceptDrop::CreateSP(this, &SModularRigHierarchy::OnAcceptDrop);
 	Delegates.OnMouseButtonClick = FOnModularRigTreeMouseButtonClick::CreateSP(this, &SModularRigHierarchy::OnItemClicked);
+	Delegates.OnMouseButtonDoubleClick = FOnModularRigTreeMouseButtonClick::CreateSP(this, &SModularRigHierarchy::OnItemDoubleClicked);
 	Delegates.OnRequestDetailsInspection = FOnModularRigTreeRequestDetailsInspection::CreateSP(this, &SModularRigHierarchy::OnRequestDetailsInspection);
+	Delegates.OnRenameElement = FOnModularRigTreeRenameElement::CreateSP(this, &SModularRigHierarchy::HandleRenameModule);
+	Delegates.OnVerifyModuleNameChanged = FOnModularRigTreeVerifyElementNameChanged::CreateSP(this, &SModularRigHierarchy::HandleVerifyNameChanged);
 	
 	ChildSlot
 	[
@@ -162,6 +165,10 @@ void SModularRigHierarchy::BindCommands()
 	CommandList->MapAction(Commands.AddModuleItem,
 		FExecuteAction::CreateSP(this, &SModularRigHierarchy::HandleNewItem),
 		FCanExecuteAction());
+
+	CommandList->MapAction(Commands.RenameModuleItem,
+		FExecuteAction::CreateSP(this, &SModularRigHierarchy::HandleRenameModule),
+		FCanExecuteAction());
 }
 
 void SModularRigHierarchy::RefreshTreeView(bool bRebuildContent)
@@ -237,18 +244,11 @@ void SModularRigHierarchy::OnItemClicked(TSharedPtr<FModularRigTreeElement> InIt
 	{
 		ControlRigEditor.Pin()->SetDetailViewForRigModules({InItem->Key});
 	}
+}
 
-	uint32 CurrentCycles = FPlatformTime::Cycles();
-	// double SecondsPassed = double(CurrentCycles - TreeView->LastClickCycles) * FPlatformTime::GetSecondsPerCycle();
-	// if (SecondsPassed > 0.5f)
-	// {
-	// 	RegisterActiveTimer(0.f, FWidgetActiveTimerDelegate::CreateLambda([this](double, float) {
-	// 		HandleRenameItem();
-	// 		return EActiveTimerReturnType::Stop;
-	// 	}));
-	// }
+void SModularRigHierarchy::OnItemDoubleClicked(TSharedPtr<FModularRigTreeElement> InItem)
+{
 
-	TreeView->LastClickCycles = CurrentCycles;
 }
 
 void SModularRigHierarchy::CreateContextMenu()
@@ -294,6 +294,7 @@ void SModularRigHierarchy::CreateContextMenu()
 							DefaultSection.AddMenuEntry(Commands.AddModuleItem);
 						})
 					);
+					ElementsSection.AddMenuEntry(Commands.RenameModuleItem);
 				}
 			})
 		);
@@ -439,17 +440,96 @@ void SModularRigHierarchy::HandleNewItem(UClass* InClass, const FString &InParen
 	{
 		return;
 	}
+
+	FSlateApplication::Get().DismissAllMenus();
 	
 	if (ControlRigBlueprint.IsValid())
 	{
 		FString ClassName = InClass->GetName();
 		ClassName.RemoveFromEnd(TEXT("_C"));
-		FString PathName = FString::Printf(TEXT("%s:%s"), *InParentPath, *ClassName);
+		FString PathName = InParentPath.IsEmpty() ? *ClassName : FString::Printf(TEXT("%s:%s"), *InParentPath, *ClassName);
 		const FName Name = CreateUniqueName(*PathName);
 		ControlRigBlueprint->GetModularRigController()->AddModule(Name, InClass, InParentPath);
+
+		FString NewPathName = InParentPath.IsEmpty() ? *Name.ToString() : FString::Printf(TEXT("%s:%s"), *InParentPath, *Name.ToString());
+		TSharedPtr<FModularRigTreeElement> Element = TreeView->FindElement(NewPathName);
+		if (Element.IsValid())
+		{
+			TreeView->SetSelection({Element});
+			TreeView->bRequestRenameSelected = true;
+		}
+	}
+}
+
+bool SModularRigHierarchy::CanRenameModule() const
+{
+	return IsSingleSelected();
+}
+
+void SModularRigHierarchy::HandleRenameModule()
+{
+	if(!ControlRigEditor.IsValid())
+	{
+		return;
+	}
+
+	if (!CanRenameModule())
+	{
+		return;
+	}
+
+	UModularRig* Rig = GetDefaultHierarchy();
+	if (Rig)
+	{
+		FScopedTransaction Transaction(LOCTEXT("ModularRigHierarchyRenameSelected", "Rename selected module"));
+
+		TArray<TSharedPtr<FModularRigTreeElement>> SelectedItems = TreeView->GetSelectedItems();
+		if (SelectedItems.Num() == 1)
+		{
+			SelectedItems[0]->RequestRename();
+		}
+	}
+
+	return;
+}
+
+FName SModularRigHierarchy::HandleRenameModule(const FString& InOldPath, const FName& InNewName)
+{
+	ClearDetailPanel();
+	
+	if (ControlRigBlueprint.IsValid())
+	{
+		FScopedTransaction Transaction(LOCTEXT("ModularRigHierarchyRename", "Rename Module"));
+
+		UModularRigController* Controller = ControlRigBlueprint->GetModularRigController();
+		check(Controller);
+
+		FName ResultingName = NAME_None;
+		if (Controller->RenameModule(InOldPath, InNewName))
+		{
+			return InNewName;
+		}
+	}
+
+	return NAME_None;
+}
+
+bool SModularRigHierarchy::HandleVerifyNameChanged(const FString& InOldPath, const FName& InNewName, FText& OutErrorMessage)
+{
+	if (InNewName.IsNone())
+	{
+		return false;
 	}
 	
-	FSlateApplication::Get().DismissAllMenus();
+	if (ControlRigBlueprint.IsValid())
+	{
+		UModularRigController* Controller = ControlRigBlueprint->GetModularRigController();
+		check(Controller);
+
+		return Controller->CanRenameModule(InOldPath, InNewName, OutErrorMessage);
+	}
+
+	return false;
 }
 
 class SModularRigHierarchyPasteTransformsErrorPipe : public FOutputDevice
@@ -517,6 +597,14 @@ void SModularRigHierarchy::OnRequestDetailsInspection(const FString& InKey)
 	ControlRigEditor.Pin()->SetDetailViewForRigModules({InKey});
 }
 
+void SModularRigHierarchy::ClearDetailPanel() const
+{
+	if(ControlRigEditor.IsValid())
+	{
+		ControlRigEditor.Pin()->ClearDetailObject();
+	}
+}
+
 void SModularRigHierarchy::PostRedo(bool bSuccess) 
 {
 	if (bSuccess)
diff --git a/Engine/Plugins/Animation/ControlRig/Source/ControlRigEditor/Private/Editor/SModularRigHierarchy.h b/Engine/Plugins/Animation/ControlRig/Source/ControlRigEditor/Private/Editor/SModularRigHierarchy.h
--- a/Engine/Plugins/Animation/ControlRig/Source/ControlRigEditor/Private/Editor/SModularRigHierarchy.h
+++ b/Engine/Plugins/Animation/ControlRig/Source/ControlRigEditor/Private/Editor/SModularRigHierarchy.h
@@ -7,7 +7,6 @@
 #include "Editor/SModularRigHierarchyTreeView.h"
 #include "ControlRigBlueprint.h"
 #include "Editor/RigVMEditor.h"
-//#include "SModularRigHierarchy.generated.h"
 
 class SModularRigHierarchy;
 class FControlRigEditor;
@@ -59,11 +58,17 @@ private:
 
 	void HandleNewItem(UClass* InClass, const FString& InParentPath);
 
+	bool CanRenameModule() const;
+	void HandleRenameModule();
+	FName HandleRenameModule(const FString& InOldPath, const FName& InNewName);
+	bool HandleVerifyNameChanged(const FString& InOldPath, const FName& InNewName, FText& OutErrorMessage);
+
 	/** Set Selection Changed */
 	void OnSelectionChanged(TSharedPtr<FModularRigTreeElement> Selection, ESelectInfo::Type SelectInfo);
 
 	TSharedPtr< SWidget > CreateContextMenuWidget();
 	void OnItemClicked(TSharedPtr<FModularRigTreeElement> InItem);
+	void OnItemDoubleClicked(TSharedPtr<FModularRigTreeElement> InItem);
 	
 	// FEditorUndoClient
 	virtual void PostUndo(bool bSuccess) override;
@@ -100,6 +105,7 @@ private:
 	const UModularRig* GetHierarchyForTreeView() const { return GetHierarchy(); }
 	FName CreateUniqueName(const FName& InBasePath) const;
 	void OnRequestDetailsInspection(const FString& InKey);
+	void ClearDetailPanel() const;
 
 	void HandleRefreshEditorFromBlueprint(URigVMBlueprint* InBlueprint);
 	void HandleSetObjectBeingDebugged(UObject* InObject);
diff --git a/Engine/Plugins/Animation/ControlRig/Source/ControlRigEditor/Private/Editor/SModularRigHierarchyTreeView.cpp b/Engine/Plugins/Animation/ControlRig/Source/ControlRigEditor/Private/Editor/SModularRigHierarchyTreeView.cpp
--- a/Engine/Plugins/Animation/ControlRig/Source/ControlRigEditor/Private/Editor/SModularRigHierarchyTreeView.cpp
+++ b/Engine/Plugins/Animation/ControlRig/Source/ControlRigEditor/Private/Editor/SModularRigHierarchyTreeView.cpp
@@ -61,6 +61,11 @@ TSharedRef<ITableRow> FModularRigTreeElement::MakeTreeRowWidget(const TSharedRef
 	return SNew(SModularRigHierarchyItem, InOwnerTable, InRigTreeElement, InTreeView, bPinned);
 }
 
+void FModularRigTreeElement::RequestRename()
+{
+	OnRenameRequested.ExecuteIfBound();
+}
+
 //////////////////////////////////////////////////////////////
 /// SModularRigHierarchyItem
 ///////////////////////////////////////////////////////////
@@ -130,6 +135,8 @@ void SModularRigHierarchyItem::Construct(const FArguments& InArgs, const TShared
 			[
 				SAssignNew(InlineWidget, SInlineEditableTextBlock)
 				.Text(this, &SModularRigHierarchyItem::GetName, true)
+				.OnVerifyTextChanged(this, &SModularRigHierarchyItem::OnVerifyNameChanged)
+				.OnTextCommitted(this, &SModularRigHierarchyItem::OnNameCommitted)
 				.ToolTipText(this, &SModularRigHierarchyItem::GetItemTooltip)
 				.MultiLine(false)
 				.ColorAndOpacity_Lambda([this]()
@@ -142,6 +149,28 @@ void SModularRigHierarchyItem::Construct(const FArguments& InArgs, const TShared
 				})
 			]
 		], OwnerTable);
+
+	InRigTreeElement->OnRenameRequested.BindSP(InlineWidget.Get(), &SInlineEditableTextBlock::EnterEditingMode);
+}
+
+void SModularRigHierarchyItem::OnNameCommitted(const FText& InText, ETextCommit::Type InCommitType) const
+{
+	// for now only allow enter
+	// because it is important to keep the unique names per pose
+	if (InCommitType == ETextCommit::OnEnter)
+	{
+		FString NewName = InText.ToString();
+		const FString OldKey = WeakRigTreeElement.Pin()->Key;
+
+		Delegates.HandleRenameElement(OldKey, *NewName);
+	}
+}
+
+bool SModularRigHierarchyItem::OnVerifyNameChanged(const FText& InText, FText& OutErrorMessage)
+{
+	const FName NewName = *InText.ToString();
+	const FString OldPath = WeakRigTreeElement.Pin()->Key;
+	return Delegates.HandleVerifyElementNameChanged(OldPath, NewName, OutErrorMessage);
 }
 
 FText SModularRigHierarchyItem::GetName(bool bUseShortName) const
@@ -183,6 +212,7 @@ void SModularRigHierarchyTreeView::Construct(const FArguments& InArgs)
 	SuperArgs.ItemHeight(24);
 	SuperArgs.AllowInvisibleItemSelection(true);  //without this we deselect everything when we filter or we collapse
 	SuperArgs.OnMouseButtonClick(Delegates.OnMouseButtonClick);
+	SuperArgs.OnMouseButtonDoubleClick(Delegates.OnMouseButtonDoubleClick);
 	
 	SuperArgs.ShouldStackHierarchyHeaders_Lambda([]() -> bool {
 		return UControlRigEditorSettings::Get()->bShowStackedHierarchy;
@@ -251,6 +281,32 @@ void SModularRigHierarchyTreeView::Tick(const FGeometry& AllottedGeometry, const
 			}
 		}
 	}
+
+	if (bRequestRenameSelected)
+	{
+		RegisterActiveTimer(0.f, FWidgetActiveTimerDelegate::CreateLambda([this](double, float) {
+			TArray<TSharedPtr<FModularRigTreeElement>> SelectedItems = GetSelectedItems();
+			if (SelectedItems.Num() == 1)
+			{
+				SelectedItems[0]->RequestRename();
+			}
+			return EActiveTimerReturnType::Stop;
+		}));
+		bRequestRenameSelected = false;
+	}
+}
+
+TSharedPtr<FModularRigTreeElement> SModularRigHierarchyTreeView::FindElement(const FString& InElementKey)
+{
+	for (TSharedPtr<FModularRigTreeElement> Root : RootElements)
+	{
+		if (TSharedPtr<FModularRigTreeElement> Found = FindElement(InElementKey, Root))
+		{
+			return Found;
+		}
+	}
+
+	return TSharedPtr<FModularRigTreeElement>();
 }
 
 TSharedPtr<FModularRigTreeElement> SModularRigHierarchyTreeView::FindElement(const FString& InElementKey, TSharedPtr<FModularRigTreeElement> CurrentItem)
@@ -506,6 +562,12 @@ TArray<FString> SModularRigHierarchyTreeView::GetSelectedKeys() const
 	return Keys;
 }
 
+void SModularRigHierarchyTreeView::SetSelection(const TArray<TSharedPtr<FModularRigTreeElement>>& InSelection) 
+{
+	ClearSelection();
+	SetItemSelection(InSelection, true, ESelectInfo::Direct);
+}
+
 const TSharedPtr<FModularRigTreeElement>* SModularRigHierarchyTreeView::FindItemAtPosition(FVector2D InScreenSpacePosition) const
 {
 	if (ItemsPanel.IsValid() && SListView<TSharedPtr<FModularRigTreeElement>>::HasValidItemsSource())
diff --git a/Engine/Plugins/Animation/ControlRig/Source/ControlRigEditor/Private/Editor/SModularRigHierarchyTreeView.h b/Engine/Plugins/Animation/ControlRig/Source/ControlRigEditor/Private/Editor/SModularRigHierarchyTreeView.h
--- a/Engine/Plugins/Animation/ControlRig/Source/ControlRigEditor/Private/Editor/SModularRigHierarchyTreeView.h
+++ b/Engine/Plugins/Animation/ControlRig/Source/ControlRigEditor/Private/Editor/SModularRigHierarchyTreeView.h
@@ -14,8 +14,12 @@ class FModularRigTreeElement;
 
 DECLARE_DELEGATE_RetVal(const UModularRig*, FOnGetModularRigTreeHierarchy);
 DECLARE_DELEGATE_OneParam(FOnModularRigTreeRequestDetailsInspection, const FString&);
+DECLARE_DELEGATE_RetVal_TwoParams(FName, FOnModularRigTreeRenameElement, const FString& /*OldPath*/, const FName& /*NewName*/);
+DECLARE_DELEGATE_RetVal_ThreeParams(bool, FOnModularRigTreeVerifyElementNameChanged, const FString& /*OldPath*/, const FName& /*NewName*/, FText& /*OutErrorMessage*/);
+
 
 typedef STreeView<TSharedPtr<FModularRigTreeElement>>::FOnMouseButtonClick FOnModularRigTreeMouseButtonClick;
+typedef STreeView<TSharedPtr<FModularRigTreeElement>>::FOnMouseButtonDoubleClick FOnModularRigTreeMouseButtonDoubleClick;
 typedef STableRow<TSharedPtr<FModularRigTreeElement>>::FOnCanAcceptDrop FOnModularRigTreeCanAcceptDrop;
 typedef STableRow<TSharedPtr<FModularRigTreeElement>>::FOnAcceptDrop FOnModularRigTreeAcceptDrop;
 
@@ -23,10 +27,13 @@ struct CONTROLRIGEDITOR_API FModularRigTreeDelegates
 {
 	FOnGetModularRigTreeHierarchy OnGetHierarchy;
 	FOnModularRigTreeMouseButtonClick OnMouseButtonClick;
+	FOnModularRigTreeMouseButtonDoubleClick OnMouseButtonDoubleClick;
 	FOnModularRigTreeCanAcceptDrop OnCanAcceptDrop;
 	FOnModularRigTreeAcceptDrop OnAcceptDrop;
 	FOnContextMenuOpening OnContextMenuOpening;
 	FOnModularRigTreeRequestDetailsInspection OnRequestDetailsInspection;
+	FOnModularRigTreeRenameElement OnRenameElement;
+	FOnModularRigTreeVerifyElementNameChanged OnVerifyModuleNameChanged;
 	
 	FModularRigTreeDelegates()
 	{
@@ -40,6 +47,24 @@ struct CONTROLRIGEDITOR_API FModularRigTreeDelegates
 		}
 		return nullptr;
 	}
+
+	FName HandleRenameElement(const FString& OldPath, const FName& NewName) const
+	{
+		if(OnRenameElement.IsBound())
+		{
+			return OnRenameElement.Execute(OldPath, NewName);
+		}
+		return *OldPath;
+	}
+
+	bool HandleVerifyElementNameChanged(const FString& OldPath, const FName& NewName, FText& OutErrorMessage) const
+	{
+		if(OnVerifyModuleNameChanged.IsBound())
+		{
+			return OnVerifyModuleNameChanged.Execute(OldPath, NewName, OutErrorMessage);
+		}
+		return false;
+	}
 };
 
 
@@ -57,8 +82,14 @@ public:
 
 	TSharedRef<ITableRow> MakeTreeRowWidget(const TSharedRef<STableViewBase>& InOwnerTable, TSharedRef<FModularRigTreeElement> InRigTreeElement, TSharedPtr<SModularRigHierarchyTreeView> InTreeView, bool bPinned);
 
+	void RequestRename();
+
 	void RefreshDisplaySettings(const UModularRig* InHierarchy);
 
+	/** Delegate for when the context menu requests a rename */
+	DECLARE_DELEGATE(FOnRenameRequested);
+	FOnRenameRequested OnRenameRequested;
+
 	/** The brush to use when rendering an icon */
 	const FSlateBrush* IconBrush;
 
@@ -74,7 +105,10 @@ class SModularRigHierarchyItem : public STableRow<TSharedPtr<FModularRigTreeElem
 public:
 	
 	void Construct(const FArguments& InArgs, const TSharedRef<STableViewBase>& OwnerTable, TSharedRef<FModularRigTreeElement> InRigTreeElement, TSharedPtr<SModularRigHierarchyTreeView> InTreeView, bool bPinned);
- 	static TPair<const FSlateBrush*, FSlateColor> GetBrushForElementType(const UModularRig* InHierarchy, const FString& InKey);
+
+	void OnNameCommitted(const FText& InText, ETextCommit::Type InCommitType) const;
+	bool OnVerifyNameChanged(const FText& InText, FText& OutErrorMessage);
+	static TPair<const FSlateBrush*, FSlateColor> GetBrushForElementType(const UModularRig* InHierarchy, const FString& InKey);
 	static FLinearColor GetColorForControlType(ERigControlType InControlType, UEnum* InControlEnum);
 
 private:
@@ -106,16 +140,7 @@ public:
 	/** Performs auto scroll */
 	virtual void Tick(const FGeometry& AllottedGeometry, const double InCurrentTime, const float InDeltaTime) override;
 
-	virtual FReply OnFocusReceived(const FGeometry& MyGeometry, const FFocusEvent& InFocusEvent) override
-	{
-		FReply Reply = STreeView<TSharedPtr<FModularRigTreeElement>>::OnFocusReceived(MyGeometry, InFocusEvent);
-
-		LastClickCycles = FPlatformTime::Cycles();
-
-		return Reply;
-	}
-
-	uint32 LastClickCycles = 0;
+	bool bRequestRenameSelected = false;
 
 	/** Save a snapshot of the internal map that tracks item expansion before tree reconstruction */
 	void SaveAndClearSparseItemInfos()
@@ -143,6 +168,7 @@ public:
 	}
 
 
+	TSharedPtr<FModularRigTreeElement> FindElement(const FString& InElementKey);
 	static TSharedPtr<FModularRigTreeElement> FindElement(const FString& InElementKey, TSharedPtr<FModularRigTreeElement> CurrentItem);
 	bool AddElement(FString InKey, FString InParentKey = FString());
 	bool AddElement(const FRigModuleInstance* InElement);
@@ -153,6 +179,7 @@ public:
 	void HandleGetChildrenForTree(TSharedPtr<FModularRigTreeElement> InItem, TArray<TSharedPtr<FModularRigTreeElement>>& OutChildren);
 
 	TArray<FString> GetSelectedKeys() const;
+	void SetSelection(const TArray<TSharedPtr<FModularRigTreeElement>>& InSelection);
 	const TArray<TSharedPtr<FModularRigTreeElement>>& GetRootElements() const { return RootElements; }
 	FModularRigTreeDelegates& GetRigTreeDelegates() { return Delegates; }
 
diff --git a/Engine/Plugins/Animation/ControlRig/Source/ControlRigEditor/Private/Editor/SRigHierarchy.cpp b/Engine/Plugins/Animation/ControlRig/Source/ControlRigEditor/Private/Editor/SRigHierarchy.cpp
--- a/Engine/Plugins/Animation/ControlRig/Source/ControlRigEditor/Private/Editor/SRigHierarchy.cpp
+++ b/Engine/Plugins/Animation/ControlRig/Source/ControlRigEditor/Private/Editor/SRigHierarchy.cpp
@@ -118,6 +118,8 @@ void SRigHierarchy::Construct(const FArguments& InArgs, TSharedRef<FControlRigEd
 	ControlRigBlueprint->OnRefreshEditor().AddRaw(this, &SRigHierarchy::HandleRefreshEditorFromBlueprint);
 	ControlRigBlueprint->OnSetObjectBeingDebugged().AddRaw(this, &SRigHierarchy::HandleSetObjectBeingDebugged);
 
+	DisplaySettings.bUseShortName = !ControlRigBlueprint->IsModularRig();
+
 	// for deleting, renaming, dragging
 	CommandList = MakeShared<FUICommandList>();
 
diff --git a/Engine/Plugins/Animation/ControlRig/Source/ControlRigEditor/Private/Editor/SRigHierarchy.h b/Engine/Plugins/Animation/ControlRig/Source/ControlRigEditor/Private/Editor/SRigHierarchy.h
--- a/Engine/Plugins/Animation/ControlRig/Source/ControlRigEditor/Private/Editor/SRigHierarchy.h
+++ b/Engine/Plugins/Animation/ControlRig/Source/ControlRigEditor/Private/Editor/SRigHierarchy.h
@@ -124,7 +124,7 @@ private:
 	/** Check whether we can deleting the selected item(s) */
 	bool CanRenameItem() const;
 
-	/** Delete Item */
+	/** Rename Item */
 	void HandleRenameItem();
 
 	bool CanPasteItems() const;
diff --git a/Engine/Plugins/Animation/ControlRig/Source/ControlRigEditor/Private/Editor/SRigHierarchyTreeView.cpp b/Engine/Plugins/Animation/ControlRig/Source/ControlRigEditor/Private/Editor/SRigHierarchyTreeView.cpp
--- a/Engine/Plugins/Animation/ControlRig/Source/ControlRigEditor/Private/Editor/SRigHierarchyTreeView.cpp
+++ b/Engine/Plugins/Animation/ControlRig/Source/ControlRigEditor/Private/Editor/SRigHierarchyTreeView.cpp
@@ -110,6 +110,7 @@ void SRigHierarchyItem::Construct(const FArguments& InArgs, const TSharedRef<STa
 {
 	WeakRigTreeElement = InRigTreeElement;
 	Delegates = InTreeView->GetRigTreeDelegates();
+	FRigTreeDisplaySettings DisplaySettings = Delegates.GetDisplaySettings();
 
 	if (!InRigTreeElement->Key.IsValid())
 	{
@@ -172,7 +173,7 @@ void SRigHierarchyItem::Construct(const FArguments& InArgs, const TSharedRef<STa
 			.VAlign(VAlign_Center)
 			[
 				SAssignNew(InlineWidget, SInlineEditableTextBlock)
-				.Text(this, &SRigHierarchyItem::GetName, true)
+				.Text(this, &SRigHierarchyItem::GetName, DisplaySettings.bUseShortName)
 				.ToolTipText(this, &SRigHierarchyItem::GetItemTooltip)
 				.OnVerifyTextChanged(this, &SRigHierarchyItem::OnVerifyNameChanged)
 				.OnTextCommitted(this, &SRigHierarchyItem::OnNameCommitted)
diff --git a/Engine/Plugins/Animation/ControlRig/Source/ControlRigEditor/Private/Editor/SRigHierarchyTreeView.h b/Engine/Plugins/Animation/ControlRig/Source/ControlRigEditor/Private/Editor/SRigHierarchyTreeView.h
--- a/Engine/Plugins/Animation/ControlRig/Source/ControlRigEditor/Private/Editor/SRigHierarchyTreeView.h
+++ b/Engine/Plugins/Animation/ControlRig/Source/ControlRigEditor/Private/Editor/SRigHierarchyTreeView.h
@@ -20,6 +20,7 @@ struct CONTROLRIGEDITOR_API FRigTreeDisplaySettings
 
 		bFlattenHierarchyOnFilter = false;
 		bHideParentsOnFilter = false;
+		bUseShortName = true;
 		bShowImportedBones = true;
 		bShowBones = true;
 		bShowControls = true;
@@ -38,6 +39,9 @@ struct CONTROLRIGEDITOR_API FRigTreeDisplaySettings
 	/** Hide parents when text filtering is active */
 	bool bHideParentsOnFilter;
 
+	/** When true, the elements will show their short name */
+	bool bUseShortName;
+
 	/** Whether or not to show imported bones in the hierarchy */
 	bool bShowImportedBones;
 
