Constraints: fixed crash when deleting a constraint on at the top of the constraint list
    - we have to force the reconstruction of the tree as FMovieSceneChannelHandle::Index will be changed when removing a curve model upper in the list

#rnx

diff --git a/Engine/Plugins/Animation/ControlRig/Source/ControlRigEditor/Private/Sequencer/ControlRigParameterTrackEditor.cpp b/Engine/Plugins/Animation/ControlRig/Source/ControlRigEditor/Private/Sequencer/ControlRigParameterTrackEditor.cpp
--- a/Engine/Plugins/Animation/ControlRig/Source/ControlRigEditor/Private/Sequencer/ControlRigParameterTrackEditor.cpp
+++ b/Engine/Plugins/Animation/ControlRig/Source/ControlRigEditor/Private/Sequencer/ControlRigParameterTrackEditor.cpp
@@ -2856,17 +2856,24 @@ void FControlRigParameterTrackEditor::HandleConstraintRemoved(IMovieSceneConstra
 								return;
 							}
 
+							TSharedPtr<ISequencer> Sequencer = GetSequencer();
+							
 							const bool bCompensate = (InNotifyType == EConstraintsManagerNotifyType::ConstraintRemovedWithCompensation);
 							if (bCompensate && ConstraintChannel->GetConstraint().Get())
 							{
 								FMovieSceneConstraintChannelHelper::HandleConstraintRemoved(
 									ConstraintChannel->GetConstraint().Get(),
 									&ConstraintChannel->ActiveChannel,
-									GetSequencer(),
+									Sequencer,
 									Section);
 							}
 
 							InSection->RemoveConstraintChannel(Constraint);
+
+							if (Sequencer)
+							{
+								Sequencer->RecreateCurveEditor();
+							}
 						}
 						break;
 					case EConstraintsManagerNotifyType::ManagerUpdated:
diff --git a/Engine/Source/Editor/MovieSceneTools/Private/TrackEditors/TransformTrackEditor.cpp b/Engine/Source/Editor/MovieSceneTools/Private/TrackEditors/TransformTrackEditor.cpp
--- a/Engine/Source/Editor/MovieSceneTools/Private/TrackEditors/TransformTrackEditor.cpp
+++ b/Engine/Source/Editor/MovieSceneTools/Private/TrackEditors/TransformTrackEditor.cpp
@@ -1895,17 +1895,25 @@ void F3DTransformTrackEditor::HandleConstraintRemoved(IMovieSceneConstrainedSect
 									return;
 								}
 
+
+								TSharedPtr<ISequencer> Sequencer = GetSequencer();
+								
 								const bool bCompensate = (InNotifyType == EConstraintsManagerNotifyType::ConstraintRemovedWithCompensation);
 								if (bCompensate && ConstraintChannel->GetConstraint().Get())
 								{
 									FMovieSceneConstraintChannelHelper::HandleConstraintRemoved(
 										ConstraintChannel->GetConstraint().Get(),
 										&ConstraintChannel->ActiveChannel,
-										GetSequencer(),
+										Sequencer,
 										Section);
 								}
 
 								InSection->RemoveConstraintChannel(Constraint);
+
+								if (Sequencer)
+								{
+									Sequencer->RecreateCurveEditor();
+								}
 							}
 							break;
 						case EConstraintsManagerNotifyType::ManagerUpdated:
