Control Rig: Delete and reparent modules in the modular rig hierarchy
#rnx

diff --git a/Engine/Plugins/Animation/ControlRig/Source/ControlRigDeveloper/Private/ControlRigBlueprint.cpp b/Engine/Plugins/Animation/ControlRig/Source/ControlRigDeveloper/Private/ControlRigBlueprint.cpp
--- a/Engine/Plugins/Animation/ControlRig/Source/ControlRigDeveloper/Private/ControlRigBlueprint.cpp
+++ b/Engine/Plugins/Animation/ControlRig/Source/ControlRigDeveloper/Private/ControlRigBlueprint.cpp
@@ -2288,6 +2288,8 @@ void UControlRigBlueprint::HandleRigModulesModified(EModularRigNotification InNo
 				{
 					if (UControlRig* CDO = GetControlRigClass()->GetDefaultObject<UControlRig>())
 					{
+						Hierarchy->Modify();
+						
 						FRigVMExtendedExecuteContext& Context = CDO->GetRigVMExtendedExecuteContext();
 						FRigHierarchyExecuteContextBracket HierarchyContextGuard(Controller->GetHierarchy(), &Context);
 			
@@ -2316,6 +2318,7 @@ void UControlRigBlueprint::HandleRigModulesModified(EModularRigNotification InNo
 			break;
 		}
 		case EModularRigNotification::ModuleRenamed:
+		case EModularRigNotification::ModuleReparented:
 		{
 			if (InModule)
 			{
@@ -2323,14 +2326,24 @@ void UControlRigBlueprint::HandleRigModulesModified(EModularRigNotification InNo
 				{
 					if (UControlRig* CDO = GetControlRigClass()->GetDefaultObject<UControlRig>())
 					{
+						Hierarchy->Modify();
+						
 						struct ConnectionInfo
 						{
 							FString NewPath;
 							FRigElementKey TargetConnection;
 							FRigConnectorSettings Settings;
 						};
-						FString OldPath = FString::Printf(TEXT("%s%s:"), *InModule->ParentNamespace, *InModule->PreviousName.ToString());
-						FString NewPath = FString::Printf(TEXT("%s%s:"), *InModule->ParentNamespace, *InModule->Name.ToString());
+						FString OldPath;
+						if (InNotification == EModularRigNotification::ModuleRenamed)
+						{
+							OldPath = (InModule->ParentPath.IsEmpty()) ? *InModule->PreviousName.ToString() : FString::Printf(TEXT("%s:%s:"), *InModule->ParentPath, *InModule->PreviousName.ToString());
+						}
+						else if(InNotification == EModularRigNotification::ModuleReparented)
+						{
+							OldPath = (InModule->PreviousParentPath.IsEmpty()) ? *InModule->Name.ToString() : FString::Printf(TEXT("%s:%s:"), *InModule->PreviousParentPath, *InModule->Name.ToString());
+						}
+						FString NewPath = (InModule->ParentPath.IsEmpty()) ? *InModule->Name.ToString() : FString::Printf(TEXT("%s:%s:"), *InModule->ParentPath, *InModule->Name.ToString());
 						TArray<FRigElementKey> Connectors = Controller->GetHierarchy()->GetKeysOfType<FRigConnectorElement>();
 						TMap<FRigElementKey, ConnectionInfo> RenamedConnectors; // old key -> new key
 						for (const FRigElementKey& Connector : Connectors)
@@ -2381,8 +2394,33 @@ void UControlRigBlueprint::HandleRigModulesModified(EModularRigNotification InNo
 			}
 			break;
 		}
+		case EModularRigNotification::ModuleRemoved:
+		{
+			if (InModule)
+			{
+				// Remove all the connectors belonging this namespace
+				if (URigHierarchyController* Controller = GetHierarchyController())
+				{
+					Hierarchy->Modify();
+					
+					FString Namespace = InModule->GetNamespace();
+					TArray<FRigElementKey> Connectors = Controller->GetHierarchy()->GetKeysOfType<FRigConnectorElement>();
+					for (const FRigElementKey& Connector : Connectors)
+					{
+						FString ConnectorName = Connector.Name.ToString();
+						if (ConnectorName.StartsWith(Namespace))
+						{
+							Controller->RemoveElement(Connector);
+						}
+					}
+				}
+			}
+			break;
+		}
 		case EModularRigNotification::ConnectionChanged:
 		{
+			Hierarchy->Modify();
+			
 			const FString Namespace = InModule->GetNamespace();
 			for (TPair<FRigElementKey, FRigElementKey> Connection : InModule->Connections)
 			{
diff --git a/Engine/Plugins/Animation/ControlRig/Source/ControlRigDeveloper/Private/ModularRigController.cpp b/Engine/Plugins/Animation/ControlRig/Source/ControlRigDeveloper/Private/ModularRigController.cpp
--- a/Engine/Plugins/Animation/ControlRig/Source/ControlRigDeveloper/Private/ModularRigController.cpp
+++ b/Engine/Plugins/Animation/ControlRig/Source/ControlRigDeveloper/Private/ModularRigController.cpp
@@ -44,7 +44,6 @@ bool UModularRigController::AddModule(const FName& InModuleName, TSubclassOf<UCo
 		if(UControlRigBlueprint* Blueprint = Cast<UControlRigBlueprint>(GetOuter()))
 		{
 			Blueprint->Modify();
-			Blueprint->Hierarchy->Modify();
 		}
 	}
 #endif 
@@ -73,7 +72,7 @@ bool UModularRigController::AddModule(const FName& InModuleName, TSubclassOf<UCo
 			}
 		}
 
-		Model->Modules.Add(FRigModuleReference(InModuleName, InClass, ParentModule->GetNamespace()));
+		Model->Modules.Add(FRigModuleReference(InModuleName, InClass, ParentModule->GetPath()));
 		NewModule = &Model->Modules.Last();
 	}
 
@@ -96,8 +95,11 @@ bool UModularRigController::AddModule(const FName& InModuleName, TSubclassOf<UCo
 
 FRigModuleReference* UModularRigController::FindModule(const FString& InPath)
 {
+	FString Path = InPath;
+	Path.RemoveFromEnd(UModularRig::NamespaceSeparator);
+	
 	TArray<FRigModuleReference*>* Children = &Model->RootModules;
-	FString Left = InPath, Right;
+	FString Left = Path, Right;
 	while (Left.Split(UModularRig::NamespaceSeparator, &Left, &Right))
 	{
 		FRigModuleReference** Cur = Children->FindByPredicate([Left](FRigModuleReference* Module)
@@ -125,17 +127,17 @@ FRigModuleReference* UModularRigController::FindModule(const FString& InPath)
 
 bool UModularRigController::ConnectModuleToElement(const FRigElementKey& InConnectorKey, const FRigElementKey& InTargetKey, bool bSetupUndo)
 {
-	FString ConnectorNameSpace, ConnectorName;
-	if (!InConnectorKey.Name.ToString().Split(UModularRig::NamespaceSeparator, &ConnectorNameSpace, &ConnectorName, ESearchCase::CaseSensitive, ESearchDir::FromEnd))
+	FString ConnectorParentPath, ConnectorName;
+	if (!InConnectorKey.Name.ToString().Split(UModularRig::NamespaceSeparator, &ConnectorParentPath, &ConnectorName, ESearchCase::CaseSensitive, ESearchDir::FromEnd))
 	{
 		UE_LOG(LogControlRig, Error, TEXT("Connector %s does not contain a namespace"), *InConnectorKey.ToString());
 		return false;
 	}
 	
-	FRigModuleReference* Module = FindModule(ConnectorNameSpace);
+	FRigModuleReference* Module = FindModule(ConnectorParentPath);
 	if (!Module)
 	{
-		UE_LOG(LogControlRig, Error, TEXT("Could not find module %s"), *ConnectorNameSpace);
+		UE_LOG(LogControlRig, Error, TEXT("Could not find module %s"), *ConnectorParentPath);
 		return false;
 	}
 
@@ -159,7 +161,6 @@ bool UModularRigController::ConnectModuleToElement(const FRigElementKey& InConne
 	{
 		TransactionPtr = MakeShared<FScopedTransaction>(NSLOCTEXT("ModularRigController", "ConnectModuleToElementTransaction", "Connect to Element"));
 		Blueprint->Modify();
-		Blueprint->Hierarchy->Modify();
 	}
 #endif 
 
@@ -232,9 +233,44 @@ bool UModularRigController::SetConfigValueInModule(const FString& InModulePath,
 	return true;
 }
 
-bool UModularRigController::RemoveModule(const FString& InModulePath, bool bSetupUndo)
+bool UModularRigController::DeleteModule(const FString& InModulePath, bool bSetupUndo)
 {
-	// todo: UE-199050
+	FRigModuleReference* Module = FindModule(InModulePath);
+	if (!Module)
+	{
+		UE_LOG(LogControlRig, Error, TEXT("Could not find module %s"), *InModulePath);
+		return false;
+	}
+
+#if WITH_EDITOR
+	TSharedPtr<FScopedTransaction> TransactionPtr;
+	if (bSetupUndo)
+	{
+		TransactionPtr = MakeShared<FScopedTransaction>(NSLOCTEXT("ModularRigController", "RenameModuleTransaction", "Rename Module"));
+		if(UControlRigBlueprint* Blueprint = Cast<UControlRigBlueprint>(GetOuter()))
+		{
+			Blueprint->Modify();
+		}
+	}
+#endif
+
+	// Unparent children (add them to root)
+	for (FRigModuleReference* Child : Module->CachedChildren)
+	{
+		ReparentModule(Child->GetPath(), FString(), bSetupUndo);
+	}
+
+	Model->DeletedModules.Add(*Module);
+	Model->Modules.RemoveSingle(*Module);
+	Model->UpdateCachedChildren();
+
+	Notify(EModularRigNotification::ModuleRemoved, &Model->DeletedModules.Last());
+
+	Model->DeletedModules.Reset();
+
+#if WITH_EDITOR
+	TransactionPtr.Reset();
+#endif
 	return false;
 }
 
@@ -253,6 +289,13 @@ bool UModularRigController::RenameModule(const FString& InModulePath, const FNam
 		UE_LOG(LogControlRig, Error, TEXT("Could not rename module %s: %s"), *InModulePath, *ErrorMessage.ToString());
 		return false;
 	}
+
+	const FString OldName = Module->Name.ToString();
+	const FString NewName = InNewName.ToString();
+	if (OldName.Equals(NewName))
+	{
+		return true;
+	}
 	
 #if WITH_EDITOR
 	TSharedPtr<FScopedTransaction> TransactionPtr;
@@ -266,8 +309,8 @@ bool UModularRigController::RenameModule(const FString& InModulePath, const FNam
 	}
 #endif
 	
-	const FString OldPath = FString::Printf(TEXT("%s%s"), *Module->ParentNamespace, *Module->Name.ToString());
-	const FString NewPath = FString::Printf(TEXT("%s%s"), *Module->ParentNamespace, *InNewName.ToString());
+	const FString OldPath = (Module->ParentPath.IsEmpty()) ? OldName : FString::Printf(TEXT("%s:%s"), *Module->ParentPath, *OldName);
+	const FString NewPath = (Module->ParentPath.IsEmpty()) ? *NewName :  FString::Printf(TEXT("%s:%s"), *Module->ParentPath, *NewName);
 	Module->PreviousName = Module->Name;
 	Module->Name = InNewName;
 	TArray<FRigModuleReference*> Children;
@@ -275,7 +318,7 @@ bool UModularRigController::RenameModule(const FString& InModulePath, const FNam
 	for (int32 i=0; i<Children.Num(); ++i)
 	{
 		FRigModuleReference* Child = Children[i];
-		Child->ParentNamespace.ReplaceInline(*OldPath, *NewPath);
+		Child->ParentPath.ReplaceInline(*OldPath, *NewPath);
 
 		Children.Append(Child->CachedChildren);
 	}
@@ -304,7 +347,8 @@ bool UModularRigController::CanRenameModule(const FString& InModulePath, const F
 		return false;
 	}
 
-	const FString NewPath = FString::Printf(TEXT("%s%s"), *Module->ParentNamespace, *InNewName.ToString());
+	const FString NewNameStr = InNewName.ToString();
+	const FString NewPath = (Module->ParentPath.IsEmpty()) ? NewNameStr : FString::Printf(TEXT("%s:%s"), *Module->ParentPath, *NewNameStr);
 	if (InModulePath == NewPath)
 	{
 		return true;
@@ -312,7 +356,7 @@ bool UModularRigController::CanRenameModule(const FString& InModulePath, const F
 
 	// Check there are no siblings with the same name
 	{
-		FRigModuleReference* Parent = FindModule(Module->ParentNamespace);
+		FRigModuleReference* Parent = FindModule(Module->ParentPath);
 		TArray<FRigModuleReference*>* Siblings = &Model->RootModules;
 		if (Parent)
 		{
@@ -320,9 +364,9 @@ bool UModularRigController::CanRenameModule(const FString& InModulePath, const F
 		}
 		for (FRigModuleReference* Sibling : *Siblings)
 		{
-			if (Sibling->Name == InNewName)
+			if (Sibling->Name.ToString() == NewNameStr)
 			{
-				OutErrorMessage = FText::FromString(FString::Printf(TEXT("Sibling with path %s already exists"), *Sibling->GetNamespace()));
+				OutErrorMessage = FText::FromString(FString::Printf(TEXT("Sibling with path %s already exists"), *Sibling->GetPath()));
 				return false;
 			}
 		}
@@ -332,14 +376,52 @@ bool UModularRigController::CanRenameModule(const FString& InModulePath, const F
 
 bool UModularRigController::ReparentModule(const FString& InModulePath, const FString& InNewParentModulePath, bool bSetupUndo)
 {
-	// todo: UE-199050
+	FRigModuleReference* Module = FindModule(InModulePath);
+	if (!Module)
+	{
+		UE_LOG(LogControlRig, Error, TEXT("Could not find module %s"), *InModulePath);
+		return false;
+	}
+
+#if WITH_EDITOR
+	TSharedPtr<FScopedTransaction> TransactionPtr;
+	if (bSetupUndo)
+	{
+		TransactionPtr = MakeShared<FScopedTransaction>(NSLOCTEXT("ModularRigController", "RenameModuleTransaction", "Rename Module"));
+		if(UControlRigBlueprint* Blueprint = Cast<UControlRigBlueprint>(GetOuter()))
+		{
+			Blueprint->Modify();
+		}
+	}
+#endif
+
+	// Reparent or unparent children
+	FRigModuleReference* NewParentModule = FindModule(InNewParentModulePath);
+	const FString OldPath = Module->GetPath();
+	Module->PreviousParentPath = Module->ParentPath;
+	Module->ParentPath = (NewParentModule) ? NewParentModule->GetPath() : FString();
+	Module->Name = GetSafeNewName(Module->GetPath());
+	const FString NewPath = Module->GetPath();
+
+	// Fix all the subtree namespaces
+	TArray<FRigModuleReference*> SubTree = Module->CachedChildren;
+	for (int32 Index=0; Index<SubTree.Num(); ++Index)
+	{
+		SubTree[Index]->ParentPath.ReplaceInline(*OldPath, *NewPath);
+		SubTree.Append(SubTree[Index]->CachedChildren);
+	}
+
+	Model->UpdateCachedChildren();
+
+	Notify(EModularRigNotification::ModuleReparented, Module);
+	
 	return false;
 }
 
 FName UModularRigController::GetSafeNewName(const FString& InModuleDesiredPath)
 {
 	FString ParentPath, DesiredName = InModuleDesiredPath;
-	InModuleDesiredPath.Split(UModularRig::NamespaceSeparator, &ParentPath, &DesiredName);
+	InModuleDesiredPath.Split(UModularRig::NamespaceSeparator, &ParentPath, &DesiredName, ESearchCase::CaseSensitive, ESearchDir::FromEnd);
 
 	TArray<FRigModuleReference*>* Children = &Model->RootModules;
 	if (!ParentPath.IsEmpty())
diff --git a/Engine/Plugins/Animation/ControlRig/Source/ControlRigDeveloper/Private/ModularRigModel.cpp b/Engine/Plugins/Animation/ControlRig/Source/ControlRigDeveloper/Private/ModularRigModel.cpp
--- a/Engine/Plugins/Animation/ControlRig/Source/ControlRigDeveloper/Private/ModularRigModel.cpp
+++ b/Engine/Plugins/Animation/ControlRig/Source/ControlRigDeveloper/Private/ModularRigModel.cpp
@@ -5,13 +5,18 @@
 
 #include UE_INLINE_GENERATED_CPP_BY_NAME(ModularRigModel)
 
-FString FRigModuleReference::GetNamespace() const
+FString FRigModuleReference::GetPath() const
 {
-	if (ParentNamespace.IsEmpty())
+	if (ParentPath.IsEmpty())
 	{
-		return Name.ToString() + UModularRig::NamespaceSeparator;
+		return Name.ToString();
 	}
-	return ParentNamespace + Name.ToString() + UModularRig::NamespaceSeparator;
+	return ParentPath + UModularRig::NamespaceSeparator + Name.ToString();
+}
+
+FString FRigModuleReference::GetNamespace() const
+{
+	return GetPath() + UModularRig::NamespaceSeparator;
 }
 
 UModularRigController* FModularRigModel::GetController(bool bCreateIfNeeded)
@@ -41,19 +46,19 @@ void FModularRigModel::UpdateCachedChildren()
 	for (FRigModuleReference& Module : Modules)
 	{
 		Module.CachedChildren.Reset();
-		PathToModule.Add(Module.GetNamespace(), &Module);
+		PathToModule.Add(Module.GetPath(), &Module);
 	}
 	
 	RootModules.Reset();
 	for (FRigModuleReference& Module : Modules)
 	{
-		if (Module.ParentNamespace.IsEmpty())
+		if (Module.ParentPath.IsEmpty())
 		{
 			RootModules.Add(&Module);
 		}
 		else
 		{
-			if (FRigModuleReference** ParentModule = PathToModule.Find(Module.ParentNamespace))
+			if (FRigModuleReference** ParentModule = PathToModule.Find(Module.ParentPath))
 			{
 				(*ParentModule)->CachedChildren.Add(&Module);
 			}
@@ -61,11 +66,11 @@ void FModularRigModel::UpdateCachedChildren()
 	}
 }
 
-FRigModuleReference* FModularRigModel::FindModule(const FString InNameSpace) const
+FRigModuleReference* FModularRigModel::FindModule(const FString InPath) const
 {
 	const TArray<FRigModuleReference*>* Children = &RootModules;
 
-	FString Left = InNameSpace, Right;
+	FString Left = InPath, Right;
 	while (Left.Split(UModularRig::NamespaceSeparator, &Left, &Right))
 	{
 		FRigModuleReference* const * Child = Children->FindByPredicate([Left](FRigModuleReference* Module)
@@ -94,11 +99,39 @@ FRigModuleReference* FModularRigModel::FindModule(const FString InNameSpace) con
 	return *Child;
 }
 
-FString FModularRigModel::FindParentNamespace(const FString InNameSpace) const
+FString FModularRigModel::FindParentPath(const FString InPath) const
 {
-	if (FRigModuleReference* Element = FindModule(InNameSpace))
+	if (FRigModuleReference* Element = FindModule(InPath))
 	{
-		return Element->ParentNamespace;
+		return Element->ParentPath;
 	}
 	return FString();
 }
+
+void FModularRigModel::ForEachModule(TFunction<bool(const FRigModuleReference*)> PerModule) const
+{
+	TArray<FRigModuleReference*> ModuleInstances = RootModules;
+	for (int32 Index=0; Index < ModuleInstances.Num(); ++Index)
+	{
+		if (!PerModule(ModuleInstances[Index]))
+		{
+			break;
+		}
+		ModuleInstances.Append(ModuleInstances[Index]->CachedChildren);
+	}
+}
+
+TArray<FString> FModularRigModel::SortPaths(const TArray<FString>& InPaths) const
+{
+	TArray<FString> Result;
+	ForEachModule([InPaths, &Result](const FRigModuleReference* Element) -> bool
+	{
+		const FString& Path = Element->GetPath();
+		if(InPaths.Contains(Path))
+		{
+			Result.AddUnique(Path);
+		}
+		return true;
+	});
+	return Result;
+}
diff --git a/Engine/Plugins/Animation/ControlRig/Source/ControlRigDeveloper/Public/ModularRigController.h b/Engine/Plugins/Animation/ControlRig/Source/ControlRigDeveloper/Public/ModularRigController.h
--- a/Engine/Plugins/Animation/ControlRig/Source/ControlRigDeveloper/Public/ModularRigController.h
+++ b/Engine/Plugins/Animation/ControlRig/Source/ControlRigDeveloper/Public/ModularRigController.h
@@ -18,6 +18,10 @@ enum class EModularRigNotification : uint8
 
 	ModuleRenamed,
 
+	ModuleRemoved,
+
+	ModuleReparented,
+
 	ConnectionChanged,
 
 	ModuleConfigValueChanged,
@@ -54,7 +58,7 @@ class CONTROLRIGDEVELOPER_API UModularRigController : public UObject
 	bool SetConfigValueInModule(const FString& InModulePath, const FName& InVariableName, const FString& InValue, bool bSetupUndo = true);
 
 	UFUNCTION(BlueprintCallable, Category = "ControlRig | Modules")
-	bool RemoveModule(const FString& InModulePath, bool bSetupUndo = true);
+	bool DeleteModule(const FString& InModulePath, bool bSetupUndo = true);
 
 	UFUNCTION(BlueprintCallable, Category = "ControlRig | Modules")
 	bool RenameModule(const FString& InModulePath, const FName& InNewName, bool bSetupUndo = true);
diff --git a/Engine/Plugins/Animation/ControlRig/Source/ControlRigDeveloper/Public/ModularRigModel.h b/Engine/Plugins/Animation/ControlRig/Source/ControlRigDeveloper/Public/ModularRigModel.h
--- a/Engine/Plugins/Animation/ControlRig/Source/ControlRigDeveloper/Public/ModularRigModel.h
+++ b/Engine/Plugins/Animation/ControlRig/Source/ControlRigDeveloper/Public/ModularRigModel.h
@@ -3,6 +3,7 @@
 #pragma once
 
 #include "ControlRig.h"
+#include "ModularRig.h"
 #include "ModularRigController.h"
 #include "ModularRigModel.generated.h"
 
@@ -13,13 +14,13 @@ struct CONTROLRIGDEVELOPER_API FRigModuleReference
 
 	FRigModuleReference()
 		: Name(NAME_None)
-		, ParentNamespace(FString())
+		, ParentPath(FString())
 		, Class(nullptr)
 	{}
 	
 	FRigModuleReference(const FName& InName, TSubclassOf<UControlRig> InClass, const FString& InParentPath)
 		: Name(InName)
-		, ParentNamespace(InParentPath)
+		, ParentPath(InParentPath)
 		, Class(InClass)
 	{}
 
@@ -27,7 +28,7 @@ struct CONTROLRIGDEVELOPER_API FRigModuleReference
 	FName Name;
 
 	UPROPERTY()
-	FString ParentNamespace;
+	FString ParentPath;
 
 	UPROPERTY()
 	TSoftClassPtr<UControlRig> Class;
@@ -41,10 +42,21 @@ struct CONTROLRIGDEVELOPER_API FRigModuleReference
 	UPROPERTY(transient)
 	FName PreviousName;
 
+	UPROPERTY(transient)
+	FString PreviousParentPath;
+
 	TArray<FRigModuleReference*> CachedChildren;
 
+	FString GetPath() const;
+
 	FString GetNamespace() const;
 	
+	friend bool operator==(const FRigModuleReference& A, const FRigModuleReference& B)
+	{
+		return A.ParentPath == B.ParentPath &&
+			A.Name == B.Name;
+	}
+	
 	friend class UModularRigController;
 };
 
@@ -59,6 +71,7 @@ public:
 	UPROPERTY()
 	TArray<FRigModuleReference> Modules;
 	TArray<FRigModuleReference*> RootModules;
+	TArray<FRigModuleReference> DeletedModules;
 
 	UPROPERTY(transient)
 	TObjectPtr<UModularRigController> Controller;
@@ -71,10 +84,14 @@ public:
 
 	void UpdateCachedChildren();
 
-	FRigModuleReference* FindModule(const FString InNameSpace) const;
+	FRigModuleReference* FindModule(const FString InPath) const;
+
+	FString FindParentPath(const FString InPath) const;
+
+	void ForEachModule(TFunction<bool(const FRigModuleReference*)> PerModule) const;
+
+	TArray<FString> SortPaths(const TArray<FString>& InPaths) const;
 
-	FString FindParentNamespace(const FString InNameSpace) const;
-	
 private:
 	TWeakObjectPtr<UObject> OuterClientHost;
 
diff --git a/Engine/Plugins/Animation/ControlRig/Source/ControlRigEditor/Private/ControlRigModularRigHierarchyCommands.cpp b/Engine/Plugins/Animation/ControlRig/Source/ControlRigEditor/Private/ControlRigModularRigHierarchyCommands.cpp
--- a/Engine/Plugins/Animation/ControlRig/Source/ControlRigEditor/Private/ControlRigModularRigHierarchyCommands.cpp
+++ b/Engine/Plugins/Animation/ControlRig/Source/ControlRigEditor/Private/ControlRigModularRigHierarchyCommands.cpp
@@ -8,6 +8,7 @@ void FControlRigModularHierarchyCommands::RegisterCommands()
 {
 	UI_COMMAND(AddModuleItem, "New Module", "Add new module to the rig.", EUserInterfaceActionType::Button, FInputChord());
 	UI_COMMAND(RenameModuleItem, "Rename", "Rename module to the rig.", EUserInterfaceActionType::Button, FInputChord(EKeys::F2));
+	UI_COMMAND(DeleteModuleItem, "Delete", "Delete module from the rig.", EUserInterfaceActionType::Button, FInputChord(EKeys::Delete));
 }
 
 #undef LOCTEXT_NAMESPACE
diff --git a/Engine/Plugins/Animation/ControlRig/Source/ControlRigEditor/Private/ControlRigModularRigHierarchyCommands.h b/Engine/Plugins/Animation/ControlRig/Source/ControlRigEditor/Private/ControlRigModularRigHierarchyCommands.h
--- a/Engine/Plugins/Animation/ControlRig/Source/ControlRigEditor/Private/ControlRigModularRigHierarchyCommands.h
+++ b/Engine/Plugins/Animation/ControlRig/Source/ControlRigEditor/Private/ControlRigModularRigHierarchyCommands.h
@@ -24,6 +24,9 @@ public:
 	/** Rename Module */
 	TSharedPtr< FUICommandInfo > RenameModuleItem;
 
+	/** Delete Module */
+	TSharedPtr< FUICommandInfo > DeleteModuleItem;
+
 	/**
 	 * Initialize commands
 	 */
diff --git a/Engine/Plugins/Animation/ControlRig/Source/ControlRigEditor/Private/Editor/SModularRigHierarchy.cpp b/Engine/Plugins/Animation/ControlRig/Source/ControlRigEditor/Private/Editor/SModularRigHierarchy.cpp
--- a/Engine/Plugins/Animation/ControlRig/Source/ControlRigEditor/Private/Editor/SModularRigHierarchy.cpp
+++ b/Engine/Plugins/Animation/ControlRig/Source/ControlRigEditor/Private/Editor/SModularRigHierarchy.cpp
@@ -57,6 +57,38 @@
 
 #define LOCTEXT_NAMESPACE "SModularRigHierarchy"
 
+//////////////////////////////////////////////////////////////
+/// FModuleRigHierarchyDragDropOp
+///////////////////////////////////////////////////////////
+TSharedRef<FModuleRigHierarchyDragDropOp> FModuleRigHierarchyDragDropOp::New(const TArray<FString>& InElements)
+{
+	TSharedRef<FModuleRigHierarchyDragDropOp> Operation = MakeShared<FModuleRigHierarchyDragDropOp>();
+	Operation->Elements = InElements;
+	Operation->Construct();
+	return Operation;
+}
+
+TSharedPtr<SWidget> FModuleRigHierarchyDragDropOp::GetDefaultDecorator() const
+{
+	return SNew(SBorder)
+		.Visibility(EVisibility::Visible)
+		.BorderImage(FAppStyle::GetBrush("Menu.Background"))
+		[
+			SNew(STextBlock)
+			.Text(FText::FromString(GetJoinedElementNames()))
+			//.Font(FAppStyle::Get().GetFontStyle("FontAwesome.10"))
+		];
+}
+
+FString FModuleRigHierarchyDragDropOp::GetJoinedElementNames() const
+{
+	TArray<FString> ElementNameStrings;
+	for (const FString& Element: Elements)
+	{
+		ElementNameStrings.Add(Element);
+	}
+	return FString::Join(ElementNameStrings, TEXT(","));
+}
 
 ///////////////////////////////////////////////////////////
 
@@ -93,6 +125,7 @@ void SModularRigHierarchy::Construct(const FArguments& InArgs, TSharedRef<FContr
 	FModularRigTreeDelegates Delegates;
 	Delegates.OnGetHierarchy = FOnGetModularRigTreeHierarchy::CreateSP(this, &SModularRigHierarchy::GetHierarchyForTreeView);
 	Delegates.OnContextMenuOpening = FOnContextMenuOpening::CreateSP(this, &SModularRigHierarchy::CreateContextMenuWidget);
+	Delegates.OnDragDetected = FOnDragDetected::CreateSP(this, &SModularRigHierarchy::OnDragDetected);
 	Delegates.OnCanAcceptDrop = FOnModularRigTreeCanAcceptDrop::CreateSP(this, &SModularRigHierarchy::OnCanAcceptDrop);
 	Delegates.OnAcceptDrop = FOnModularRigTreeAcceptDrop::CreateSP(this, &SModularRigHierarchy::OnAcceptDrop);
 	Delegates.OnMouseButtonClick = FOnModularRigTreeMouseButtonClick::CreateSP(this, &SModularRigHierarchy::OnItemClicked);
@@ -169,6 +202,19 @@ void SModularRigHierarchy::BindCommands()
 	CommandList->MapAction(Commands.RenameModuleItem,
 		FExecuteAction::CreateSP(this, &SModularRigHierarchy::HandleRenameModule),
 		FCanExecuteAction());
+
+	CommandList->MapAction(Commands.DeleteModuleItem,
+		FExecuteAction::CreateSP(this, &SModularRigHierarchy::HandleDeleteModules),
+		FCanExecuteAction());
+}
+
+FReply SModularRigHierarchy::OnKeyDown(const FGeometry& MyGeometry, const FKeyEvent& InKeyEvent)
+{
+	if (CommandList.IsValid() && CommandList->ProcessCommandBindings(InKeyEvent))
+	{
+		return FReply::Handled();
+	}
+	return FReply::Unhandled();
 }
 
 void SModularRigHierarchy::RefreshTreeView(bool bRebuildContent)
@@ -295,6 +341,7 @@ void SModularRigHierarchy::CreateContextMenu()
 						})
 					);
 					ElementsSection.AddMenuEntry(Commands.RenameModuleItem);
+					ElementsSection.AddMenuEntry(Commands.DeleteModuleItem);
 				}
 			})
 		);
@@ -532,6 +579,69 @@ bool SModularRigHierarchy::HandleVerifyNameChanged(const FString& InOldPath, con
 	return false;
 }
 
+void SModularRigHierarchy::HandleDeleteModules()
+{
+	if(!ControlRigEditor.IsValid())
+	{
+		return;
+	}
+
+	UModularRig* Rig = GetDefaultHierarchy();
+	if (Rig)
+	{
+		FScopedTransaction Transaction(LOCTEXT("ModularRigHierarchyDeleteSelected", "Delete selected modules"));
+
+		TArray<TSharedPtr<FModularRigTreeElement>> SelectedItems = TreeView->GetSelectedItems();
+		TArray<FString> SelectedPaths;
+		Algo::Transform(SelectedItems, SelectedPaths, [](const TSharedPtr<FModularRigTreeElement>& Element)
+		{
+			return Element->Key;
+		});
+		HandleDeleteModules(SelectedPaths);
+	}
+
+	return;
+}
+
+void SModularRigHierarchy::HandleDeleteModules(const TArray<FString>& InPaths)
+{
+	ClearDetailPanel();
+	
+	if (ControlRigBlueprint.IsValid())
+	{
+		FScopedTransaction Transaction(LOCTEXT("ModularRigHierarchyDelete", "Delete Modules"));
+
+		UModularRigController* Controller = ControlRigBlueprint->GetModularRigController();
+		check(Controller);
+
+		// Make sure we delete the modules from children to root
+		TArray<FString> SortedPaths = Controller->Model->SortPaths(InPaths);
+		Algo::Reverse(SortedPaths);
+		for (const FString& Path : SortedPaths)
+		{
+			Controller->DeleteModule(Path);
+		}
+	}
+}
+
+void SModularRigHierarchy::HandleReparentModules(const TArray<FString>& InPaths, const FString& InParentPath)
+{
+	ClearDetailPanel();
+	
+	if (ControlRigBlueprint.IsValid())
+	{
+		FScopedTransaction Transaction(LOCTEXT("ModularRigHierarchyReparent", "Reparent Modules"));
+
+		UModularRigController* Controller = ControlRigBlueprint->GetModularRigController();
+		check(Controller);
+
+		for (const FString& Path : InPaths)
+		{
+			Controller->ReparentModule(Path, InParentPath);
+		}
+	}
+}
+
 class SModularRigHierarchyPasteTransformsErrorPipe : public FOutputDevice
 {
 public:
@@ -621,12 +731,57 @@ void SModularRigHierarchy::PostUndo(bool bSuccess)
 	}
 }
 
+FReply SModularRigHierarchy::OnDragDetected(const FGeometry& MyGeometry, const FPointerEvent& MouseEvent)
+{
+	TArray<FString> DraggedElements = GetSelectedKeys();
+	if (MouseEvent.IsMouseButtonDown(EKeys::LeftMouseButton) && DraggedElements.Num() > 0)
+	{
+		if (ControlRigEditor.IsValid())
+		{
+			TSharedRef<FModuleRigHierarchyDragDropOp> DragDropOp = FModuleRigHierarchyDragDropOp::New(MoveTemp(DraggedElements));
+			return FReply::Handled().BeginDragDrop(DragDropOp);
+		}
+	}
+
+	return FReply::Unhandled();
+}
+
 TOptional<EItemDropZone> SModularRigHierarchy::OnCanAcceptDrop(const FDragDropEvent& DragDropEvent, EItemDropZone DropZone, TSharedPtr<FModularRigTreeElement> TargetItem)
 {
 	const TOptional<EItemDropZone> InvalidDropZone;
-	TOptional<EItemDropZone> ReturnDropZone;
+	TOptional<EItemDropZone> ReturnDropZone = DropZone;
+
+	TSharedPtr<FAssetDragDropOp> AssetDragDropOperation = DragDropEvent.GetOperationAs<FAssetDragDropOp>();
+	TSharedPtr<FModuleRigHierarchyDragDropOp> ModuleDragDropOperation = DragDropEvent.GetOperationAs<FModuleRigHierarchyDragDropOp>();
+	if (AssetDragDropOperation)
+	{
+		for (const FAssetData& AssetData : AssetDragDropOperation->GetAssets())
+		{
+			static const UEnum* ControlTypeEnum = StaticEnum<EControlRigType>();
+			const FString ControlRigTypeStr = AssetData.GetTagValueRef<FString>(TEXT("ControlRigType"));
+			if (ControlRigTypeStr.IsEmpty())
+			{
+				ReturnDropZone.Reset();
+				break;
+			}
+
+			const EControlRigType ControlRigType = (EControlRigType)(ControlTypeEnum->GetValueByName(*ControlRigTypeStr));
+			if (ControlRigType != EControlRigType::RigModule)
+			{
+				ReturnDropZone.Reset();
+				break;
+			}
+		}
+	}
+	else if(ModuleDragDropOperation)
+	{
+		// Accept this drop
+	}
+	else
+	{
+		ReturnDropZone.Reset();
+	}
 
-	ReturnDropZone = DropZone;
 	return ReturnDropZone;
 }
 
@@ -640,6 +795,7 @@ FReply SModularRigHierarchy::OnAcceptDrop(const FDragDropEvent& DragDropEvent, E
 	}
 
 	TSharedPtr<FAssetDragDropOp> AssetDragDropOperation = DragDropEvent.GetOperationAs<FAssetDragDropOp>();
+	TSharedPtr<FModuleRigHierarchyDragDropOp> ModuleDragDropOperation = DragDropEvent.GetOperationAs<FModuleRigHierarchyDragDropOp>();
 	if (AssetDragDropOperation)
 	{
 		for (const FAssetData& AssetData : AssetDragDropOperation->GetAssets())
@@ -671,6 +827,11 @@ FReply SModularRigHierarchy::OnAcceptDrop(const FDragDropEvent& DragDropEvent, E
 
 		FReply::Handled();
 	}
+	else if(ModuleDragDropOperation)
+	{
+		const TArray<FString> Paths = ModuleDragDropOperation->GetElements();
+		HandleReparentModules(Paths, ParentPath);
+	}
 	
 	return FReply::Unhandled();
 }
diff --git a/Engine/Plugins/Animation/ControlRig/Source/ControlRigEditor/Private/Editor/SModularRigHierarchy.h b/Engine/Plugins/Animation/ControlRig/Source/ControlRigEditor/Private/Editor/SModularRigHierarchy.h
--- a/Engine/Plugins/Animation/ControlRig/Source/ControlRigEditor/Private/Editor/SModularRigHierarchy.h
+++ b/Engine/Plugins/Animation/ControlRig/Source/ControlRigEditor/Private/Editor/SModularRigHierarchy.h
@@ -6,6 +6,7 @@
 #include "EditorUndoClient.h"
 #include "Editor/SModularRigHierarchyTreeView.h"
 #include "ControlRigBlueprint.h"
+#include "DragAndDrop/GraphNodeDragDropOp.h"
 #include "Editor/RigVMEditor.h"
 
 class SModularRigHierarchy;
@@ -19,6 +20,34 @@ class FMenuBuilder;
 class UToolMenu;
 struct FToolMenuContext;
 
+class FModuleRigHierarchyDragDropOp : public FDragDropOperation
+{
+public:
+	DRAG_DROP_OPERATOR_TYPE(FModuleRigHierarchyDragDropOp, FDragDropOperation)
+
+	static TSharedRef<FModuleRigHierarchyDragDropOp> New(const TArray<FString>& InElements);
+
+	virtual TSharedPtr<SWidget> GetDefaultDecorator() const override;
+
+	/** @return true if this drag operation contains property paths */
+	bool HasElements() const
+	{
+		return Elements.Num() > 0;
+	}
+
+	/** @return The property paths from this drag operation */
+	const TArray<FString>& GetElements() const
+	{
+		return Elements;
+	}
+
+	FString GetJoinedElementNames() const;
+
+private:
+
+	/** Data for the property paths this item represents */
+	TArray<FString> Elements;
+};
 
 /** Widget allowing editing of a control rig's structure */
 class SModularRigHierarchy : public SCompoundWidget, public FEditorUndoClient
@@ -47,6 +76,9 @@ private:
 	/** Bind commands that this widget handles */
 	void BindCommands();
 
+	/** SWidget interface */
+	virtual FReply OnKeyDown(const FGeometry& MyGeometry, const FKeyEvent& InKeyEvent);
+
 	/** Rebuild the tree view */
 	void RefreshTreeView(bool bRebuildContent = true);
 
@@ -55,14 +87,21 @@ private:
 
 	/** Create a new item */
 	void HandleNewItem();
-
 	void HandleNewItem(UClass* InClass, const FString& InParentPath);
 
+	/** Rename item */
 	bool CanRenameModule() const;
 	void HandleRenameModule();
 	FName HandleRenameModule(const FString& InOldPath, const FName& InNewName);
 	bool HandleVerifyNameChanged(const FString& InOldPath, const FName& InNewName, FText& OutErrorMessage);
 
+	/** Delete items */
+	void HandleDeleteModules();
+	void HandleDeleteModules(const TArray<FString>& InPaths);
+
+	/** Reparent items */
+	void HandleReparentModules(const TArray<FString>& InPaths, const FString& InParentPath);
+
 	/** Set Selection Changed */
 	void OnSelectionChanged(TSharedPtr<FModularRigTreeElement> Selection, ESelectInfo::Type SelectInfo);
 
@@ -74,6 +113,9 @@ private:
 	virtual void PostUndo(bool bSuccess) override;
 	virtual void PostRedo(bool bSuccess) override;
 
+	// reply to a drag operation
+	FReply OnDragDetected(const FGeometry& MyGeometry, const FPointerEvent& MouseEvent);
+
 	// reply to a drop operation on item
 	TOptional<EItemDropZone> OnCanAcceptDrop(const FDragDropEvent& DragDropEvent, EItemDropZone DropZone, TSharedPtr<FModularRigTreeElement> TargetItem);
 	FReply OnAcceptDrop(const FDragDropEvent& DragDropEvent, EItemDropZone DropZone, TSharedPtr<FModularRigTreeElement> TargetItem);
diff --git a/Engine/Plugins/Animation/ControlRig/Source/ControlRigEditor/Private/Editor/SModularRigHierarchyTreeView.cpp b/Engine/Plugins/Animation/ControlRig/Source/ControlRigEditor/Private/Editor/SModularRigHierarchyTreeView.cpp
--- a/Engine/Plugins/Animation/ControlRig/Source/ControlRigEditor/Private/Editor/SModularRigHierarchyTreeView.cpp
+++ b/Engine/Plugins/Animation/ControlRig/Source/ControlRigEditor/Private/Editor/SModularRigHierarchyTreeView.cpp
@@ -98,6 +98,7 @@ void SModularRigHierarchyItem::Construct(const FArguments& InArgs, const TShared
 
 	STableRow<TSharedPtr<FModularRigTreeElement>>::Construct(
 		STableRow<TSharedPtr<FModularRigTreeElement>>::FArguments()
+		.OnDragDetected(Delegates.OnDragDetected)
 		.OnCanAcceptDrop(Delegates.OnCanAcceptDrop)
 		.OnAcceptDrop(Delegates.OnAcceptDrop)
 		.ShowWires(true)
diff --git a/Engine/Plugins/Animation/ControlRig/Source/ControlRigEditor/Private/Editor/SModularRigHierarchyTreeView.h b/Engine/Plugins/Animation/ControlRig/Source/ControlRigEditor/Private/Editor/SModularRigHierarchyTreeView.h
--- a/Engine/Plugins/Animation/ControlRig/Source/ControlRigEditor/Private/Editor/SModularRigHierarchyTreeView.h
+++ b/Engine/Plugins/Animation/ControlRig/Source/ControlRigEditor/Private/Editor/SModularRigHierarchyTreeView.h
@@ -28,6 +28,7 @@ struct CONTROLRIGEDITOR_API FModularRigTreeDelegates
 	FOnGetModularRigTreeHierarchy OnGetHierarchy;
 	FOnModularRigTreeMouseButtonClick OnMouseButtonClick;
 	FOnModularRigTreeMouseButtonDoubleClick OnMouseButtonDoubleClick;
+	FOnDragDetected OnDragDetected;
 	FOnModularRigTreeCanAcceptDrop OnCanAcceptDrop;
 	FOnModularRigTreeAcceptDrop OnAcceptDrop;
 	FOnContextMenuOpening OnContextMenuOpening;
