Control Rig: Enable additive control rigs by default
#rnx

diff --git a/Engine/Plugins/Animation/ControlRig/Source/ControlRigEditor/Private/Sequencer/ControlRigParameterTrackEditor.cpp b/Engine/Plugins/Animation/ControlRig/Source/ControlRigEditor/Private/Sequencer/ControlRigParameterTrackEditor.cpp
--- a/Engine/Plugins/Animation/ControlRig/Source/ControlRigEditor/Private/Sequencer/ControlRigParameterTrackEditor.cpp
+++ b/Engine/Plugins/Animation/ControlRig/Source/ControlRigEditor/Private/Sequencer/ControlRigParameterTrackEditor.cpp
@@ -105,7 +105,7 @@ TAutoConsoleVariable<bool> CVarSelectedKeysSelectControls(TEXT("ControlRig.Seque
 
 TAutoConsoleVariable<bool> CVarSelectedSectionSetsSectionToKey(TEXT("ControlRig.Sequencer.SelectedSectionSetsSectionToKey"), false, TEXT("When true when we select a channel in a section, if it's the only section selected we set it as the Section To Key, by default false."));
 
-TAutoConsoleVariable<bool> CVarEnableAdditiveControlRigs(TEXT("ControlRig.Sequencer.EnableAdditiveControlRigs"), false, TEXT("When true it is possible to add an additive control rig to a skeletal mesh component."));
+TAutoConsoleVariable<bool> CVarEnableAdditiveControlRigs(TEXT("ControlRig.Sequencer.EnableAdditiveControlRigs"), true, TEXT("When true it is possible to add an additive control rig to a skeletal mesh component."));
 
 static USkeletalMeshComponent* AcquireSkeletalMeshFromObject(UObject* BoundObject, TSharedPtr<ISequencer> SequencerPtr)
 {
