Control Rig: Modify rig element key in details panel
#rnx

diff --git a/Engine/Plugins/Animation/ControlRig/Source/ControlRig/Private/ModularRig.cpp b/Engine/Plugins/Animation/ControlRig/Source/ControlRig/Private/ModularRig.cpp
--- a/Engine/Plugins/Animation/ControlRig/Source/ControlRig/Private/ModularRig.cpp
+++ b/Engine/Plugins/Animation/ControlRig/Source/ControlRig/Private/ModularRig.cpp
@@ -17,32 +17,32 @@ const FString UModularRig::NamespaceSeparator = TEXT(":");
 // FModuleInstanceHandle
 ////////////////////////////////////////////////////////////////////////////////
 
-FModuleInstanceHandle::FModuleInstanceHandle(UModularRig* InRig, const FString& InPath)
-: Rig(InRig)
+FModuleInstanceHandle::FModuleInstanceHandle(UModularRig* InModularRig, const FString& InPath)
+: ModularRig(InModularRig)
 , Path(InPath)
 {
 }
 
-FModuleInstanceHandle::FModuleInstanceHandle(UModularRig* InRig, const FRigModuleInstance* InModule)
-: Rig(InRig)
+FModuleInstanceHandle::FModuleInstanceHandle(UModularRig* InModularRig, const FRigModuleInstance* InModule)
+: ModularRig(InModularRig)
 , Path(InModule->GetPath())
 {
 }
 
 const FRigModuleInstance* FModuleInstanceHandle::Get() const
 {
-	if(Rig.IsValid())
+	if(ModularRig.IsValid())
 	{
-		return Rig->FindModule(Path);
+		return ModularRig->FindModule(Path);
 	}
 	return nullptr;
 }
 
 FRigModuleInstance* FModuleInstanceHandle::Get()
 {
-	if(Rig.IsValid())
+	if(ModularRig.IsValid())
 	{
-		return Rig->FindModule(Path);
+		return ModularRig->FindModule(Path);
 	}
 	return nullptr;
 }
diff --git a/Engine/Plugins/Animation/ControlRig/Source/ControlRig/Private/Rigs/RigHierarchyCache.cpp b/Engine/Plugins/Animation/ControlRig/Source/ControlRig/Private/Rigs/RigHierarchyCache.cpp
--- a/Engine/Plugins/Animation/ControlRig/Source/ControlRig/Private/Rigs/RigHierarchyCache.cpp
+++ b/Engine/Plugins/Animation/ControlRig/Source/ControlRig/Private/Rigs/RigHierarchyCache.cpp
@@ -112,13 +112,16 @@ const FRigElementKey* FRigElementKeyRedirector::FindReverse(const FRigElementKey
 
 void FRigElementKeyRedirector::Add(const FRigElementKey& InSource, const FRigElementKey& InTarget, const URigHierarchy* InHierarchy)
 {
-	if(!InSource.IsValid() || !InTarget.IsValid() || InSource == InTarget)
+	if(!InSource.IsValid() || InSource == InTarget)
 	{
 		return;
 	}
 
-	InternalKeyToExternalKey.Add(InSource, FCachedRigElement(InTarget, InHierarchy, true));
-	ExternalKeys.Add(InTarget);
+	if (InTarget.IsValid())
+	{
+		InternalKeyToExternalKey.Add(InSource, FCachedRigElement(InTarget, InHierarchy, true));
+	}
+	ExternalKeys.Add(InSource, InTarget);
 	Hash = HashCombine(Hash, HashCombine(GetTypeHash(InSource), GetTypeHash(InTarget)));
 }
 
diff --git a/Engine/Plugins/Animation/ControlRig/Source/ControlRig/Public/ModularRig.h b/Engine/Plugins/Animation/ControlRig/Source/ControlRig/Public/ModularRig.h
--- a/Engine/Plugins/Animation/ControlRig/Source/ControlRig/Public/ModularRig.h
+++ b/Engine/Plugins/Animation/ControlRig/Source/ControlRig/Public/ModularRig.h
@@ -12,18 +12,18 @@ struct CONTROLRIG_API FModuleInstanceHandle
 public:
 
 	FModuleInstanceHandle()
-		: Rig(nullptr)
+		: ModularRig(nullptr)
 		, Path()
 	{}
 
-	FModuleInstanceHandle(UModularRig* InRig, const FString& InPath);
-	FModuleInstanceHandle(UModularRig* InRig, const FRigModuleInstance* InElement);
+	FModuleInstanceHandle(UModularRig* InModularRig, const FString& InPath);
+	FModuleInstanceHandle(UModularRig* InModularRig, const FRigModuleInstance* InElement);
 
 	bool IsValid() const { return Get() != nullptr; }
 	operator bool() const { return IsValid(); }
 	
-	const UModularRig* GetRig() const { return Rig.Get(); }
-	UModularRig* GetHierarchy() { return Rig.Get(); }
+	const UModularRig* GetModularRig() const { return ModularRig.Get(); }
+	UModularRig* GetHierarchy() { return ModularRig.Get(); }
 	const FString& GetPath() const { return Path; }
 
 	const FRigModuleInstance* Get() const;
@@ -31,7 +31,7 @@ public:
 
 private:
 
-	TWeakObjectPtr<UModularRig> Rig;
+	TWeakObjectPtr<UModularRig> ModularRig;
 	FString Path;
 };
 
diff --git a/Engine/Plugins/Animation/ControlRig/Source/ControlRig/Public/Rigs/RigHierarchyCache.h b/Engine/Plugins/Animation/ControlRig/Source/ControlRig/Public/Rigs/RigHierarchyCache.h
--- a/Engine/Plugins/Animation/ControlRig/Source/ControlRig/Public/Rigs/RigHierarchyCache.h
+++ b/Engine/Plugins/Animation/ControlRig/Source/ControlRig/Public/Rigs/RigHierarchyCache.h
@@ -168,8 +168,8 @@ public:
 	FRigElementKeyRedirector(const FRigElementKeyRedirector& InOther, const URigHierarchy* InHierarchy);
 
 	bool Contains(const FRigElementKey& InKey) const { return InternalKeyToExternalKey.Contains(InKey); }
-	bool ContainsExternalKey(const FRigElementKey& InKey) const { return ExternalKeys.Contains(InKey); }
 	const FCachedRigElement* Find(const FRigElementKey& InKey) const { return InternalKeyToExternalKey.Find(InKey); }
+	const FRigElementKey* FindExternalKey(const FRigElementKey& InKey) const { return ExternalKeys.Find(InKey); }
 	FCachedRigElement* Find(const FRigElementKey& InKey) { return InternalKeyToExternalKey.Find(InKey); }
 	const FRigElementKey* FindReverse(const FRigElementKey& InKey) const;
 	uint32 GetHash() const { return Hash; }
@@ -179,7 +179,7 @@ private:
 	void Add(const FRigElementKey& InSource, const FRigElementKey& InTarget, const URigHierarchy* InHierarchy);
 
 	TMap<FRigElementKey, FCachedRigElement> InternalKeyToExternalKey;
-	TSet<FRigElementKey> ExternalKeys;
+	TMap<FRigElementKey, FRigElementKey> ExternalKeys;
 	uint32 Hash;
 
 	friend class URigHierarchy;
diff --git a/Engine/Plugins/Animation/ControlRig/Source/ControlRigDeveloper/Private/ControlRigBlueprint.cpp b/Engine/Plugins/Animation/ControlRig/Source/ControlRigDeveloper/Private/ControlRigBlueprint.cpp
--- a/Engine/Plugins/Animation/ControlRig/Source/ControlRigDeveloper/Private/ControlRigBlueprint.cpp
+++ b/Engine/Plugins/Animation/ControlRig/Source/ControlRigDeveloper/Private/ControlRigBlueprint.cpp
@@ -1572,6 +1572,7 @@ static void BuildModularRig(UModularRig* Rig, const FModularRigModel& Model)
 
 void UControlRigBlueprint::RecompileModularRig()
 {
+	OnModularRigPreCompiled().Broadcast(this);
 	if (UClass* MyControlRigClass = GeneratedClass)
 	{
 		if (UModularRig* DefaultObject = Cast<UModularRig>(MyControlRigClass->GetDefaultObject(false)))
diff --git a/Engine/Plugins/Animation/ControlRig/Source/ControlRigDeveloper/Private/ModularRigController.cpp b/Engine/Plugins/Animation/ControlRig/Source/ControlRigDeveloper/Private/ModularRigController.cpp
--- a/Engine/Plugins/Animation/ControlRig/Source/ControlRigDeveloper/Private/ModularRigController.cpp
+++ b/Engine/Plugins/Animation/ControlRig/Source/ControlRigDeveloper/Private/ModularRigController.cpp
@@ -149,12 +149,6 @@ bool UModularRigController::ConnectModuleToElement(const FRigElementKey& InConne
 		return false;
 	}
 
-	if (!Blueprint->Hierarchy->Contains(InTargetKey))
-	{
-		UE_LOG(LogControlRig, Error, TEXT("Could not find target %s"), *InTargetKey.ToString());
-		return false;
-	}
-
 #if WITH_EDITOR
 	TSharedPtr<FScopedTransaction> TransactionPtr;
 	if (bSetupUndo)
diff --git a/Engine/Plugins/Animation/ControlRig/Source/ControlRigDeveloper/Public/ControlRigBlueprint.h b/Engine/Plugins/Animation/ControlRig/Source/ControlRigDeveloper/Public/ControlRigBlueprint.h
--- a/Engine/Plugins/Animation/ControlRig/Source/ControlRigDeveloper/Public/ControlRigBlueprint.h
+++ b/Engine/Plugins/Animation/ControlRig/Source/ControlRigDeveloper/Public/ControlRigBlueprint.h
@@ -298,6 +298,7 @@ public:
 	*/
 	FRigHierarchyModifiedEvent& OnHierarchyModified() { return HierarchyModifiedEvent; }
 
+	FOnRigVMRefreshEditorEvent& OnModularRigPreCompiled() { return ModularRigPreCompiled; }
 	FOnRigVMRefreshEditorEvent& OnModularRigCompiled() { return ModularRigCompiled; }
 
 private:
@@ -306,6 +307,7 @@ private:
 	TObjectPtr<UControlRigValidator> Validator;
 
 	FRigHierarchyModifiedEvent	HierarchyModifiedEvent;
+	FOnRigVMRefreshEditorEvent ModularRigPreCompiled;
 	FOnRigVMRefreshEditorEvent ModularRigCompiled;
 
 	void HandleHierarchyModified(ERigHierarchyNotification InNotification, URigHierarchy* InHierarchy, const FRigBaseElement* InElement);
diff --git a/Engine/Plugins/Animation/ControlRig/Source/ControlRigEditor/Private/ControlRigElementDetails.cpp b/Engine/Plugins/Animation/ControlRig/Source/ControlRigEditor/Private/ControlRigElementDetails.cpp
--- a/Engine/Plugins/Animation/ControlRig/Source/ControlRigEditor/Private/ControlRigElementDetails.cpp
+++ b/Engine/Plugins/Animation/ControlRig/Source/ControlRigEditor/Private/ControlRigElementDetails.cpp
@@ -61,18 +61,6 @@ struct FRigElementTransformWidgetSettings
 TMap<uint32, FRigElementTransformWidgetSettings> FRigElementTransformWidgetSettings::sSettings;
 
 
-namespace FRigElementKeyDetailsDefs
-{
-	// Active foreground pin alpha
-	static const float ActivePinForegroundAlpha = 1.f;
-	// InActive foreground pin alpha
-	static const float InactivePinForegroundAlpha = 0.15f;
-	// Active background pin alpha
-	static const float ActivePinBackgroundAlpha = 0.8f;
-	// InActive background pin alpha
-	static const float InactivePinBackgroundAlpha = 0.4f;
-};
-
 void RigElementKeyDetails_GetCustomizedInfo(TSharedRef<IPropertyHandle> InStructPropertyHandle, UControlRigBlueprint*& OutBlueprint)
 {
 	TArray<UObject*> Objects;
@@ -157,6 +145,168 @@ UControlRigBlueprint* RigElementDetails_GetBlueprintFromHierarchy(URigHierarchy*
 	return Blueprint;
 }
 
+void SRigElementKeyWidget::Construct(const FArguments& InArgs, TSharedPtr<IPropertyHandle> InNameHandle, TSharedPtr<IPropertyHandle> InTypeHandle)
+{
+	NameHandle = InNameHandle;
+	TypeHandle = InTypeHandle;
+	Construct(InArgs);
+}
+
+void SRigElementKeyWidget::Construct(const FArguments& InArgs)
+{
+	BlueprintBeingCustomized = InArgs._Blueprint;
+	OnGetElementType = InArgs._OnGetElementType;
+	OnElementNameChanged = InArgs._OnElementNameChanged;
+	OnElementTypeChanged = InArgs._OnElementTypeChanged;
+	
+	UpdateElementNameList();
+	
+	ChildSlot
+	[
+		SNew(SHorizontalBox)
+		+ SHorizontalBox::Slot()
+		.AutoWidth()
+		[
+			TypeHandle.IsValid() ?
+				TypeHandle->CreatePropertyValueWidget()
+					:
+				SNew(SEnumComboBox, StaticEnum<ERigElementType>())
+				.CurrentValue_Lambda([this]()
+				{
+					if (OnGetElementType.IsBound())
+					{
+						return (int32)OnGetElementType.Execute();
+					}
+					return (int32)ERigElementType::None;
+				})
+				.OnEnumSelectionChanged_Lambda([this](int32 InEnumValue, ESelectInfo::Type SelectInfo)
+				{
+					ERigElementType EnumValue = static_cast<ERigElementType>(InEnumValue);
+					OnElementTypeChanged.ExecuteIfBound(EnumValue);
+					UpdateElementNameList();
+					SearchableComboBox->ClearSelection();
+					OnElementNameChanged.ExecuteIfBound(nullptr, ESelectInfo::Direct);
+				})
+		]
+
+		+ SHorizontalBox::Slot()
+		.AutoWidth()
+		.Padding(4.f, 0.f, 0.f, 0.f)
+		[
+			SAssignNew(SearchableComboBox, SSearchableComboBox)
+			.OptionsSource(&ElementNameList)
+			.OnSelectionChanged(InArgs._OnElementNameChanged)
+			.OnGenerateWidget_Lambda([](TSharedPtr<FString> InItem)
+			{
+				return SNew(STextBlock)
+				.Text(FText::FromString(InItem.IsValid() ? *InItem : FString()))
+				.Font(IDetailLayoutBuilder::GetDetailFont());
+			})
+			.Content()
+			[
+				SNew(STextBlock)
+				.Text_Lambda([InArgs]()
+				{
+					if (InArgs._OnGetElementNameAsText.IsBound())
+					{
+						return InArgs._OnGetElementNameAsText.Execute();
+					}
+					return FText();
+				})
+				.Font(IDetailLayoutBuilder::GetDetailFont())
+			]
+		]
+		// Use button
+		+SHorizontalBox::Slot()
+		.AutoWidth()
+		.Padding(1,0)
+		.VAlign(VAlign_Center)
+		[
+			SAssignNew(UseSelectedButton, SButton)
+			.ButtonStyle( FAppStyle::Get(), "NoBorder" )
+			.ButtonColorAndOpacity_Lambda([this, InArgs]() { return UseSelectedButton.IsValid() && UseSelectedButton->IsHovered() ? InArgs._ActiveBackgroundColor : InArgs._InactiveBackgroundColor; })
+			.OnClicked(InArgs._OnGetSelectedClicked)
+			.ContentPadding(1.f)
+			.ToolTipText(NSLOCTEXT("ControlRigElementDetails", "ObjectGraphPin_Use_Tooltip", "Use item selected"))
+			[
+				SNew(SImage)
+				.ColorAndOpacity_Lambda( [this, InArgs]() { return UseSelectedButton.IsValid() && UseSelectedButton->IsHovered() ? InArgs._ActiveForegroundColor : InArgs._InactiveForegroundColor; })
+				.Image(FAppStyle::GetBrush("Icons.CircleArrowLeft"))
+			]
+		]
+		// Select in hierarchy button
+		+SHorizontalBox::Slot()
+		.AutoWidth()
+		.Padding(1,0)
+		.VAlign(VAlign_Center)
+		[
+			SAssignNew(SelectElementButton, SButton)
+			.ButtonStyle( FAppStyle::Get(), "NoBorder" )
+			.ButtonColorAndOpacity_Lambda([this, InArgs]() { return SelectElementButton.IsValid() && SelectElementButton->IsHovered() ? InArgs._ActiveBackgroundColor : InArgs._InactiveBackgroundColor; })
+			.OnClicked(InArgs._OnSelectInHierarchyClicked)
+			.ContentPadding(0)
+			.ToolTipText(NSLOCTEXT("ControlRigElementDetails", "ObjectGraphPin_Browse_Tooltip", "Select in hierarchy"))
+			[
+				SNew(SImage)
+				.ColorAndOpacity_Lambda( [this, InArgs]() { return SelectElementButton.IsValid() && SelectElementButton->IsHovered() ? InArgs._ActiveForegroundColor : InArgs._InactiveForegroundColor; })
+				.Image(FAppStyle::GetBrush("Icons.Search"))
+			]
+		]
+	];
+
+	if (TypeHandle)
+	{
+		TypeHandle->SetOnPropertyValueChanged(FSimpleDelegate::CreateLambda(
+			[this, InArgs]()
+			{
+				int32 EnumValue;
+				TypeHandle->GetValue(EnumValue);
+				OnElementTypeChanged.ExecuteIfBound(static_cast<ERigElementType>(EnumValue));
+				UpdateElementNameList();
+				SearchableComboBox->ClearSelection();
+				OnElementNameChanged.ExecuteIfBound(nullptr, ESelectInfo::Direct);
+			}
+		));
+	}
+}
+
+void SRigElementKeyWidget::UpdateElementNameList()
+{
+	ElementNameList.Reset();
+
+	if (BlueprintBeingCustomized)
+	{
+		for (UEdGraph* Graph : BlueprintBeingCustomized->UbergraphPages)
+		{
+			if (UControlRigGraph* RigGraph = Cast<UControlRigGraph>(Graph))
+			{
+				
+				const TArray<TSharedPtr<FRigVMStringWithTag>>* NameList = nullptr;
+				if (OnGetElementType.IsBound())
+				{
+					NameList = RigGraph->GetElementNameList(OnGetElementType.Execute());
+				}
+
+				ElementNameList.Reset();
+				if (NameList)
+				{
+					ElementNameList.Reserve(NameList->Num());
+					for(const TSharedPtr<FRigVMStringWithTag>& Name : *NameList)
+					{
+						ElementNameList.Add(MakeShared<FString>(Name->GetString()));
+					}
+				}
+				
+				if(SearchableComboBox.IsValid())
+				{
+					SearchableComboBox->RefreshOptions();
+				}
+				return;
+			}
+		}
+	}
+}
+
 void FRigElementKeyDetails::CustomizeHeader(TSharedRef<IPropertyHandle> InStructPropertyHandle, FDetailWidgetRow& HeaderRow, IPropertyTypeCustomizationUtils& StructCustomizationUtils)
 {
 	BlueprintBeingCustomized = nullptr;
@@ -198,16 +348,6 @@ void FRigElementKeyDetails::CustomizeHeader(TSharedRef<IPropertyHandle> InStruct
 		TypeHandle = InStructPropertyHandle->GetChildHandle(TEXT("Type"));
 		NameHandle = InStructPropertyHandle->GetChildHandle(TEXT("Name"));
 
-		TypeHandle->SetOnPropertyValueChanged(FSimpleDelegate::CreateLambda(
-			[this]()
-			{
-				this->UpdateElementNameList();
-				SetElementName(FString());
-			}
-		));
-
-		UpdateElementNameList();
-
 		HeaderRow
 		.NameContent()
 		[
@@ -216,66 +356,19 @@ void FRigElementKeyDetails::CustomizeHeader(TSharedRef<IPropertyHandle> InStruct
 		.ValueContent()
 		.MinDesiredWidth(250.f)
 		[
-			SNew(SHorizontalBox)
-
-			+ SHorizontalBox::Slot()
-			.AutoWidth()
-			[
-				TypeHandle->CreatePropertyValueWidget()
-			]
-
-			+ SHorizontalBox::Slot()
-			.AutoWidth()
-			.Padding(4.f, 0.f, 0.f, 0.f)
-			[
-				SAssignNew(SearchableComboBox, SSearchableComboBox)
-				.OptionsSource(&ElementNameList)
-				.OnSelectionChanged(this, &FRigElementKeyDetails::OnElementNameChanged)
-				.OnGenerateWidget(this, &FRigElementKeyDetails::OnGetElementNameWidget)
-				.IsEnabled(!NameHandle->IsEditConst())
-				.Content()
-				[
-					SNew(STextBlock)
-					.Text(this, &FRigElementKeyDetails::GetElementNameAsText)
-					.Font(IDetailLayoutBuilder::GetDetailFont())
-				]
-			]
-			// Use button
-			+SHorizontalBox::Slot()
-			.AutoWidth()
-			.Padding(1,0)
-			.VAlign(VAlign_Center)
-			[
-				SAssignNew(UseSelectedButton, SButton)
-				.ButtonStyle( FAppStyle::Get(), "NoBorder" )
-				.ButtonColorAndOpacity_Lambda([this]() { return OnGetWidgetBackground(UseSelectedButton); })
-				.OnClicked(this, &FRigElementKeyDetails::OnGetSelectedClicked)
-				.ContentPadding(1.f)
-				.ToolTipText(NSLOCTEXT("ControlRigElementDetails", "ObjectGraphPin_Use_Tooltip", "Use item selected"))
-				[
-					SNew(SImage)
-					.ColorAndOpacity_Lambda( [this]() { return OnGetWidgetForeground(UseSelectedButton); })
-					.Image(FAppStyle::GetBrush("Icons.CircleArrowLeft"))
-				]
-			]
-			// Select in hierarchy button
-			+SHorizontalBox::Slot()
-			.AutoWidth()
-			.Padding(1,0)
-			.VAlign(VAlign_Center)
-			[
-				SAssignNew(SelectElementButton, SButton)
-				.ButtonStyle( FAppStyle::Get(), "NoBorder" )
-				.ButtonColorAndOpacity_Lambda([this]() { return OnGetWidgetBackground(SelectElementButton); })
-				.OnClicked(this, &FRigElementKeyDetails::OnSelectInHierarchyClicked)
-				.ContentPadding(0)
-				.ToolTipText(NSLOCTEXT("ControlRigElementDetails", "ObjectGraphPin_Browse_Tooltip", "Select in hierarchy"))
-				[
-					SNew(SImage)
-					.ColorAndOpacity_Lambda( [this]() { return OnGetWidgetForeground(SelectElementButton); })
-					.Image(FAppStyle::GetBrush("Icons.Search"))
-				]
-			]			
+			SAssignNew(RigElementKeyWidget, SRigElementKeyWidget, NameHandle, TypeHandle)
+			.Blueprint(BlueprintBeingCustomized)
+			.IsEnabled_Lambda([this](){ return !NameHandle->IsEditConst(); })
+			.ActiveBackgroundColor(FSlateColor(FLinearColor(1.f, 1.f, 1.f, FRigElementKeyDetailsDefs::ActivePinBackgroundAlpha)))
+			.ActiveForegroundColor(FSlateColor(FLinearColor(1.f, 1.f, 1.f, FRigElementKeyDetailsDefs::ActivePinForegroundAlpha)))
+			.InactiveBackgroundColor(FSlateColor(FLinearColor(1.f, 1.f, 1.f, FRigElementKeyDetailsDefs::InactivePinBackgroundAlpha)))
+			.InactiveForegroundColor(FSlateColor(FLinearColor(1.f, 1.f, 1.f, FRigElementKeyDetailsDefs::InactivePinForegroundAlpha)))
+			.OnElementNameChanged(this, &FRigElementKeyDetails::OnElementNameChanged)
+			.OnGetSelectedClicked(this, &FRigElementKeyDetails::OnGetSelectedClicked)
+			//.OnGetElementNameWidget(this, &FRigElementKeyDetails::OnGetElementNameWidget)
+			.OnSelectInHierarchyClicked(this, &FRigElementKeyDetails::OnSelectInHierarchyClicked)
+			.OnGetElementNameAsText_Raw(this, &FRigElementKeyDetails::GetElementNameAsText)
+			.OnGetElementType(this, &FRigElementKeyDetails::GetElementType)
 		];
 	}
 }
@@ -380,41 +473,6 @@ void FRigElementKeyDetails::SetElementName(FString InName)
 	}
 }
 
-void FRigElementKeyDetails::UpdateElementNameList()
-{
-	if (!TypeHandle.IsValid())
-	{
-		return;
-	}
-
-	ElementNameList.Reset();
-
-	if (BlueprintBeingCustomized)
-	{
-		for (UEdGraph* Graph : BlueprintBeingCustomized->UbergraphPages)
-		{
-			if (UControlRigGraph* RigGraph = Cast<UControlRigGraph>(Graph))
-			{
-				const TArray<TSharedPtr<FRigVMStringWithTag>>* NameList =
-					RigGraph->GetElementNameList(GetElementType());
-
-				ElementNameList.Reset();
-				ElementNameList.Reserve(NameList->Num());
-				for(const TSharedPtr<FRigVMStringWithTag>& Name : *NameList)
-				{
-					ElementNameList.Add(MakeShared<FString>(Name->GetString()));
-				}
-				
-				if(SearchableComboBox.IsValid())
-				{
-					SearchableComboBox->RefreshOptions();
-				}
-				return;
-			}
-		}
-	}
-}
-
 void FRigElementKeyDetails::OnElementNameChanged(TSharedPtr<FString> InItem, ESelectInfo::Type InSelectionInfo)
 {
 	if (InItem.IsValid())
@@ -427,13 +485,6 @@ void FRigElementKeyDetails::OnElementNameChanged(TSharedPtr<FString> InItem, ESe
 	}
 }
 
-TSharedRef<SWidget> FRigElementKeyDetails::OnGetElementNameWidget(TSharedPtr<FString> InItem)
-{
-	return SNew(STextBlock)
-		.Text(FText::FromString(InItem.IsValid() ? *InItem : FString()))
-		.Font(IDetailLayoutBuilder::GetDetailFont());
-}
-
 FText FRigElementKeyDetails::GetElementNameAsText() const
 {
 	return FText::FromString(GetElementName());
diff --git a/Engine/Plugins/Animation/ControlRig/Source/ControlRigEditor/Private/ControlRigElementDetails.h b/Engine/Plugins/Animation/ControlRig/Source/ControlRigEditor/Private/ControlRigElementDetails.h
--- a/Engine/Plugins/Animation/ControlRig/Source/ControlRigEditor/Private/ControlRigElementDetails.h
+++ b/Engine/Plugins/Animation/ControlRig/Source/ControlRigEditor/Private/ControlRigElementDetails.h
@@ -30,6 +30,67 @@
 
 class IPropertyHandle;
 
+namespace FRigElementKeyDetailsDefs
+{
+	// Active foreground pin alpha
+	static const float ActivePinForegroundAlpha = 1.f;
+	// InActive foreground pin alpha
+	static const float InactivePinForegroundAlpha = 0.15f;
+	// Active background pin alpha
+	static const float ActivePinBackgroundAlpha = 0.8f;
+	// InActive background pin alpha
+	static const float InactivePinBackgroundAlpha = 0.4f;
+};
+
+class SRigElementKeyWidget : public SCompoundWidget
+{
+public:
+
+	DECLARE_DELEGATE_RetVal(FText, FGetElementNameAsText);
+	DECLARE_DELEGATE_RetVal(ERigElementType, FGetElementType);
+	DECLARE_DELEGATE_RetVal(bool, FIsEnabled);
+	DECLARE_DELEGATE_OneParam(FOnElementTypeChanged, ERigElementType);
+	
+	SLATE_BEGIN_ARGS(SRigElementKeyWidget)
+	{
+	}
+	SLATE_ARGUMENT(UControlRigBlueprint*, Blueprint)
+	SLATE_ARGUMENT(FSlateColor, ActiveBackgroundColor)
+	SLATE_ARGUMENT(FSlateColor, InactiveBackgroundColor)
+	SLATE_ARGUMENT(FSlateColor, ActiveForegroundColor)
+	SLATE_ARGUMENT(FSlateColor, InactiveForegroundColor)
+	SLATE_EVENT(SSearchableComboBox::FOnSelectionChanged, OnElementNameChanged)
+	SLATE_EVENT(FOnClicked, OnGetSelectedClicked)
+	SLATE_EVENT(FOnClicked, OnSelectInHierarchyClicked)
+	SLATE_EVENT(FGetElementNameAsText, OnGetElementNameAsText)
+	SLATE_EVENT(FGetElementType, OnGetElementType)
+	SLATE_EVENT(FOnElementTypeChanged, OnElementTypeChanged)
+	SLATE_EVENT(FIsEnabled, IsEnabled);
+	
+	SLATE_END_ARGS()
+
+	void Construct(const FArguments& InArgs, TSharedPtr<IPropertyHandle> InNameHandle, TSharedPtr<IPropertyHandle> InTypeHandle);
+	void Construct(const FArguments& InArgs);
+
+	void UpdateElementNameList();
+
+private:
+	TSharedPtr<IPropertyHandle> NameHandle;
+	TSharedPtr<IPropertyHandle> TypeHandle;
+
+	/** Helper buttons. */
+	TSharedPtr<SButton> UseSelectedButton;
+	TSharedPtr<SButton> SelectElementButton;
+
+	FGetElementType OnGetElementType;
+	FOnElementTypeChanged OnElementTypeChanged;
+	SSearchableComboBox::FOnSelectionChanged OnElementNameChanged;
+
+	UControlRigBlueprint* BlueprintBeingCustomized;
+	TArray<TSharedPtr<FString>> ElementNameList;
+	TSharedPtr<SSearchableComboBox> SearchableComboBox;
+};
+
 class FRigElementKeyDetails : public IPropertyTypeCustomization
 {
 public:
@@ -48,9 +109,7 @@ protected:
 	ERigElementType GetElementType() const;
 	FString GetElementName() const;
 	void SetElementName(FString InName);
-	void UpdateElementNameList();
 	void OnElementNameChanged(TSharedPtr<FString> InItem, ESelectInfo::Type InSelectionInfo);
-	TSharedRef<SWidget> OnGetElementNameWidget(TSharedPtr<FString> InItem);
 	FText GetElementNameAsText() const;
 
 	/** Helper buttons. */
@@ -69,9 +128,8 @@ protected:
 	
 	TSharedPtr<IPropertyHandle> TypeHandle;
 	TSharedPtr<IPropertyHandle> NameHandle;
-	TArray<TSharedPtr<FString>> ElementNameList;
 	UControlRigBlueprint* BlueprintBeingCustomized;
-	TSharedPtr<SSearchableComboBox> SearchableComboBox;
+	TSharedPtr<SRigElementKeyWidget> RigElementKeyWidget;
 };
 
 UENUM()
diff --git a/Engine/Plugins/Animation/ControlRig/Source/ControlRigEditor/Private/ControlRigModuleDetails.cpp b/Engine/Plugins/Animation/ControlRig/Source/ControlRigEditor/Private/ControlRigModuleDetails.cpp
--- a/Engine/Plugins/Animation/ControlRig/Source/ControlRigEditor/Private/ControlRigModuleDetails.cpp
+++ b/Engine/Plugins/Animation/ControlRig/Source/ControlRigEditor/Private/ControlRigModuleDetails.cpp
@@ -119,7 +119,10 @@ void FRigModuleInstanceDetails::CustomizeDetails(IDetailLayoutBuilder& DetailBui
 
 		FPerModuleInfo Info;
 		Info.WrapperObject = WrapperObject;
-		Info.Module = Cast<UModularRig>(WrapperObject->GetSubject())->GetHandle(Path);
+		if (UModularRig* Subject = Cast<UModularRig>(WrapperObject->GetSubject()))
+		{
+			Info.Module = Subject->GetHandle(Path);
+		}
 
 		if(!Info.Module.IsValid())
 		{
@@ -177,15 +180,21 @@ void FRigModuleInstanceDetails::CustomizeDetails(IDetailLayoutBuilder& DetailBui
 	IDetailCategoryBuilder& ConnectionsCategory = DetailBuilder.EditCategory(TEXT("Connections"), LOCTEXT("Connections", "Connections"));
 
 	{
+		TArray<FRigModuleConnector> Connectors = GetConnectors();
 		FRigElementKeyRedirector Redirector = GetConnections();
-		for(TPair<FRigElementKey, FCachedRigElement>& Connection : Redirector.InternalKeyToExternalKey)
+		for(FRigModuleConnector& Connector : Connectors)
 		{
-			if (!Connection.Key.IsTypeOf(ERigElementType::Connector))
+			const FText Label = FText::FromString(Connector.Name);
+			FRigElementKey ConnectorKey(*Connector.Name, ERigElementType::Connector);
+			const FRigElementKey* TargetKey = Redirector.FindExternalKey(ConnectorKey);
+			if (TargetKey)
 			{
-				continue;
+				Connections.Add(ConnectorKey, *TargetKey);
+			}
+			else
+			{
+				Connections.Add(ConnectorKey, FRigElementKey());
 			}
-			const FText Label = FText::FromName(Connection.Key.Name);
-			const FRigElementKey& Target = Connection.Value.GetKey();
 			ConnectionsCategory.AddCustomRow(Label)
 				.NameContent()
 				[
@@ -196,10 +205,19 @@ void FRigModuleInstanceDetails::CustomizeDetails(IDetailLayoutBuilder& DetailBui
 				]
 				.ValueContent()
 				[
-					SNew(STextBlock)
-					.Text(FText::FromName(Target.Name))
-					.Font(IDetailLayoutBuilder::GetDetailFont())
-					.IsEnabled(true)
+					SAssignNew(RigElementKeyWidget, SRigElementKeyWidget)
+					.Blueprint(PerModuleInfos[0].GetBlueprint())
+					.IsEnabled_Lambda([this](){ return true; })
+					.ActiveBackgroundColor(FSlateColor(FLinearColor(1.f, 1.f, 1.f, FRigElementKeyDetailsDefs::ActivePinBackgroundAlpha)))
+					.ActiveForegroundColor(FSlateColor(FLinearColor(1.f, 1.f, 1.f, FRigElementKeyDetailsDefs::ActivePinForegroundAlpha)))
+					.InactiveBackgroundColor(FSlateColor(FLinearColor(1.f, 1.f, 1.f, FRigElementKeyDetailsDefs::InactivePinBackgroundAlpha)))
+					.InactiveForegroundColor(FSlateColor(FLinearColor(1.f, 1.f, 1.f, FRigElementKeyDetailsDefs::InactivePinForegroundAlpha)))
+					.OnElementNameChanged(this, &FRigModuleInstanceDetails::OnElementNameChanged, ConnectorKey)
+					.OnGetSelectedClicked(this, &FRigModuleInstanceDetails::OnGetSelectedClicked, ConnectorKey)
+					.OnSelectInHierarchyClicked(this, &FRigModuleInstanceDetails::OnSelectInHierarchyClicked, ConnectorKey)
+					.OnGetElementNameAsText_Raw(this, &FRigModuleInstanceDetails::GetElementNameAsText, ConnectorKey)
+					.OnGetElementType(this, &FRigModuleInstanceDetails::GetElementType, ConnectorKey)
+					.OnElementTypeChanged(this, &FRigModuleInstanceDetails::OnElementTypeChanged, ConnectorKey)
 				];
 		}
 	}
@@ -317,6 +335,27 @@ FText FRigModuleInstanceDetails::GetRigClassPath() const
 	return FText();
 }
 
+TArray<FRigModuleConnector> FRigModuleInstanceDetails::GetConnectors() const
+{
+	if(PerModuleInfos.Num() > 1)
+	{
+		return TArray<FRigModuleConnector>();
+	}
+
+	if (FRigModuleInstance* Module = PerModuleInfos[0].GetModule())
+	{
+		if (TSoftObjectPtr<UControlRig> Rig = Module->Rig)
+		{
+			if (Rig.IsValid())
+			{
+				return Rig->GetRigModuleSettings().ExposedConnectors;
+			}
+		}
+	}
+
+	return TArray<FRigModuleConnector>();
+}
+
 FRigElementKeyRedirector FRigModuleInstanceDetails::GetConnections() const
 {
 	if(PerModuleInfos.Num() > 1)
@@ -367,7 +406,7 @@ void FRigModuleInstanceDetails::OnConfigValueChanged(const FName InVariableName)
 		if (Rig.IsValid())
 		{
 			FString ValueStr = Rig->GetVariableAsString(InVariableName);
-			if (UControlRigBlueprint* Blueprint = RigModuleDetails_GetBlueprintFromRig(PerModuleInfos[0].GetRig()))
+			if (UControlRigBlueprint* Blueprint = RigModuleDetails_GetBlueprintFromRig(PerModuleInfos[0].GetModularRig()))
 			{
 				UModularRigController* Controller = Blueprint->GetModularRigController();
 				Controller->SetConfigValueInModule(ModuleInstance->GetPath(), InVariableName, ValueStr);
@@ -418,5 +457,95 @@ void FRigModuleInstanceDetails::RegisterSectionMappings(FPropertyEditorModule& P
 	MetadataSection->AddCategory("Metadata");
 }
 
+void FRigModuleInstanceDetails::OnElementNameChanged(TSharedPtr<FString> InItem, ESelectInfo::Type InSelectionInfo, FRigElementKey Connector)
+{
+	for (FPerModuleInfo& Info : PerModuleInfos)
+	{
+		if (UControlRigBlueprint* Blueprint = RigModuleDetails_GetBlueprintFromRig(Info.GetModularRig()))
+		{
+			UModularRigController* Controller = Blueprint->GetModularRigController();
+			FRigElementKey Key;
+			FRigElementKey* TargetKey = Connections.Find(Connector);
+			if (TargetKey)
+			{
+				if (InItem)
+				{
+					TargetKey->Name = **InItem;
+				}
+
+				FRigElementKey NamespacedConnector(*FString::Printf(TEXT("%s:%s"), *Info.GetModule()->GetPath(), *Connector.Name.ToString()), ERigElementType::Connector);
+				Controller->ConnectModuleToElement(NamespacedConnector, *TargetKey);
+			}
+		}
+	}
+}
+
+void FRigModuleInstanceDetails::OnElementTypeChanged(ERigElementType InElementType, FRigElementKey Connector)
+{
+	for (FPerModuleInfo& Info : PerModuleInfos)
+	{
+		if (UControlRigBlueprint* Blueprint = RigModuleDetails_GetBlueprintFromRig(Info.GetModularRig()))
+		{
+			UModularRigController* Controller = Blueprint->GetModularRigController();
+			FRigElementKey Key;
+			FRigElementKey* TargetKey = Connections.Find(Connector);
+			if (TargetKey)
+			{
+				TargetKey->Type = InElementType;
+
+				FRigElementKey NamespacedConnector(*FString::Printf(TEXT("%s:%s"), *Info.GetModule()->GetPath(), *Connector.Name.ToString()), ERigElementType::Connector);
+				Controller->ConnectModuleToElement(NamespacedConnector, *TargetKey);
+			}
+		}
+	}
+}
+
+FText FRigModuleInstanceDetails::GetElementNameAsText(FRigElementKey Connector) const
+{
+	if (const FRigElementKey* TargetKey = Connections.Find(Connector))
+	{
+		return FText::FromName(TargetKey->Name);
+	}
+	return FText();
+}
+
+FReply FRigModuleInstanceDetails::OnGetSelectedClicked(FRigElementKey Connector)
+{
+	if (UControlRigBlueprint* Blueprint = RigModuleDetails_GetBlueprintFromRig(PerModuleInfos[0].GetModularRig()))
+	{
+		const TArray<FRigElementKey>& Selected = Blueprint->Hierarchy->GetSelectedKeys();
+		if (Selected.Num() > 0)
+		{
+			FRigElementKey NamespacedConnector(*FString::Printf(TEXT("%s:%s"), *PerModuleInfos[0].GetModule()->GetPath(), *Connector.Name.ToString()), ERigElementType::Connector);
+			Blueprint->GetModularRigController()->ConnectModuleToElement(NamespacedConnector, Selected[0]);
+		}
+	}
+	return FReply::Handled();
+}
+
+FReply FRigModuleInstanceDetails::OnSelectInHierarchyClicked(FRigElementKey Connector)
+{
+	if (UControlRigBlueprint* Blueprint = RigModuleDetails_GetBlueprintFromRig(PerModuleInfos[0].GetModularRig()))
+	{
+		if (FRigElementKey* Target = Connections.Find(Connector))
+		{
+			if (Target->IsValid())
+			{
+				Blueprint->GetHierarchyController()->SetSelection({*Target});
+			}
+		}
+	}
+	return FReply::Handled();
+}
+
+ERigElementType FRigModuleInstanceDetails::GetElementType(FRigElementKey Connector) const
+{
+	if (const FRigElementKey* TargetKey = Connections.Find(Connector))
+	{
+		return TargetKey->Type;
+	}
+	return ERigElementType::None;
+}
+
 
 #undef LOCTEXT_NAMESPACE
diff --git a/Engine/Plugins/Animation/ControlRig/Source/ControlRigEditor/Private/ControlRigModuleDetails.h b/Engine/Plugins/Animation/ControlRig/Source/ControlRigEditor/Private/ControlRigModuleDetails.h
--- a/Engine/Plugins/Animation/ControlRig/Source/ControlRigEditor/Private/ControlRigModuleDetails.h
+++ b/Engine/Plugins/Animation/ControlRig/Source/ControlRigEditor/Private/ControlRigModuleDetails.h
@@ -8,6 +8,7 @@
 #include "ControlRig.h"
 #include "ModularRig.h"
 #include "ControlRigBlueprint.h"
+#include "ControlRigElementDetails.h"
 #include "Editor/ControlRigWrapperObject.h"
 #include "Styling/SlateTypes.h"
 #include "IPropertyUtilities.h"
@@ -41,6 +42,7 @@ public:
 	FString GetModulePath() const;
 	FText GetName() const;
 	FText GetRigClassPath() const;
+	TArray<FRigModuleConnector> GetConnectors() const;
 	FRigElementKeyRedirector GetConnections() const;
 	TArray<FRigVMExternalVariable> GetConfigValues() const;
 
@@ -74,19 +76,19 @@ public:
 		bool IsValid() const { return Module.IsValid(); }
 		operator bool() const { return IsValid(); }
 
-		UModularRig* GetRig() const { return (UModularRig*)Module.GetRig(); }
+		UModularRig* GetModularRig() const { return (UModularRig*)Module.GetModularRig(); }
 		UModularRig* GetDefaultRig() const
 		{
 			if(DefaultModule.IsValid())
 			{
-				return (UModularRig*)DefaultModule.GetRig();
+				return (UModularRig*)DefaultModule.GetModularRig();
 			}
-			return GetRig();
+			return GetModularRig();
 		}
 
 		UControlRigBlueprint* GetBlueprint() const
 		{
-			if(const UModularRig* ControlRig = GetRig()->GetTypedOuter<UModularRig>())
+			if(const UModularRig* ControlRig = GetModularRig())
 			{
 				return Cast<UControlRigBlueprint>(ControlRig->GetClass()->ClassGeneratedBy);
 			}
@@ -118,7 +120,16 @@ public:
 
 	virtual void RegisterSectionMappings(FPropertyEditorModule& PropertyEditorModule, UClass* InClass);
 
+	void OnElementNameChanged(TSharedPtr<FString> InItem, ESelectInfo::Type InSelectionInfo, FRigElementKey Connector);
+	void OnElementTypeChanged(ERigElementType InElementType, FRigElementKey Connector);
+	FText GetElementNameAsText(FRigElementKey Connector) const;
+	FReply OnGetSelectedClicked(FRigElementKey Connector);
+	FReply OnSelectInHierarchyClicked(FRigElementKey Connector);
+	ERigElementType GetElementType(FRigElementKey Connector) const;
+
 protected:
 
 	TArray<FPerModuleInfo> PerModuleInfos;
+	TSharedPtr<SRigElementKeyWidget> RigElementKeyWidget;
+	TMap<FRigElementKey, FRigElementKey> Connections;
 };
diff --git a/Engine/Plugins/Animation/ControlRig/Source/ControlRigEditor/Private/Editor/ControlRigEditor.cpp b/Engine/Plugins/Animation/ControlRig/Source/ControlRigEditor/Private/Editor/ControlRigEditor.cpp
--- a/Engine/Plugins/Animation/ControlRig/Source/ControlRigEditor/Private/Editor/ControlRigEditor.cpp
+++ b/Engine/Plugins/Animation/ControlRig/Source/ControlRigEditor/Private/Editor/ControlRigEditor.cpp
@@ -149,6 +149,10 @@ FControlRigEditor::~FControlRigEditor()
 		}
 
 		RigBlueprint->OnRigTypeChanged().RemoveAll(this);
+		if (RigBlueprint->IsModularRig())
+		{
+			RigBlueprint->GetModularRigController()->OnModified().RemoveAll(this);
+		}
 	}
 
 	if (PersonaToolkit.IsValid())
@@ -313,6 +317,10 @@ void FControlRigEditor::InitRigVMEditor(const EToolkitMode::Type Mode, const TSh
 		}
 
 		ControlRigBlueprint->OnRigTypeChanged().AddSP(this, &FControlRigEditor::HandleRigTypeChanged);
+		if (ControlRigBlueprint->IsModularRig())
+		{
+			ControlRigBlueprint->GetModularRigController()->OnModified().AddSP(this, &FControlRigEditor::HandleModularRigModified);
+		}
 	}
 
 	CreateRigHierarchyToGraphDragAndDropMenu();
@@ -1096,6 +1104,31 @@ void FControlRigEditor::SetDetailViewForRigElements(const TArray<FRigElementKey>
 	SetDetailObjects(Objects);
 }
 
+void FControlRigEditor::SetDetailObjects(const TArray<UObject*>& InObjects)
+{
+	bool bAnyModule = false;
+	for (const UObject* Object : InObjects)
+	{
+		if (const URigVMDetailsViewWrapperObject* WrapperObject = Cast<URigVMDetailsViewWrapperObject>(Object))
+		{
+			if (const UScriptStruct* WrappedStruct = WrapperObject->GetWrappedStruct())
+			{
+				if (WrappedStruct->IsChildOf(FRigModuleInstance::StaticStruct()))
+				{
+					bAnyModule = true;
+				}
+			}
+		}
+	}
+	
+	if (!bAnyModule)
+	{
+		ModulesSelected.Reset();
+	}
+	
+	IControlRigEditor::SetDetailObjects(InObjects);
+}
+
 void FControlRigEditor::RefreshDetailView()
 {
 	if(DetailViewShowsAnyRigElement())
@@ -1103,6 +1136,11 @@ void FControlRigEditor::RefreshDetailView()
 		SetDetailViewForRigElements();
 		return;
 	}
+	else if(!ModulesSelected.IsEmpty())
+	{
+		SetDetailViewForRigModules();
+		return;
+	}
 
 	FRigVMEditor::RefreshDetailView();
 }
@@ -1137,7 +1175,12 @@ bool FControlRigEditor::DetailViewShowsRigElement(FRigElementKey InKey) const
 	return false;
 }
 
-void FControlRigEditor::SetDetailViewForRigModules(const TArray<FString>& InKeys)
+void FControlRigEditor::SetDetailViewForRigModules()
+{
+	SetDetailViewForRigModules(ModulesSelected);
+}
+
+void FControlRigEditor::SetDetailViewForRigModules(const TArray<FString> InKeys)
 {
 	if(IsDetailsPanelRefreshSuspended())
 	{
@@ -1148,6 +1191,13 @@ void FControlRigEditor::SetDetailViewForRigModules(const TArray<FString>& InKeys
 
 	UControlRigBlueprint* RigBlueprint = Cast<UControlRigBlueprint>(GetBlueprintObj());
 	UModularRig* RigBeingDebugged = Cast<UModularRig>(RigBlueprint->GetDebuggedControlRig());
+
+	if (!RigBeingDebugged)
+	{
+		return;
+	}
+
+	ModulesSelected = InKeys;
 	TArray<UObject*> Objects;
 
 	for(const FString& Key : InKeys)
@@ -1166,6 +1216,12 @@ void FControlRigEditor::SetDetailViewForRigModules(const TArray<FString>& InKeys
 	}
 	
 	SetDetailObjects(Objects);
+
+	if (Objects.IsEmpty())
+	{
+		// In case the modules selected are still not available, lets set them again
+		ModulesSelected = InKeys;
+	}
 }
 
 bool FControlRigEditor::DetailViewShowsAnyRigModule() const
@@ -3510,6 +3566,42 @@ void FControlRigEditor::HandleRigTypeChanged(UControlRigBlueprint* InBlueprint)
 	Compile();
 }
 
+void FControlRigEditor::HandleModularRigModified(EModularRigNotification InNotification, const FRigModuleReference* InModule)
+{
+	switch(InNotification)
+	{
+		case EModularRigNotification::ModuleAdded:
+		{
+			ModulesSelected = {InModule->GetPath()};
+			break;
+		}
+		case EModularRigNotification::ModuleRemoved:
+		{
+			if (DetailViewShowsAnyRigModule())
+			{
+				ClearDetailObject();
+			}
+			break;
+		}
+		case EModularRigNotification::ModuleReparented:
+		case EModularRigNotification::ModuleRenamed:
+		{
+			FString OldPath;
+			if (InNotification == EModularRigNotification::ModuleRenamed)
+			{
+				OldPath = FString::Printf(TEXT("%s:%s"), *InModule->ParentPath, *InModule->PreviousName.ToString());
+			}
+			else
+			{
+				OldPath = FString::Printf(TEXT("%s:%s"), *InModule->PreviousParentPath, *InModule->Name.ToString());
+			}
+			ModulesSelected.Remove(OldPath);
+			ModulesSelected.Add(InModule->GetPath());
+			break;
+		}
+	}
+}
+
 void FControlRigEditor::SynchronizeViewportBoneSelection()
 {
 	UControlRigBlueprint* RigBlueprint = GetControlRigBlueprint();
diff --git a/Engine/Plugins/Animation/ControlRig/Source/ControlRigEditor/Private/Editor/ControlRigEditor.h b/Engine/Plugins/Animation/ControlRig/Source/ControlRigEditor/Private/Editor/ControlRigEditor.h
--- a/Engine/Plugins/Animation/ControlRig/Source/ControlRigEditor/Private/Editor/ControlRigEditor.h
+++ b/Engine/Plugins/Animation/ControlRig/Source/ControlRigEditor/Private/Editor/ControlRigEditor.h
@@ -22,6 +22,7 @@
 #include "RigVMModel/RigVMController.h"
 #include "Editor/RigVMDetailsViewWrapperObject.h"
 #include "ControlRigTestData.h"
+#include "ModularRigController.h"
 #include "RigVMHost.h"
 #include "Units/RigUnit.h"
 
@@ -127,11 +128,13 @@ public:
 	bool DetailViewShowsAnyRigElement() const;
 	bool DetailViewShowsRigElement(FRigElementKey InKey) const;
 
-	//void SetDetailViewForRigModules();
-	void SetDetailViewForRigModules(const TArray<FString>& InKeys);
+	void SetDetailViewForRigModules();
+	void SetDetailViewForRigModules(const TArray<FString> InKeys);
 	bool DetailViewShowsAnyRigModule() const;
 	bool DetailViewShowsRigModule(FString InKey) const;
+	TArray<FString> ModulesSelected;
 
+	virtual void SetDetailObjects(const TArray<UObject*>& InObjects) override;
 	virtual void RefreshDetailView() override;
 
 	void CreatePersonaToolKitIfRequired();
@@ -154,6 +157,8 @@ public:
 
 	void HandleRigTypeChanged(UControlRigBlueprint* InBlueprint);
 
+	void HandleModularRigModified(EModularRigNotification InNotification, const FRigModuleReference* InModule);
+
 	const FName RigHierarchyToGraphDragAndDropMenuName = TEXT("ControlRigEditor.RigHierarchyToGraphDragAndDropMenu");
 	void CreateRigHierarchyToGraphDragAndDropMenu() const;
 	void OnGraphNodeDropToPerform(TSharedPtr<FGraphNodeDragDropOp> InDragDropOp, UEdGraph* InGraph, const FVector2D& InNodePosition, const FVector2D& InScreenPosition);
diff --git a/Engine/Plugins/Animation/ControlRig/Source/ControlRigEditor/Private/Editor/SModularRigHierarchy.cpp b/Engine/Plugins/Animation/ControlRig/Source/ControlRigEditor/Private/Editor/SModularRigHierarchy.cpp
--- a/Engine/Plugins/Animation/ControlRig/Source/ControlRigEditor/Private/Editor/SModularRigHierarchy.cpp
+++ b/Engine/Plugins/Animation/ControlRig/Source/ControlRigEditor/Private/Editor/SModularRigHierarchy.cpp
@@ -108,7 +108,8 @@ void SModularRigHierarchy::Construct(const FArguments& InArgs, TSharedRef<FContr
 
 	ControlRigBlueprint->OnRefreshEditor().AddRaw(this, &SModularRigHierarchy::HandleRefreshEditorFromBlueprint);
 	ControlRigBlueprint->OnSetObjectBeingDebugged().AddRaw(this, &SModularRigHierarchy::HandleSetObjectBeingDebugged);
-	ControlRigBlueprint->OnModularRigCompiled().AddRaw(this, &SModularRigHierarchy::HandleRefreshEditorFromBlueprint);
+	ControlRigBlueprint->OnModularRigPreCompiled().AddRaw(this, &SModularRigHierarchy::HandlePreCompileModularRigs);
+	ControlRigBlueprint->OnModularRigCompiled().AddRaw(this, &SModularRigHierarchy::HandlePostCompileModularRigs);
 
 	// for deleting, renaming, dragging
 	CommandList = MakeShared<FUICommandList>();
@@ -184,6 +185,8 @@ void SModularRigHierarchy::OnEditorClose(const FRigVMEditor* InEditor, URigVMBlu
 	{
 		InBlueprint->OnRefreshEditor().RemoveAll(this);
 		InBlueprint->OnSetObjectBeingDebugged().RemoveAll(this);
+		BP->OnModularRigPreCompiled().RemoveAll(this);
+		BP->OnModularRigCompiled().RemoveAll(this);
 	}
 	
 	ControlRigEditor.Reset();
@@ -247,6 +250,27 @@ TArray<FString> SModularRigHierarchy::GetSelectedKeys() const
 	return SelectedKeys;
 }
 
+
+void SModularRigHierarchy::HandlePreCompileModularRigs(URigVMBlueprint* InBlueprint)
+{
+	ClearDetailPanel();
+}
+
+void SModularRigHierarchy::HandlePostCompileModularRigs(URigVMBlueprint* InBlueprint)
+{
+	RefreshTreeView();
+	if (ControlRigEditor.IsValid())
+	{
+		TArray<TSharedPtr<FModularRigTreeElement>> SelectedElements;
+		Algo::Transform(ControlRigEditor.Pin()->ModulesSelected, SelectedElements, [this](const FString& Path)
+		{
+			return TreeView->FindElement(Path);
+		});
+		TreeView->SetSelection(SelectedElements);
+		ControlRigEditor.Pin()->RefreshDetailView();
+	}
+}
+
 void SModularRigHierarchy::HandleRefreshEditorFromBlueprint(URigVMBlueprint* InBlueprint)
 {
 	RefreshTreeView();
@@ -330,13 +354,6 @@ void SModularRigHierarchy::CreateContextMenu()
 						FNewToolMenuDelegate::CreateLambda([Commands, RigHierarchyPanel](UToolMenu* InSubMenu)
 						{
 							FToolMenuSection& DefaultSection = InSubMenu->AddSection(NAME_None);
-							FString SelectedKey;
-							TArray<TSharedPtr<FModularRigTreeElement>> SelectedItems = RigHierarchyPanel->TreeView->GetSelectedItems();
-							if (SelectedItems.Num() > 0)
-							{
-								SelectedKey = SelectedItems[0]->Key;
-							}
-							
 							DefaultSection.AddMenuEntry(Commands.AddModuleItem);
 						})
 					);
diff --git a/Engine/Plugins/Animation/ControlRig/Source/ControlRigEditor/Private/Editor/SModularRigHierarchy.h b/Engine/Plugins/Animation/ControlRig/Source/ControlRigEditor/Private/Editor/SModularRigHierarchy.h
--- a/Engine/Plugins/Animation/ControlRig/Source/ControlRigEditor/Private/Editor/SModularRigHierarchy.h
+++ b/Engine/Plugins/Animation/ControlRig/Source/ControlRigEditor/Private/Editor/SModularRigHierarchy.h
@@ -149,6 +149,9 @@ private:
 	void OnRequestDetailsInspection(const FString& InKey);
 	void ClearDetailPanel() const;
 
+	void HandlePreCompileModularRigs(URigVMBlueprint* InBlueprint);
+	void HandlePostCompileModularRigs(URigVMBlueprint* InBlueprint);
+	
 	void HandleRefreshEditorFromBlueprint(URigVMBlueprint* InBlueprint);
 	void HandleSetObjectBeingDebugged(UObject* InObject);
 
