Control Rig: Propagate UserData from ModularRig to modules
#rnx

diff --git a/Engine/Plugins/Animation/ControlRig/Source/ControlRig/Private/ControlRig.cpp b/Engine/Plugins/Animation/ControlRig/Source/ControlRig/Private/ControlRig.cpp
--- a/Engine/Plugins/Animation/ControlRig/Source/ControlRig/Private/ControlRig.cpp
+++ b/Engine/Plugins/Animation/ControlRig/Source/ControlRig/Private/ControlRig.cpp
@@ -2875,6 +2875,18 @@ const TArray<UAssetUserData*>* UControlRig::GetAssetUserDataArray() const
 		}
 	}
 
+	// Propagate the outer rig user data into the child modules
+	if (UControlRig* OuterCR = GetTypedOuter<UControlRig>())
+	{
+		if(const TArray<UAssetUserData*>* OuterUserDataArray = OuterCR->GetAssetUserDataArray())
+		{
+			for(UAssetUserData* OuterUserData : *OuterUserDataArray)
+			{
+				CombinedAssetUserData.Add(OuterUserData);
+			}
+		}
+	}
+	
 	if(OuterSceneComponent.IsValid())
 	{
 		if(const USkeletalMeshComponent* SkeletalMeshComponent = Cast<USkeletalMeshComponent>(OuterSceneComponent))
diff --git a/Engine/Plugins/Animation/ControlRig/Source/ControlRig/Private/ModularRig.cpp b/Engine/Plugins/Animation/ControlRig/Source/ControlRig/Private/ModularRig.cpp
--- a/Engine/Plugins/Animation/ControlRig/Source/ControlRig/Private/ModularRig.cpp
+++ b/Engine/Plugins/Animation/ControlRig/Source/ControlRig/Private/ModularRig.cpp
@@ -141,8 +141,32 @@ bool UModularRig::Execute_Internal(const FName& InEventName)
 				// Make sure the hierarchy has the correct element redirector from this module rig
 				FRigHierarchyRedirectorGuard ElementRedirectorGuard(Rig);
 
+				FRigVMExtendedExecuteContext& RigExtendedExecuteContext= Rig->GetRigVMExtendedExecuteContext();
+
 				// Make sure the hierarchy has the correct execute context with the rig module namespace
-				FRigHierarchyExecuteContextBracket ExecuteContextBracket(Hierarchy, &Rig->GetRigVMExtendedExecuteContext());
+				FRigHierarchyExecuteContextBracket ExecuteContextBracket(Hierarchy, &RigExtendedExecuteContext);
+
+				// Make sure the module's rig has the corrct user data
+				// The rig will combine the user data of the
+				// - skeleton
+				// - skeletalmesh
+				// - SkeletalMeshComponent
+				// - default control rig module
+				// - outer modular rig
+				// - external variables
+				{
+					FControlRigExecuteContext& RigPublicContext = RigExtendedExecuteContext.GetPublicDataSafe<FControlRigExecuteContext>();
+					RigPublicContext.AssetUserData.Reset();
+					if(const TArray<UAssetUserData*>* ControlRigUserDataArray = Rig->GetAssetUserDataArray())
+					{
+					   for(const UAssetUserData* ControlRigUserData : *ControlRigUserDataArray)
+					   {
+						   RigPublicContext.AssetUserData.Add(ControlRigUserData);
+					   }
+					}
+					RigPublicContext.AssetUserData.Remove(nullptr);
+				}
+				
 				
 				Rig->Execute_Internal(InEventName);
 			}
