Sequencer: Fix default interpolation not being used on baking/loading animation onto control rigs.

diff --git a/Engine/Plugins/Animation/ControlRig/Source/ControlRig/Public/Sequencer/MovieSceneControlRigParameterSection.h b/Engine/Plugins/Animation/ControlRig/Source/ControlRig/Public/Sequencer/MovieSceneControlRigParameterSection.h
--- a/Engine/Plugins/Animation/ControlRig/Source/ControlRig/Public/Sequencer/MovieSceneControlRigParameterSection.h
+++ b/Engine/Plugins/Animation/ControlRig/Source/ControlRig/Public/Sequencer/MovieSceneControlRigParameterSection.h
@@ -315,7 +315,7 @@ public:
 
 	//Function to load an Anim Sequence into this section. It will automatically resize to the section size.
 	//Will return false if fails or is canceled
-	virtual bool LoadAnimSequenceIntoThisSection(UAnimSequence* Sequence, UMovieScene* MovieScene, UObject* BoundObject, bool bKeyReduce, float Tolerance, FFrameNumber InStartFrame = 0, EMovieSceneKeyInterpolation InInterolation = EMovieSceneKeyInterpolation::SmartAuto);
+	virtual bool LoadAnimSequenceIntoThisSection(UAnimSequence* Sequence, UMovieScene* MovieScene, UObject* BoundObject, bool bKeyReduce, float Tolerance, FFrameNumber InStartFrame , EMovieSceneKeyInterpolation InInterolation);
 #endif
 	const TArray<bool>& GetControlsMask() 
 	{
diff --git a/Engine/Plugins/Animation/ControlRig/Source/ControlRigEditor/Private/ControlRigEditorModule.cpp b/Engine/Plugins/Animation/ControlRig/Source/ControlRigEditor/Private/ControlRigEditorModule.cpp
--- a/Engine/Plugins/Animation/ControlRig/Source/ControlRigEditor/Private/ControlRigEditorModule.cpp
+++ b/Engine/Plugins/Animation/ControlRig/Source/ControlRigEditor/Private/ControlRigEditorModule.cpp
@@ -1089,8 +1089,9 @@ void FControlRigEditorModule::BakeToControlRig(UClass* ControlRigClass, UAnimSeq
 				{
 					if (ParamSection)
 					{
+						EMovieSceneKeyInterpolation DefaultInterpolation = WeakSequencer.Pin()->GetKeyInterpolation();
 						ParamSection->LoadAnimSequenceIntoThisSection(AnimSequence, MovieScene, SkelMeshComp, bKeyReduce,
-							KeyReduceTolerance);
+							KeyReduceTolerance, FFrameNumber(0), DefaultInterpolation);
 					}
 					WeakSequencer.Pin()->EmptySelection();
 					WeakSequencer.Pin()->SelectSection(ParamSection);
diff --git a/Engine/Plugins/Animation/ControlRig/Source/ControlRigEditor/Private/ControlRigSequencerEditorLibrary.cpp b/Engine/Plugins/Animation/ControlRig/Source/ControlRigEditor/Private/ControlRigSequencerEditorLibrary.cpp
--- a/Engine/Plugins/Animation/ControlRig/Source/ControlRigEditor/Private/ControlRigSequencerEditorLibrary.cpp
+++ b/Engine/Plugins/Animation/ControlRig/Source/ControlRigEditor/Private/ControlRigSequencerEditorLibrary.cpp
@@ -1278,15 +1278,18 @@ bool UControlRigSequencerEditorLibrary::BakeToControlRig(UWorld* World, ULevelSe
 					Track->SetTrackName(FName(*ObjectName));
 					Track->SetDisplayName(FText::FromString(ObjectName));
 
+					EMovieSceneKeyInterpolation DefaultInterpolation = EMovieSceneKeyInterpolation::SmartAuto;
 					if (WeakSequencer.IsValid())
 					{
 						WeakSequencer.Pin()->EmptySelection();
 						WeakSequencer.Pin()->SelectSection(NewSection);
 						WeakSequencer.Pin()->ThrobSectionSelection();
 						WeakSequencer.Pin()->NotifyMovieSceneDataChanged(EMovieSceneDataChangeType::MovieSceneStructureItemAdded);
+						DefaultInterpolation = WeakSequencer.Pin()->GetKeyInterpolation();
 					}
+
 					ParamSection->LoadAnimSequenceIntoThisSection(TempAnimSequence, MovieScene, SkeletalMeshComp,
-						bReduceKeys, Tolerance);
+						bReduceKeys, Tolerance, FFrameNumber(0), DefaultInterpolation);
 
 					//Turn Off Any Skeletal Animation Tracks
 					UMovieSceneSkeletalAnimationTrack* SkelTrack = Cast<UMovieSceneSkeletalAnimationTrack>(MovieScene->FindTrack(UMovieSceneSkeletalAnimationTrack::StaticClass(), Binding.BindingID, NAME_None));
@@ -1334,7 +1337,7 @@ bool UControlRigSequencerEditorLibrary::BakeToControlRig(UWorld* World, ULevelSe
 }
 
 bool UControlRigSequencerEditorLibrary::LoadAnimSequenceIntoControlRigSection(UMovieSceneSection* MovieSceneSection, UAnimSequence* AnimSequence, USkeletalMeshComponent* SkelMeshComp,
-	FFrameNumber InStartFrame, ESequenceTimeUnit TimeUnit,bool bKeyReduce, float Tolerance)
+	FFrameNumber InStartFrame, ESequenceTimeUnit TimeUnit,bool bKeyReduce, float Tolerance, EMovieSceneKeyInterpolation Interpolation)
 {
 	if (MovieSceneSection == nullptr || AnimSequence == nullptr || SkelMeshComp == nullptr)
 	{
@@ -1351,7 +1354,7 @@ bool UControlRigSequencerEditorLibrary::LoadAnimSequenceIntoControlRigSection(UM
 		{
 			InStartFrame = FFrameRate::TransformTime(FFrameTime(InStartFrame, 0), MovieScene->GetDisplayRate(),MovieScene->GetTickResolution()).RoundToFrame();
 		}
-		return Section->LoadAnimSequenceIntoThisSection(AnimSequence, MovieScene, SkelMeshComp, bKeyReduce, Tolerance, InStartFrame);
+		return Section->LoadAnimSequenceIntoThisSection(AnimSequence, MovieScene, SkelMeshComp, bKeyReduce, Tolerance, InStartFrame, Interpolation);
 	}
 	return false;
 }
diff --git a/Engine/Plugins/Animation/ControlRig/Source/ControlRigEditor/Private/Sequencer/ControlRigParameterTrackEditor.cpp b/Engine/Plugins/Animation/ControlRig/Source/ControlRigEditor/Private/Sequencer/ControlRigParameterTrackEditor.cpp
--- a/Engine/Plugins/Animation/ControlRig/Source/ControlRigEditor/Private/Sequencer/ControlRigParameterTrackEditor.cpp
+++ b/Engine/Plugins/Animation/ControlRig/Source/ControlRigEditor/Private/Sequencer/ControlRigParameterTrackEditor.cpp
@@ -1079,8 +1079,9 @@ void FControlRigParameterTrackEditor::BakeToControlRig(UClass* InClass, FGuid Ob
 					GetSequencer()->SelectSection(NewSection);
 					GetSequencer()->ThrobSectionSelection();
 					GetSequencer()->NotifyMovieSceneDataChanged(EMovieSceneDataChangeType::MovieSceneStructureItemAdded);
+					EMovieSceneKeyInterpolation DefaultInterpolation = SequencerParent->GetKeyInterpolation();
 					ParamSection->LoadAnimSequenceIntoThisSection(TempAnimSequence, OwnerMovieScene, SkelMeshComp,
-						BakeSettings->bReduceKeys, BakeSettings->Tolerance);
+						BakeSettings->bReduceKeys, BakeSettings->Tolerance, FFrameNumber(0), DefaultInterpolation);
 
 					//Turn Off Any Skeletal Animation Tracks
 					TArray<UMovieSceneSkeletalAnimationTrack*> SkelAnimationTracks;
@@ -5246,7 +5247,8 @@ void FControlRigParameterSection::OnAnimationAssetSelectedForFK(const FAssetData
 			Section->Modify();
 			UMovieScene* MovieScene = SequencerPtr->GetFocusedMovieSceneSequence()->GetMovieScene();
 			FFrameNumber StartFrame = SequencerPtr->GetLocalTime().Time.GetFrame();
-			if (!Section->LoadAnimSequenceIntoThisSection(AnimSequence, MovieScene, SkelMeshComp, false, 0.1f, StartFrame))
+			EMovieSceneKeyInterpolation DefaultInterpolation = SequencerPtr->GetKeyInterpolation();
+			if (!Section->LoadAnimSequenceIntoThisSection(AnimSequence, MovieScene, SkelMeshComp, false, 0.1f, StartFrame, DefaultInterpolation))
 			{
 				Transaction.Cancel();
 			}
diff --git a/Engine/Plugins/Animation/ControlRig/Source/ControlRigEditor/Public/ControlRigSequencerEditorLibrary.h b/Engine/Plugins/Animation/ControlRig/Source/ControlRigEditor/Public/ControlRigSequencerEditorLibrary.h
--- a/Engine/Plugins/Animation/ControlRig/Source/ControlRigEditor/Public/ControlRigSequencerEditorLibrary.h
+++ b/Engine/Plugins/Animation/ControlRig/Source/ControlRigEditor/Public/ControlRigSequencerEditorLibrary.h
@@ -121,12 +121,13 @@ public:
 	* @param TimeUnit Unit for all frame and time values, either in display rate or tick resolution
 	* @param bKeyReduce If true do key reduction based upon Tolerance, if false don't
 	* @param Tolerance If reducing keys, tolerance about which keys will be removed, smaller tolerance, more keys usually.
-
+	* @param Interpolation The key interpolation type to set the keys, defaults to EMovieSceneKeyInterpolation::SmartAuto
 	* @return returns True if successful, False otherwise
 	*/
 	UFUNCTION(BlueprintCallable, Category = "Editor Scripting | Sequencer Tools | Control Rig")
 	static bool LoadAnimSequenceIntoControlRigSection(UMovieSceneSection* MovieSceneSection, UAnimSequence* AnimSequence, USkeletalMeshComponent* SkelMeshComp,
-		FFrameNumber InStartFrame, ESequenceTimeUnit TimeUnit = ESequenceTimeUnit::DisplayRate, bool bKeyReduce = false, float Tolerance = 0.001f);
+		FFrameNumber InStartFrame, ESequenceTimeUnit TimeUnit = ESequenceTimeUnit::DisplayRate, bool bKeyReduce = false, float Tolerance = 0.001f,
+		EMovieSceneKeyInterpolation Interpolation = EMovieSceneKeyInterpolation::SmartAuto);
 
 	/**
 	* Bake the current animation in the binding to a Control Rig track
diff --git a/Engine/Plugins/Importers/USDImporter/Source/USDStage/Private/USDLevelSequenceHelper.cpp b/Engine/Plugins/Importers/USDImporter/Source/USDStage/Private/USDLevelSequenceHelper.cpp
--- a/Engine/Plugins/Importers/USDImporter/Source/USDStage/Private/USDLevelSequenceHelper.cpp
+++ b/Engine/Plugins/Importers/USDImporter/Source/USDStage/Private/USDLevelSequenceHelper.cpp
@@ -368,7 +368,8 @@ namespace UsdLevelSequenceHelperImpl
 			// of the playback range, which we do here
 			ControlRigSectionStartFrame -= UE::MovieScene::DiscreteInclusiveLower(MovieScene->GetPlaybackRange());
 
-			ParamSection->LoadAnimSequenceIntoThisSection(AnimSequence, MovieScene, SkeletalMeshComp, bReduceKeys, Tolerance, ControlRigSectionStartFrame);
+			ParamSection->LoadAnimSequenceIntoThisSection(AnimSequence, MovieScene, SkeletalMeshComp, bReduceKeys, Tolerance, ControlRigSectionStartFrame,
+				EMovieSceneKeyInterpolation::SmartAuto);
 
 			// Disable Skeletal Animation Tracks
 			if (SkelTrack)
