Sequencer: Control Rig: We only set the blend weight on the control rig layer from an absolute control rig, the addtives don't effect it, but wil blend off of the absolute

diff --git a/Engine/Plugins/Animation/ControlRig/Source/ControlRig/Private/Sequencer/MovieSceneControlRigParameterTemplate.cpp b/Engine/Plugins/Animation/ControlRig/Source/ControlRig/Private/Sequencer/MovieSceneControlRigParameterTemplate.cpp
--- a/Engine/Plugins/Animation/ControlRig/Source/ControlRig/Private/Sequencer/MovieSceneControlRigParameterTemplate.cpp
+++ b/Engine/Plugins/Animation/ControlRig/Source/ControlRig/Private/Sequencer/MovieSceneControlRigParameterTemplate.cpp
@@ -1287,7 +1287,11 @@ struct FControlRigParameterExecutionToken : IMovieSceneExecutionToken
 							FControlRigIOSettings InputSettings;
 							InputSettings.bUpdateCurves = true;
 							InputSettings.bUpdatePose = true;
-							AnimInstance->UpdateControlRigTrack(ControlRig->GetUniqueID(), Weight, InputSettings, true);
+							//this is not great but assumes we have 1 absolute track that will be used for weighting
+							if (Section->GetBlendType() == EMovieSceneBlendType::Absolute)
+							{
+								AnimInstance->UpdateControlRigTrack(ControlRig->GetUniqueID(), Weight, InputSettings, true);
+							}
 						}
 					}
 					else
