Control Rig: Add Pre/Post forward solve events

diff --git a/Engine/Plugins/Animation/ControlRig/Source/ControlRig/Private/ControlRig.cpp b/Engine/Plugins/Animation/ControlRig/Source/ControlRig/Private/ControlRig.cpp
--- a/Engine/Plugins/Animation/ControlRig/Source/ControlRig/Private/ControlRig.cpp
+++ b/Engine/Plugins/Animation/ControlRig/Source/ControlRig/Private/ControlRig.cpp
@@ -631,6 +631,12 @@ bool UControlRig::Execute(const FName& InEventName)
 		return false;
 	}
 
+	// Only top-level rigs should execute this function.
+	// Rig modules/nested rigs should run through Execute_Internal
+	ensureMsgf(GetTypedOuter<UControlRig>() == nullptr, TEXT("UControlRig::Execute running from a nested rig in %s"), *GetPackage()->GetPathName());
+	ensureMsgf(InEventName != FRigUnit_PreBeginExecution::EventName &&
+						InEventName != FRigUnit_PostBeginExecution::EventName, TEXT("Requested execution of invalid event %s on top level rig in %s"), *InEventName.ToString(), *GetPackage()->GetPathName());
+
 	bool bJustRanInit = false;
 	if(bRequiresInitExecution)
 	{
@@ -1016,10 +1022,17 @@ bool UControlRig::Execute(const FName& InEventName)
 				FControlRigBracketScope BracketScope(PreForwardsSolveBracket);
 				PreForwardsSolveEvent.Broadcast(this, FRigUnit_BeginExecution::EventName);
 			}
+
+			Execute_Internal(FRigUnit_PreBeginExecution::EventName);
 		}
 
 		bSuccess = Execute_Internal(InEventName);
 
+		if (bIsForwardSolve)
+		{
+			Execute_Internal(FRigUnit_PostBeginExecution::EventName);
+		}
+
 #if WITH_EDITOR
 		if (bEnableAnimAttributeTrace && ExternalAnimAttributeContainer != nullptr)
 		{
diff --git a/Engine/Plugins/Animation/ControlRig/Source/ControlRig/Private/ModularRig.cpp b/Engine/Plugins/Animation/ControlRig/Source/ControlRig/Private/ModularRig.cpp
--- a/Engine/Plugins/Animation/ControlRig/Source/ControlRig/Private/ModularRig.cpp
+++ b/Engine/Plugins/Animation/ControlRig/Source/ControlRig/Private/ModularRig.cpp
@@ -137,6 +137,11 @@ bool UModularRig::Execute_Internal(const FName& InEventName)
 			if (Module->Rig.IsValid())
 			{
 				UControlRig* Rig = Module->Rig.Get();
+
+				if (!Rig->SupportsEvent(InEventName))
+				{
+					return true;
+				}
 				
 				// Make sure the hierarchy has the correct element redirector from this module rig
 				FRigHierarchyRedirectorGuard ElementRedirectorGuard(Rig);
diff --git a/Engine/Plugins/Animation/ControlRig/Source/ControlRig/Private/Units/Execution/RigUnit_BeginExecution.cpp b/Engine/Plugins/Animation/ControlRig/Source/ControlRig/Private/Units/Execution/RigUnit_BeginExecution.cpp
--- a/Engine/Plugins/Animation/ControlRig/Source/ControlRig/Private/Units/Execution/RigUnit_BeginExecution.cpp
+++ b/Engine/Plugins/Animation/ControlRig/Source/ControlRig/Private/Units/Execution/RigUnit_BeginExecution.cpp
@@ -6,6 +6,8 @@
 #include UE_INLINE_GENERATED_CPP_BY_NAME(RigUnit_BeginExecution)
 
 FName FRigUnit_BeginExecution::EventName = TEXT("Forwards Solve");
+FName FRigUnit_PreBeginExecution::EventName = TEXT("Pre Forwards Solve");
+FName FRigUnit_PostBeginExecution::EventName = TEXT("Post Forwards Solve");
 
 FRigUnit_BeginExecution_Execute()
 {
@@ -13,3 +15,15 @@ FRigUnit_BeginExecution_Execute()
 	ExecuteContext.SetEventName(FRigUnit_BeginExecution::EventName);
 }
 
+FRigUnit_PreBeginExecution_Execute()
+{
+	DECLARE_SCOPE_HIERARCHICAL_COUNTER_RIGUNIT()
+	ExecuteContext.SetEventName(FRigUnit_PreBeginExecution::EventName);
+}
+
+FRigUnit_PostBeginExecution_Execute()
+{
+	DECLARE_SCOPE_HIERARCHICAL_COUNTER_RIGUNIT()
+	ExecuteContext.SetEventName(FRigUnit_PostBeginExecution::EventName);
+}
+
diff --git a/Engine/Plugins/Animation/ControlRig/Source/ControlRig/Public/Units/Execution/RigUnit_BeginExecution.h b/Engine/Plugins/Animation/ControlRig/Source/ControlRig/Public/Units/Execution/RigUnit_BeginExecution.h
--- a/Engine/Plugins/Animation/ControlRig/Source/ControlRig/Public/Units/Execution/RigUnit_BeginExecution.h
+++ b/Engine/Plugins/Animation/ControlRig/Source/ControlRig/Public/Units/Execution/RigUnit_BeginExecution.h
@@ -26,3 +26,45 @@ struct CONTROLRIG_API FRigUnit_BeginExecution : public FRigUnit
 
 	static FName EventName;
 };
+
+/**
+ * Event always executed before the forward solve
+ */
+USTRUCT(meta=(DisplayName="Pre Forwards Solve", Category="Events", NodeColor="1, 0, 0", Keywords="Begin,Update,Tick,PreForward,Event"))
+struct CONTROLRIG_API FRigUnit_PreBeginExecution : public FRigUnit
+{
+	GENERATED_BODY()
+
+	RIGVM_METHOD()
+	virtual void Execute() override;
+
+	virtual FName GetEventName() const override { return EventName; }
+	virtual bool CanOnlyExistOnce() const override { return true; }
+
+	// The execution result
+	UPROPERTY(EditAnywhere, Transient, DisplayName = "Execute", Category = "BeginExecution", meta = (Output))
+	FControlRigExecuteContext ExecuteContext;
+
+	static FName EventName;
+};
+
+/**
+ * Event always executed after the forward solve
+ */
+USTRUCT(meta=(DisplayName="Post Forwards Solve", Category="Events", NodeColor="1, 0, 0", Keywords="Begin,Update,Tick,PostForward,Event"))
+struct CONTROLRIG_API FRigUnit_PostBeginExecution : public FRigUnit
+{
+	GENERATED_BODY()
+
+	RIGVM_METHOD()
+	virtual void Execute() override;
+
+	virtual FName GetEventName() const override { return EventName; }
+	virtual bool CanOnlyExistOnce() const override { return true; }
+
+	// The execution result
+	UPROPERTY(EditAnywhere, Transient, DisplayName = "Execute", Category = "BeginExecution", meta = (Output))
+	FControlRigExecuteContext ExecuteContext;
+
+	static FName EventName;
+};
diff --git a/Engine/Plugins/Animation/ControlRig/Source/ControlRigDeveloper/Private/Graph/ControlRigGraphSchema.cpp b/Engine/Plugins/Animation/ControlRig/Source/ControlRigDeveloper/Private/Graph/ControlRigGraphSchema.cpp
--- a/Engine/Plugins/Animation/ControlRig/Source/ControlRigDeveloper/Private/Graph/ControlRigGraphSchema.cpp
+++ b/Engine/Plugins/Animation/ControlRig/Source/ControlRigDeveloper/Private/Graph/ControlRigGraphSchema.cpp
@@ -50,6 +50,8 @@ bool UControlRigGraphSchema::IsRigVMDefaultEvent(const FName& InEventName) const
 	}
 	
 	return InEventName == FRigUnit_BeginExecution::EventName ||
+		InEventName == FRigUnit_PreBeginExecution::EventName ||
+		InEventName == FRigUnit_PostBeginExecution::EventName ||
 		InEventName == FRigUnit_InverseExecution::EventName ||
 		InEventName == FRigUnit_PrepareForExecution::EventName ||
 		InEventName == FRigUnit_InteractionExecution::EventName;
diff --git a/Engine/Plugins/Animation/ControlRig/Source/ControlRigEditor/Private/ControlRigAnimGraphDetails.cpp b/Engine/Plugins/Animation/ControlRig/Source/ControlRigEditor/Private/ControlRigAnimGraphDetails.cpp
--- a/Engine/Plugins/Animation/ControlRig/Source/ControlRigEditor/Private/ControlRigAnimGraphDetails.cpp
+++ b/Engine/Plugins/Animation/ControlRig/Source/ControlRigEditor/Private/ControlRigAnimGraphDetails.cpp
@@ -9,6 +9,7 @@
 #include "PropertyCustomizationHelpers.h"
 #include "AnimGraphNode_ControlRig.h"
 #include "Algo/Transform.h"
+#include "Units/Execution/RigUnit_BeginExecution.h"
 
 #define LOCTEXT_NAMESPACE "ControlRigAnimGraphDetails"
 static const FText ControlRigAnimDetailsMultipleValues = LOCTEXT("MultipleValues", "Multiple Values");
@@ -134,7 +135,13 @@ void FControlRigAnimNodeEventNameDetails::UpdateEntryNameList()
 		{
 			if(const UControlRig* CDO = Cast<UControlRig>(Class->GetDefaultObject(true)))
 			{
-				Algo::Transform(CDO->GetSupportedEvents(), EntryNameList,[](const FName& InEntryName)
+				TArray<FName> SupportedEvents = CDO->GetSupportedEvents();
+
+				// Remove Pre/Post forward solve
+				SupportedEvents.Remove(FRigUnit_PreBeginExecution::EventName);
+				SupportedEvents.Remove(FRigUnit_PostBeginExecution::EventName);
+				
+				Algo::Transform(SupportedEvents, EntryNameList,[](const FName& InEntryName)
 				{
 					return MakeShareable(new FString(InEntryName.ToString()));
 				});
