Control Rig: Prevent manual creation of invalid rig elements in rig module
#rnx

diff --git a/Engine/Plugins/Animation/ControlRig/Source/ControlRig/Private/Rigs/RigHierarchyController.cpp b/Engine/Plugins/Animation/ControlRig/Source/ControlRig/Private/Rigs/RigHierarchyController.cpp
--- a/Engine/Plugins/Animation/ControlRig/Source/ControlRig/Private/Rigs/RigHierarchyController.cpp
+++ b/Engine/Plugins/Animation/ControlRig/Source/ControlRig/Private/Rigs/RigHierarchyController.cpp
@@ -1194,6 +1194,11 @@ public:
 };
 
 TArray<FRigElementKey> URigHierarchyController::ImportFromText(FString InContent, bool bReplaceExistingElements, bool bSelectNewElements, bool bSetupUndo, bool bPrintPythonCommands)
+{
+	return ImportFromText(InContent, ERigElementType::All, bReplaceExistingElements, bSelectNewElements, bSetupUndo, bPrintPythonCommands);
+}
+
+TArray<FRigElementKey> URigHierarchyController::ImportFromText(FString InContent, ERigElementType InAllowedTypes, bool bReplaceExistingElements, bool bSelectNewElements, bool bSetupUndo, bool bPrintPythonCommands)
 {
 	TArray<FRigElementKey> PastedKeys;
 	if(!IsValid())
@@ -1214,6 +1219,23 @@ TArray<FRigElementKey> URigHierarchyController::ImportFromText(FString InContent
 		// check if this is a copy & paste buffer from pre-5.0
 		if(Data.Contents.Num() > 0)
 		{
+			const int32 OriginalNumElements = Data.Elements.Num();
+			for (int32 i=0; i<Data.Types.Num(); )
+			{
+				if (((uint8)InAllowedTypes & (uint8)Data.Types[i]) == 0)
+				{
+					Data.Contents.RemoveAt(i);
+					Data.Types.RemoveAt(i);
+					Data.LocalTransforms.RemoveAt(i);
+					Data.GlobalTransforms.RemoveAt(i);
+					continue;
+				}
+				++i;
+			}
+			if (OriginalNumElements > Data.Types.Num())
+			{
+				ReportAndNotifyErrorf(TEXT("Some elements were not allowed to be pasted."));
+			}
 			FRigHierarchyContainer OldHierarchy;
 			if(OldHierarchy.ImportFromText(Data).Num() > 0)
 			{
@@ -1224,11 +1246,21 @@ TArray<FRigElementKey> URigHierarchyController::ImportFromText(FString InContent
 		return PastedKeys;
 	}
 
+	const int32 OriginalNumElements = Data.Elements.Num();
+	Data.Elements = Data.Elements.FilterByPredicate([InAllowedTypes](const FRigHierarchyCopyPasteContentPerElement& Element)
+	{
+		return ((uint8)InAllowedTypes & (uint8)Element.Key.Type) != 0;
+	});
+	if (OriginalNumElements > Data.Elements.Num())
+	{
+		ReportAndNotifyErrorf(TEXT("Some elements were not allowed to be pasted."));
+	}
+
 #if WITH_EDITOR
 	TSharedPtr<FScopedTransaction> TransactionPtr;
 	if(bSetupUndo)
 	{
-		TransactionPtr = MakeShared<FScopedTransaction>(NSLOCTEXT("RigHierarchyController", "Add Bone", "Add Bone"));
+		TransactionPtr = MakeShared<FScopedTransaction>(NSLOCTEXT("RigHierarchyController", "Add Elements", "Add Elements"));
 		Hierarchy->Modify();
 	}
 #endif
diff --git a/Engine/Plugins/Animation/ControlRig/Source/ControlRig/Public/Rigs/RigHierarchyController.h b/Engine/Plugins/Animation/ControlRig/Source/ControlRig/Public/Rigs/RigHierarchyController.h
--- a/Engine/Plugins/Animation/ControlRig/Source/ControlRig/Public/Rigs/RigHierarchyController.h
+++ b/Engine/Plugins/Animation/ControlRig/Source/ControlRig/Public/Rigs/RigHierarchyController.h
@@ -402,6 +402,14 @@ public:
 		bool bSetupUndo = false,
 		bool bPrintPythonCommands = false);
 
+	TArray<FRigElementKey> ImportFromText(
+		FString InContent,
+		ERigElementType InAllowedTypes,
+		bool bReplaceExistingElements = false,
+		bool bSelectNewElements = true,
+		bool bSetupUndo = false,
+		bool bPrintPythonCommands = false);
+
 	/**
 	* Imports the content of a RigHierachyContainer (the hierarchy v1 pre 5.0)
 	* This is used for backwards compatbility only during load and does not support undo.
diff --git a/Engine/Plugins/Animation/ControlRig/Source/ControlRigEditor/Private/Editor/SRigHierarchy.cpp b/Engine/Plugins/Animation/ControlRig/Source/ControlRigEditor/Private/Editor/SRigHierarchy.cpp
--- a/Engine/Plugins/Animation/ControlRig/Source/ControlRigEditor/Private/Editor/SRigHierarchy.cpp
+++ b/Engine/Plugins/Animation/ControlRig/Source/ControlRigEditor/Private/Editor/SRigHierarchy.cpp
@@ -316,27 +316,27 @@ void SRigHierarchy::BindCommands()
 
 	CommandList->MapAction(Commands.AddBoneItem,
 		FExecuteAction::CreateSP(this, &SRigHierarchy::HandleNewItem, ERigElementType::Bone, false),
-		FCanExecuteAction::CreateSP(this, &SRigHierarchy::IsNonProceduralElementSelected));
+		FCanExecuteAction::CreateSP(this, &SRigHierarchy::CanAddElement, ERigElementType::Bone));
 
 	CommandList->MapAction(Commands.AddControlItem,
 		FExecuteAction::CreateSP(this, &SRigHierarchy::HandleNewItem, ERigElementType::Control, false),
-		FCanExecuteAction::CreateSP(this, &SRigHierarchy::IsNonProceduralElementSelected));
+		FCanExecuteAction::CreateSP(this, &SRigHierarchy::CanAddElement, ERigElementType::Control));
 
 	CommandList->MapAction(Commands.AddAnimationChannelItem,
 		FExecuteAction::CreateSP(this, &SRigHierarchy::HandleNewItem, ERigElementType::Control, true),
-		FCanExecuteAction::CreateSP(this, &SRigHierarchy::IsControlSelected, false));
+		FCanExecuteAction::CreateSP(this, &SRigHierarchy::CanAddAnimationChannel));
 
 	CommandList->MapAction(Commands.AddNullItem,
 		FExecuteAction::CreateSP(this, &SRigHierarchy::HandleNewItem, ERigElementType::Null, false),
-		FCanExecuteAction::CreateSP(this, &SRigHierarchy::IsNonProceduralElementSelected));
+		FCanExecuteAction::CreateSP(this, &SRigHierarchy::CanAddElement, ERigElementType::Null));
 
 	CommandList->MapAction(Commands.AddConnectorItem,
 		FExecuteAction::CreateSP(this, &SRigHierarchy::HandleNewItem, ERigElementType::Connector, false),
-		FCanExecuteAction::CreateSP(this, &SRigHierarchy::IsNonProceduralElementSelected));
+		FCanExecuteAction::CreateSP(this, &SRigHierarchy::CanAddElement, ERigElementType::Connector));
 
 	CommandList->MapAction(Commands.AddSocketItem,
 		FExecuteAction::CreateSP(this, &SRigHierarchy::HandleNewItem, ERigElementType::Socket, false),
-		FCanExecuteAction::CreateSP(this, &SRigHierarchy::IsNonProceduralElementSelected));
+		FCanExecuteAction::CreateSP(this, &SRigHierarchy::CanAddElement, ERigElementType::Socket));
 
 	CommandList->MapAction(Commands.DuplicateItem,
 		FExecuteAction::CreateSP(this, &SRigHierarchy::HandleDuplicateItem),
@@ -1779,6 +1779,11 @@ bool SRigHierarchy::IsControlOrNullSelected(bool bIncludeProcedural) const
 bool SRigHierarchy::IsProceduralElementSelected() const
 {
 	TArray<FRigElementKey> SelectedKeys = GetSelectedKeys();
+	if (SelectedKeys.IsEmpty())
+	{
+		return false;
+	}
+	
 	for (const FRigElementKey& SelectedKey : SelectedKeys)
 	{
 		if(!GetHierarchy()->IsProcedural(SelectedKey))
@@ -1792,6 +1797,11 @@ bool SRigHierarchy::IsProceduralElementSelected() const
 bool SRigHierarchy::IsNonProceduralElementSelected() const
 {
 	TArray<FRigElementKey> SelectedKeys = GetSelectedKeys();
+	if (SelectedKeys.IsEmpty())
+	{
+		return false;
+	}
+	
 	for (const FRigElementKey& SelectedKey : SelectedKeys)
 	{
 		if(GetHierarchy()->IsProcedural(SelectedKey))
@@ -1802,6 +1812,32 @@ bool SRigHierarchy::IsNonProceduralElementSelected() const
 	return true;
 }
 
+bool SRigHierarchy::CanAddElement(const ERigElementType ElementType) const
+{
+	if (IsProceduralElementSelected())
+	{
+		return false;
+	}
+
+	// Always allow connectors 
+	if (ElementType == ERigElementType::Connector)
+	{
+		return true;
+	}
+	
+	return !ControlRigBlueprint->IsControlRigModule();
+}
+
+bool SRigHierarchy::CanAddAnimationChannel() const
+{
+	if (!IsControlSelected(false))
+	{
+		return false;
+	}
+
+	return !ControlRigBlueprint->IsControlRigModule();
+}
+
 void SRigHierarchy::HandleDeleteItem()
 {
 	if(!ControlRigEditor.IsValid())
@@ -2041,7 +2077,22 @@ void SRigHierarchy::HandleNewItem(ERigElementType InElementType, bool bIsAnimati
 /** Check whether we can deleting the selected item(s) */
 bool SRigHierarchy::CanDuplicateItem() const
 {
-	return IsMultiSelected(false);
+	if (!IsMultiSelected(false))
+	{
+		return false;
+	}
+
+	if (ControlRigBlueprint->IsControlRigModule())
+	{
+		bool bAnyNonConnector = GetSelectedKeys().ContainsByPredicate([](const FRigElementKey& Key)
+		{
+			return Key.Type != ERigElementType::Connector;
+		});
+		
+		return !bAnyNonConnector;
+	}
+
+	return true;
 }
 
 /** Duplicate Item */
@@ -2247,7 +2298,8 @@ void SRigHierarchy::HandlePasteItems()
 		URigHierarchyController* Controller = Hierarchy->GetController(true);
 		check(Controller);
 
-		Controller->ImportFromText(Content, false, true, true, true);
+		ERigElementType AllowedTypes = ControlRigBlueprint->IsControlRigModule() ? ERigElementType::Connector : ERigElementType::All;
+		Controller->ImportFromText(Content, AllowedTypes, false, true, true, true);
 	}
 
 	//ControlRigBlueprint->PropagateHierarchyFromBPToInstances();
@@ -2515,8 +2567,16 @@ TOptional<EItemDropZone> SRigHierarchy::OnCanAcceptDrop(const FDragDropEvent& Dr
 			}
 			case ERigElementType::Socket:
 			{
-				// You cannot parent anything under a socket
-				return InvalidDropZone;
+				// Only connectors can be parented under a socket
+				if (RigDragDropOp->IsDraggingSingleConnector())
+				{
+					ReturnDropZone = DropZone;
+				}
+				else
+				{
+					return InvalidDropZone;
+				}
+				break;
 			}
 			default:
 			{
diff --git a/Engine/Plugins/Animation/ControlRig/Source/ControlRigEditor/Private/Editor/SRigHierarchy.h b/Engine/Plugins/Animation/ControlRig/Source/ControlRigEditor/Private/Editor/SRigHierarchy.h
--- a/Engine/Plugins/Animation/ControlRig/Source/ControlRigEditor/Private/Editor/SRigHierarchy.h
+++ b/Engine/Plugins/Animation/ControlRig/Source/ControlRigEditor/Private/Editor/SRigHierarchy.h
@@ -193,6 +193,8 @@ private:
 	bool IsControlOrNullSelected(bool bIncludeProcedural) const;
 	bool IsProceduralElementSelected() const;
 	bool IsNonProceduralElementSelected() const;
+	bool CanAddElement(const ERigElementType ElementType) const;
+	bool CanAddAnimationChannel() const;
 
 	URigHierarchy* GetHierarchy() const;
 	URigHierarchy* GetDefaultHierarchy() const;
