Control Rig: Modular rig cleanup
#rnx

diff --git a/Engine/Plugins/Animation/ControlRig/Source/ControlRigDeveloper/Private/ControlRigBlueprint.cpp b/Engine/Plugins/Animation/ControlRig/Source/ControlRigDeveloper/Private/ControlRigBlueprint.cpp
--- a/Engine/Plugins/Animation/ControlRig/Source/ControlRigDeveloper/Private/ControlRigBlueprint.cpp
+++ b/Engine/Plugins/Animation/ControlRig/Source/ControlRigDeveloper/Private/ControlRigBlueprint.cpp
@@ -2304,13 +2304,6 @@ void UControlRigBlueprint::HandleRigModulesModified(EModularRigNotification InNo
 						{
 							Controller->AddConnector(*Connector.Name, Connector.Settings);
 						}
-			
-						// todo: copy the sockets
-						// TArray<FRigSocketElement*> Sockets = DefaultModule->GetHierarchy()->GetElementsOfType<FRigSocketElement>();
-						// for (FRigSocketElement* Socket : Sockets)
-						// {
-						// 	Controller->CopySocket(NewModule->GetPath(), Socket);
-						// }
 
 						PropagateHierarchyFromBPToInstances();
 					}
@@ -2335,25 +2328,31 @@ void UControlRigBlueprint::HandleRigModulesModified(EModularRigNotification InNo
 							FRigElementKey TargetConnection;
 							FRigConnectorSettings Settings;
 						};
-						FString OldPath;
+						FString OldNamespace;
 						if (InNotification == EModularRigNotification::ModuleRenamed)
 						{
-							OldPath = (InModule->ParentPath.IsEmpty()) ? *InModule->PreviousName.ToString() : FString::Printf(TEXT("%s:%s:"), *InModule->ParentPath, *InModule->PreviousName.ToString());
+							OldNamespace = (InModule->ParentPath.IsEmpty()) ? *InModule->PreviousName.ToString() : FString::Printf(TEXT("%s:%s"), *InModule->ParentPath, *InModule->PreviousName.ToString());
 						}
 						else if(InNotification == EModularRigNotification::ModuleReparented)
 						{
-							OldPath = (InModule->PreviousParentPath.IsEmpty()) ? *InModule->Name.ToString() : FString::Printf(TEXT("%s:%s:"), *InModule->PreviousParentPath, *InModule->Name.ToString());
+							OldNamespace = (InModule->PreviousParentPath.IsEmpty()) ? *InModule->Name.ToString() : FString::Printf(TEXT("%s:%s"), *InModule->PreviousParentPath, *InModule->Name.ToString());
 						}
-						FString NewPath = (InModule->ParentPath.IsEmpty()) ? *InModule->Name.ToString() : FString::Printf(TEXT("%s:%s:"), *InModule->ParentPath, *InModule->Name.ToString());
+						FString NewNamespace = (InModule->ParentPath.IsEmpty()) ? *InModule->Name.ToString() : FString::Printf(TEXT("%s:%s"), *InModule->ParentPath, *InModule->Name.ToString());
+						OldNamespace.Append(UModularRig::NamespaceSeparator);
+						NewNamespace.Append(UModularRig::NamespaceSeparator);
+						
 						TArray<FRigElementKey> Connectors = Controller->GetHierarchy()->GetKeysOfType<FRigConnectorElement>();
 						TMap<FRigElementKey, ConnectionInfo> RenamedConnectors; // old key -> new key
 						for (const FRigElementKey& Connector : Connectors)
 						{
 							FString OldConnectorName = Connector.Name.ToString();
-							if (OldConnectorName.StartsWith(OldPath))
+							FString OldConnectorNamespace, ShortName;
+							OldConnectorName.Split(UModularRig::NamespaceSeparator, &OldConnectorNamespace, &ShortName, ESearchCase::CaseSensitive, ESearchDir::FromEnd);
+							OldConnectorNamespace.Append(UModularRig::NamespaceSeparator);
+							if (OldConnectorNamespace == OldNamespace)
 							{
 								ConnectionInfo& Info = RenamedConnectors.FindOrAdd(Connector);
-								Info.NewPath = OldConnectorName.Replace(*OldPath, *NewPath);
+								Info.NewPath = OldConnectorName.Replace(*OldNamespace, *NewNamespace);
 								Info.Settings = CastChecked<FRigConnectorElement>(Controller->GetHierarchy()->FindChecked(Connector))->Settings;
 								if (FRigElementKey* TargetKey = ConnectionMap.Find(Connector))
 								{
@@ -2377,7 +2376,7 @@ void UControlRigBlueprint::HandleRigModulesModified(EModularRigNotification InNo
 							{
 								FString Namespace, ConnectorName;
 								Pair.Value.NewPath.Split(UModularRig::NamespaceSeparator, &Namespace, &ConnectorName, ESearchCase::CaseSensitive, ESearchDir::FromEnd);
-								Namespace.AppendChar(TCHAR(':'));
+								Namespace.Append(UModularRig::NamespaceSeparator);
 								FControlRigExecuteContextRigModuleGuard RigModuleGuard(PublicContext, Namespace);
 								Controller->AddConnector(*ConnectorName, Pair.Value.Settings);
 							}
@@ -2404,12 +2403,16 @@ void UControlRigBlueprint::HandleRigModulesModified(EModularRigNotification InNo
 				{
 					Hierarchy->Modify();
 					
-					FString Namespace = InModule->GetNamespace();
+					const FString Namespace = InModule->GetNamespace();
 					TArray<FRigElementKey> Connectors = Controller->GetHierarchy()->GetKeysOfType<FRigConnectorElement>();
+					Connectors.Append(Controller->GetHierarchy()->GetKeysOfType<FRigSocketElement>());
 					for (const FRigElementKey& Connector : Connectors)
 					{
 						FString ConnectorName = Connector.Name.ToString();
-						if (ConnectorName.StartsWith(Namespace))
+						FString ConnectorNamespace, ShortName;
+						ConnectorName.Split(UModularRig::NamespaceSeparator, &ConnectorNamespace, &ShortName, ESearchCase::CaseSensitive, ESearchDir::FromEnd);
+						ConnectorNamespace.Append(UModularRig::NamespaceSeparator);
+						if (ConnectorNamespace == Namespace)
 						{
 							Controller->RemoveElement(Connector);
 						}
diff --git a/Engine/Plugins/Animation/ControlRig/Source/ControlRigEditor/Private/ControlRigBlueprintFactory.cpp b/Engine/Plugins/Animation/ControlRig/Source/ControlRigEditor/Private/ControlRigBlueprintFactory.cpp
--- a/Engine/Plugins/Animation/ControlRig/Source/ControlRigEditor/Private/ControlRigBlueprintFactory.cpp
+++ b/Engine/Plugins/Animation/ControlRig/Source/ControlRigEditor/Private/ControlRigBlueprintFactory.cpp
@@ -30,9 +30,6 @@
 
 #define LOCTEXT_NAMESPACE "ControlRigBlueprintFactory"
 
-static TAutoConsoleVariable<bool> CVarControlRigEnableModularRigs(TEXT("ControlRig.Test.EnableModularRigs"), false, TEXT("When true we will allow the creation of modular rigs."));
-
-
 /** Dialog to configure creation properties */
 class SControlRigBlueprintCreateDialog : public SCompoundWidget
 {
@@ -274,7 +271,7 @@ UControlRigBlueprintFactory::UControlRigBlueprintFactory()
 
 bool UControlRigBlueprintFactory::ConfigureProperties()
 {
-	if (CVarControlRigEnableModularRigs->GetBool())
+	if (CVarControlRigHierarchyEnableModules.GetValueOnAnyThread())
 	{
 		TSharedRef<SControlRigBlueprintCreateDialog> Dialog = SNew(SControlRigBlueprintCreateDialog);
 		return Dialog->ConfigureProperties(this);
diff --git a/Engine/Plugins/Animation/ControlRig/Source/ControlRigEditor/Private/Editor/SModularRigHierarchy.cpp b/Engine/Plugins/Animation/ControlRig/Source/ControlRigEditor/Private/Editor/SModularRigHierarchy.cpp
--- a/Engine/Plugins/Animation/ControlRig/Source/ControlRigEditor/Private/Editor/SModularRigHierarchy.cpp
+++ b/Engine/Plugins/Animation/ControlRig/Source/ControlRigEditor/Private/Editor/SModularRigHierarchy.cpp
@@ -241,9 +241,12 @@ TArray<FString> SModularRigHierarchy::GetSelectedKeys() const
 	TArray<FString> SelectedKeys;
 	for (const TSharedPtr<FModularRigTreeElement>& SelectedItem : SelectedItems)
 	{
-		if(!SelectedItem->Key.IsEmpty())
+		if (SelectedItem.IsValid())
 		{
-			SelectedKeys.AddUnique(SelectedItem->Key);
+			if(!SelectedItem->Key.IsEmpty())
+			{
+				SelectedKeys.AddUnique(SelectedItem->Key);
+			}
 		}
 	}
 
