Control Rig: Make sure to rely on the same shape name list everywhere

#rnx

diff --git a/Engine/Plugins/Animation/ControlRig/Source/ControlRigEditor/Private/ControlRigElementDetails.cpp b/Engine/Plugins/Animation/ControlRig/Source/ControlRigEditor/Private/ControlRigElementDetails.cpp
--- a/Engine/Plugins/Animation/ControlRig/Source/ControlRigEditor/Private/ControlRigElementDetails.cpp
+++ b/Engine/Plugins/Animation/ControlRig/Source/ControlRigEditor/Private/ControlRigElementDetails.cpp
@@ -3855,20 +3855,47 @@ void FRigControlElementDetails::CustomizeShape(IDetailLayoutBuilder& DetailBuild
 	
 	if (UControlRigBlueprint* Blueprint = PerElementInfos[0].GetBlueprint())
 	{
-		const bool bUseNameSpace = Blueprint->ShapeLibraries.Num() > 1;
-		for(TSoftObjectPtr<UControlRigShapeLibrary>& ShapeLibrary : Blueprint->ShapeLibraries)
+		if(UEdGraph* RootEdGraph = Blueprint->GetEdGraph(Blueprint->GetModel()))
 		{
-			if (!ShapeLibrary.IsValid())
+			if(UControlRigGraph* RigGraph = Cast<UControlRigGraph>(RootEdGraph))
 			{
-				ShapeLibrary.LoadSynchronous();
+				URigHierarchy* Hierarchy = Blueprint->Hierarchy;
+				if(UControlRig* RigBeingDebugged = Cast<UControlRig>(Blueprint->GetObjectBeingDebugged()))
+				{
+					Hierarchy = RigBeingDebugged->GetHierarchy();
+				}
+
+				const TArray<TSoftObjectPtr<UControlRigShapeLibrary>>* ShapeLibraries = &Blueprint->ShapeLibraries;
+				if(const UControlRig* DebuggedControlRig = Hierarchy->GetTypedOuter<UControlRig>())
+				{
+					ShapeLibraries = &DebuggedControlRig->GetShapeLibraries();
+				}
+				RigGraph->CacheNameLists(Hierarchy, &Blueprint->DrawContainer, *ShapeLibraries);
+				
+				if(const TArray<TSharedPtr<FRigVMStringWithTag>>* GraphShapeNameListPtr = RigGraph->GetShapeNameList())
+				{
+					ShapeNameList = *GraphShapeNameListPtr;
+				}
 			}
-			if (ShapeLibrary.IsValid())
+		}
+
+		if(ShapeNameList.IsEmpty())
+		{
+			const bool bUseNameSpace = Blueprint->ShapeLibraries.Num() > 1;
+			for(TSoftObjectPtr<UControlRigShapeLibrary>& ShapeLibrary : Blueprint->ShapeLibraries)
 			{
-				const FString NameSpace = bUseNameSpace ? ShapeLibrary->GetName() + TEXT(".") : FString();
-				ShapeNameList.Add(MakeShared<FRigVMStringWithTag>(NameSpace + ShapeLibrary->DefaultShape.ShapeName.ToString()));
-				for (const FControlRigShapeDefinition& Shape : ShapeLibrary->Shapes)
+				if (!ShapeLibrary.IsValid())
 				{
-					ShapeNameList.Add(MakeShared<FRigVMStringWithTag>(NameSpace + Shape.ShapeName.ToString()));
+					ShapeLibrary.LoadSynchronous();
+				}
+				if (ShapeLibrary.IsValid())
+				{
+					const FString NameSpace = bUseNameSpace ? ShapeLibrary->GetName() + TEXT(".") : FString();
+					ShapeNameList.Add(MakeShared<FRigVMStringWithTag>(NameSpace + ShapeLibrary->DefaultShape.ShapeName.ToString()));
+					for (const FControlRigShapeDefinition& Shape : ShapeLibrary->Shapes)
+					{
+						ShapeNameList.Add(MakeShared<FRigVMStringWithTag>(NameSpace + Shape.ShapeName.ToString()));
+					}
 				}
 			}
 		}
