Control Rig: Fix undo/redo in curve container not working
#rnx

diff --git a/Engine/Plugins/Animation/ControlRig/Source/ControlRigEditor/Private/Editor/SRigCurveContainer.cpp b/Engine/Plugins/Animation/ControlRig/Source/ControlRigEditor/Private/Editor/SRigCurveContainer.cpp
--- a/Engine/Plugins/Animation/ControlRig/Source/ControlRigEditor/Private/Editor/SRigCurveContainer.cpp
+++ b/Engine/Plugins/Animation/ControlRig/Source/ControlRigEditor/Private/Editor/SRigCurveContainer.cpp
@@ -147,6 +147,13 @@ void SRigCurveContainer::Construct(const FArguments& InArgs, TSharedRef<FControl
 	ControlRigBlueprint->Hierarchy->OnModified().AddRaw(this, &SRigCurveContainer::OnHierarchyModified);
 	ControlRigBlueprint->OnRefreshEditor().AddRaw(this, &SRigCurveContainer::HandleRefreshEditorFromBlueprint);
 
+	UEditorEngine* Editor = Cast<UEditorEngine>(GEngine);
+	if (Editor != nullptr)
+	{
+		Editor->RegisterForUndo(this);
+	}
+
+
 	// Register and bind all our menu commands
 	FCurveContainerCommands::Register();
 	BindCommands();
@@ -455,6 +462,22 @@ void SRigCurveContainer::OnNameCommitted(const FText& InNewName, ETextCommit::Ty
 	}
 }
 
+void SRigCurveContainer::PostUndo(bool bSuccess)
+{
+	if (bSuccess)
+	{
+		RefreshCurveList();
+	}
+}
+
+void SRigCurveContainer::PostRedo(bool bSuccess)
+{
+	if (bSuccess)
+	{
+		RefreshCurveList();
+	}
+}
+
 void SRigCurveContainer::OnDeleteNameClicked()
 {
 	URigHierarchy* Hierarchy = GetHierarchy();
diff --git a/Engine/Plugins/Animation/ControlRig/Source/ControlRigEditor/Private/Editor/SRigCurveContainer.h b/Engine/Plugins/Animation/ControlRig/Source/ControlRigEditor/Private/Editor/SRigCurveContainer.h
--- a/Engine/Plugins/Animation/ControlRig/Source/ControlRigEditor/Private/Editor/SRigCurveContainer.h
+++ b/Engine/Plugins/Animation/ControlRig/Source/ControlRigEditor/Private/Editor/SRigCurveContainer.h
@@ -9,6 +9,7 @@
 #include "Widgets/Views/SListView.h"
 #include "Rigs/RigHierarchy.h"
 #include "ControlRigBlueprint.h"
+#include "EditorUndoClient.h"
 
 //////////////////////////////////////////////////////////////////////////
 // FDisplayedRigCurveInfo
@@ -121,7 +122,7 @@ private:
 //////////////////////////////////////////////////////////////////////////
 // SRigCurveContainer
 
-class SRigCurveContainer : public SCompoundWidget
+class SRigCurveContainer : public SCompoundWidget, public FEditorUndoClient
 {
 public:
 	SLATE_BEGIN_ARGS( SRigCurveContainer )
@@ -201,6 +202,10 @@ public:
 	// When a name is committed after being edited in the list
 	virtual void OnNameCommitted(const FText& NewName, ETextCommit::Type CommitType, FDisplayedRigCurveInfoPtr Item);
 
+	// FEditorUndoClient
+	virtual void PostUndo(bool bSuccess) override;
+	virtual void PostRedo(bool bSuccess) override;
+
 private:
 
 	void BindCommands();
