Control Rig: Disable import skeleton button for rig modules

#rnx

diff --git a/Engine/Plugins/Animation/ControlRig/Source/ControlRigEditor/Private/Editor/SRigHierarchy.cpp b/Engine/Plugins/Animation/ControlRig/Source/ControlRigEditor/Private/Editor/SRigHierarchy.cpp
--- a/Engine/Plugins/Animation/ControlRig/Source/ControlRigEditor/Private/Editor/SRigHierarchy.cpp
+++ b/Engine/Plugins/Animation/ControlRig/Source/ControlRigEditor/Private/Editor/SRigHierarchy.cpp
@@ -179,7 +179,8 @@ void SRigHierarchy::Construct(const FArguments& InArgs, TSharedRef<FControlRigEd
 						.ButtonStyle(FAppStyle::Get(), "FlatButton.Success")
 						.ForegroundColor(FLinearColor::White)
 						.OnClicked(FOnClicked::CreateSP(this, &SRigHierarchy::OnImportSkeletonClicked))
-						.Text(FText::FromString(TEXT("Import Hierarchy")))
+						.Text(this, &SRigHierarchy::GetImportHierarchyText)
+						.IsEnabled(this, &SRigHierarchy::IsImportHierarchyEnabled)
 					]
 				]
 
@@ -535,7 +536,15 @@ FReply SRigHierarchy::OnImportSkeletonClicked()
 	{
 		if (Settings.Mesh != nullptr)
 		{
-			ImportHierarchy(FAssetData(Settings.Mesh));
+			if(ControlRigBlueprint->IsControlRigModule())
+			{
+				//ControlRigBlueprint->SetPreviewMesh(Settings.Mesh);
+				UpdateMesh(Settings.Mesh, true);
+			}
+			else
+			{
+				ImportHierarchy(FAssetData(Settings.Mesh));
+			}
 		}
 	});
 	
@@ -544,6 +553,21 @@ FReply SRigHierarchy::OnImportSkeletonClicked()
 	return FReply::Handled();
 }
 
+FText SRigHierarchy::GetImportHierarchyText() const
+{
+	if(ControlRigBlueprint->IsControlRigModule())
+	{
+		return LOCTEXT("SetPreviewMesh", "Set Preview Mesh");
+	}
+	return LOCTEXT("ImportHierarchy", "Import Hierarchy");
+}
+
+bool SRigHierarchy::IsImportHierarchyEnabled() const
+{
+	// for now we'll enable this always
+	return true;
+}
+
 void SRigHierarchy::OnFilterTextChanged(const FText& SearchText)
 {
 	DisplaySettings.FilterText = SearchText;
diff --git a/Engine/Plugins/Animation/ControlRig/Source/ControlRigEditor/Private/Editor/SRigHierarchy.h b/Engine/Plugins/Animation/ControlRig/Source/ControlRigEditor/Private/Editor/SRigHierarchy.h
--- a/Engine/Plugins/Animation/ControlRig/Source/ControlRigEditor/Private/Editor/SRigHierarchy.h
+++ b/Engine/Plugins/Animation/ControlRig/Source/ControlRigEditor/Private/Editor/SRigHierarchy.h
@@ -174,6 +174,8 @@ private:
 	EVisibility IsToolbarVisible() const;
 	EVisibility IsSearchbarVisible() const;
 	FReply OnImportSkeletonClicked();
+	FText GetImportHierarchyText() const;
+	bool IsImportHierarchyEnabled() const;
 	void OnFilterTextChanged(const FText& SearchText);
 
 	/** Tree view widget */
