Control Rig: Add execution queue to modular rigs
#rnx

diff --git a/Engine/Plugins/Animation/ControlRig/Source/ControlRig/Private/ModularRig.cpp b/Engine/Plugins/Animation/ControlRig/Source/ControlRig/Private/ModularRig.cpp
--- a/Engine/Plugins/Animation/ControlRig/Source/ControlRig/Private/ModularRig.cpp
+++ b/Engine/Plugins/Animation/ControlRig/Source/ControlRig/Private/ModularRig.cpp
@@ -131,7 +131,7 @@ bool UModularRig::Execute_Internal(const FName& InEventName)
 		}
 		
 #endif
-		
+
 		ForEachModule([&InEventName, this, Hierarchy](FRigModuleInstance* Module) -> bool
 		{
 			if (Module->Rig.IsValid())
@@ -142,47 +142,87 @@ bool UModularRig::Execute_Internal(const FName& InEventName)
 				{
 					return true;
 				}
+
+				ExecutionQueue.Add(FRigModuleExecutionElement(Module, InEventName));
+			}
+			return true;
+		});
+
+		ExecuteQueue();
+		return true;
+	}
+	return false;
+}
+
+void UModularRig::Evaluate_AnyThread()
+{
+	ResetExecutionQueue();
+	Super::Evaluate_AnyThread();
+}
+
+void UModularRig::ExecuteQueue()
+{
+	FRigVMExtendedExecuteContext& Context = GetRigVMExtendedExecuteContext();
+	URigHierarchy* Hierarchy = GetHierarchy();
+	
+	while(ExecutionQueue.IsValidIndex(ExecutionQueueFront))
+	{
+		FRigModuleExecutionElement& ExecutionElement = ExecutionQueue[ExecutionQueueFront];
+		if (ExecutionElement.ModuleInstance->Rig.IsValid())
+		{
+			UControlRig* Rig = ExecutionElement.ModuleInstance->Rig.Get();
+
+			if (!Rig->SupportsEvent(ExecutionElement.EventName))
+			{
+				ExecutionQueueFront++;
+				continue;
+			}
 				
-				// Make sure the hierarchy has the correct element redirector from this module rig
-				FRigHierarchyRedirectorGuard ElementRedirectorGuard(Rig);
-
-				FRigVMExtendedExecuteContext& RigExtendedExecuteContext= Rig->GetRigVMExtendedExecuteContext();
-
-				// Make sure the hierarchy has the correct execute context with the rig module namespace
-				FRigHierarchyExecuteContextBracket ExecuteContextBracket(Hierarchy, &RigExtendedExecuteContext);
-
-				// Make sure the module's rig has the corrct user data
-				// The rig will combine the user data of the
-				// - skeleton
-				// - skeletalmesh
-				// - SkeletalMeshComponent
-				// - default control rig module
-				// - outer modular rig
-				// - external variables
+			// Make sure the hierarchy has the correct element redirector from this module rig
+			FRigHierarchyRedirectorGuard ElementRedirectorGuard(Rig);
+
+			FRigVMExtendedExecuteContext& RigExtendedExecuteContext= Rig->GetRigVMExtendedExecuteContext();
+
+			// Make sure the hierarchy has the correct execute context with the rig module namespace
+			FRigHierarchyExecuteContextBracket ExecuteContextBracket(Hierarchy, &RigExtendedExecuteContext);
+
+			// Make sure the module's rig has the corrct user data
+			// The rig will combine the user data of the
+			// - skeleton
+			// - skeletalmesh
+			// - SkeletalMeshComponent
+			// - default control rig module
+			// - outer modular rig
+			// - external variables
+			{
+				FControlRigExecuteContext& RigPublicContext = RigExtendedExecuteContext.GetPublicDataSafe<FControlRigExecuteContext>();
+				RigPublicContext.AssetUserData.Reset();
+				if(const TArray<UAssetUserData*>* ControlRigUserDataArray = Rig->GetAssetUserDataArray())
 				{
-					FControlRigExecuteContext& RigPublicContext = RigExtendedExecuteContext.GetPublicDataSafe<FControlRigExecuteContext>();
-					RigPublicContext.AssetUserData.Reset();
-					if(const TArray<UAssetUserData*>* ControlRigUserDataArray = Rig->GetAssetUserDataArray())
+					for(const UAssetUserData* ControlRigUserData : *ControlRigUserDataArray)
 					{
-					   for(const UAssetUserData* ControlRigUserData : *ControlRigUserDataArray)
-					   {
-						   RigPublicContext.AssetUserData.Add(ControlRigUserData);
-					   }
+						RigPublicContext.AssetUserData.Add(ControlRigUserData);
 					}
-					RigPublicContext.AssetUserData.Remove(nullptr);
 				}
+				RigPublicContext.AssetUserData.Remove(nullptr);
+			}
 				
 				
-				Rig->Execute_Internal(InEventName);
-			}
-			return true;
-		});
-
-		return true;
+			Rig->Execute_Internal(ExecutionElement.EventName);
+			ExecutionElement.bExecuted = true;
+		}
+		
+		ExecutionQueueFront++;
 	}
-	return false;
 }
 
+void UModularRig::ResetExecutionQueue()
+{
+	ExecutionQueue.Reset();
+	ExecutionQueueFront = 0;
+}
+
+
 void UModularRig::OnObjectsReplaced(const TMap<UObject*, UObject*>& OldToNewInstanceMap)
 {
 	for (FRigModuleInstance& ModuleInstance : Modules)
diff --git a/Engine/Plugins/Animation/ControlRig/Source/ControlRig/Public/ModularRig.h b/Engine/Plugins/Animation/ControlRig/Source/ControlRig/Public/ModularRig.h
--- a/Engine/Plugins/Animation/ControlRig/Source/ControlRig/Public/ModularRig.h
+++ b/Engine/Plugins/Animation/ControlRig/Source/ControlRig/Public/ModularRig.h
@@ -61,6 +61,36 @@ struct CONTROLRIG_API FRigModuleInstance
 	FString GetPath() const;
 };
 
+USTRUCT(BlueprintType)
+struct CONTROLRIG_API FRigModuleExecutionElement
+{
+	GENERATED_USTRUCT_BODY()
+
+	FRigModuleExecutionElement()
+		: ModulePath(FString())
+		, ModuleInstance(nullptr)
+		, EventName(NAME_None)
+		, bExecuted(false)
+	{}
+
+	FRigModuleExecutionElement(FRigModuleInstance* InModule, FName InEvent)
+		: ModulePath(InModule->GetPath())
+		, ModuleInstance(InModule)
+		, EventName(InEvent)
+		, bExecuted(false)
+	{}
+
+	UPROPERTY()
+	FString ModulePath;
+	FRigModuleInstance* ModuleInstance;
+
+	UPROPERTY()
+	FName EventName;
+
+	UPROPERTY()
+	bool bExecuted;
+};
+
 /** Runs logic for mapping input data to transforms (the "Rig") */
 UCLASS(Blueprintable, Abstract, editinlinenew)
 class CONTROLRIG_API UModularRig : public UControlRig
@@ -79,9 +109,16 @@ public:
 	virtual void InitializeVMsFromCDO() override { URigVMHost::InitializeFromCDO(); }
 	virtual void RequestInitVMs() override { URigVMHost::RequestInit(); }
 	virtual bool Execute_Internal(const FName& InEventName) override;
+	virtual void Evaluate_AnyThread() override;
 	virtual FRigElementKeyRedirector& GetElementKeyRedirector() override { return ElementKeyRedirector; }
 	// END ControlRig
 
+	UPROPERTY()
+	TArray<FRigModuleExecutionElement> ExecutionQueue;
+	int32 ExecutionQueueFront = 0;
+	void ExecuteQueue();
+	void ResetExecutionQueue();
+
 	// BEGIN UObject
 	virtual void BeginDestroy() override;
 	// END UObject
