Constraint baking:
    - only remove transform keys in active intervals
    - cleanup useless inactive keys

diff --git a/Engine/Plugins/Animation/ControlRig/Source/ControlRigEditor/Private/Sequencer/ControlRigParameterTrackEditor.cpp b/Engine/Plugins/Animation/ControlRig/Source/ControlRigEditor/Private/Sequencer/ControlRigParameterTrackEditor.cpp
--- a/Engine/Plugins/Animation/ControlRig/Source/ControlRigEditor/Private/Sequencer/ControlRigParameterTrackEditor.cpp
+++ b/Engine/Plugins/Animation/ControlRig/Source/ControlRigEditor/Private/Sequencer/ControlRigParameterTrackEditor.cpp
@@ -2779,6 +2779,11 @@ void FControlRigParameterTrackEditor::HandleConstraintKeyDeleted(
 	const FMovieSceneConstraintChannel* InConstraintChannel,
 	const TArray<FKeyAddOrDeleteEventItem>& InDeletedItems) const
 {
+	if (FMovieSceneConstraintChannelHelper::bDoNotCompensate)
+	{
+		return;
+	}
+	
 	UMovieSceneControlRigParameterSection* Section = Cast<UMovieSceneControlRigParameterSection>(InSection);
 
 	const UControlRig* ControlRig = Section ? Section->GetControlRig() : nullptr;
diff --git a/Engine/Plugins/Animation/ControlRig/Source/ControlRigEditor/Private/Tools/ConstraintBaker.cpp b/Engine/Plugins/Animation/ControlRig/Source/ControlRigEditor/Private/Tools/ConstraintBaker.cpp
--- a/Engine/Plugins/Animation/ControlRig/Source/ControlRigEditor/Private/Tools/ConstraintBaker.cpp
+++ b/Engine/Plugins/Animation/ControlRig/Source/ControlRigEditor/Private/Tools/ConstraintBaker.cpp
@@ -57,11 +57,12 @@ void FConstraintBaker::GetMinimalFramesToBake(
 	}
 
 	// default init bounds to the scene bounds
-	FFrameNumber FirstBakingFrame = InSettings.IsSet() ? InSettings.GetValue().StartFrame : MovieScene->GetPlaybackRange().GetLowerBoundValue();
-	FFrameNumber LastBakingFrame = InSettings.IsSet() ? InSettings.GetValue().EndFrame : MovieScene->GetPlaybackRange().GetUpperBoundValue();
+	const bool bHasSettings = InSettings.IsSet();
+	FFrameNumber FirstBakingFrame = bHasSettings ? InSettings->StartFrame : MovieScene->GetPlaybackRange().GetLowerBoundValue();
+	FFrameNumber LastBakingFrame = bHasSettings ? InSettings->EndFrame : MovieScene->GetPlaybackRange().GetUpperBoundValue();
 	
 	// set start to first active frame if any, if we have no range
-	if (InSettings.IsSet() == false)
+	if (!bHasSettings)
 	{
 		const int32 FirstActiveIndex = ConstraintValues.IndexOfByKey(true);
 		if (ConstraintValues.IsValidIndex(FirstActiveIndex))
@@ -87,7 +88,7 @@ void FConstraintBaker::GetMinimalFramesToBake(
 		return InConstraint->IsFullyActive();
 	};
 
-	if (InSettings.IsSet() == false || InSettings.GetValue().BakingKeySettings != EBakingKeySettings::KeysOnly)
+	if (!bHasSettings || InSettings->BakingKeySettings != EBakingKeySettings::KeysOnly)
 	{
 		TArray<FFrameNumber> Frames;
 		MovieSceneToolHelpers::CalculateFramesBetween(MovieScene, FirstBakingFrame, LastBakingFrame,  1 /* use increment of 1 since we reduce later*/,
@@ -101,7 +102,7 @@ void FConstraintBaker::GetMinimalFramesToBake(
 			OutFramesToBake.Add(ConstraintFrame);
 		}
 
-		int32 FrameInc = (InSettings.IsSet() == true) ? InSettings.GetValue().FrameIncrement : 1;
+		const int32 FrameInc = bHasSettings ? InSettings->FrameIncrement : 1;
 		for (int32 FrameIndex = 0; FrameIndex < Frames.Num(); ++FrameIndex)
 		{
 			if (FrameIndex % FrameInc == 0)
@@ -113,10 +114,11 @@ void FConstraintBaker::GetMinimalFramesToBake(
 				}
 			}
 		}
-		if (FrameInc > 1) //if skipping make sure we get all transforms
+		
+		if (FrameInc > 1 && bHasSettings) //if skipping make sure we get all transforms
 		{
 			FMovieSceneConstraintChannelHelper::GetTransformFramesForConstraintHandles(InConstraint, InSequencer,
-				InSettings.GetValue().StartFrame, InSettings.GetValue().EndFrame, OutFramesToBake);
+				InSettings->StartFrame, InSettings->EndFrame, OutFramesToBake);
 		}
 		
 	}
@@ -131,7 +133,7 @@ void FConstraintBaker::GetMinimalFramesToBake(
 		}
 
 		FMovieSceneConstraintChannelHelper::GetTransformFramesForConstraintHandles(InConstraint, InSequencer,
-			InSettings.GetValue().StartFrame, InSettings.GetValue().EndFrame, OutFramesToBake);
+			InSettings->StartFrame, InSettings->EndFrame, OutFramesToBake);
 
 	}
 
@@ -159,11 +161,10 @@ void FConstraintBaker::GetMinimalFramesToBake(
 				}
 
 				//if doing keys or per frame increments we need to also 
-				if (InSettings.IsSet() == true &&
-					(InSettings.GetValue().BakingKeySettings == EBakingKeySettings::KeysOnly || InSettings.GetValue().FrameIncrement > 1))
+				if (bHasSettings && (InSettings->BakingKeySettings == EBakingKeySettings::KeysOnly || InSettings->FrameIncrement > 1))
 				{
 					FMovieSceneConstraintChannelHelper::GetTransformFramesForConstraintHandles(TransformConstraint, InSequencer,
-						InSettings.GetValue().StartFrame, InSettings.GetValue().EndFrame, OutFramesToBake);
+						InSettings->StartFrame, InSettings->EndFrame, OutFramesToBake);
 				}
 			}
 		}
@@ -291,6 +292,7 @@ void FConstraintBaker::Bake(UWorld* InWorld,
 	// disable constraint and delete extra transform keys
 	TMovieSceneChannelData<bool> ConstraintChannelData = ConstraintSections.ActiveChannel->ActiveChannel.GetData();
 	const TArrayView<const FFrameNumber> ConstraintFrames = ConstraintChannelData.GetTimes();
+	const TArray<bool> PreviousConstraintValues(ConstraintChannelData.GetValues());
 	
 	// get transform channels
 	const TArrayView<FMovieSceneFloatChannel*> FloatTransformChannels = Handle->GetFloatChannels(ConstraintSections.ChildTransformSection);
@@ -327,32 +329,23 @@ void FConstraintBaker::Bake(UWorld* InWorld,
 				FramesToBake.RemoveAt(RemoveIndex);
 				Transforms.RemoveAt(RemoveIndex);
 			}
-			
 		}
 	}
-	//remove transform keys if baking frame
-	const EMovieSceneTransformChannel ChannelsToKey = InConstraint->GetChannelsToKey();
-
-	if (InSettings.IsSet() && InSettings.GetValue().BakingKeySettings == EBakingKeySettings::AllFrames)
+	
+	// remove transform keys in active sections if baking all frames
+	if (InSettings.IsSet() && InSettings->BakingKeySettings == EBakingKeySettings::AllFrames)
 	{
-		if (FloatTransformChannels.Num() > 0)
-		{
-			FMovieSceneConstraintChannelHelper::DeleteTransformTimes(FloatTransformChannels, FramesToBake[0] +1, FramesToBake[FramesToBake.Num() -1],
-				ChannelsToKey);
-		}
-		else if (DoubleTransformChannels.Num() > 0)
-		{
-			FMovieSceneConstraintChannelHelper::DeleteTransformTimes(DoubleTransformChannels, FramesToBake[0] + 1, FramesToBake[FramesToBake.Num() - 1],
-				ChannelsToKey);
-		}
+		DeleteTransformKeysInActiveRanges(
+			FloatTransformChannels, DoubleTransformChannels, InConstraint->GetChannelsToKey(),
+			FramesToBake, ConstraintFrames, PreviousConstraintValues);
 	}
 
 	// now bake to channel curves
 	const EMovieSceneTransformChannel Channels = InConstraint->GetChannelsToKey();
 	ConstraintSections.Interface->AddHandleTransformKeys(InSequencer, Handle, FramesToBake, Transforms, Channels);
 
-	if (InSettings.IsSet() && InSettings.GetValue().BakingKeySettings != EBakingKeySettings::KeysOnly
-		&& InSettings.GetValue().bReduceKeys == true)
+	// reduce keys
+	if (InSettings.IsSet() && InSettings->BakingKeySettings != EBakingKeySettings::KeysOnly && InSettings->bReduceKeys)
 	{
 		FKeyDataOptimizationParams Param;
 		Param.bAutoSetInterpolation = true;
@@ -361,10 +354,119 @@ void FConstraintBaker::Bake(UWorld* InWorld,
 		Param.Range = Range;
 		MovieSceneToolHelpers::OptimizeSection(Param, ConstraintSections.ConstraintSection);
 	}
+
+	// remove useless constraint keys
+	CleanupConstraintKeys(*ConstraintSections.ActiveChannel);
+	
 	// notify
 	InSequencer->RequestEvaluate();
 }
 
+void FConstraintBaker::DeleteTransformKeysInActiveRanges(
+	const TArrayView<FMovieSceneFloatChannel*>& InFloatTransformChannels,
+	const TArrayView<FMovieSceneDoubleChannel*>& InDoubleTransformChannels,
+	const EMovieSceneTransformChannel& InChannels,
+	const TArray<FFrameNumber>& InFramesToBake,
+	const TArrayView<const FFrameNumber>& InConstraintFrames,
+	const TArray<bool>& InConstraintValues)
+{
+	// get transform channels
+	auto CleanRange = [&](const FFrameNumber& StartTime, const FFrameNumber& EndTime)
+	{
+		if (InFloatTransformChannels.Num() > 0)
+		{
+			FMovieSceneConstraintChannelHelper::DeleteTransformTimes(InFloatTransformChannels, StartTime, EndTime, InChannels);
+		}
+		else if (InDoubleTransformChannels.Num() > 0)
+		{
+			FMovieSceneConstraintChannelHelper::DeleteTransformTimes(InDoubleTransformChannels, StartTime, EndTime, InChannels);
+		}
+	};
+
+	auto GetRange = [&](const int32 FirstActiveIndex)
+	{
+		TArray<int32> Range;
+		if (InConstraintValues.IsValidIndex(FirstActiveIndex))
+		{
+			Range.Add(FirstActiveIndex);
+			for (int32 Index = FirstActiveIndex+1; Index < InConstraintValues.Num(); ++Index)
+			{
+				if (InConstraintValues[Index] != InConstraintValues[Range.Last()])
+				{
+					Range.Add(Index);
+				}
+				else if (InConstraintValues[Index] == false)
+				{
+					Range.Last() = Index;
+				}
+			}
+		}
+		return Range;
+	};
+		
+	const int32 ActiveIndex = InConstraintValues.IndexOfByKey(true);
+	if (ActiveIndex != INDEX_NONE)
+	{
+		const TArray<int32> Range = GetRange(ActiveIndex);
+		if (Range.Num() == 1)
+		{
+			// one single active key, clean all range
+			CleanRange(InConstraintFrames[ActiveIndex] + 1, InFramesToBake.Last());
+		}
+		else
+		{
+			// clean closed ranges
+			const int32 NumClosedRange = Range.Num() / 2;
+			for (int32 RangeIndex = 0; RangeIndex < NumClosedRange; ++RangeIndex)
+			{
+				const int32 Start = Range[2*RangeIndex], End = Range[2*RangeIndex+1];    
+				CleanRange(InConstraintFrames[Start] + 1, InConstraintFrames[End]);
+			}
+
+			// clean last opened range
+			if (Range.Num() % 2 == 1)
+			{
+				const int32 Start = Range.Last();
+				CleanRange(InConstraintFrames[Start] + 1, InFramesToBake.Last());
+			}
+		}
+	}
+}
+
+void FConstraintBaker::CleanupConstraintKeys(FConstraintAndActiveChannel& InOutActiveChannel)
+{
+	TMovieSceneChannelData<bool> ConstraintChannelData = InOutActiveChannel.ActiveChannel.GetData();
+
+	const bool bHasAnyActiveKey = ConstraintChannelData.GetValues().Contains(true);
+	if (!bHasAnyActiveKey)
+	{
+		// all the keys are inactive so clear the channel without notification at all
+		InOutActiveChannel.ActiveChannel.Reset();
+		return;
+	}
+		
+	// disable extra compensation when removing keys
+	TGuardValue<bool> CompensateGuard(FMovieSceneConstraintChannelHelper::bDoNotCompensate, true);
+			
+	// store a copy to keep track of values as we remove them
+	const TArray<bool> ConstraintValues(ConstraintChannelData.GetValues());
+	for (int32 Index = ConstraintValues.Num()-1; Index >= 0; Index--)
+	{
+		// if inactive, check if the previous one is also inactive. if that's the case, then this key is useless
+		if (!ConstraintValues[Index])
+		{
+			if (Index == 0)
+			{
+				ConstraintChannelData.RemoveKey(Index);
+			}
+			else if (!ConstraintValues[Index-1])
+			{
+				ConstraintChannelData.RemoveKey(Index);
+			}
+		}
+	}
+}
+
 #undef LOCTEXT_NAMESPACE
 
 
diff --git a/Engine/Plugins/Animation/ControlRig/Source/ControlRigEditor/Public/Tools/ConstraintBaker.h b/Engine/Plugins/Animation/ControlRig/Source/ControlRigEditor/Public/Tools/ConstraintBaker.h
--- a/Engine/Plugins/Animation/ControlRig/Source/ControlRigEditor/Public/Tools/ConstraintBaker.h
+++ b/Engine/Plugins/Animation/ControlRig/Source/ControlRigEditor/Public/Tools/ConstraintBaker.h
@@ -6,12 +6,17 @@
 #include "Templates/SharedPointer.h"
 #include "BakingAnimationKeySettings.h"
 
+struct FConstraintAndActiveChannel;
 class UMovieSceneSection;
 class UTransformableHandle;
 class UTickableTransformConstraint;
 class ISequencer;
 class UWorld;
 struct FFrameNumber;
+struct FMovieSceneFloatChannel;
+struct FMovieSceneDoubleChannel;
+enum class EMovieSceneTransformChannel : uint32;
+
 struct FConstraintBaker
 {
 public:
@@ -40,5 +45,17 @@ private:
 		UMovieSceneSection* InSection,
 		const TOptional< FBakingAnimationKeySettings>& Settings,
 		TArray<FFrameNumber>& OutFramesToBake);
+
+	/** Delete the handle's transform keys within active ranges. */
+	static void DeleteTransformKeysInActiveRanges(
+		const TArrayView<FMovieSceneFloatChannel*>& InFloatTransformChannels,
+		const TArrayView<FMovieSceneDoubleChannel*>& InDoubleTransformChannels,
+		const EMovieSceneTransformChannel& InChannels,
+		const TArray<FFrameNumber>& InFramesToBake,
+		const TArrayView<const FFrameNumber>& InConstraintFrames,
+		const TArray<bool>& InConstraintValues);
+
+	/** Removes unnecessary inactive keys to ensure the cleanest possible active channel once baked */
+	static void CleanupConstraintKeys(FConstraintAndActiveChannel& InOutActiveChannel);
 };
 
diff --git a/Engine/Source/Editor/MovieSceneTools/Private/Constraints/MovieSceneConstraintChannelHelper.cpp b/Engine/Source/Editor/MovieSceneTools/Private/Constraints/MovieSceneConstraintChannelHelper.cpp
--- a/Engine/Source/Editor/MovieSceneTools/Private/Constraints/MovieSceneConstraintChannelHelper.cpp
+++ b/Engine/Source/Editor/MovieSceneTools/Private/Constraints/MovieSceneConstraintChannelHelper.cpp
@@ -767,6 +767,11 @@ void FMovieSceneConstraintChannelHelper::HandleConstraintKeyDeleted(
 	UMovieSceneSection* InSection,
 	const FFrameNumber& InTime)
 {
+	if (bDoNotCompensate)
+	{
+		return;
+	}
+	
 	if (!InConstraint || !InConstraint->NeedsCompensation())
 	{
 		return;
diff --git a/Engine/Source/Editor/MovieSceneTools/Private/TrackEditors/TransformTrackEditor.cpp b/Engine/Source/Editor/MovieSceneTools/Private/TrackEditors/TransformTrackEditor.cpp
--- a/Engine/Source/Editor/MovieSceneTools/Private/TrackEditors/TransformTrackEditor.cpp
+++ b/Engine/Source/Editor/MovieSceneTools/Private/TrackEditors/TransformTrackEditor.cpp
@@ -1819,6 +1819,11 @@ static UTickableTransformConstraint* GetTickableTransformConstraint(IMovieSceneC
 void F3DTransformTrackEditor::HandleConstraintKeyDeleted(IMovieSceneConstrainedSection* InSection, const FMovieSceneConstraintChannel* InConstraintChannel,
 	const TArray<FKeyAddOrDeleteEventItem>& InDeletedItems) const
 {
+	if (FMovieSceneConstraintChannelHelper::bDoNotCompensate)
+	{
+		return;
+	}
+	
 	if (!InConstraintChannel)
 	{
 		return;
