Control Rig: Allow binding of variables in modules
#rnx

diff --git a/Engine/Plugins/Animation/ControlRig/Source/ControlRig/Private/ModularRig.cpp b/Engine/Plugins/Animation/ControlRig/Source/ControlRig/Private/ModularRig.cpp
--- a/Engine/Plugins/Animation/ControlRig/Source/ControlRig/Private/ModularRig.cpp
+++ b/Engine/Plugins/Animation/ControlRig/Source/ControlRig/Private/ModularRig.cpp
@@ -206,8 +206,17 @@ void UModularRig::ExecuteQueue()
 				}
 				RigPublicContext.AssetUserData.Remove(nullptr);
 			}
-				
-				
+
+			// Copy variable bindings
+			for (TPair<FName, FRigVMExternalVariable>& Pair : ExecutionElement.ModuleInstance->VariableBindings)
+			{
+				FRigVMExternalVariable TargetVariable = ExecutionElement.ModuleInstance->Rig->GetPublicVariableByName(Pair.Key);
+				if (RigVMTypeUtils::AreCompatible(Pair.Value.Property, TargetVariable.Property))
+				{
+					Pair.Value.Property->CopyCompleteValue(TargetVariable.Memory, Pair.Value.Memory);
+				}
+			}
+			
 			Rig->Execute_Internal(ExecutionElement.EventName);
 			ExecutionElement.bExecuted = true;
 		}
@@ -286,14 +295,14 @@ void UModularRig::UpdateCachedChildren()
 }
 
 bool UModularRig::AddModuleInstance(const FName& InModuleName, TSubclassOf<UControlRig> InModuleClass, FString InParentPath,
-	const TMap<FRigElementKey, FRigElementKey>& InConnectionMap, const TMap<FName, FString>& InVariableDefaultValues )
+	const TMap<FRigElementKey, FRigElementKey>& InConnectionMap, const TMap<FName, FString>& InVariableDefaultValues, const TMap<FName, FString>& InVariableBindings )
 {
 	FRigModuleInstance* ParentModule = FindModule(InParentPath);
-	return AddModuleInstance(InModuleName, InModuleClass, ParentModule, InConnectionMap, InVariableDefaultValues ) != nullptr;
+	return AddModuleInstance(InModuleName, InModuleClass, ParentModule, InConnectionMap, InVariableDefaultValues, InVariableBindings) != nullptr;
 }
 
 FRigModuleInstance* UModularRig::AddModuleInstance(const FName& InModuleName, TSubclassOf<UControlRig> InModuleClass, FRigModuleInstance* InParent,
-	const TMap<FRigElementKey, FRigElementKey>& InConnectionMap, const TMap<FName, FString>& InVariableDefaultValues ) 
+	const TMap<FRigElementKey, FRigElementKey>& InConnectionMap, const TMap<FName, FString>& InVariableDefaultValues, const TMap<FName, FString>& InVariableBindings ) 
 {
 	// Make sure there are no name clashes
 	if (InParent)
@@ -346,10 +355,29 @@ FRigModuleInstance* UModularRig::AddModuleInstance(const FName& InModuleName, TS
 		ModulePublicContext.RigModuleNameSpaceHash = GetTypeHash(ModulePublicContext.RigModuleNameSpace);
 		NewModule.Rig->SetElementKeyRedirector(FRigElementKeyRedirector(InConnectionMap, Hierarchy));
 
-		for (TPair<FName, FString> Variable : InVariableDefaultValues )
+		for (const TPair<FName, FString>& Variable : InVariableDefaultValues )
 		{
 			NewModule.Rig->SetVariableFromString(Variable.Key, Variable.Value);
 		}
+
+		for (const TPair<FName, FString>& Pair : InVariableBindings)
+		{
+			FString SourceModulePath, SourceVariableName = Pair.Value;
+			Pair.Value.Split(NamespaceSeparator, &SourceModulePath, &SourceVariableName);
+			FRigVMExternalVariable SourceVariable;
+			if (SourceModulePath.IsEmpty())
+			{
+				if (FProperty* Property = GetClass()->FindPropertyByName(*SourceVariableName))
+				{
+					SourceVariable = FRigVMExternalVariable::Make(Property, (UObject*)this);
+				}
+			}
+			else if(FRigModuleInstance* SourceModule = FindModule(SourceModulePath))
+			{
+				SourceVariable = SourceModule->Rig->GetPublicVariableByName(*SourceVariableName);
+			}
+			NewModule.VariableBindings.Add(Pair.Key, SourceVariable);
+		}
 	}
 	
 	return &NewModule;
diff --git a/Engine/Plugins/Animation/ControlRig/Source/ControlRig/Public/ModularRig.h b/Engine/Plugins/Animation/ControlRig/Source/ControlRig/Public/ModularRig.h
--- a/Engine/Plugins/Animation/ControlRig/Source/ControlRig/Public/ModularRig.h
+++ b/Engine/Plugins/Animation/ControlRig/Source/ControlRig/Public/ModularRig.h
@@ -56,6 +56,9 @@ struct CONTROLRIG_API FRigModuleInstance
 	UPROPERTY()
 	FString ParentPath;
 
+	UPROPERTY()
+	TMap<FName, FRigVMExternalVariable> VariableBindings;
+
 	TArray<FRigModuleInstance*> CachedChildren;
 
 	FString GetPath() const;
@@ -130,8 +133,8 @@ public:
 	void UpdateCachedChildren();
 
 	/** Adds a module to the rig*/
-	bool AddModuleInstance(const FName& InModuleName, TSubclassOf<UControlRig> InModuleClass, FString InParentPath, const TMap<FRigElementKey, FRigElementKey>& InConnectionMap, const TMap<FName, FString>& InVariables);
-	FRigModuleInstance* AddModuleInstance(const FName& InModuleName, TSubclassOf<UControlRig> InModuleClass, FRigModuleInstance* InParent, const TMap<FRigElementKey, FRigElementKey>& InConnectionMap, const TMap<FName, FString>& InVariables);
+	bool AddModuleInstance(const FName& InModuleName, TSubclassOf<UControlRig> InModuleClass, FString InParentPath, const TMap<FRigElementKey, FRigElementKey>& InConnectionMap, const TMap<FName, FString>& InVariableDefaultValues, const TMap<FName, FString>& InVariableBindings);
+	FRigModuleInstance* AddModuleInstance(const FName& InModuleName, TSubclassOf<UControlRig> InModuleClass, FRigModuleInstance* InParent, const TMap<FRigElementKey, FRigElementKey>& InConnectionMap, const TMap<FName, FString>& InVariableDefaultValues, const TMap<FName, FString>& InVariableBindings);
 
 	FRigModuleInstance* FindModule(const FString& InPath) const;
 	FString GetParentPath(const FString& InPath) const;
diff --git a/Engine/Plugins/Animation/ControlRig/Source/ControlRigDeveloper/Private/ControlRigBlueprint.cpp b/Engine/Plugins/Animation/ControlRig/Source/ControlRigDeveloper/Private/ControlRigBlueprint.cpp
--- a/Engine/Plugins/Animation/ControlRig/Source/ControlRigDeveloper/Private/ControlRigBlueprint.cpp
+++ b/Engine/Plugins/Animation/ControlRig/Source/ControlRigDeveloper/Private/ControlRigBlueprint.cpp
@@ -1552,7 +1552,7 @@ static void AddModulesRecursively(UModularRig* Rig, const FRigModuleReference* I
 		}
 	}
 	
-	FRigModuleInstance* NewModule = Rig->AddModuleInstance(InModule->Name, InModule->Class.Get(), InParent, InModule->Connections, InModule->ConfigValues);
+	FRigModuleInstance* NewModule = Rig->AddModuleInstance(InModule->Name, InModule->Class.Get(), InParent, InModule->Connections, InModule->ConfigValues, InModule->Bindings);
 	
 	for (const FRigModuleReference* ChildModule : InModule->CachedChildren)
 	{
diff --git a/Engine/Plugins/Animation/ControlRig/Source/ControlRigDeveloper/Private/ModularRigController.cpp b/Engine/Plugins/Animation/ControlRig/Source/ControlRigDeveloper/Private/ModularRigController.cpp
--- a/Engine/Plugins/Animation/ControlRig/Source/ControlRigDeveloper/Private/ModularRigController.cpp
+++ b/Engine/Plugins/Animation/ControlRig/Source/ControlRigDeveloper/Private/ModularRigController.cpp
@@ -193,6 +193,12 @@ bool UModularRigController::SetConfigValueInModule(const FString& InModulePath,
 		return false;
 	}
 
+	if (Property->HasAllPropertyFlags(CPF_BlueprintReadOnly))
+	{
+		UE_LOG(LogControlRig, Error, TEXT("The target variable %s in module %s is read only"), *InVariableName.ToString(), *InModulePath);
+		return false;
+	}
+
 	TArray<uint8, TAlignedHeapAllocator<16>> TempStorage;
 	TempStorage.AddZeroed(Property->GetSize());
 	uint8* TempMemory = TempStorage.GetData();
@@ -227,6 +233,249 @@ bool UModularRigController::SetConfigValueInModule(const FString& InModulePath,
 	return true;
 }
 
+TArray<FString> UModularRigController::GetPossibleBindings(const FString& InModulePath, const FName& InVariableName)
+{
+	TArray<FString> PossibleBindings;
+	FRigModuleReference* Module = FindModule(InModulePath);
+	if (!Module)
+	{
+		return PossibleBindings;
+	}
+
+	if (!Module->Class.IsValid())
+	{
+		return PossibleBindings;
+	}
+
+	const FProperty* TargetProperty = Module->Class->FindPropertyByName(InVariableName);
+	if (!TargetProperty)
+	{
+		return PossibleBindings;
+	}
+
+	if (TargetProperty->HasAnyPropertyFlags(CPF_BlueprintReadOnly | CPF_DisableEditOnInstance))
+	{
+		return PossibleBindings;
+	}
+
+	// Add possible blueprint variables
+	if(UControlRigBlueprint* Blueprint = Cast<UControlRigBlueprint>(GetOuter()))
+	{
+		TArray<FRigVMExternalVariable> Variables = Blueprint->GetControlRigClass()->GetDefaultObject<UControlRig>()->GetExternalVariables();
+		for (const FRigVMExternalVariable& Variable : Variables)
+		{
+			FText ErrorMessage;
+			const FString VariableName = Variable.Name.ToString();
+			if (CanBindModuleVariable(InModulePath, InVariableName, VariableName, ErrorMessage))
+			{
+				PossibleBindings.Add(VariableName);
+			}
+		}
+	}
+
+	// Add possible module variables
+	TArray<FRigModuleReference*> Modules = Model->RootModules;
+	const FString ModulePath = Module->GetPath();
+	for (int32 i=0; i<Modules.Num(); ++i)
+	{
+		const FString CurModulePath = Modules[i]->GetPath();
+		if (ModulePath == CurModulePath)
+		{
+			continue;
+		}
+
+		TArray<FRigVMExternalVariable> Variables = Modules[i]->Class->GetDefaultObject<UControlRig>()->GetExternalVariables();
+		for (const FRigVMExternalVariable& Variable : Variables)
+		{
+			FText ErrorMessage;
+			const FString VariablePath = FString::Printf(TEXT("%s:%s"), *CurModulePath, *Variable.Name.ToString());
+			if (CanBindModuleVariable(InModulePath, InVariableName, VariablePath, ErrorMessage))
+			{
+				PossibleBindings.Add(VariablePath);
+			}
+		}
+
+		Modules.Append(Modules[i]->CachedChildren);
+	}
+
+	return PossibleBindings;
+}
+
+bool UModularRigController::CanBindModuleVariable(const FString& InModulePath, const FName& InVariableName, const FString& InSourcePath, FText& OutErrorMessage)
+{
+	FRigModuleReference* Module = FindModule(InModulePath);
+	if (!Module)
+	{
+		OutErrorMessage = FText::FromString(FString::Printf(TEXT("Could not find module %s"), *InModulePath));
+		return false;
+	}
+
+	if (!Module->Class.IsValid())
+	{
+		OutErrorMessage = FText::FromString(FString::Printf(TEXT("Class defined in module %s is not valid"), *InModulePath));
+		return false;
+	}
+
+	const FProperty* TargetProperty = Module->Class->FindPropertyByName(InVariableName);
+	if (!TargetProperty)
+	{
+		OutErrorMessage = FText::FromString(FString::Printf(TEXT("Could not find variable %s in module %s"), *InVariableName.ToString(), *InModulePath));
+		return false;
+	}
+
+	if (TargetProperty->HasAnyPropertyFlags(CPF_BlueprintReadOnly | CPF_DisableEditOnInstance))
+	{
+		OutErrorMessage = FText::FromString(FString::Printf(TEXT("The target variable %s in module %s is read only"), *InVariableName.ToString(), *InModulePath));
+		return false;
+	}
+
+	FString SourceModulePath, SourceVariableName = InSourcePath;
+	InSourcePath.Split(UModularRig::NamespaceSeparator, &SourceModulePath, &SourceVariableName);
+
+	FRigModuleReference* SourceModule = nullptr;
+	if (!SourceModulePath.IsEmpty())
+	{
+		SourceModule = FindModule(SourceModulePath);
+		if (!SourceModule)
+		{
+			OutErrorMessage = FText::FromString(FString::Printf(TEXT("Could not find source module %s"), *SourceModulePath));
+			return false;
+		}
+
+		if (SourceModulePath.StartsWith(InModulePath))
+		{
+			OutErrorMessage = FText::FromString(FString::Printf(TEXT("Cannot bind variable of module %s to a variable of module %s because the source module is a child of the target module"), *InModulePath, *SourceModulePath));
+			return false;
+		}
+	}
+
+	const FProperty* SourceProperty = nullptr;
+	if (SourceModule)
+	{
+		SourceProperty = SourceModule->Class->FindPropertyByName(*SourceVariableName);
+	}
+	else
+	{
+		if(UControlRigBlueprint* Blueprint = Cast<UControlRigBlueprint>(GetOuter()))
+		{
+			SourceProperty = Blueprint->GetControlRigClass()->FindPropertyByName(*SourceVariableName);
+		}
+	}
+	if (!SourceProperty)
+	{
+		OutErrorMessage = FText::FromString(FString::Printf(TEXT("Could not find source variable %s"), *InSourcePath));
+		return false;
+	}
+
+	FString SourcePath = (SourceModulePath.IsEmpty()) ? SourceVariableName : FString::Printf(TEXT("%s:%s"), *SourceModulePath, *SourceVariableName);
+	if (!RigVMTypeUtils::AreCompatible(SourceProperty, TargetProperty))
+	{
+		FString TargetPath = FString::Printf(TEXT("%s.%s"), *InModulePath, *InVariableName.ToString());
+		OutErrorMessage = FText::FromString(FString::Printf(TEXT("Property %s of type %s and %s of type %s are not compatible"), *SourcePath, *SourceProperty->GetCPPType(), *TargetPath, *TargetProperty->GetCPPType()));
+		return false;
+	}
+
+	return true;
+}
+
+bool UModularRigController::BindModuleVariable(const FString& InModulePath, const FName& InVariableName, const FString& InSourcePath, bool bSetupUndo)
+{
+	FText ErrorMessage;
+	if (!CanBindModuleVariable(InModulePath, InVariableName, InSourcePath, ErrorMessage))
+	{
+		UE_LOG(LogControlRig, Error, TEXT("Could not bind module variable %s:%s : %s"), *InModulePath, *InVariableName.ToString(), *ErrorMessage.ToString());
+		return false;
+	}
+	
+	FRigModuleReference* Module = FindModule(InModulePath);
+	const FProperty* TargetProperty = Module->Class->FindPropertyByName(InVariableName);
+
+	FString SourceModulePath, SourceVariableName = InSourcePath;
+	InSourcePath.Split(UModularRig::NamespaceSeparator, &SourceModulePath, &SourceVariableName);
+
+	FRigModuleReference* SourceModule = nullptr;
+	if (!SourceModulePath.IsEmpty())
+	{
+		SourceModule = FindModule(SourceModulePath);
+	}
+
+	const FProperty* SourceProperty = nullptr;
+	if (SourceModule)
+	{
+		SourceProperty = SourceModule->Class->FindPropertyByName(*SourceVariableName);
+	}
+	else
+	{
+		if(UControlRigBlueprint* Blueprint = Cast<UControlRigBlueprint>(GetOuter()))
+		{
+			SourceProperty = Blueprint->GetControlRigClass()->FindPropertyByName(*SourceVariableName);
+		}
+	}
+
+	FString SourcePath = (SourceModulePath.IsEmpty()) ? SourceVariableName : FString::Printf(TEXT("%s:%s"), *SourceModulePath, *SourceVariableName);
+
+#if WITH_EDITOR
+	TSharedPtr<FScopedTransaction> TransactionPtr;
+	if (bSetupUndo)
+	{
+		TransactionPtr = MakeShared<FScopedTransaction>(NSLOCTEXT("ModularRigController", "BindModuleVariableTransaction", "Bind Module Variable"));
+		if(UControlRigBlueprint* Blueprint = Cast<UControlRigBlueprint>(GetOuter()))
+		{
+			Blueprint->Modify();
+		}
+	}
+#endif
+
+	FString& SourceStr = Module->Bindings.FindOrAdd(InVariableName);
+	SourceStr = SourcePath;
+
+	Notify(EModularRigNotification::ModuleConfigValueChanged, Module);
+
+#if WITH_EDITOR
+	TransactionPtr.Reset();
+#endif
+
+	return true;
+}
+
+bool UModularRigController::UnBindModuleVariable(const FString& InModulePath, const FName& InVariableName, bool bSetupUndo)
+{
+	FRigModuleReference* Module = FindModule(InModulePath);
+	if (!Module)
+	{
+		UE_LOG(LogControlRig, Error, TEXT("Could not find module %s"), *InModulePath);
+		return false;
+	}
+
+	if (!Module->Bindings.Contains(InVariableName))
+	{
+		UE_LOG(LogControlRig, Error, TEXT("Variable %s in module %s is not bound"), *InVariableName.ToString(), *InModulePath);
+		return false;
+	}
+
+#if WITH_EDITOR
+	TSharedPtr<FScopedTransaction> TransactionPtr;
+	if (bSetupUndo)
+	{
+		TransactionPtr = MakeShared<FScopedTransaction>(NSLOCTEXT("ModularRigController", "BindModuleVariableTransaction", "Bind Module Variable"));
+		if(UControlRigBlueprint* Blueprint = Cast<UControlRigBlueprint>(GetOuter()))
+		{
+			Blueprint->Modify();
+		}
+	}
+#endif
+
+	Module->Bindings.Remove(InVariableName);
+
+	Notify(EModularRigNotification::ModuleConfigValueChanged, Module);
+
+#if WITH_EDITOR
+	TransactionPtr.Reset();
+#endif
+
+	return true;
+}
+
 bool UModularRigController::DeleteModule(const FString& InModulePath, bool bSetupUndo)
 {
 	FRigModuleReference* Module = FindModule(InModulePath);
@@ -258,6 +507,21 @@ bool UModularRigController::DeleteModule(const FString& InModulePath, bool bSetu
 	Model->Modules.RemoveSingle(*Module);
 	Model->UpdateCachedChildren();
 
+	// Fix bindings
+	for (FRigModuleReference& Reference : Model->Modules)
+	{
+		Reference.Bindings = Reference.Bindings.FilterByPredicate([InModulePath](const TPair<FName, FString>& Binding)
+		{
+			FString ModulePath, VariableName = Binding.Value;
+			Binding.Value.Split(UModularRig::NamespaceSeparator, &ModulePath, &VariableName);
+			if (ModulePath == InModulePath)
+			{
+				return false;
+			}
+			return true;
+		});
+	}
+
 	Notify(EModularRigNotification::ModuleRemoved, &Model->DeletedModules.Last());
 
 	Model->DeletedModules.Reset();
@@ -317,6 +581,20 @@ bool UModularRigController::RenameModule(const FString& InModulePath, const FNam
 		Children.Append(Child->CachedChildren);
 	}
 
+	// Fix bindings
+	for (FRigModuleReference& Reference : Model->Modules)
+	{
+		for (TPair<FName, FString>& Binding : Reference.Bindings)
+		{
+			FString ModulePath, VariableName = Binding.Value;
+			Binding.Value.Split(UModularRig::NamespaceSeparator, &ModulePath, &VariableName);
+			if (ModulePath == OldPath)
+			{
+				Binding.Value = FString::Printf(TEXT("%s:%s"), *NewPath, *VariableName);
+			}
+		};
+	}
+
 	Notify(EModularRigNotification::ModuleRenamed, Module);
 
 #if WITH_EDITOR
@@ -407,6 +685,34 @@ bool UModularRigController::ReparentModule(const FString& InModulePath, const FS
 
 	Model->UpdateCachedChildren();
 
+	// Fix bindings
+	for (FRigModuleReference& Reference : Model->Modules)
+	{
+		const FString ReferencePath = Reference.GetPath();
+		for (TPair<FName, FString>& Binding : Reference.Bindings)
+		{
+			FString ModulePath, VariableName = Binding.Value;
+			Binding.Value.Split(UModularRig::NamespaceSeparator, &ModulePath, &VariableName);
+			if (ModulePath == OldPath)
+			{
+				Binding.Value = FString::Printf(TEXT("%s:%s"), *NewPath, *VariableName);
+				ModulePath = NewPath;
+			}
+
+			// Remove any child dependency
+			if (ModulePath.Contains(ReferencePath))
+			{
+				UE_LOG(LogControlRig, Warning, TEXT("Binding lost due to source %s contained in child module of %s"), *Binding.Value, *ReferencePath);
+				Binding.Value.Reset();
+			}
+		};
+
+		Reference.Bindings = Reference.Bindings.FilterByPredicate([](const TPair<FName, FString>& Binding)
+		{
+			return !Binding.Value.IsEmpty();
+		});
+	}
+
 	Notify(EModularRigNotification::ModuleReparented, Module);
 	
 	return false;
diff --git a/Engine/Plugins/Animation/ControlRig/Source/ControlRigDeveloper/Public/ModularRigController.h b/Engine/Plugins/Animation/ControlRig/Source/ControlRigDeveloper/Public/ModularRigController.h
--- a/Engine/Plugins/Animation/ControlRig/Source/ControlRigDeveloper/Public/ModularRigController.h
+++ b/Engine/Plugins/Animation/ControlRig/Source/ControlRigDeveloper/Public/ModularRigController.h
@@ -57,6 +57,14 @@ class CONTROLRIGDEVELOPER_API UModularRigController : public UObject
 	UFUNCTION(BlueprintCallable, Category = "Control Rig | Modules")
 	bool SetConfigValueInModule(const FString& InModulePath, const FName& InVariableName, const FString& InValue, bool bSetupUndo = true);
 
+	UFUNCTION(BlueprintCallable, Category = "ControlRig | Modules")
+	bool BindModuleVariable(const FString& InModulePath, const FName& InVariableName, const FString& InSourcePath, bool bSetupUndo = true);
+	bool CanBindModuleVariable(const FString& InModulePath, const FName& InVariableName, const FString& InSourcePath, FText& OutErrorMessage);
+	TArray<FString> GetPossibleBindings(const FString& InModulePath, const FName& InVariableName);
+
+	UFUNCTION(BlueprintCallable, Category = "ControlRig | Modules")
+	bool UnBindModuleVariable(const FString& InModulePath, const FName& InVariableName, bool bSetupUndo = true);
+
 	UFUNCTION(BlueprintCallable, Category = "ControlRig | Modules")
 	bool DeleteModule(const FString& InModulePath, bool bSetupUndo = true);
 
diff --git a/Engine/Plugins/Animation/ControlRig/Source/ControlRigDeveloper/Public/ModularRigModel.h b/Engine/Plugins/Animation/ControlRig/Source/ControlRigDeveloper/Public/ModularRigModel.h
--- a/Engine/Plugins/Animation/ControlRig/Source/ControlRigDeveloper/Public/ModularRigModel.h
+++ b/Engine/Plugins/Animation/ControlRig/Source/ControlRigDeveloper/Public/ModularRigModel.h
@@ -39,6 +39,9 @@ struct CONTROLRIGDEVELOPER_API FRigModuleReference
 	UPROPERTY()
 	TMap<FName, FString> ConfigValues;
 
+	UPROPERTY()
+	TMap<FName, FString> Bindings; // ExternalVariableName (current module) -> SourceExternalVariableNamespacedPath (root rig or other module)
+
 	UPROPERTY(transient)
 	FName PreviousName;
 
diff --git a/Engine/Plugins/Animation/ControlRig/Source/ControlRigEditor/Private/Editor/SModularRigHierarchy.cpp b/Engine/Plugins/Animation/ControlRig/Source/ControlRigEditor/Private/Editor/SModularRigHierarchy.cpp
--- a/Engine/Plugins/Animation/ControlRig/Source/ControlRigEditor/Private/Editor/SModularRigHierarchy.cpp
+++ b/Engine/Plugins/Animation/ControlRig/Source/ControlRigEditor/Private/Editor/SModularRigHierarchy.cpp
@@ -615,7 +615,11 @@ void SModularRigHierarchy::HandleDeleteModules()
 		TArray<FString> SelectedPaths;
 		Algo::Transform(SelectedItems, SelectedPaths, [](const TSharedPtr<FModularRigTreeElement>& Element)
 		{
-			return Element->Key;
+			if (Element.IsValid())
+			{
+				return Element->Key;
+			}
+			return FString();
 		});
 		HandleDeleteModules(SelectedPaths);
 	}
diff --git a/Engine/Plugins/Runtime/RigVM/Source/RigVM/Public/RigVMTypeUtils.h b/Engine/Plugins/Runtime/RigVM/Source/RigVM/Public/RigVMTypeUtils.h
--- a/Engine/Plugins/Runtime/RigVM/Source/RigVM/Public/RigVMTypeUtils.h
+++ b/Engine/Plugins/Runtime/RigVM/Source/RigVM/Public/RigVMTypeUtils.h
@@ -507,4 +507,51 @@ namespace RigVMTypeUtils
 	{
 		return Cast<T>(FindObjectFromCPPTypeObjectPath(InObjectPath));
 	}
+
+	static bool AreCompatible(const FProperty* InSourceProperty, const FProperty* InTargetProperty)
+	{
+		bool bCompatible = InSourceProperty->SameType(InTargetProperty);
+		if (!bCompatible)
+		{
+			if(const FFloatProperty* TargetFloatProperty = CastField<FFloatProperty>(InTargetProperty))
+			{
+				bCompatible = InSourceProperty->IsA<FDoubleProperty>();
+			}
+			else if(const FDoubleProperty* TargetDoubleProperty = CastField<FDoubleProperty>(InTargetProperty))
+			{
+				bCompatible = InSourceProperty->IsA<FFloatProperty>();
+			}
+			else if (const FByteProperty* TargetByteProperty = CastField<FByteProperty>(InTargetProperty))
+			{
+				bCompatible = InSourceProperty->IsA<FEnumProperty>();
+			}
+			else if (const FEnumProperty* TargetEnumProperty = CastField<FEnumProperty>(InTargetProperty))
+			{
+				bCompatible = InSourceProperty->IsA<FByteProperty>();
+			}
+			else if(const FArrayProperty* TargetArrayProperty = CastField<FArrayProperty>(InTargetProperty))
+			{
+				if(const FArrayProperty* SourceArrayProperty = CastField<FArrayProperty>(InSourceProperty))
+				{
+					if(TargetArrayProperty->Inner->IsA<FFloatProperty>())
+					{
+						bCompatible = SourceArrayProperty->Inner->IsA<FDoubleProperty>();
+					}
+					else if(TargetArrayProperty->Inner->IsA<FDoubleProperty>())
+					{
+						bCompatible = SourceArrayProperty->Inner->IsA<FFloatProperty>();
+					}
+					else if(FByteProperty* TargetArrayInnerByteProperty = CastField<FByteProperty>(TargetArrayProperty->Inner))
+					{
+						bCompatible = SourceArrayProperty->Inner->IsA<FEnumProperty>();
+					}
+					else if(FEnumProperty* TargetArrayInnerEnumProperty = CastField<FEnumProperty>(TargetArrayProperty->Inner))
+					{
+						bCompatible = SourceArrayProperty->Inner->IsA<FByteProperty>();
+					}
+				}
+			}
+		}
+		return bCompatible;
+	}
 }
\ No newline at end of file
