ControlRig: Add option to prevent pose from being reset when the event queue changes

#rnx

diff --git a/Engine/Plugins/Animation/ControlRig/Source/ControlRig/Private/Settings/ControlRigSettings.cpp b/Engine/Plugins/Animation/ControlRig/Source/ControlRig/Private/Settings/ControlRigSettings.cpp
--- a/Engine/Plugins/Animation/ControlRig/Source/ControlRig/Private/Settings/ControlRigSettings.cpp
+++ b/Engine/Plugins/Animation/ControlRig/Source/ControlRig/Private/Settings/ControlRigSettings.cpp
@@ -21,6 +21,7 @@ UControlRigEditorSettings::UControlRigEditorSettings(const FObjectInitializer& O
 #if WITH_EDITORONLY_DATA
 	bResetControlsOnCompile = true;
 	bResetControlsOnPinValueInteraction = false;
+	bResetPoseWhenTogglingEventQueue = true;
 	bEnableUndoForPoseInteraction = true;
 
 	ConstructionEventBorderColor = FLinearColor::Red;
diff --git a/Engine/Plugins/Animation/ControlRig/Source/ControlRig/Public/Settings/ControlRigSettings.h b/Engine/Plugins/Animation/ControlRig/Source/ControlRig/Public/Settings/ControlRigSettings.h
--- a/Engine/Plugins/Animation/ControlRig/Source/ControlRig/Public/Settings/ControlRigSettings.h
+++ b/Engine/Plugins/Animation/ControlRig/Source/ControlRig/Public/Settings/ControlRigSettings.h
@@ -66,6 +66,11 @@ public:
 	// value as the user interacts with a pin value
 	UPROPERTY(EditAnywhere, config, Category = Interaction)
 	bool bResetControlsOnPinValueInteraction;
+	
+	// When this is checked all elements will be reset to their initial value
+	// if the user changes the event queue (for example between forward / backward solve)
+	UPROPERTY(EditAnywhere, config, Category = Interaction)
+	bool bResetPoseWhenTogglingEventQueue;
 
 	// When this is checked any hierarchy interaction within the Control Rig
 	// Editor will be stored on the undo stack
diff --git a/Engine/Plugins/Animation/ControlRig/Source/ControlRigEditor/Private/Editor/ControlRigEditor.cpp b/Engine/Plugins/Animation/ControlRig/Source/ControlRigEditor/Private/Editor/ControlRigEditor.cpp
--- a/Engine/Plugins/Animation/ControlRig/Source/ControlRigEditor/Private/Editor/ControlRigEditor.cpp
+++ b/Engine/Plugins/Animation/ControlRig/Source/ControlRigEditor/Private/Editor/ControlRigEditor.cpp
@@ -873,7 +873,10 @@ void FControlRigEditor::SetEventQueue(TArray<FName> InEventQueue, bool bCompile)
 		if (InEventQueue.Contains(FRigUnit_PrepareForExecution::EventName) ||
 			InEventQueue.Contains(FRigUnit_BeginExecution::EventName))
 		{
-			ControlRig->GetHierarchy()->ResetPoseToInitial(ERigElementType::All);
+			if(UControlRigEditorSettings::Get()->bResetPoseWhenTogglingEventQueue)
+			{
+				ControlRig->GetHierarchy()->ResetPoseToInitial(ERigElementType::All);
+			}
 		}
 	}
 
