ControlRig: added Parent constraint math Node

diff --git a/Engine/Plugins/Animation/ControlRig/Source/ControlRig/Private/Units/Highlevel/Hierarchy/RigUnit_TransformConstraint.cpp b/Engine/Plugins/Animation/ControlRig/Source/ControlRig/Private/Units/Highlevel/Hierarchy/RigUnit_TransformConstraint.cpp
--- a/Engine/Plugins/Animation/ControlRig/Source/ControlRig/Private/Units/Highlevel/Hierarchy/RigUnit_TransformConstraint.cpp
+++ b/Engine/Plugins/Animation/ControlRig/Source/ControlRig/Private/Units/Highlevel/Hierarchy/RigUnit_TransformConstraint.cpp
@@ -439,6 +439,170 @@ FRigUnit_ParentConstraint_Execute()
 
 #endif
 
+FRigUnit_ParentConstraintMath_Execute()
+ {
+ 	URigHierarchy* Hierarchy = ExecuteContext.Hierarchy;	
+	if (Hierarchy)
+	{
+
+		if(ParentCaches.Num() != Parents.Num())
+		{
+			ParentCaches.SetNumZeroed(Parents.Num());
+		}
+
+		// calculate total weight
+		float OverallWeight = 0;
+		for (int32 ParentIndex = 0; ParentIndex < Parents.Num(); ParentIndex++)
+		{
+			const FConstraintParent& Parent = Parents[ParentIndex];
+			if (!ParentCaches[ParentIndex].UpdateCache(Parent.Item, Hierarchy))
+			{
+				continue;
+			}
+
+			const float ClampedWeight = FMath::Max(Parent.Weight, 0.f);
+			if (ClampedWeight < KINDA_SMALL_NUMBER)
+			{
+				continue;
+			}
+
+			OverallWeight += ClampedWeight;
+		}
+
+		if (OverallWeight > KINDA_SMALL_NUMBER)
+		{
+			const float WeightNormalizer = 1.0f / OverallWeight;
+
+			FTransform MixedGlobalTransform;
+
+			// initial rotation needs to be (0,0,0,0) instead of (0,0,0,1) due to Quaternion Blending Math
+			MixedGlobalTransform.SetLocation(FVector::ZeroVector);
+			MixedGlobalTransform.SetRotation(FQuat(0.f, 0.f, 0.f, 0.f));
+			MixedGlobalTransform.SetScale3D(FVector::ZeroVector);
+
+			float AccumulatedWeight = 0.0f;
+
+			for (int32 ParentIndex = 0; ParentIndex < Parents.Num(); ParentIndex++)
+			{
+				if (!ParentCaches[ParentIndex].IsValid())
+				{
+					continue;
+				}
+
+				const FConstraintParent& Parent = Parents[ParentIndex];
+
+				const float ClampedWeight = FMath::Max(Parent.Weight, 0.f);
+				if (ClampedWeight < KINDA_SMALL_NUMBER)
+				{
+					continue;
+				}
+
+				const float NormalizedWeight = ClampedWeight * WeightNormalizer;
+				AccumulatedWeight += NormalizedWeight;
+
+				FTransform OffsetParentTransform = Hierarchy->GetGlobalTransformByIndex(ParentCaches[ParentIndex], false);
+
+				// Input initial
+				const FTransform ChildInitialGlobalTransform = Input;
+				const FTransform ParentInitialGlobalTransform = Hierarchy->GetInitialGlobalTransform(ParentCaches[ParentIndex]);
+				
+				// offset transform is a transform that transforms parent to child
+				const FTransform OffsetTransform = ChildInitialGlobalTransform.GetRelativeTransform(ParentInitialGlobalTransform);
+
+				OffsetParentTransform = OffsetTransform * OffsetParentTransform;
+				OffsetParentTransform.NormalizeRotation();
+
+				// deal with different interpolation types
+				if (AdvancedSettings.InterpolationType == EConstraintInterpType::Average)
+				{
+					// component-wise average
+					MixedGlobalTransform.AccumulateWithShortestRotation(OffsetParentTransform, ScalarRegister(NormalizedWeight));
+				}
+				else if (AdvancedSettings.InterpolationType == EConstraintInterpType::Shortest)
+				{
+					FQuat MixedGlobalQuat = MixedGlobalTransform.GetRotation();
+
+					if (MixedGlobalQuat == FQuat(0.0f, 0.0f, 0.0f, 0.0f))
+					{
+						MixedGlobalTransform = OffsetParentTransform;
+					}
+					else
+					{
+						const float Alpha = NormalizedWeight / AccumulatedWeight;
+						const FQuat OffsetParentQuat = OffsetParentTransform.GetRotation();
+
+						MixedGlobalTransform.LerpTranslationScale3D(MixedGlobalTransform, OffsetParentTransform, ScalarRegister(Alpha));
+						MixedGlobalQuat = FQuat::Slerp(MixedGlobalQuat, OffsetParentQuat, NormalizedWeight);
+						MixedGlobalTransform.SetRotation(MixedGlobalQuat);
+					}
+				}
+				else
+				{
+					// invalid interpolation type
+					ensure(false);
+					MixedGlobalTransform = Input;
+					break;
+				}
+			}
+
+			MixedGlobalTransform.NormalizeRotation();
+
+			Output = MixedGlobalTransform;
+		}
+	} 
+ }
+
+#if WITH_DEV_AUTOMATION_TESTS
+#include "Units/RigUnitTest.h"
+
+ IMPLEMENT_RIGUNIT_AUTOMATION_TEST(FRigUnit_ParentConstraintMath)
+ {
+	// use euler rotation here to match other software's rotation representation more easily
+	EEulerRotationOrder Order = EEulerRotationOrder::XZY;
+	const FRigElementKey Child = Controller->AddBone(TEXT("Child"), FRigElementKey(), FTransform( AnimationCore::QuatFromEuler( FVector(-10, -10, -10), Order), FVector(0.f, 0.f, 0.f)), true, ERigBoneType::User);
+	const FRigElementKey Parent1 = Controller->AddBone(TEXT("Parent1"), FRigElementKey(), FTransform( AnimationCore::QuatFromEuler( FVector(30, -30, -30), Order), FVector(20.f, 20.f, 20.f)), true, ERigBoneType::User);
+	const FRigElementKey Parent2 = Controller->AddBone(TEXT("Parent2"), FRigElementKey(), FTransform( AnimationCore::QuatFromEuler( FVector(-40, -40, 40), Order), FVector(40.f, 40.f, 40.f)), true, ERigBoneType::User);
+	const FRigElementKey Parent3 = Controller->AddBone(TEXT("Parent3"), FRigElementKey(), FTransform( AnimationCore::QuatFromEuler( FVector(-50, 50, -50), Order), FVector(60.f, 60.f, 60.f)), true, ERigBoneType::User);
+	const FRigElementKey Parent4 = Controller->AddBone(TEXT("Parent4"), FRigElementKey(), FTransform( AnimationCore::QuatFromEuler( FVector(60, 60, 60), Order), FVector(80.f, 80.f, 80.f)), true, ERigBoneType::User);
+	
+	ExecuteContext.Hierarchy = Hierarchy;
+	Unit.Input = Hierarchy->GetInitialGlobalTransform(Child);
+
+	Unit.Parents.Add(FConstraintParent(Parent1, 1.0));
+	Unit.Parents.Add(FConstraintParent(Parent2, 1.0));
+	Unit.Parents.Add(FConstraintParent(Parent3, 1.0));
+	Unit.Parents.Add(FConstraintParent(Parent4, 1.0));
+
+	Hierarchy->ResetPoseToInitial(ERigElementType::Bone);
+	Unit.AdvancedSettings.InterpolationType = EConstraintInterpType::Average;
+	
+	Hierarchy->ResetPoseToInitial(ERigElementType::Bone);
+	Hierarchy->SetGlobalTransform(2, FTransform(AnimationCore::QuatFromEuler( FVector(100, 100, -100), Order), FVector(-40.f, -40.f, -40.f)));
+	Unit.AdvancedSettings.InterpolationType = EConstraintInterpType::Average;
+	
+	Execute();
+	AddErrorIfFalse(Unit.Output.GetTranslation().Equals(FVector(-8.66f, 7.01f, -13.0f), 0.02f),
+                    TEXT("unexpected translation for average interpolation type"));
+	
+	FQuat Result = Unit.Output.GetRotation();
+	FQuat Expected = AnimationCore::QuatFromEuler( FVector(5.408f, -5.679f, -34.44f), Order);
+	AddErrorIfFalse(Result.Equals(Expected, 0.001f), TEXT("unexpected rotation for average interpolation type"));
+
+	
+	Hierarchy->ResetPoseToInitial(ERigElementType::Bone);
+	Hierarchy->SetGlobalTransform(2, FTransform(AnimationCore::QuatFromEuler( FVector(100.0f, 100.0f, -100.0f), Order), FVector(-40.f, -40.f, -40.f)));
+	Unit.AdvancedSettings.InterpolationType = EConstraintInterpType::Shortest;
+	
+	Execute();
+	Result = Unit.Output.GetRotation();
+	Expected = AnimationCore::QuatFromEuler( FVector(-1.209f, -8.332f, -25.022f), Order);
+	AddErrorIfFalse(Result.Equals(Expected, 0.001f), TEXT("unexpected rotation for shortest interpolation type"));
+	
+	return true;
+ }
+
+#endif
+
 FRigUnit_PositionConstraint_Execute()
 {
 	if (Weight < KINDA_SMALL_NUMBER)
diff --git a/Engine/Plugins/Animation/ControlRig/Source/ControlRig/Public/Units/Highlevel/Hierarchy/RigUnit_TransformConstraint.h b/Engine/Plugins/Animation/ControlRig/Source/ControlRig/Public/Units/Highlevel/Hierarchy/RigUnit_TransformConstraint.h
--- a/Engine/Plugins/Animation/ControlRig/Source/ControlRig/Public/Units/Highlevel/Hierarchy/RigUnit_TransformConstraint.h
+++ b/Engine/Plugins/Animation/ControlRig/Source/ControlRig/Public/Units/Highlevel/Hierarchy/RigUnit_TransformConstraint.h
@@ -249,6 +249,59 @@ struct FRigUnit_ParentConstraint : public FRigUnit_HighlevelBaseMutable
 	TArray<FCachedRigElement> ParentCaches;
 };
 
+USTRUCT()
+struct FRigUnit_ParentConstraintMath_AdvancedSettings
+{
+	GENERATED_BODY()
+
+	FRigUnit_ParentConstraintMath_AdvancedSettings()
+		: InterpolationType(EConstraintInterpType::Average)
+	{}
+	
+	/**
+	*	Options for interpolating rotations
+	*/
+	UPROPERTY(EditAnywhere, Category = "Constraint", meta = (Input))
+	EConstraintInterpType InterpolationType;
+};
+
+/**
+* Computes the output transform by constraining the input transform to multiple parents 
+*/
+USTRUCT(meta=(DisplayName="Parent Constraint Math", Category="Constraints", Keywords = "Parent,Orient,Scale"))
+struct FRigUnit_ParentConstraintMath : public FRigUnit_HighlevelBase
+{
+	GENERATED_BODY()
+
+	FRigUnit_ParentConstraintMath()
+		: Input(FTransform::Identity)
+		, Output(FTransform::Identity)
+	{
+		Parents.Add(FConstraintParent());
+	}
+	
+	RIGVM_METHOD()
+	virtual void Execute() override;
+
+	/**
+	*	Input is used to calculate offsets from parents' initial transform
+	*/
+	UPROPERTY(EditAnywhere, Category = "Constraint", meta = (Input, ExpandByDefault))
+	FTransform Input;
+	
+	UPROPERTY(EditAnywhere, Category = "Constraint", meta = (Input, ExpandByDefault, DefaultArraySize = 1))
+	TArray<FConstraintParent> Parents;
+
+	UPROPERTY(EditAnywhere, Category = "Constraint", meta = (Input)) 
+	FRigUnit_ParentConstraintMath_AdvancedSettings AdvancedSettings;
+	
+	UPROPERTY(EditAnywhere, Category = "Constraint", meta = (Output))
+	FTransform Output;
+
+	UPROPERTY()
+	TArray<FCachedRigElement> ParentCaches;
+};
+
 /**
 * Constrains an item's position to multiple items' positions 
 */
