Control Rig: Only reset recorded read / written transforms for correct events

#rnx

diff --git a/Engine/Plugins/Animation/ControlRig/Source/ControlRig/Private/ControlRig.cpp b/Engine/Plugins/Animation/ControlRig/Source/ControlRig/Private/ControlRig.cpp
--- a/Engine/Plugins/Animation/ControlRig/Source/ControlRig/Private/ControlRig.cpp
+++ b/Engine/Plugins/Animation/ControlRig/Source/ControlRig/Private/ControlRig.cpp
@@ -96,6 +96,35 @@ UControlRig::UControlRig(const FObjectInitializer& ObjectInitializer)
 	SetRigVMExtendedExecuteContext(&RigVMExtendedExecuteContext);
 }
 
+#if WITH_EDITOR
+
+void UControlRig::ResetRecordedTransforms(const FName& InEventName)
+{
+	if(const URigHierarchy* Hierarchy = GetHierarchy())
+	{
+		if(Hierarchy->bRecordTransformsAtRuntime)
+		{
+			bool bResetRecordedTransforms = SupportsEvent(InEventName);;
+			if(InEventName == FRigUnit_PostBeginExecution::EventName)
+			{
+				bResetRecordedTransforms = false;
+			}
+			else if(InEventName == FRigUnit_BeginExecution::EventName)
+			{
+				bResetRecordedTransforms = !SupportsEvent(FRigUnit_PreBeginExecution::EventName);
+			}
+
+			if(bResetRecordedTransforms)
+			{
+				Hierarchy->ReadTransformsAtRuntime.Reset();
+				Hierarchy->WrittenTransformsAtRuntime.Reset();
+			}
+		}
+	}
+}
+
+#endif
+
 void UControlRig::BeginDestroy()
 {
 	Super::BeginDestroy();
@@ -846,9 +875,26 @@ bool UControlRig::Execute(const FName& InEventName)
 	PublicContext.AssetUserData.Remove(nullptr);
 
 	// if we have any referenced elements dirty them
-	if(GetHierarchy())
+	// also reset the recorded read / written transforms as needed
+#if WITH_EDITOR
+	TSharedPtr<TGuardValue<bool>> RecordTransformsPerInstructionGuard;
+#endif
+	if(URigHierarchy* Hierarchy = GetHierarchy())
 	{
-		GetHierarchy()->UpdateReferences(&PublicContext);
+		Hierarchy->UpdateReferences(&PublicContext);
+
+#if WITH_EDITOR
+		bool bRecordTransformsAtRuntime = true;
+		if(const UObject* Outer = GetOuter())
+		{
+			if(Outer->IsA<UControlRigComponent>())
+			{
+				bRecordTransformsAtRuntime = false;
+			}
+		}
+		RecordTransformsPerInstructionGuard = MakeShared<TGuardValue<bool>>(Hierarchy->bRecordTransformsAtRuntime, bRecordTransformsAtRuntime);
+		ResetRecordedTransforms(InEventName);
+#endif
 	}
 
 	// guard against recursion
@@ -1352,6 +1398,11 @@ bool UControlRig::Execute(const FName& InEventName)
 
 bool UControlRig::Execute_Internal(const FName& InEventName)
 {
+	if(!SupportsEvent(InEventName))
+	{
+		return false;
+	}
+	
 	if(IsRigModule())
 	{
 		FString ConnectorWarning;
@@ -1426,25 +1477,6 @@ bool UControlRig::Execute_Internal(const FName& InEventName)
 #endif
 
 		URigHierarchy* Hierarchy = GetHierarchy();
-#if WITH_EDITOR
-
-		bool bRecordTransformsAtRuntime = true;
-		if(const UObject* Outer = GetOuter())
-		{
-			if(Outer->IsA<UControlRigComponent>())
-			{
-				bRecordTransformsAtRuntime = false;
-			}
-		}
-		TGuardValue<bool> RecordTransformsPerInstructionGuard(Hierarchy->bRecordTransformsAtRuntime, bRecordTransformsAtRuntime);
-		
-		if(Hierarchy->bRecordTransformsAtRuntime)
-		{
-			Hierarchy->ReadTransformsAtRuntime.Reset();
-			Hierarchy->WrittenTransformsAtRuntime.Reset();
-		}
-		
-#endif
 		FRigHierarchyExecuteContextBracket HierarchyContextGuard(Hierarchy, &Context);
 
 		// setup the module information
diff --git a/Engine/Plugins/Animation/ControlRig/Source/ControlRig/Private/ModularRig.cpp b/Engine/Plugins/Animation/ControlRig/Source/ControlRig/Private/ModularRig.cpp
--- a/Engine/Plugins/Animation/ControlRig/Source/ControlRig/Private/ModularRig.cpp
+++ b/Engine/Plugins/Animation/ControlRig/Source/ControlRig/Private/ModularRig.cpp
@@ -112,25 +112,6 @@ bool UModularRig::Execute_Internal(const FName& InEventName)
 	{
 		FRigVMExtendedExecuteContext& Context = GetRigVMExtendedExecuteContext();
 		URigHierarchy* Hierarchy = GetHierarchy();
-#if WITH_EDITOR
-
-		bool bRecordTransformsAtRuntime = true;
-		if(const UObject* Outer = GetOuter())
-		{
-			if(Outer->IsA<UControlRigComponent>())
-			{
-				bRecordTransformsAtRuntime = false;
-			}
-		}
-		TGuardValue<bool> RecordTransformsPerInstructionGuard(Hierarchy->bRecordTransformsAtRuntime, bRecordTransformsAtRuntime);
-		
-		if(Hierarchy->bRecordTransformsAtRuntime)
-		{
-			Hierarchy->ReadTransformsAtRuntime.Reset();
-			Hierarchy->WrittenTransformsAtRuntime.Reset();
-		}
-		
-#endif
 
 		ForEachModule([&InEventName, this, Hierarchy](FRigModuleInstance* Module) -> bool
 		{
diff --git a/Engine/Plugins/Animation/ControlRig/Source/ControlRig/Public/ControlRig.h b/Engine/Plugins/Animation/ControlRig/Source/ControlRig/Public/ControlRig.h
--- a/Engine/Plugins/Animation/ControlRig/Source/ControlRig/Public/ControlRig.h
+++ b/Engine/Plugins/Animation/ControlRig/Source/ControlRig/Public/ControlRig.h
@@ -179,6 +179,9 @@ public:
 	// called after post reinstance when compilng blueprint by Sequencer
 	void PostReinstanceCallback(const UControlRig* Old);
 
+	// resets the recorded transform changes
+	void ResetRecordedTransforms(const FName& InEventName);
+
 #endif // WITH_EDITOR
 	
 	// BEGIN UObject interface
