Control Rig: Show short names on connectors
#rnx

diff --git a/Engine/Plugins/Animation/ControlRig/Source/ControlRig/Private/Rigs/RigHierarchyController.cpp b/Engine/Plugins/Animation/ControlRig/Source/ControlRig/Private/Rigs/RigHierarchyController.cpp
--- a/Engine/Plugins/Animation/ControlRig/Source/ControlRig/Private/Rigs/RigHierarchyController.cpp
+++ b/Engine/Plugins/Animation/ControlRig/Source/ControlRig/Private/Rigs/RigHierarchyController.cpp
@@ -1919,10 +1919,13 @@ int32 URigHierarchyController::AddElement(FRigBaseElement* InElementToAdd, FRigB
 	FRigName DesiredName = InDesiredName;
 	URigHierarchy::SanitizeName(DesiredName);
 
-	int32 NameSpaceTokenIndex = INDEX_NONE; 
-	if(InElementToAdd->GetName().FindLastChar(':', NameSpaceTokenIndex))
+	// Short name of an element is ModuleName:ElementName (or ElementName if not in a module)
+	FString ModulePath, ElementName = InElementToAdd->GetName();
+	if (ElementName.Split(UModularRig::NamespaceSeparator, &ModulePath, &ElementName, ESearchCase::CaseSensitive, ESearchDir::FromEnd))
 	{
-		const FString ShortName = InElementToAdd->GetName().RightChop(NameSpaceTokenIndex + 1);
+		FString ModuleName = ModulePath, ModuleParentPath;
+		ModulePath.Split(UModularRig::NamespaceSeparator, &ModuleParentPath, &ModuleName, ESearchCase::CaseSensitive, ESearchDir::FromEnd);
+		const FString ShortName = FString::Printf(TEXT("%s:%s"), *ModuleName, *ElementName);
 		if(!ShortName.IsEmpty())
 		{
 			Hierarchy->SetNameMetadata(InElementToAdd->Key, URigHierarchy::ShortNameMetadataName, *ShortName);
diff --git a/Engine/Plugins/Animation/ControlRig/Source/ControlRigEditor/Private/ControlRigHierarchyCommands.cpp b/Engine/Plugins/Animation/ControlRig/Source/ControlRigEditor/Private/ControlRigHierarchyCommands.cpp
--- a/Engine/Plugins/Animation/ControlRig/Source/ControlRigEditor/Private/ControlRigHierarchyCommands.cpp
+++ b/Engine/Plugins/Animation/ControlRig/Source/ControlRigEditor/Private/ControlRigHierarchyCommands.cpp
@@ -32,6 +32,7 @@ void FControlRigHierarchyCommands::RegisterCommands()
 	UI_COMMAND(Unparent, "Unparent", "Unparents the selected elements from the hierarchy", EUserInterfaceActionType::Button, FInputChord(EKeys::P, EModifierKey::Shift));
 	UI_COMMAND(FilteringFlattensHierarchy, "Filtering Flattens Hierarchy", "Whether to keep the hierarchy or flatten it when searching for tree items", EUserInterfaceActionType::ToggleButton, FInputChord());
 	UI_COMMAND(HideParentsWhenFiltering, "Hide Parents When Filtering", "Whether to show parent items grayed out, or hide them entirely when filtering", EUserInterfaceActionType::ToggleButton, FInputChord());
+	UI_COMMAND(ShowShortNames, "Show Short Names", "Whether to show the full path or the short name of each element", EUserInterfaceActionType::ToggleButton, FInputChord());
 	UI_COMMAND(ShowImportedBones, "Show Imported Bones", "Whether to show or hide imported bones", EUserInterfaceActionType::ToggleButton, FInputChord());
 	UI_COMMAND(ShowBones, "Show Bones", "Whether to show or hide bones", EUserInterfaceActionType::ToggleButton, FInputChord());
 	UI_COMMAND(ShowControls, "Show Controls", "Whether to show or hide controls", EUserInterfaceActionType::ToggleButton, FInputChord());
diff --git a/Engine/Plugins/Animation/ControlRig/Source/ControlRigEditor/Private/ControlRigHierarchyCommands.h b/Engine/Plugins/Animation/ControlRig/Source/ControlRigEditor/Private/ControlRigHierarchyCommands.h
--- a/Engine/Plugins/Animation/ControlRig/Source/ControlRigEditor/Private/ControlRigHierarchyCommands.h
+++ b/Engine/Plugins/Animation/ControlRig/Source/ControlRigEditor/Private/ControlRigHierarchyCommands.h
@@ -96,6 +96,9 @@ public:
 	/* Hide parents on filter */
 	TSharedPtr< FUICommandInfo > HideParentsWhenFiltering;
 
+	/* Show short names */
+	TSharedPtr< FUICommandInfo > ShowShortNames;
+
 	/* Show imported bones */
 	TSharedPtr< FUICommandInfo > ShowImportedBones;
 
diff --git a/Engine/Plugins/Animation/ControlRig/Source/ControlRigEditor/Private/Editor/SRigHierarchy.cpp b/Engine/Plugins/Animation/ControlRig/Source/ControlRigEditor/Private/Editor/SRigHierarchy.cpp
--- a/Engine/Plugins/Animation/ControlRig/Source/ControlRigEditor/Private/Editor/SRigHierarchy.cpp
+++ b/Engine/Plugins/Animation/ControlRig/Source/ControlRigEditor/Private/Editor/SRigHierarchy.cpp
@@ -118,8 +118,6 @@ void SRigHierarchy::Construct(const FArguments& InArgs, TSharedRef<FControlRigEd
 	ControlRigBlueprint->OnRefreshEditor().AddRaw(this, &SRigHierarchy::HandleRefreshEditorFromBlueprint);
 	ControlRigBlueprint->OnSetObjectBeingDebugged().AddRaw(this, &SRigHierarchy::HandleSetObjectBeingDebugged);
 
-	DisplaySettings.bUseShortName = !ControlRigBlueprint->IsModularRig();
-
 	// for deleting, renaming, dragging
 	CommandList = MakeShared<FUICommandList>();
 
@@ -425,6 +423,12 @@ void SRigHierarchy::BindCommands()
 		FCanExecuteAction(),
 		FIsActionChecked::CreateLambda([this]() { return DisplaySettings.bHideParentsOnFilter; }));
 
+	CommandList->MapAction(
+		Commands.ShowShortNames,
+		FExecuteAction::CreateLambda([this]() { DisplaySettings.bUseShortName = !DisplaySettings.bUseShortName; RefreshTreeView(); }),
+		FCanExecuteAction(),
+		FIsActionChecked::CreateLambda([this]() { return DisplaySettings.bUseShortName; }));
+
 	CommandList->MapAction(
 		Commands.ShowImportedBones,
 		FExecuteAction::CreateLambda([this]() { DisplaySettings.bShowImportedBones = !DisplaySettings.bShowImportedBones; RefreshTreeView(); }),
@@ -1044,6 +1048,7 @@ TSharedRef< SWidget > SRigHierarchy::CreateFilterMenu()
 	{
 		MenuBuilder.AddMenuEntry(Actions.FilteringFlattensHierarchy);
 		MenuBuilder.AddMenuEntry(Actions.HideParentsWhenFiltering);
+		MenuBuilder.AddMenuEntry(Actions.ShowShortNames);
 	}
 	MenuBuilder.EndSection();
 
@@ -1838,11 +1843,6 @@ bool SRigHierarchy::IsNonProceduralElementSelected() const
 
 bool SRigHierarchy::CanAddElement(const ERigElementType ElementType) const
 {
-	if (IsProceduralElementSelected())
-	{
-		return false;
-	}
-
 	// Always allow connectors 
 	if (ElementType == ERigElementType::Connector)
 	{
diff --git a/Engine/Plugins/Animation/ControlRig/Source/ControlRigEditor/Private/Editor/SRigHierarchyTreeView.cpp b/Engine/Plugins/Animation/ControlRig/Source/ControlRigEditor/Private/Editor/SRigHierarchyTreeView.cpp
--- a/Engine/Plugins/Animation/ControlRig/Source/ControlRigEditor/Private/Editor/SRigHierarchyTreeView.cpp
+++ b/Engine/Plugins/Animation/ControlRig/Source/ControlRigEditor/Private/Editor/SRigHierarchyTreeView.cpp
@@ -515,7 +515,14 @@ bool SRigHierarchyTreeView::AddElement(const FRigBaseElement* InElement)
 								const FRigElementKey ConnectorKey = InElement->GetKey();
 
 								SRigHierarchyTagWidget::FArguments TagArguments;
-								TagArguments.Text(FText::FromString(InElement->GetName()));
+
+								FName Name = ConnectorKey.Name;
+								if (GetRigTreeDelegates().GetDisplaySettings().bUseShortName)
+								{
+									Name = Hierarchy->GetNameMetadata(ConnectorKey, URigHierarchy::ShortNameMetadataName, ConnectorKey.Name);
+								}
+								TagArguments.Text(FText::FromName(Name));
+								TagArguments.TooltipText(FText::FromName(ConnectorKey.Name));
 								TagArguments.TextColor(FSlateColor(FLinearColor::White));
 								TagArguments.Color(FLinearColor(0.0, 112.f/255.f, 224.f/255.f));
 								TagArguments.AllowDragDrop(true);
