Control Rig: Add interim property binding UI for modular rigs

#rnx

diff --git a/Engine/Plugins/Animation/ControlRig/Source/ControlRig/Private/ModularRig.cpp b/Engine/Plugins/Animation/ControlRig/Source/ControlRig/Private/ModularRig.cpp
--- a/Engine/Plugins/Animation/ControlRig/Source/ControlRig/Private/ModularRig.cpp
+++ b/Engine/Plugins/Animation/ControlRig/Source/ControlRig/Private/ModularRig.cpp
@@ -38,15 +38,6 @@ const FRigModuleInstance* FModuleInstanceHandle::Get() const
 	return nullptr;
 }
 
-FRigModuleInstance* FModuleInstanceHandle::Get()
-{
-	if(ModularRig.IsValid())
-	{
-		return ModularRig->FindModule(Path);
-	}
-	return nullptr;
-}
-
 ////////////////////////////////////////////////////////////////////////////////
 // UModularRig
 ////////////////////////////////////////////////////////////////////////////////
@@ -191,10 +182,13 @@ void UModularRig::ExecuteQueue()
 			// Copy variable bindings
 			for (TPair<FName, FRigVMExternalVariable>& Pair : ExecutionElement.ModuleInstance->VariableBindings)
 			{
-				FRigVMExternalVariable TargetVariable = ExecutionElement.ModuleInstance->Rig->GetPublicVariableByName(Pair.Key);
-				if (RigVMTypeUtils::AreCompatible(Pair.Value.Property, TargetVariable.Property))
+				const FRigVMExternalVariable TargetVariable = ExecutionElement.ModuleInstance->Rig->GetPublicVariableByName(Pair.Key);
+				if(ensure(TargetVariable.Property))
 				{
-					Pair.Value.Property->CopyCompleteValue(TargetVariable.Memory, Pair.Value.Memory);
+					if (RigVMTypeUtils::AreCompatible(Pair.Value.Property, TargetVariable.Property))
+					{
+						Pair.Value.Property->CopyCompleteValue(TargetVariable.Memory, Pair.Value.Memory);
+					}
 				}
 			}
 			
@@ -278,7 +272,7 @@ void UModularRig::UpdateCachedChildren()
 bool UModularRig::AddModuleInstance(const FName& InModuleName, TSubclassOf<UControlRig> InModuleClass, FString InParentPath,
 	const TMap<FRigElementKey, FRigElementKey>& InConnectionMap, const TMap<FName, FString>& InVariableDefaultValues, const TMap<FName, FString>& InVariableBindings )
 {
-	FRigModuleInstance* ParentModule = FindModule(InParentPath);
+	FRigModuleInstance* ParentModule = const_cast<FRigModuleInstance*>(FindModule(InParentPath));
 	return AddModuleInstance(InModuleName, InModuleClass, ParentModule, InConnectionMap, InVariableDefaultValues, InVariableBindings) != nullptr;
 }
 
@@ -353,7 +347,7 @@ FRigModuleInstance* UModularRig::AddModuleInstance(const FName& InModuleName, TS
 					SourceVariable = FRigVMExternalVariable::Make(Property, (UObject*)this);
 				}
 			}
-			else if(FRigModuleInstance* SourceModule = FindModule(SourceModulePath))
+			else if(const FRigModuleInstance* SourceModule = FindModule(SourceModulePath))
 			{
 				SourceVariable = SourceModule->Rig->GetPublicVariableByName(*SourceVariableName);
 			}
@@ -364,7 +358,7 @@ FRigModuleInstance* UModularRig::AddModuleInstance(const FName& InModuleName, TS
 	return &NewModule;
 }
 
-FRigModuleInstance* UModularRig::FindModule(const FString& InPath) const
+const FRigModuleInstance* UModularRig::FindModule(const FString& InPath) const
 {
 	const TArray<FRigModuleInstance*>* Children = &RootModules;
 	FString Left = InPath, Right;
@@ -393,9 +387,29 @@ FRigModuleInstance* UModularRig::FindModule(const FString& InPath) const
 	return *Cur;
 }
 
+const FRigModuleInstance* UModularRig::FindModule(const UControlRig* InModuleInstance) const
+{
+	const FRigModuleInstance* FoundModule = nullptr;
+	ForEachModule([InModuleInstance, &FoundModule](const FRigModuleInstance* Module) -> bool
+	{
+		if (Module->Rig.IsValid())
+		{
+			if(Module->Rig.Get() == InModuleInstance)
+			{
+				FoundModule = Module;
+				// don't continue
+				return false;
+			}
+		}
+		return true;
+	});
+
+	return FoundModule;
+}
+
 FString UModularRig::GetParentPath(const FString& InPath) const
 {
-	if (FRigModuleInstance* Element = FindModule(InPath))
+	if (const FRigModuleInstance* Element = FindModule(InPath))
 	{
 		return Element->ParentPath;
 	}
diff --git a/Engine/Plugins/Animation/ControlRig/Source/ControlRig/Public/ModularRig.h b/Engine/Plugins/Animation/ControlRig/Source/ControlRig/Public/ModularRig.h
--- a/Engine/Plugins/Animation/ControlRig/Source/ControlRig/Public/ModularRig.h
+++ b/Engine/Plugins/Animation/ControlRig/Source/ControlRig/Public/ModularRig.h
@@ -22,16 +22,17 @@ public:
 	bool IsValid() const { return Get() != nullptr; }
 	operator bool() const { return IsValid(); }
 	
-	const UModularRig* GetModularRig() const { return ModularRig.Get(); }
-	UModularRig* GetHierarchy() { return ModularRig.Get(); }
+	const UModularRig* GetModularRig() const
+	{
+		return ModularRig.Get();
+	}
 	const FString& GetPath() const { return Path; }
 
 	const FRigModuleInstance* Get() const;
-	FRigModuleInstance* Get();
 
 private:
 
-	TWeakObjectPtr<UModularRig> ModularRig;
+	mutable TSoftObjectPtr<UModularRig> ModularRig;
 	FString Path;
 };
 
@@ -136,7 +137,8 @@ public:
 	bool AddModuleInstance(const FName& InModuleName, TSubclassOf<UControlRig> InModuleClass, FString InParentPath, const TMap<FRigElementKey, FRigElementKey>& InConnectionMap, const TMap<FName, FString>& InVariableDefaultValues, const TMap<FName, FString>& InVariableBindings);
 	FRigModuleInstance* AddModuleInstance(const FName& InModuleName, TSubclassOf<UControlRig> InModuleClass, FRigModuleInstance* InParent, const TMap<FRigElementKey, FRigElementKey>& InConnectionMap, const TMap<FName, FString>& InVariableDefaultValues, const TMap<FName, FString>& InVariableBindings);
 
-	FRigModuleInstance* FindModule(const FString& InPath) const;
+	const FRigModuleInstance* FindModule(const FString& InPath) const;
+	const FRigModuleInstance* FindModule(const UControlRig* InModuleInstance) const;
 	FString GetParentPath(const FString& InPath) const;
 
 	void ForEachModule(TFunctionRef<bool(FRigModuleInstance*)> PerModuleFunction);
@@ -149,7 +151,7 @@ public:
 	 */
 	FModuleInstanceHandle GetHandle(const FString& InPath) const
 	{
-		if(FRigModuleInstance* Module = FindModule(InPath))
+		if(FindModule(InPath))
 		{
 			return FModuleInstanceHandle((UModularRig*)this, InPath);
 		}
diff --git a/Engine/Plugins/Animation/ControlRig/Source/ControlRigDeveloper/Private/ModularRigController.cpp b/Engine/Plugins/Animation/ControlRig/Source/ControlRigDeveloper/Private/ModularRigController.cpp
--- a/Engine/Plugins/Animation/ControlRig/Source/ControlRigDeveloper/Private/ModularRigController.cpp
+++ b/Engine/Plugins/Animation/ControlRig/Source/ControlRigDeveloper/Private/ModularRigController.cpp
@@ -236,7 +236,7 @@ bool UModularRigController::SetConfigValueInModule(const FString& InModulePath,
 TArray<FString> UModularRigController::GetPossibleBindings(const FString& InModulePath, const FName& InVariableName)
 {
 	TArray<FString> PossibleBindings;
-	FRigModuleReference* Module = FindModule(InModulePath);
+	const FRigModuleReference* Module = FindModule(InModulePath);
 	if (!Module)
 	{
 		return PossibleBindings;
@@ -274,29 +274,24 @@ TArray<FString> UModularRigController::GetPossibleBindings(const FString& InModu
 	}
 
 	// Add possible module variables
-	TArray<FRigModuleReference*> Modules = Model->RootModules;
-	const FString ModulePath = Module->GetPath();
-	for (int32 i=0; i<Modules.Num(); ++i)
+	Model->ForEachModule([this, &PossibleBindings, InModulePath, InVariableName](const FRigModuleReference* InModule) -> bool
 	{
-		const FString CurModulePath = Modules[i]->GetPath();
-		if (ModulePath == CurModulePath)
-		{
-			continue;
-		}
-
-		TArray<FRigVMExternalVariable> Variables = Modules[i]->Class->GetDefaultObject<UControlRig>()->GetExternalVariables();
-		for (const FRigVMExternalVariable& Variable : Variables)
+		const FString CurModulePath = InModule->GetPath();
+		if (InModulePath != CurModulePath)
 		{
-			FText ErrorMessage;
-			const FString VariablePath = FString::Printf(TEXT("%s:%s"), *CurModulePath, *Variable.Name.ToString());
-			if (CanBindModuleVariable(InModulePath, InVariableName, VariablePath, ErrorMessage))
+			TArray<FRigVMExternalVariable> Variables = InModule->Class->GetDefaultObject<UControlRig>()->GetExternalVariables();
+			for (const FRigVMExternalVariable& Variable : Variables)
 			{
-				PossibleBindings.Add(VariablePath);
+				FText ErrorMessage;
+				const FString SourceVariablePath = FString::Printf(TEXT("%s:%s"), *CurModulePath, *Variable.Name.ToString());
+				if (CanBindModuleVariable(InModulePath, InVariableName, SourceVariablePath, ErrorMessage))
+				{
+					PossibleBindings.Add(SourceVariablePath);
+				}
 			}
-		}
-
-		Modules.Append(Modules[i]->CachedChildren);
-	}
+		}		
+		return true;
+	});
 
 	return PossibleBindings;
 }
@@ -330,7 +325,7 @@ bool UModularRigController::CanBindModuleVariable(const FString& InModulePath, c
 	}
 
 	FString SourceModulePath, SourceVariableName = InSourcePath;
-	InSourcePath.Split(UModularRig::NamespaceSeparator, &SourceModulePath, &SourceVariableName);
+	InSourcePath.Split(UModularRig::NamespaceSeparator, &SourceModulePath, &SourceVariableName, ESearchCase::CaseSensitive, ESearchDir::FromEnd);
 
 	FRigModuleReference* SourceModule = nullptr;
 	if (!SourceModulePath.IsEmpty())
diff --git a/Engine/Plugins/Animation/ControlRig/Source/ControlRigEditor/Private/ControlRigEditorModule.cpp b/Engine/Plugins/Animation/ControlRig/Source/ControlRigEditor/Private/ControlRigEditorModule.cpp
--- a/Engine/Plugins/Animation/ControlRig/Source/ControlRigEditor/Private/ControlRigEditorModule.cpp
+++ b/Engine/Plugins/Animation/ControlRig/Source/ControlRigEditor/Private/ControlRigEditorModule.cpp
@@ -124,6 +124,9 @@ void FControlRigEditorModule::StartupModule()
 	ClassesToUnregisterOnShutdown.Add(UControlRigBlueprint::StaticClass()->GetFName());
 	PropertyEditorModule.RegisterCustomClassLayout(ClassesToUnregisterOnShutdown.Last(), FOnGetDetailCustomizationInstance::CreateStatic(&FControlRigBlueprintDetails::MakeInstance));
 
+	ClassesToUnregisterOnShutdown.Add(UControlRig::StaticClass()->GetFName());
+	PropertyEditorModule.RegisterCustomClassLayout(ClassesToUnregisterOnShutdown.Last(), FOnGetDetailCustomizationInstance::CreateStatic(&FRigModuleInstanceDetails::MakeInstance));
+
 	ClassesToUnregisterOnShutdown.Add(UControlRig::StaticClass()->GetFName());
 
 	// same as ClassesToUnregisterOnShutdown but for properties, there is none right now
diff --git a/Engine/Plugins/Animation/ControlRig/Source/ControlRigEditor/Private/ControlRigModuleDetails.cpp b/Engine/Plugins/Animation/ControlRig/Source/ControlRigEditor/Private/ControlRigModuleDetails.cpp
--- a/Engine/Plugins/Animation/ControlRig/Source/ControlRigEditor/Private/ControlRigModuleDetails.cpp
+++ b/Engine/Plugins/Animation/ControlRig/Source/ControlRigEditor/Private/ControlRigModuleDetails.cpp
@@ -20,6 +20,9 @@
 #include "Editor/SModularRigHierarchyTreeView.h"
 #include "StructViewerFilter.h"
 #include "StructViewerModule.h"
+#include "Features/IModularFeatures.h"
+#include "IPropertyAccessEditor.h"
+#include "ScopedTransaction.h"
 
 #define LOCTEXT_NAMESPACE "ControlRigModuleDetails"
 
@@ -113,36 +116,43 @@ void FRigModuleInstanceDetails::CustomizeDetails(IDetailLayoutBuilder& DetailBui
 	DetailBuilder.GetObjectsBeingCustomized(DetailObjects);
 	for(TWeakObjectPtr<UObject> DetailObject : DetailObjects)
 	{
-		URigVMDetailsViewWrapperObject* WrapperObject = CastChecked<URigVMDetailsViewWrapperObject>(DetailObject.Get());
-
-		const FString Path = WrapperObject->GetContent<FRigModuleInstance>().GetPath();
-
-		FPerModuleInfo Info;
-		Info.WrapperObject = WrapperObject;
-		if (UModularRig* Subject = Cast<UModularRig>(WrapperObject->GetSubject()))
-		{
-			Info.Module = Subject->GetHandle(Path);
-		}
-
-		if(!Info.Module.IsValid())
+		if(UControlRig* ModuleInstance = Cast<UControlRig>(DetailObject))
 		{
-			return;
-		}
-		if(const UControlRigBlueprint* Blueprint = Info.GetBlueprint())
-		{
-			if (UModularRig* ModularRig = Cast<UModularRig>(Blueprint->GetObjectBeingDebugged()))
+			if(const UModularRig* ModularRig = Cast<UModularRig>(ModuleInstance->GetOuter()))
 			{
-				Info.DefaultModule = ModularRig->GetHandle(Path);
+				if(const FRigModuleInstance* Module = ModularRig->FindModule(ModuleInstance))
+				{
+					const FString Path = Module->GetPath();
+
+					FPerModuleInfo Info;
+					Info.Path = Path;
+					Info.Module = ModularRig->GetHandle(Path);
+					if(!Info.Module.IsValid())
+					{
+						return;
+					}
+					
+					if(const UControlRigBlueprint* Blueprint = Info.GetBlueprint())
+					{
+						if(const UModularRig* DefaultModularRig = Cast<UModularRig>(Blueprint->GeneratedClass->GetDefaultObject()))
+						{
+							Info.DefaultModule = DefaultModularRig->GetHandle(Path);
+						}
+					}
+
+					PerModuleInfos.Add(Info);
+				}
 			}
 		}
-
-		PerModuleInfos.Add(Info);
 	}
 
-	DetailBuilder.HideCategory(TEXT("RigModuleInstance"));
+	// don't customize if the 
+	if(PerModuleInfos.IsEmpty())
+	{
+		return;
+	}
 
 	IDetailCategoryBuilder& GeneralCategory = DetailBuilder.EditCategory(TEXT("General"), LOCTEXT("General", "General"));
-
 	{
 		GeneralCategory.AddCustomRow(FText::FromString(TEXT("Name")))
 		.NameContent()
@@ -178,7 +188,6 @@ void FRigModuleInstanceDetails::CustomizeDetails(IDetailLayoutBuilder& DetailBui
 	}
 
 	IDetailCategoryBuilder& ConnectionsCategory = DetailBuilder.EditCategory(TEXT("Connections"), LOCTEXT("Connections", "Connections"));
-
 	{
 		TArray<FRigModuleConnector> Connectors = GetConnectors();
 		FRigElementKeyRedirector Redirector = GetConnections();
@@ -222,65 +231,104 @@ void FRigModuleInstanceDetails::CustomizeDetails(IDetailLayoutBuilder& DetailBui
 		}
 	}
 
-	IDetailCategoryBuilder& ConfigValuesCategory = DetailBuilder.EditCategory(TEXT("Config Values"), LOCTEXT("ConfigValues", "Config Values"));
+	IDetailCategoryBuilder& ConfigValuesCategory = DetailBuilder.EditCategory(TEXT("Default"), LOCTEXT("ConfigValues", "Config Values"));
 	{
-		TArray<UObject*> DetailObjectsRaw;
-		DetailObjectsRaw.Reserve(DetailObjects.Num());
-		for (TWeakObjectPtr<UObject> Obj : DetailObjects)
+		IPropertyAccessEditor& PropertyAccessEditor = IModularFeatures::Get().GetModularFeature<IPropertyAccessEditor>("PropertyAccessEditor");
+
+		TArray<TSharedRef<IPropertyHandle>> DefaultProperties;
+		ConfigValuesCategory.GetDefaultProperties(DefaultProperties, true, true);
+
+		for(const TSharedRef<IPropertyHandle>& DefaultProperty : DefaultProperties)
 		{
-			if (Obj.IsValid())
+			const FProperty* Property = DefaultProperty->GetProperty();
+			if(Property == nullptr)
 			{
-				DetailObjectsRaw.Add(Obj.Get());
+				DetailBuilder.HideProperty(DefaultProperty);
+				continue;
 			}
-		}
-		
-		TArray<FRigVMExternalVariable> Variables = GetConfigValues();
-		for(FRigVMExternalVariable& Variable : Variables)
-		{
-			if (!Variable.bIsPublic)
+
+			// skip advanced properties for now
+			const bool bAdvancedDisplay = Property->HasAnyPropertyFlags(CPF_AdvancedDisplay);
+			if(bAdvancedDisplay)
 			{
+				DetailBuilder.HideProperty(DefaultProperty);
 				continue;
 			}
-			const FText Label = FText::FromName(Variable.Name);
 
-			if (FRigModuleInstance* ModuleInstance = PerModuleInfos[0].GetModule())
+			// skip non-public properties for now
+			const bool bIsPublic = Property->HasAnyPropertyFlags(CPF_Edit | CPF_EditConst);
+			const bool bIsInstanceEditable = !Property->HasAnyPropertyFlags(CPF_DisableEditOnInstance);
+			if(!bIsPublic || !bIsInstanceEditable)
 			{
-				TSoftObjectPtr<UControlRig> Rig = ModuleInstance->Rig;
-				if (Rig.IsValid())
-				{
-					if (FProperty* Property = Rig.Get()->GetClass()->FindPropertyByName(Variable.Name))
-					{
-						uint8* Container = (uint8*)Rig.Get();
-
-						TSharedPtr<FStructOnScope> StructOnScope = MakeShareable(new FStructOnScope(Rig.Get()->GetClass(), Container));
-						IDetailPropertyRow* Row = ConfigValuesCategory.AddExternalObjectProperty({Rig.Get()}, Variable.Name);
-						
-
-						Row->DisplayName(FText::FromName(Variable.Name));
+				DetailBuilder.HideProperty(DefaultProperty);
+				continue;
+			}
 
-						const FSimpleDelegate OnValueChangedDelegate = FSimpleDelegate::CreateSP(this, &FRigModuleInstanceDetails::OnConfigValueChanged, Variable.Name);
+			const FSimpleDelegate OnValueChangedDelegate = FSimpleDelegate::CreateSP(this, &FRigModuleInstanceDetails::OnConfigValueChanged, Property->GetFName());
+			DefaultProperty->SetOnPropertyValueChanged(OnValueChangedDelegate);
+			DefaultProperty->SetOnChildPropertyValueChanged(OnValueChangedDelegate);
 
-						TSharedPtr<IPropertyHandle> Handle = Row->GetPropertyHandle();
-						Handle->SetOnPropertyValueChanged(OnValueChangedDelegate);
-						Handle->SetOnChildPropertyValueChanged(OnValueChangedDelegate);
-						
-					}
-				}
-			}
+			FPropertyBindingWidgetArgs BindingArgs;
+			BindingArgs.Property = (FProperty*)Property;
+			BindingArgs.CurrentBindingText = TAttribute<FText>::CreateLambda([this, Property]()
+			{
+				return GetBindingText(Property);
+			});
+			BindingArgs.CurrentBindingImage = TAttribute<const FSlateBrush*>::CreateLambda([this, Property]()
+			{
+				return GetBindingImage(Property);
+			});
+			BindingArgs.CurrentBindingColor = TAttribute<FLinearColor>::CreateLambda([this, Property]()
+			{
+				return GetBindingColor(Property);
+			});
+
+			BindingArgs.OnCanBindProperty.BindLambda([](const FProperty* InProperty) -> bool { return true; });
+			BindingArgs.OnCanBindToClass.BindLambda([](UClass* InClass) -> bool { return false; });
+			BindingArgs.OnCanRemoveBinding.BindRaw(this, &FRigModuleInstanceDetails::CanRemoveBinding);
+			BindingArgs.OnRemoveBinding.BindSP(this, &FRigModuleInstanceDetails::HandleRemoveBinding);
+
+			BindingArgs.bGeneratePureBindings = true;
+			BindingArgs.bAllowNewBindings = true;
+			BindingArgs.bAllowArrayElementBindings = false;
+			BindingArgs.bAllowStructMemberBindings = false;
+			BindingArgs.bAllowUObjectFunctions = false;
+
+			BindingArgs.MenuExtender = MakeShareable(new FExtender);
+			BindingArgs.MenuExtender->AddMenuExtension(
+				"Properties",
+				EExtensionHook::After,
+				nullptr,
+				FMenuExtensionDelegate::CreateSPLambda(this, [this, Property](FMenuBuilder& MenuBuilder)
+				{
+					FillBindingMenu(MenuBuilder, Property);
+				})
+			);
+
+			// todo: remove the original row.
+
+			ConfigValuesCategory.AddCustomRow(DefaultProperty->GetPropertyDisplayName())
+			.NameContent()
+			[
+				DefaultProperty->CreatePropertyNameWidget()
+			]
+
+			// note: this doesn't work for some reason. seeking help from the editor team
+			.ValueContent()
+			[
+				DefaultProperty->CreatePropertyValueWidget()
+				// todo: if the property is bound / or partially bound
+				// mark the property value widget as disabled / read only.
+			]
+			
+			.ExtensionContent()
+			[
+				PropertyAccessEditor.MakePropertyBindingWidget(nullptr, BindingArgs)
+			];
 		}
 	}
 }
 
-FString FRigModuleInstanceDetails::GetModulePath() const
-{
-	check(PerModuleInfos.Num() == 1);
-	if (FRigModuleInstance* Module = PerModuleInfos[0].GetModule())
-	{
-		return Module->GetPath();
-	}
-	return FString();
-}
-
 FText FRigModuleInstanceDetails::GetName() const
 {
 	if(PerModuleInfos.Num() > 1)
@@ -377,54 +425,26 @@ FRigElementKeyRedirector FRigModuleInstanceDetails::GetConnections() const
 	return FRigElementKeyRedirector();
 }
 
-TArray<FRigVMExternalVariable> FRigModuleInstanceDetails::GetConfigValues() const
-{
-	if(PerModuleInfos.Num() > 1)
-	{
-		return TArray<FRigVMExternalVariable>();
-	}
-
-	if (FRigModuleInstance* Module = PerModuleInfos[0].GetModule())
-	{
-		if (TSoftObjectPtr<UControlRig> Rig = Module->Rig)
-		{
-			if (Rig.IsValid())
-			{
-				return Rig->GetExternalVariables();
-			}
-		}
-	}
-
-	return TArray<FRigVMExternalVariable>();
-}
-
 void FRigModuleInstanceDetails::OnConfigValueChanged(const FName InVariableName)
 {
-	if (FRigModuleInstance* ModuleInstance = PerModuleInfos[0].GetModule())
+	for(const FPerModuleInfo& Info : PerModuleInfos)
 	{
-		TSoftObjectPtr<UControlRig> Rig = ModuleInstance->Rig;
-		if (Rig.IsValid())
+		if (const FRigModuleInstance* ModuleInstance = Info.GetModule())
 		{
-			FString ValueStr = Rig->GetVariableAsString(InVariableName);
-			if (UControlRigBlueprint* Blueprint = RigModuleDetails_GetBlueprintFromRig(PerModuleInfos[0].GetModularRig()))
+			TSoftObjectPtr<UControlRig> Rig = ModuleInstance->Rig;
+			if (Rig.IsValid())
 			{
-				UModularRigController* Controller = Blueprint->GetModularRigController();
-				Controller->SetConfigValueInModule(ModuleInstance->GetPath(), InVariableName, ValueStr);
+				FString ValueStr = Rig->GetVariableAsString(InVariableName);
+				if (UControlRigBlueprint* Blueprint = Info.GetBlueprint())
+				{
+					UModularRigController* Controller = Blueprint->GetModularRigController();
+					Controller->SetConfigValueInModule(ModuleInstance->GetPath(), InVariableName, ValueStr);
+				}
 			}
 		}
 	}
 }
 
-TArray<FString> FRigModuleInstanceDetails::GetModulePaths() const
-{
-	TArray<FString> Paths;
-	Algo::Transform(PerModuleInfos, Paths, [](const FPerModuleInfo& Info)
-	{
-		return Info.GetModule()->GetPath();
-	});
-	return Paths;
-}
-
 const FRigModuleInstanceDetails::FPerModuleInfo& FRigModuleInstanceDetails::FindModule(const FString& InPath) const
 {
 	const FPerModuleInfo* Info = FindModuleByPredicate([InPath](const FPerModuleInfo& Info)
@@ -547,5 +567,256 @@ ERigElementType FRigModuleInstanceDetails::GetElementType(FRigElementKey Connect
 	return ERigElementType::None;
 }
 
+FText FRigModuleInstanceDetails::GetBindingText(const FProperty* InProperty) const
+{
+	const FName VariableName = InProperty->GetFName();
+	FText FirstValue;
+	for(const FPerModuleInfo& Info : PerModuleInfos)
+	{
+		if (const FRigModuleReference* ModuleReference = Info.GetReference())
+		{
+			if(ModuleReference->Bindings.Contains(VariableName))
+			{
+				const FText BindingText = FText::FromString(ModuleReference->Bindings.FindChecked(VariableName));
+				if(FirstValue.IsEmpty())
+				{
+					FirstValue = BindingText;
+				}
+				else if(!FirstValue.EqualTo(BindingText))
+				{
+					return ControlRigModuleDetailsMultipleValues;
+				}
+			}
+		}
+	}
+	return FirstValue;
+}
+
+const FSlateBrush* FRigModuleInstanceDetails::GetBindingImage(const FProperty* InProperty) const
+{
+	static FName TypeIcon(TEXT("Kismet.VariableList.TypeIcon"));
+	static FName ArrayTypeIcon(TEXT("Kismet.VariableList.ArrayTypeIcon"));
+
+	if(CastField<FArrayProperty>(InProperty))
+	{
+		return FAppStyle::GetBrush(ArrayTypeIcon);
+	}
+	return FAppStyle::GetBrush(TypeIcon);
+}
+
+FLinearColor FRigModuleInstanceDetails::GetBindingColor(const FProperty* InProperty) const
+{
+	if(InProperty)
+	{
+		FEdGraphPinType PinType;
+		const UEdGraphSchema_K2* Schema_K2 = GetDefault<UEdGraphSchema_K2>();
+		if (Schema_K2->ConvertPropertyToPinType(InProperty, PinType))
+		{
+			const URigVMEdGraphSchema* Schema = GetDefault<URigVMEdGraphSchema>();
+			return Schema->GetPinTypeColor(PinType);
+		}
+	}
+	return FLinearColor::White;
+}
+
+void FRigModuleInstanceDetails::FillBindingMenu(FMenuBuilder& MenuBuilder, const FProperty* InProperty) const
+{
+	if(PerModuleInfos.IsEmpty())
+	{
+		return;
+	}
+
+	UControlRigBlueprint* Blueprint = PerModuleInfos[0].GetBlueprint();
+	UModularRigController* Controller = Blueprint->GetModularRigController();
+
+	TArray<FString> CombinedBindings;
+	for(int32 Index = 0; Index < PerModuleInfos.Num(); Index++)
+	{
+		const FPerModuleInfo& Info  = PerModuleInfos[Index];
+		const TArray<FString> Bindings = Controller->GetPossibleBindings(Info.GetPath(), InProperty->GetFName());
+		if(Index == 0)
+		{
+			CombinedBindings = Bindings;
+		}
+		else
+		{
+			// reduce the set of bindings to the overall possible bindings
+			CombinedBindings.RemoveAll([Bindings](const FString& Binding)
+			{
+				return !Bindings.Contains(Binding);
+			});
+		}
+	}
+
+	// sort lexically
+	CombinedBindings.Sort();
+
+	// create a map of all of the variables per menu prefix (the module path the variables belong to)
+	struct FPerMenuData
+	{
+		FString Name;
+		FString ParentMenuPath;
+		TArray<FString> SubMenuPaths;
+		TArray<FString> Variables;
+
+		static void SetupMenu(
+			TSharedRef<FRigModuleInstanceDetails const> ThisDetails,
+			const FProperty* InProperty,
+			FMenuBuilder& InMenuBuilder,
+			const FString& InMenuPath,
+			TSharedRef<TMap<FString, FPerMenuData>> PerMenuData)
+		{
+			FPerMenuData& Data = PerMenuData->FindChecked((InMenuPath));
+
+			Data.SubMenuPaths.Sort();
+			Data.Variables.Sort();
+
+			for(const FString& VariablePath : Data.Variables)
+			{
+				FString VariableName = VariablePath;
+				(void)VariablePath.Split(UModularRig::NamespaceSeparator, nullptr, &VariableName, ESearchCase::CaseSensitive, ESearchDir::FromEnd);
+				
+				InMenuBuilder.AddMenuEntry(
+					FUIAction(FExecuteAction::CreateLambda([ThisDetails, InProperty, VariablePath]()
+					{
+						ThisDetails->HandleChangeBinding(InProperty, VariablePath);
+					})),
+					SNew(SHorizontalBox)
+						+SHorizontalBox::Slot()
+						.AutoWidth()
+						.VAlign(VAlign_Center)
+						.Padding(1.0f, 0.0f)
+						[
+							SNew(SImage)
+							.Image(ThisDetails->GetBindingImage(InProperty))
+							.ColorAndOpacity(ThisDetails->GetBindingColor(InProperty))
+						]
+						+SHorizontalBox::Slot()
+						.AutoWidth()
+						.VAlign(VAlign_Center)
+						.Padding(4.0f, 0.0f)
+						[
+							SNew(STextBlock)
+							.Text(FText::FromString(VariableName))
+							.ColorAndOpacity(FLinearColor::White)
+						]
+					);
+			}
+
+			for(const FString& SubMenuPath : Data.SubMenuPaths)
+			{
+				const FPerMenuData& SubMenuData = PerMenuData->FindChecked(SubMenuPath);
+
+				const FText Label = FText::FromString(SubMenuData.Name);
+				static const FText TooltipFormat = LOCTEXT("BindingMenuTooltipFormat", "Access to all variables of the {0} module");
+				const FText Tooltip = FText::Format(TooltipFormat, Label);  
+				InMenuBuilder.AddSubMenu(Label, Tooltip, FNewMenuDelegate::CreateLambda([ThisDetails, InProperty, SubMenuPath, PerMenuData](FMenuBuilder& SubMenuBuilder)
+				{
+					SetupMenu(ThisDetails, InProperty, SubMenuBuilder, SubMenuPath, PerMenuData);
+				}));
+			}
+		}
+	};
+	
+	// define the root menu
+	const TSharedRef<TMap<FString, FPerMenuData>> PerMenuData = MakeShared<TMap<FString, FPerMenuData>>();
+	PerMenuData->FindOrAdd(FString());
+
+	// make sure all levels of the menu are known and we have the variables available
+	for(const FString& BindingPath : CombinedBindings)
+	{
+		FString MenuPath;
+		(void)BindingPath.Split(UModularRig::NamespaceSeparator, &MenuPath, nullptr, ESearchCase::CaseSensitive, ESearchDir::FromEnd);
+
+		FString PreviousMenuPath = MenuPath;
+		FString ParentMenuPath = MenuPath, RemainingPath;
+		while(ParentMenuPath.Split(UModularRig::NamespaceSeparator, &ParentMenuPath, &RemainingPath, ESearchCase::CaseSensitive, ESearchDir::FromEnd))
+		{
+			// scope since the map may change at the end of this block
+			{
+				FPerMenuData& Data = PerMenuData->FindOrAdd(MenuPath);
+				if(Data.Name.IsEmpty())
+				{
+					Data.Name = RemainingPath;
+				}
+			}
+			
+			PerMenuData->FindOrAdd(ParentMenuPath).SubMenuPaths.AddUnique(PreviousMenuPath);
+			PerMenuData->FindOrAdd(PreviousMenuPath).ParentMenuPath = ParentMenuPath;
+			PerMenuData->FindOrAdd(PreviousMenuPath).Name = RemainingPath;
+			if(!ParentMenuPath.Contains(UModularRig::NamespaceSeparator))
+			{
+				PerMenuData->FindOrAdd(FString()).SubMenuPaths.AddUnique(ParentMenuPath);
+				PerMenuData->FindOrAdd(ParentMenuPath).Name = ParentMenuPath;
+			}
+			PreviousMenuPath = ParentMenuPath;
+		}
+
+		FPerMenuData& Data = PerMenuData->FindOrAdd(MenuPath);
+		if(Data.Name.IsEmpty())
+		{
+			Data.Name = MenuPath;
+		}
+
+		Data.Variables.Add(BindingPath);
+		if(!MenuPath.IsEmpty())
+		{
+			PerMenuData->FindChecked(Data.ParentMenuPath).SubMenuPaths.AddUnique(MenuPath);
+		}
+	}
+
+	// build the menu
+	FPerMenuData::SetupMenu(SharedThis(this), InProperty, MenuBuilder, FString(), PerMenuData);
+}
+
+bool FRigModuleInstanceDetails::CanRemoveBinding(FName InPropertyName) const
+{
+	// offer the "removing binding" button if any of the selected module instances
+	// has a binding for the given variable
+	for(const FPerModuleInfo& Info : PerModuleInfos)
+	{
+		if (const FRigModuleInstance* ModuleInstance = Info.GetModule())
+		{
+			if(ModuleInstance->VariableBindings.Contains(InPropertyName))
+			{
+				return true;
+			}
+		}
+	}
+	return false; 
+}
+
+void FRigModuleInstanceDetails::HandleRemoveBinding(FName InPropertyName) const
+{
+	FScopedTransaction Transaction(LOCTEXT("BindModuleVariableTransaction", "Remove Binding"));
+	for(const FPerModuleInfo& Info : PerModuleInfos)
+	{
+		if (UControlRigBlueprint* Blueprint = Info.GetBlueprint())
+		{
+			if (const FRigModuleInstance* ModuleInstance = Info.GetModule())
+			{
+				UModularRigController* Controller = Blueprint->GetModularRigController();
+				Controller->UnBindModuleVariable(ModuleInstance->GetPath(), InPropertyName);
+			}
+		}
+	}
+}
+
+void FRigModuleInstanceDetails::HandleChangeBinding(const FProperty* InProperty, const FString& InNewVariablePath) const
+{
+	FScopedTransaction Transaction(LOCTEXT("BindModuleVariableTransaction", "Bind Module Variable"));
+	for(const FPerModuleInfo& Info : PerModuleInfos)
+	{
+		if (UControlRigBlueprint* Blueprint = Info.GetBlueprint())
+		{
+			if (const FRigModuleInstance* ModuleInstance = Info.GetModule())
+			{
+				UModularRigController* Controller = Blueprint->GetModularRigController();
+				Controller->BindModuleVariable(ModuleInstance->GetPath(), InProperty->GetFName(), InNewVariablePath);
+			}
+		}
+	}
+}
+
 
 #undef LOCTEXT_NAMESPACE
diff --git a/Engine/Plugins/Animation/ControlRig/Source/ControlRigEditor/Private/ControlRigModuleDetails.h b/Engine/Plugins/Animation/ControlRig/Source/ControlRigEditor/Private/ControlRigModuleDetails.h
--- a/Engine/Plugins/Animation/ControlRig/Source/ControlRigEditor/Private/ControlRigModuleDetails.h
+++ b/Engine/Plugins/Animation/ControlRig/Source/ControlRigEditor/Private/ControlRigModuleDetails.h
@@ -31,7 +31,6 @@ public:
 
 	/** IDetailCustomization interface */
 	virtual void CustomizeDetails(IDetailLayoutBuilder& DetailBuilder) override;
-	virtual void BeginDestroy() {};
 
 	// Makes a new instance of this detail layout class for a specific detail view requesting it
 	static TSharedRef<IDetailCustomization> MakeInstance()
@@ -39,36 +38,17 @@ public:
 		return MakeShareable(new FRigModuleInstanceDetails);
 	}
 
-	FString GetModulePath() const;
 	FText GetName() const;
 	FText GetRigClassPath() const;
 	TArray<FRigModuleConnector> GetConnectors() const;
 	FRigElementKeyRedirector GetConnections() const;
-	TArray<FRigVMExternalVariable> GetConfigValues() const;
 
 	void OnConfigValueChanged(const FName InVariableName);
 	
-	TArray<FString> GetModulePaths() const;
-
-	TArray<FRigModuleInstance> GetModulesInDetailsView(const TArray<FString>& InFilter = TArray<FString>()) const
-	{
-		TArray<FRigModuleInstance> Elements;
-		for(const FPerModuleInfo& Info : PerModuleInfos)
-		{
-			FRigModuleInstance Content = Info.WrapperObject->GetContent<FRigModuleInstance>();
-			if(!InFilter.IsEmpty() && !InFilter.Contains(Content.GetPath()))
-			{
-				continue;
-			}
-			Elements.Add(Content);
-		}
-		return Elements;
-	}
-
 	struct FPerModuleInfo
 	{
 		FPerModuleInfo()
-			: WrapperObject()
+			: Path()
 			, Module()
 			, DefaultModule()
 		{}
@@ -76,7 +56,16 @@ public:
 		bool IsValid() const { return Module.IsValid(); }
 		operator bool() const { return IsValid(); }
 
-		UModularRig* GetModularRig() const { return (UModularRig*)Module.GetModularRig(); }
+		const FString& GetPath() const
+		{
+			return Path;
+		}
+
+		UModularRig* GetModularRig() const
+		{
+			return (UModularRig*)Module.GetModularRig();
+		}
+		
 		UModularRig* GetDefaultRig() const
 		{
 			if(DefaultModule.IsValid())
@@ -92,7 +81,7 @@ public:
 			{
 				return Cast<UControlRigBlueprint>(ControlRig->GetClass()->ClassGeneratedBy);
 			}
-			return GetDefaultRig()->GetTypedOuter<UControlRigBlueprint>();
+			return nullptr;
 		}
 
 		FRigModuleInstance* GetModule() const
@@ -100,7 +89,7 @@ public:
 			return (FRigModuleInstance*)Module.Get();
 		}
 
-		FRigModuleInstance* GetDefaultElement() const
+		FRigModuleInstance* GetDefaultModule() const
 		{
 			if(DefaultModule)
 			{
@@ -109,7 +98,16 @@ public:
 			return GetModule();
 		}
 
-		TWeakObjectPtr<URigVMDetailsViewWrapperObject> WrapperObject;
+		const FRigModuleReference* GetReference() const
+		{
+			if(const UControlRigBlueprint* Blueprint = GetBlueprint())
+			{
+				return Blueprint->ModularRigModel.FindModule(Path);
+			}
+			return nullptr;
+		}
+
+		FString Path;
 		FModuleInstanceHandle Module;
 		FModuleInstanceHandle DefaultModule;
 	};
@@ -129,6 +127,14 @@ public:
 
 protected:
 
+	FText GetBindingText(const FProperty* InProperty) const;
+	const FSlateBrush* GetBindingImage(const FProperty* InProperty) const;
+	FLinearColor GetBindingColor(const FProperty* InProperty) const;
+	void FillBindingMenu(FMenuBuilder& MenuBuilder, const FProperty* InProperty) const;
+	bool CanRemoveBinding(FName InPropertyName) const;
+	void HandleRemoveBinding(FName InPropertyName) const;
+	void HandleChangeBinding(const FProperty* InProperty, const FString& InNewVariablePath) const;
+
 	TArray<FPerModuleInfo> PerModuleInfos;
 	TSharedPtr<SRigElementKeyWidget> RigElementKeyWidget;
 	TMap<FRigElementKey, FRigElementKey> Connections;
diff --git a/Engine/Plugins/Animation/ControlRig/Source/ControlRigEditor/Private/Editor/ControlRigEditor.cpp b/Engine/Plugins/Animation/ControlRig/Source/ControlRigEditor/Private/Editor/ControlRigEditor.cpp
--- a/Engine/Plugins/Animation/ControlRig/Source/ControlRigEditor/Private/Editor/ControlRigEditor.cpp
+++ b/Engine/Plugins/Animation/ControlRig/Source/ControlRigEditor/Private/Editor/ControlRigEditor.cpp
@@ -152,6 +152,7 @@ FControlRigEditor::~FControlRigEditor()
 		if (RigBlueprint->IsModularRig())
 		{
 			RigBlueprint->GetModularRigController()->OnModified().RemoveAll(this);
+			RigBlueprint->OnModularRigCompiled().RemoveAll(this);
 		}
 	}
 
@@ -320,6 +321,7 @@ void FControlRigEditor::InitRigVMEditor(const EToolkitMode::Type Mode, const TSh
 		if (ControlRigBlueprint->IsModularRig())
 		{
 			ControlRigBlueprint->GetModularRigController()->OnModified().AddSP(this, &FControlRigEditor::HandleModularRigModified);
+			ControlRigBlueprint->OnModularRigCompiled().AddRaw(this, &FControlRigEditor::HandlePostCompileModularRigs);
 		}
 	}
 
@@ -869,7 +871,7 @@ void FControlRigEditor::SetEventQueue(TArray<FName> InEventQueue, bool bCompile)
 			}
 		}
 
-		// Reset transforms only for construction and forward solve to not inturrupt any animation that might be playing
+		// Reset transforms only for construction and forward solve to not interrupt any animation that might be playing
 		if (InEventQueue.Contains(FRigUnit_PrepareForExecution::EventName) ||
 			InEventQueue.Contains(FRigUnit_BeginExecution::EventName))
 		{
@@ -1109,22 +1111,11 @@ void FControlRigEditor::SetDetailViewForRigElements(const TArray<FRigElementKey>
 
 void FControlRigEditor::SetDetailObjects(const TArray<UObject*>& InObjects)
 {
-	bool bAnyModule = false;
-	for (const UObject* Object : InObjects)
+	// if no modules should be selected - we need to deselect all modules
+	if (!InObjects.ContainsByPredicate([](const UObject* InObject) -> bool
 	{
-		if (const URigVMDetailsViewWrapperObject* WrapperObject = Cast<URigVMDetailsViewWrapperObject>(Object))
-		{
-			if (const UScriptStruct* WrappedStruct = WrapperObject->GetWrappedStruct())
-			{
-				if (WrappedStruct->IsChildOf(FRigModuleInstance::StaticStruct()))
-				{
-					bAnyModule = true;
-				}
-			}
-		}
-	}
-	
-	if (!bAnyModule)
+		return IsValid(InObject) && InObject->IsA<UControlRig>();
+	}))
 	{
 		ModulesSelected.Reset();
 	}
@@ -1205,24 +1196,23 @@ void FControlRigEditor::SetDetailViewForRigModules(const TArray<FString> InKeys)
 
 	for(const FString& Key : InKeys)
 	{
-		FRigModuleInstance* Element = RigBeingDebugged->FindModule(Key);
+		const FRigModuleInstance* Element = RigBeingDebugged->FindModule(Key);
 		if (Element == nullptr)
 		{
 			continue;
 		}
 
-		URigVMDetailsViewWrapperObject* WrapperObject = URigVMDetailsViewWrapperObject::MakeInstance(GetDetailWrapperClass(), GetBlueprintObj(), FRigModuleInstance::StaticStruct(), (uint8*)Element, RigBeingDebugged);
-		WrapperObject->GetWrappedPropertyChangedChainEvent().AddSP(this, &FControlRigEditor::OnWrappedPropertyChangedChainEvent);
-		WrapperObject->AddToRoot();
-
-		Objects.Add(WrapperObject);
+		if(UControlRig* ModuleInstance = Element->Rig.Get())
+		{
+			Objects.Add(ModuleInstance);
+		}
 	}
 	
 	SetDetailObjects(Objects);
 
+	// In case the modules selected are still not available, lets set them again
 	if (Objects.IsEmpty())
 	{
-		// In case the modules selected are still not available, lets set them again
 		ModulesSelected = InKeys;
 	}
 }
@@ -3605,6 +3595,11 @@ void FControlRigEditor::HandleModularRigModified(EModularRigNotification InNotif
 	}
 }
 
+void FControlRigEditor::HandlePostCompileModularRigs(URigVMBlueprint* InBlueprint)
+{
+	RefreshDetailView();
+}
+
 void FControlRigEditor::SynchronizeViewportBoneSelection()
 {
 	UControlRigBlueprint* RigBlueprint = GetControlRigBlueprint();
diff --git a/Engine/Plugins/Animation/ControlRig/Source/ControlRigEditor/Private/Editor/ControlRigEditor.h b/Engine/Plugins/Animation/ControlRig/Source/ControlRigEditor/Private/Editor/ControlRigEditor.h
--- a/Engine/Plugins/Animation/ControlRig/Source/ControlRigEditor/Private/Editor/ControlRigEditor.h
+++ b/Engine/Plugins/Animation/ControlRig/Source/ControlRigEditor/Private/Editor/ControlRigEditor.h
@@ -158,6 +158,7 @@ public:
 	void HandleRigTypeChanged(UControlRigBlueprint* InBlueprint);
 
 	void HandleModularRigModified(EModularRigNotification InNotification, const FRigModuleReference* InModule);
+	void HandlePostCompileModularRigs(URigVMBlueprint* InBlueprint);
 
 	const FName RigHierarchyToGraphDragAndDropMenuName = TEXT("ControlRigEditor.RigHierarchyToGraphDragAndDropMenu");
 	void CreateRigHierarchyToGraphDragAndDropMenu() const;
diff --git a/Engine/Plugins/Animation/ControlRig/Source/ControlRigEditor/Private/Editor/ControlRigWrapperObject.cpp b/Engine/Plugins/Animation/ControlRig/Source/ControlRigEditor/Private/Editor/ControlRigWrapperObject.cpp
--- a/Engine/Plugins/Animation/ControlRig/Source/ControlRigEditor/Private/Editor/ControlRigWrapperObject.cpp
+++ b/Engine/Plugins/Animation/ControlRig/Source/ControlRigEditor/Private/Editor/ControlRigWrapperObject.cpp
@@ -47,14 +47,6 @@ UClass* UControlRigWrapperObject::GetClassForStruct(UScriptStruct* InStruct, boo
 			}
 		}
 	}
-	else if (InStruct->IsChildOf(FRigModuleInstance::StaticStruct()))
-	{
-		FPropertyEditorModule& PropertyEditorModule = FModuleManager::LoadModuleChecked<FPropertyEditorModule>("PropertyEditor");
-		if (!PropertyEditorModule.GetClassNameToDetailLayoutNameMap().Contains(WrapperClassName))
-		{
-			PropertyEditorModule.RegisterCustomClassLayout(WrapperClassName, FOnGetDetailCustomizationInstance::CreateStatic(&FRigModuleInstanceDetails::MakeInstance));
-		}
-	}
 #endif
 
 	return Class;
diff --git a/Engine/Plugins/Animation/ControlRig/Source/ControlRigEditor/Private/Editor/SModularRigHierarchyTreeView.cpp b/Engine/Plugins/Animation/ControlRig/Source/ControlRigEditor/Private/Editor/SModularRigHierarchyTreeView.cpp
--- a/Engine/Plugins/Animation/ControlRig/Source/ControlRigEditor/Private/Editor/SModularRigHierarchyTreeView.cpp
+++ b/Engine/Plugins/Animation/ControlRig/Source/ControlRigEditor/Private/Editor/SModularRigHierarchyTreeView.cpp
@@ -601,7 +601,7 @@ TPair<const FSlateBrush*, FSlateColor> SModularRigHierarchyItem::GetBrushForElem
 	const FSlateBrush* Brush = nullptr;
 	FSlateColor Color = FSlateColor::UseForeground();
 
-	if (FRigModuleInstance* Module = InHierarchy->FindModule(InKey))
+	if (const FRigModuleInstance* Module = InHierarchy->FindModule(InKey))
 	{
 		if (Module->Rig.IsValid())
 		{
