Get animation attributes working with control rig component

diff --git a/Engine/Plugins/Animation/ControlRig/Source/ControlRig/Private/ControlRigAnimInstance.cpp b/Engine/Plugins/Animation/ControlRig/Source/ControlRig/Private/ControlRigAnimInstance.cpp
--- a/Engine/Plugins/Animation/ControlRig/Source/ControlRig/Private/ControlRigAnimInstance.cpp
+++ b/Engine/Plugins/Animation/ControlRig/Source/ControlRig/Private/ControlRigAnimInstance.cpp
@@ -15,7 +15,7 @@ void FControlRigAnimInstanceProxy::Initialize(UAnimInstance* InAnimInstance)
 
 bool FControlRigAnimInstanceProxy::Evaluate(FPoseContext& Output)
 {
-	if (StoredTransforms.Num() == 0 && StoredCurves.Num() == 0)
+	if (StoredTransforms.Num() == 0 && StoredCurves.Num() == 0 && !StoredAttributes.ContainsData())
 	{
 		return false;
 	}
@@ -55,6 +55,8 @@ bool FControlRigAnimInstanceProxy::Evaluate(FPoseContext& Output)
 	UE::Anim::FCurveUtils::BuildUnsorted(Curve, StoredCurves);
 	Output.Curve.Combine(Curve);
 
+	Output.CustomAttributes.CopyFrom(StoredAttributes);
+
 	return true;
 }
 
diff --git a/Engine/Plugins/Animation/ControlRig/Source/ControlRig/Private/ControlRigComponent.cpp b/Engine/Plugins/Animation/ControlRig/Source/ControlRig/Private/ControlRigComponent.cpp
--- a/Engine/Plugins/Animation/ControlRig/Source/ControlRig/Private/ControlRigComponent.cpp
+++ b/Engine/Plugins/Animation/ControlRig/Source/ControlRig/Private/ControlRigComponent.cpp
@@ -12,6 +12,7 @@
 #include "ControlRigObjectBinding.h"
 #include "SceneManagement.h"
 #include "Math/ControlRigMathLibrary.h"
+#include "ControlRig.h"
 
 #include UE_INLINE_GENERATED_CPP_BY_NAME(ControlRigComponent)
 
@@ -492,7 +493,16 @@ void UControlRigComponent::Update(float DeltaTime)
 				{
 					CR->ClearPoseBeforeBackwardsSolve();
 				}
-				CR->Evaluate_AnyThread();
+
+				{
+					// Necessary for FStackAttributeContainer that uses a FAnimStackAllocator (TMemStackAllocator) which allocates from FMemStack.
+					// When allocating memory from FMemStack we need to explicitly use FMemMark to ensure items are freed when the scope exits. 
+					FMemMark Mark(FMemStack::Get());
+					TempAttributeContainer = MakeUnique<UE::Anim::FStackAttributeContainer>();
+					UControlRig::FAnimAttributeContainerPtrScope AttributeScope(CR, *TempAttributeContainer);
+				
+					CR->Evaluate_AnyThread();
+				}
 
 #if WITH_EDITOR
 				if(URigHierarchy* Hierarchy = CR->GetHierarchy())
@@ -1532,10 +1542,25 @@ void UControlRigComponent::TransferOutputs()
 				{
 					if (USkeletalMesh* SkeletalMesh = SkeletalMeshComponent->GetSkeletalMeshAsset())
 					{
+						USceneComponent* LastComponent = nullptr;
+						FControlRigAnimInstanceProxy* Proxy = nullptr;
+						
 						TArray<FTransform> BoneSpaceTransforms;
 						BoneSpaceTransforms.SetNumUninitialized(SkeletalMesh->GetRefSkeleton().GetNum());
 						for (const FControlRigComponentMappedElement& MappedElement : MappedElements)
 						{
+							if (LastComponent != MappedElement.SceneComponent || Proxy == nullptr)
+							{
+								Proxy = MappedElement.GetAnimProxyOnGameThread();
+								if (Proxy)
+								{
+									Proxy->StoredTransforms.Reset();
+									Proxy->StoredCurves.Reset();
+									Proxy->StoredAttributes.Empty();
+									LastComponent = MappedElement.SceneComponent;
+								}
+							}
+							
 							if (MappedElement.ElementIndex == INDEX_NONE || MappedElement.Direction == EControlRigComponentMapDirection::Input)
 							{
 								continue;
@@ -1556,9 +1581,11 @@ void UControlRigComponent::TransferOutputs()
 						TArray<FTransform> OutSpaceBases;
 						OutSpaceBases.SetNumUninitialized(BoneSpaceTransforms.Num());
 						SkeletalMesh->FillComponentSpaceTransforms(BoneSpaceTransforms, SkeletalMeshComponent->FillComponentSpaceTransformsRequiredBones, SkeletalMeshComponent->GetEditableComponentSpaceTransforms());
-#if WITH_EDITOR
-						SkeletalMeshComponent->ApplyEditedComponentSpaceTransforms();
-#endif
+
+						if (Proxy)
+						{
+							Proxy->StoredAttributes.CopyFrom(*TempAttributeContainer);
+						}	
 					}
 				}
 			}
@@ -1577,6 +1604,7 @@ void UControlRigComponent::TransferOutputs()
 					{
 						Proxy->StoredTransforms.Reset();
 						Proxy->StoredCurves.Reset();
+						Proxy->StoredAttributes.Empty();
 						LastComponent = MappedElement.SceneComponent;
 					}
 				}
@@ -1661,6 +1689,11 @@ void UControlRigComponent::TransferOutputs()
 					}
 				}
 			}
+
+			if (Proxy)
+			{
+				Proxy->StoredAttributes.CopyFrom(*TempAttributeContainer);
+			}			
 		}
 
 #if WITH_EDITOR
@@ -1743,6 +1776,7 @@ void UControlRigComponent::HandleControlRigPreConstructionEvent(UControlRig* InC
 			{
 				Proxy->StoredTransforms.Reset();
 				Proxy->StoredCurves.Reset();
+				Proxy->StoredAttributes.Empty();
 				LastComponent = MappedElement.SceneComponent;
 			}
 		}
diff --git a/Engine/Plugins/Animation/ControlRig/Source/ControlRig/Public/ControlRigAnimInstance.h b/Engine/Plugins/Animation/ControlRig/Source/ControlRig/Public/ControlRigAnimInstance.h
--- a/Engine/Plugins/Animation/ControlRig/Source/ControlRig/Public/ControlRigAnimInstance.h
+++ b/Engine/Plugins/Animation/ControlRig/Source/ControlRig/Public/ControlRigAnimInstance.h
@@ -33,6 +33,7 @@ public:
 
 	TMap<FMeshPoseBoneIndex, FTransform> StoredTransforms;
 	TMap<FName, float> StoredCurves;
+	UE::Anim::FHeapAttributeContainer StoredAttributes;
 };
 
 UCLASS(transient, NotBlueprintable)
diff --git a/Engine/Plugins/Animation/ControlRig/Source/ControlRig/Public/ControlRigComponent.h b/Engine/Plugins/Animation/ControlRig/Source/ControlRig/Public/ControlRigComponent.h
--- a/Engine/Plugins/Animation/ControlRig/Source/ControlRig/Public/ControlRigComponent.h
+++ b/Engine/Plugins/Animation/ControlRig/Source/ControlRig/Public/ControlRigComponent.h
@@ -661,6 +661,8 @@ private:
 
 	UPROPERTY(transient)
 	TObjectPtr<UControlRig> ControlRig;
+	
+	TUniquePtr<UE::Anim::FStackAttributeContainer> TempAttributeContainer;
 
 	TMap<USkeletalMeshComponent*, FCachedSkeletalMeshComponentSettings> CachedSkeletalMeshComponentSettings;
 
diff --git a/Engine/Plugins/Animation/ControlRig/Source/ControlRig/Public/Units/Animation/RigUnit_AnimAttribute.h b/Engine/Plugins/Animation/ControlRig/Source/ControlRig/Public/Units/Animation/RigUnit_AnimAttribute.h
--- a/Engine/Plugins/Animation/ControlRig/Source/ControlRig/Public/Units/Animation/RigUnit_AnimAttribute.h
+++ b/Engine/Plugins/Animation/ControlRig/Source/ControlRig/Public/Units/Animation/RigUnit_AnimAttribute.h
@@ -6,6 +6,7 @@
 #include "Components/SkeletalMeshComponent.h"
 #include "Units/RigDispatchFactory.h"
 #include "Animation/BuiltInAttributeTypes.h"
+#include "ControlRigComponent.h"
 #include "RigUnit_AnimAttribute.generated.h"
 
 namespace RigUnit_AnimAttribute
@@ -76,6 +77,16 @@ namespace RigUnit_AnimAttribute
 		}
 	
 		const USkeletalMeshComponent* OwningComponent = Cast<USkeletalMeshComponent>(Context.GetOwningComponent());
+		if (!OwningComponent)
+		{
+			if (const UControlRigComponent* ControlRigComponent = Cast<UControlRigComponent>(Context.GetOwningComponent()))
+			{
+				if (ControlRigComponent->MappedElements.Num() > 0)
+				{
+					OwningComponent = Cast<USkeletalMeshComponent>(ControlRigComponent->MappedElements[0].SceneComponent);
+				}				
+			}
+		}
 
 		if (!OwningComponent ||
 			!OwningComponent->GetSkeletalMeshAsset())
