Control Rig: Update UI code modular rig terminoligy

#rnx

diff --git a/Engine/Plugins/Animation/ControlRig/Source/ControlRigEditor/Private/ControlRigBlueprintEditorLibrary.cpp b/Engine/Plugins/Animation/ControlRig/Source/ControlRigEditor/Private/ControlRigBlueprintEditorLibrary.cpp
--- a/Engine/Plugins/Animation/ControlRig/Source/ControlRigEditor/Private/ControlRigBlueprintEditorLibrary.cpp
+++ b/Engine/Plugins/Animation/ControlRig/Source/ControlRigEditor/Private/ControlRigBlueprintEditorLibrary.cpp
@@ -2,7 +2,7 @@
 
 #include "ControlRigBlueprintEditorLibrary.h"
 #include "Editor/SRigHierarchy.h"
-#include "Editor/SModularRigHierarchy.h"
+#include "Editor/SModularRigModel.h"
 #include "AssetRegistry/AssetRegistryModule.h"
 #include "RigVMModel/RigVMPin.h"
 
@@ -89,7 +89,7 @@ void UControlRigBlueprintEditorLibrary::SetupAllEditorMenus()
 {
 	SRigHierarchy::CreateContextMenu();
 	SRigHierarchy::CreateDragDropMenu();
-	SModularRigHierarchy::CreateContextMenu();
+	SModularRigModel::CreateContextMenu();
 }
 
 TArray<FRigModuleDescription> UControlRigBlueprintEditorLibrary::GetAvailableRigModules()
diff --git a/Engine/Plugins/Animation/ControlRig/Source/ControlRigEditor/Private/ControlRigEditorModule.cpp b/Engine/Plugins/Animation/ControlRig/Source/ControlRigEditor/Private/ControlRigEditorModule.cpp
--- a/Engine/Plugins/Animation/ControlRig/Source/ControlRigEditor/Private/ControlRigEditorModule.cpp
+++ b/Engine/Plugins/Animation/ControlRig/Source/ControlRigEditor/Private/ControlRigEditorModule.cpp
@@ -38,7 +38,7 @@
 #include "Graph/ControlRigGraphPanelPinFactory.h"
 #include <Editor/ControlRigEditorCommands.h>
 #include "ControlRigHierarchyCommands.h"
-#include "ControlRigModularRigHierarchyCommands.h"
+#include "ControlRigModularRigCommands.h"
 #include "Animation/AnimSequence.h"
 #include "Editor/ControlRigEditorEditMode.h"
 #include "ControlRigElementDetails.h"
@@ -106,7 +106,7 @@ void FControlRigEditorModule::StartupModule()
 	FControlRigEditModeCommands::Register();
 	FControlRigEditorCommands::Register();
 	FControlRigHierarchyCommands::Register();
-	FControlRigModularHierarchyCommands::Register();
+	FControlRigModularRigCommands::Register();
 	FControlRigEditorStyle::Get();
 
 	EdGraphPanelPinFactory = MakeShared<FControlRigGraphPanelPinFactory>();
diff --git a/Engine/Plugins/Animation/ControlRig/Source/ControlRigEditor/Private/ControlRigModularRigCommands.cpp b/Engine/Plugins/Animation/ControlRig/Source/ControlRigEditor/Private/ControlRigModularRigCommands.cpp
new file mode 100644
--- /dev/null
+++ b/Engine/Plugins/Animation/ControlRig/Source/ControlRigEditor/Private/ControlRigModularRigCommands.cpp
@@ -0,0 +1,14 @@
+// Copyright Epic Games, Inc. All Rights Reserved.
+
+#include "ControlRigModularRigCommands.h"
+
+#define LOCTEXT_NAMESPACE "ControlRigModularRigCommands"
+
+void FControlRigModularRigCommands::RegisterCommands()
+{
+	UI_COMMAND(AddModuleItem, "New Module", "Add new module to the rig.", EUserInterfaceActionType::Button, FInputChord());
+	UI_COMMAND(RenameModuleItem, "Rename", "Rename module to the rig.", EUserInterfaceActionType::Button, FInputChord(EKeys::F2));
+	UI_COMMAND(DeleteModuleItem, "Delete", "Delete module from the rig.", EUserInterfaceActionType::Button, FInputChord(EKeys::Delete));
+}
+
+#undef LOCTEXT_NAMESPACE
diff --git a/Engine/Plugins/Animation/ControlRig/Source/ControlRigEditor/Private/ControlRigModularRigCommands.h b/Engine/Plugins/Animation/ControlRig/Source/ControlRigEditor/Private/ControlRigModularRigCommands.h
new file mode 100644
--- /dev/null
+++ b/Engine/Plugins/Animation/ControlRig/Source/ControlRigEditor/Private/ControlRigModularRigCommands.h
@@ -0,0 +1,34 @@
+// Copyright Epic Games, Inc. All Rights Reserved.
+
+#pragma once
+
+#include "CoreMinimal.h"
+#include "Framework/Commands/Commands.h"
+#include "ControlRigEditorStyle.h"
+
+class FControlRigModularRigCommands : public TCommands<FControlRigModularRigCommands>
+{
+public:
+	FControlRigModularRigCommands() : TCommands<FControlRigModularRigCommands>
+	(
+		"ControlRigModularRigModel",
+		NSLOCTEXT("Contexts", "ModularRigModel", "Modular Rig Modules"),
+		NAME_None, // "MainFrame" // @todo Fix this crash
+		FControlRigEditorStyle::Get().GetStyleSetName() // Icon Style Set
+	)
+	{}
+	
+	/** Add Module at root */
+	TSharedPtr< FUICommandInfo > AddModuleItem;
+
+	/** Rename Module */
+	TSharedPtr< FUICommandInfo > RenameModuleItem;
+
+	/** Delete Module */
+	TSharedPtr< FUICommandInfo > DeleteModuleItem;
+
+	/**
+	 * Initialize commands
+	 */
+	virtual void RegisterCommands() override;
+};
diff --git a/Engine/Plugins/Animation/ControlRig/Source/ControlRigEditor/Private/ControlRigModularRigHierarchyCommands.cpp b/Engine/Plugins/Animation/ControlRig/Source/ControlRigEditor/Private/ControlRigModularRigHierarchyCommands.cpp
deleted file mode 100644
diff --git a/Engine/Plugins/Animation/ControlRig/Source/ControlRigEditor/Private/ControlRigModularRigHierarchyCommands.h b/Engine/Plugins/Animation/ControlRig/Source/ControlRigEditor/Private/ControlRigModularRigHierarchyCommands.h
deleted file mode 100644
diff --git a/Engine/Plugins/Animation/ControlRig/Source/ControlRigEditor/Private/ControlRigModuleDetails.cpp b/Engine/Plugins/Animation/ControlRig/Source/ControlRigEditor/Private/ControlRigModuleDetails.cpp
--- a/Engine/Plugins/Animation/ControlRig/Source/ControlRigEditor/Private/ControlRigModuleDetails.cpp
+++ b/Engine/Plugins/Animation/ControlRig/Source/ControlRigEditor/Private/ControlRigModuleDetails.cpp
@@ -17,7 +17,7 @@
 #include "SEnumCombo.h"
 #include "HAL/PlatformApplicationMisc.h"
 #include "Styling/AppStyle.h"
-#include "Editor/SModularRigHierarchyTreeView.h"
+#include "Editor/SModularRigTreeView.h"
 #include "StructViewerFilter.h"
 #include "StructViewerModule.h"
 #include "Features/IModularFeatures.h"
diff --git a/Engine/Plugins/Animation/ControlRig/Source/ControlRigEditor/Private/Editor/ControlRigContextMenuContext.cpp b/Engine/Plugins/Animation/ControlRig/Source/ControlRigEditor/Private/Editor/ControlRigContextMenuContext.cpp
--- a/Engine/Plugins/Animation/ControlRig/Source/ControlRigEditor/Private/Editor/ControlRigContextMenuContext.cpp
+++ b/Engine/Plugins/Animation/ControlRig/Source/ControlRigEditor/Private/Editor/ControlRigContextMenuContext.cpp
@@ -80,11 +80,11 @@ SRigHierarchy* UControlRigContextMenuContext::GetRigHierarchyPanel() const
 	return nullptr;
 }
 
-SModularRigHierarchy* UControlRigContextMenuContext::GetModularRigHierarchyPanel() const
+SModularRigModel* UControlRigContextMenuContext::GetModularRigModelPanel() const
 {
-	if (const TSharedPtr<SModularRigHierarchy> ModularRigHierarchyPanel = MenuSpecificContext.ModularRigHierarchyPanel.Pin())
+	if (const TSharedPtr<SModularRigModel> ModularRigModelPanel = MenuSpecificContext.ModularRigModelPanel.Pin())
 	{
-		return ModularRigHierarchyPanel.Get();
+		return ModularRigModelPanel.Get();
 	}
 	return nullptr;
 }
diff --git a/Engine/Plugins/Animation/ControlRig/Source/ControlRigEditor/Private/Editor/ControlRigContextMenuContext.h b/Engine/Plugins/Animation/ControlRig/Source/ControlRigEditor/Private/Editor/ControlRigContextMenuContext.h
--- a/Engine/Plugins/Animation/ControlRig/Source/ControlRigEditor/Private/Editor/ControlRigContextMenuContext.h
+++ b/Engine/Plugins/Animation/ControlRig/Source/ControlRigEditor/Private/Editor/ControlRigContextMenuContext.h
@@ -15,7 +15,7 @@ class URigVMGraph;
 class URigVMNode;
 class URigVMPin;
 class SRigHierarchy;
-class SModularRigHierarchy;
+class SModularRigModel;
 
 USTRUCT(BlueprintType)
 struct FControlRigRigHierarchyDragAndDropContext
@@ -98,7 +98,7 @@ struct FControlRigMenuSpecificContext
 	
 	FControlRigRigHierarchyDragAndDropContext RigHierarchyDragAndDropContext;
 
-	TWeakPtr<SModularRigHierarchy> ModularRigHierarchyPanel;
+	TWeakPtr<SModularRigModel> ModularRigModelPanel;
 	
 	FControlRigGraphNodeContextMenuContext GraphNodeContextMenuContext;
 
@@ -144,7 +144,7 @@ public:
 
 	SRigHierarchy* GetRigHierarchyPanel() const;
 
-	SModularRigHierarchy* GetModularRigHierarchyPanel() const;
+	SModularRigModel* GetModularRigModelPanel() const;
 	
 	FControlRigEditor* GetControlRigEditor() const;
 
diff --git a/Engine/Plugins/Animation/ControlRig/Source/ControlRigEditor/Private/Editor/ControlRigEditor.h b/Engine/Plugins/Animation/ControlRig/Source/ControlRigEditor/Private/Editor/ControlRigEditor.h
--- a/Engine/Plugins/Animation/ControlRig/Source/ControlRigEditor/Private/Editor/ControlRigEditor.h
+++ b/Engine/Plugins/Animation/ControlRig/Source/ControlRigEditor/Private/Editor/ControlRigEditor.h
@@ -348,7 +348,7 @@ protected:
 	friend class FControlRigEditorMode;
 	friend class SControlRigStackView;
 	friend class SRigHierarchy;
-	friend class SModularRigHierarchy;
+	friend class SModularRigModel;
 	friend struct FRigHierarchyTabSummoner;
-	friend struct FModularRigHierarchyTabSummoner;
+	friend struct FModularRigModelTabSummoner;
 };
diff --git a/Engine/Plugins/Animation/ControlRig/Source/ControlRigEditor/Private/Editor/ControlRigEditorMode.cpp b/Engine/Plugins/Animation/ControlRig/Source/ControlRigEditor/Private/Editor/ControlRigEditorMode.cpp
--- a/Engine/Plugins/Animation/ControlRig/Source/ControlRigEditor/Private/Editor/ControlRigEditorMode.cpp
+++ b/Engine/Plugins/Animation/ControlRig/Source/ControlRigEditor/Private/Editor/ControlRigEditorMode.cpp
@@ -7,7 +7,7 @@
 #include "IPersonaToolkit.h"
 #include "PersonaTabs.h"
 #include "Editor/RigHierarchyTabSummoner.h"
-#include "Editor/ModularRigHierarchyTabSummoner.h"
+#include "Editor/ModularRigModelTabSummoner.h"
 #include "Editor/RigModuleAssetBrowserTabSummoner.h"
 #include "Editor/RigVMExecutionStackTabSummoner.h"
 #include "Editor/RigCurveContainerTabSummoner.h"
@@ -145,7 +145,7 @@ FModularRigEditorMode::FModularRigEditorMode(const TSharedRef<FControlRigEditor>
 	: FControlRigEditorMode(InControlRigEditor, false)
 {
 
-	TabFactories.RegisterFactory(MakeShared<FModularRigHierarchyTabSummoner>(InControlRigEditor));
+	TabFactories.RegisterFactory(MakeShared<FModularRigModelTabSummoner>(InControlRigEditor));
 	TabFactories.RegisterFactory(MakeShared<FRigModuleAssetBrowserTabSummoner>(InControlRigEditor));
 	
 	TabLayout = FTabManager::NewLayout("Standalone_ModularRigEditMode_Layout_v1.2")
@@ -193,7 +193,7 @@ FModularRigEditorMode::FModularRigEditorMode(const TSharedRef<FControlRigEditor>
 							// Left bottom right
 							FTabManager::NewStack()
 							->SetSizeCoefficient(0.5f)
-							->AddTab(FModularRigHierarchyTabSummoner::TabID, ETabState::OpenedTab)
+							->AddTab(FModularRigModelTabSummoner::TabID, ETabState::OpenedTab)
 						)
 					)
 				)
diff --git a/Engine/Plugins/Animation/ControlRig/Source/ControlRigEditor/Private/Editor/ModularRigHierarchyTabSummoner.cpp b/Engine/Plugins/Animation/ControlRig/Source/ControlRigEditor/Private/Editor/ModularRigHierarchyTabSummoner.cpp
deleted file mode 100644
diff --git a/Engine/Plugins/Animation/ControlRig/Source/ControlRigEditor/Private/Editor/ModularRigHierarchyTabSummoner.h b/Engine/Plugins/Animation/ControlRig/Source/ControlRigEditor/Private/Editor/ModularRigHierarchyTabSummoner.h
deleted file mode 100644
diff --git a/Engine/Plugins/Animation/ControlRig/Source/ControlRigEditor/Private/Editor/ModularRigModelTabSummoner.cpp b/Engine/Plugins/Animation/ControlRig/Source/ControlRigEditor/Private/Editor/ModularRigModelTabSummoner.cpp
new file mode 100644
--- /dev/null
+++ b/Engine/Plugins/Animation/ControlRig/Source/ControlRigEditor/Private/Editor/ModularRigModelTabSummoner.cpp
@@ -0,0 +1,77 @@
+// Copyright Epic Games, Inc. All Rights Reserved.
+
+#include "Editor/ModularRigModelTabSummoner.h"
+#include "Editor/SModularRigModel.h"
+#include "ControlRigEditorStyle.h"
+#include "Editor/ControlRigEditor.h"
+#include "Widgets/Docking/SDockTab.h"
+
+#define LOCTEXT_NAMESPACE "ModularRigHierarchyTabSummoner"
+
+const FName FModularRigModelTabSummoner::TabID(TEXT("ModularRigModel"));
+
+FModularRigModelTabSummoner::FModularRigModelTabSummoner(const TSharedRef<FControlRigEditor>& InControlRigEditor)
+	: FWorkflowTabFactory(TabID, InControlRigEditor)
+	, ControlRigEditor(InControlRigEditor)
+{
+	TabLabel = LOCTEXT("ModularRigHierarchyTabLabel", "Modular Rig Hierarchy");
+	TabIcon = FSlateIcon(FControlRigEditorStyle::Get().GetStyleSetName(), "ModularRigHierarchy.TabIcon");
+
+	ViewMenuDescription = LOCTEXT("ModularRigHierarchy_ViewMenu_Desc", "Modular Rig Hierarchy");
+	ViewMenuTooltip = LOCTEXT("ModularRigHierarchy_ViewMenu_ToolTip", "Show the Modular Rig Hierarchy tab");
+}
+
+FTabSpawnerEntry& FModularRigModelTabSummoner::RegisterTabSpawner(TSharedRef<FTabManager> InTabManager, const FApplicationMode* CurrentApplicationMode) const
+{
+	FTabSpawnerEntry& SpawnerEntry = FWorkflowTabFactory::RegisterTabSpawner(InTabManager, CurrentApplicationMode);
+
+	SpawnerEntry.SetReuseTabMethod(FOnFindTabToReuse::CreateLambda([](const FTabId& InTabId) ->TSharedPtr<SDockTab> {
+	
+		return TSharedPtr<SDockTab>();
+
+	}));
+
+	return SpawnerEntry;
+}
+
+TSharedRef<SWidget> FModularRigModelTabSummoner::CreateTabBody(const FWorkflowTabSpawnInfo& Info) const
+{
+	ControlRigEditor.Pin()->ModularRigHierarchyTabCount++;
+	return SNew(SModularRigModel, ControlRigEditor.Pin().ToSharedRef());
+}
+
+TSharedRef<SDockTab> FModularRigModelTabSummoner::SpawnTab(const FWorkflowTabSpawnInfo& Info) const
+{
+	TSharedRef<SDockTab>  DockTab = FWorkflowTabFactory::SpawnTab(Info);
+	TWeakPtr<SDockTab> WeakDockTab = DockTab;
+	DockTab->SetCanCloseTab(SDockTab::FCanCloseTab::CreateLambda([WeakDockTab]()
+    {
+		int32 HierarchyTabCount = 0;
+		if (TSharedPtr<SDockTab> SharedDocTab = WeakDockTab.Pin())
+		{
+			if(SWidget* Content = &SharedDocTab->GetContent().Get())
+			{
+				SModularRigModel* RigHierarchy = (SModularRigModel*)Content;
+				if(FControlRigEditor* ControlRigEditorForTab = RigHierarchy->GetControlRigEditor())
+				{
+					HierarchyTabCount = ControlRigEditorForTab->GetModularRigHierarchyTabCount();
+				}
+			}
+		}
+		return HierarchyTabCount > 0;
+    }));
+	DockTab->SetOnTabClosed( SDockTab::FOnTabClosedCallback::CreateLambda([](TSharedRef<SDockTab> DockTab)
+	{
+		if(SWidget* Content = &DockTab->GetContent().Get())
+		{
+			SModularRigModel* RigHierarchy = (SModularRigModel*)Content;
+			if(FControlRigEditor* ControlRigEditorForTab = RigHierarchy->GetControlRigEditor())
+			{
+				ControlRigEditorForTab->ModularRigHierarchyTabCount--;
+			}
+		}
+	}));
+	return DockTab;
+}
+
+#undef LOCTEXT_NAMESPACE 
diff --git a/Engine/Plugins/Animation/ControlRig/Source/ControlRigEditor/Private/Editor/ModularRigModelTabSummoner.h b/Engine/Plugins/Animation/ControlRig/Source/ControlRigEditor/Private/Editor/ModularRigModelTabSummoner.h
new file mode 100644
--- /dev/null
+++ b/Engine/Plugins/Animation/ControlRig/Source/ControlRigEditor/Private/Editor/ModularRigModelTabSummoner.h
@@ -0,0 +1,24 @@
+// Copyright Epic Games, Inc. All Rights Reserved.
+
+#pragma once
+
+#include "CoreMinimal.h"
+#include "WorkflowOrientedApp/WorkflowTabFactory.h"
+
+class FControlRigEditor;
+
+struct FModularRigModelTabSummoner : public FWorkflowTabFactory
+{
+public:
+	static const FName TabID;
+	
+public:
+	FModularRigModelTabSummoner(const TSharedRef<FControlRigEditor>& InControlRigEditor);
+	
+	virtual FTabSpawnerEntry& RegisterTabSpawner(TSharedRef<FTabManager> TabManager, const FApplicationMode* CurrentApplicationMode) const;
+	virtual TSharedRef<SWidget> CreateTabBody(const FWorkflowTabSpawnInfo& Info) const override;
+	virtual TSharedRef<SDockTab> SpawnTab(const FWorkflowTabSpawnInfo& Info) const override;
+
+protected:
+	TWeakPtr<FControlRigEditor> ControlRigEditor;
+};
diff --git a/Engine/Plugins/Animation/ControlRig/Source/ControlRigEditor/Private/Editor/SModularRigHierarchy.cpp b/Engine/Plugins/Animation/ControlRig/Source/ControlRigEditor/Private/Editor/SModularRigHierarchy.cpp
deleted file mode 100644
diff --git a/Engine/Plugins/Animation/ControlRig/Source/ControlRigEditor/Private/Editor/SModularRigHierarchy.h b/Engine/Plugins/Animation/ControlRig/Source/ControlRigEditor/Private/Editor/SModularRigHierarchy.h
deleted file mode 100644
diff --git a/Engine/Plugins/Animation/ControlRig/Source/ControlRigEditor/Private/Editor/SModularRigHierarchyTreeView.cpp b/Engine/Plugins/Animation/ControlRig/Source/ControlRigEditor/Private/Editor/SModularRigHierarchyTreeView.cpp
deleted file mode 100644
diff --git a/Engine/Plugins/Animation/ControlRig/Source/ControlRigEditor/Private/Editor/SModularRigHierarchyTreeView.h b/Engine/Plugins/Animation/ControlRig/Source/ControlRigEditor/Private/Editor/SModularRigHierarchyTreeView.h
deleted file mode 100644
diff --git a/Engine/Plugins/Animation/ControlRig/Source/ControlRigEditor/Private/Editor/SModularRigModel.cpp b/Engine/Plugins/Animation/ControlRig/Source/ControlRigEditor/Private/Editor/SModularRigModel.cpp
new file mode 100644
--- /dev/null
+++ b/Engine/Plugins/Animation/ControlRig/Source/ControlRigEditor/Private/Editor/SModularRigModel.cpp
@@ -0,0 +1,884 @@
+// Copyright Epic Games, Inc. All Rights Reserved.
+
+#include "Editor/SModularRigModel.h"
+#include "Widgets/Input/SComboButton.h"
+#include "Styling/AppStyle.h"
+#include "Widgets/Input/SButton.h"
+#include "Widgets/Input/SSearchBox.h"
+#include "Kismet2/BlueprintEditorUtils.h"
+#include "ScopedTransaction.h"
+#include "Editor/ControlRigEditor.h"
+#include "BlueprintActionDatabase.h"
+#include "BlueprintVariableNodeSpawner.h"
+#include "Widgets/Layout/SSpacer.h"
+#include "Widgets/Images/SImage.h"
+#include "Widgets/Text/STextBlock.h"
+#include "K2Node_VariableGet.h"
+#include "RigVMBlueprintUtils.h"
+#include "ControlRigModularRigCommands.h"
+#include "ControlRigBlueprint.h"
+#include "Graph/ControlRigGraph.h"
+#include "Graph/ControlRigGraphNode.h"
+#include "Graph/ControlRigGraphSchema.h"
+#include "GraphEditorModule.h"
+#include "Framework/MultiBox/MultiBoxBuilder.h"
+#include "AnimationRuntime.h"
+#include "ClassViewerFilter.h"
+#include "PropertyCustomizationHelpers.h"
+#include "Framework/Application/SlateApplication.h"
+#include "Editor/EditorEngine.h"
+#include "HelperUtil.h"
+#include "Widgets/Text/SInlineEditableTextBlock.h"
+#include "ControlRig.h"
+#include "HAL/PlatformApplicationMisc.h"
+#include "HAL/PlatformTime.h"
+#include "Dialogs/Dialogs.h"
+#include "IPersonaToolkit.h"
+#include "SKismetInspector.h"
+#include "Types/WidgetActiveTimerDelegate.h"
+#include "Dialog/SCustomDialog.h"
+#include "EditMode/ControlRigEditMode.h"
+#include "ToolMenus.h"
+#include "Editor/ControlRigContextMenuContext.h"
+#include "Editor/SRigSpacePickerWidget.h"
+#include "Settings/ControlRigSettings.h"
+#include "Framework/Notifications/NotificationManager.h"
+#include "Widgets/Notifications/SNotificationList.h"
+#include "Styling/AppStyle.h"
+#include "ControlRigSkeletalMeshComponent.h"
+#include "Sequencer/ControlRigLayerInstance.h"
+#include "Algo/MinElement.h"
+#include "Algo/MaxElement.h"
+#include "AssetRegistry/AssetRegistryModule.h"
+#include "DragAndDrop/AssetDragDropOp.h"
+#include "Kismet2/SClassPickerDialog.h"
+#include "RigVMFunctions/Math/RigVMMathLibrary.h"
+#include "Preferences/PersonaOptions.h"
+
+#define LOCTEXT_NAMESPACE "SModularRigModel"
+
+//////////////////////////////////////////////////////////////
+/// FModularRigModuleDragDropOp
+///////////////////////////////////////////////////////////
+TSharedRef<FModularRigModuleDragDropOp> FModularRigModuleDragDropOp::New(const TArray<FString>& InElements)
+{
+	TSharedRef<FModularRigModuleDragDropOp> Operation = MakeShared<FModularRigModuleDragDropOp>();
+	Operation->Elements = InElements;
+	Operation->Construct();
+	return Operation;
+}
+
+TSharedPtr<SWidget> FModularRigModuleDragDropOp::GetDefaultDecorator() const
+{
+	return SNew(SBorder)
+		.Visibility(EVisibility::Visible)
+		.BorderImage(FAppStyle::GetBrush("Menu.Background"))
+		[
+			SNew(STextBlock)
+			.Text(FText::FromString(GetJoinedElementNames()))
+			//.Font(FAppStyle::Get().GetFontStyle("FontAwesome.10"))
+		];
+}
+
+FString FModularRigModuleDragDropOp::GetJoinedElementNames() const
+{
+	TArray<FString> ElementNameStrings;
+	for (const FString& Element: Elements)
+	{
+		ElementNameStrings.Add(Element);
+	}
+	return FString::Join(ElementNameStrings, TEXT(","));
+}
+
+///////////////////////////////////////////////////////////
+
+const FName SModularRigModel::ContextMenuName = TEXT("ControlRigEditor.ModularRigModel.ContextMenu");
+
+SModularRigModel::~SModularRigModel()
+{
+	const FControlRigEditor* Editor = ControlRigEditor.IsValid() ? ControlRigEditor.Pin().Get() : nullptr;
+	OnEditorClose(Editor, ControlRigBlueprint.Get());
+}
+
+void SModularRigModel::Construct(const FArguments& InArgs, TSharedRef<FControlRigEditor> InControlRigEditor)
+{
+	ControlRigEditor = InControlRigEditor;
+
+	ControlRigBlueprint = ControlRigEditor.Pin()->GetControlRigBlueprint();
+
+	ControlRigBlueprint->OnRefreshEditor().AddRaw(this, &SModularRigModel::HandleRefreshEditorFromBlueprint);
+	ControlRigBlueprint->OnSetObjectBeingDebugged().AddRaw(this, &SModularRigModel::HandleSetObjectBeingDebugged);
+	ControlRigBlueprint->OnModularRigPreCompiled().AddRaw(this, &SModularRigModel::HandlePreCompileModularRigs);
+	ControlRigBlueprint->OnModularRigCompiled().AddRaw(this, &SModularRigModel::HandlePostCompileModularRigs);
+
+	// for deleting, renaming, dragging
+	CommandList = MakeShared<FUICommandList>();
+
+	UEditorEngine* Editor = Cast<UEditorEngine>(GEngine);
+	if (Editor != nullptr)
+	{
+		Editor->RegisterForUndo(this);
+	}
+
+	BindCommands();
+
+	// setup all delegates for the modular rig model widget
+	FModularRigTreeDelegates Delegates;
+	Delegates.OnGetModularRig = FOnGetModularRigTreeRig::CreateSP(this, &SModularRigModel::GetModularRigForTreeView);
+	Delegates.OnContextMenuOpening = FOnContextMenuOpening::CreateSP(this, &SModularRigModel::CreateContextMenuWidget);
+	Delegates.OnDragDetected = FOnDragDetected::CreateSP(this, &SModularRigModel::OnDragDetected);
+	Delegates.OnCanAcceptDrop = FOnModularRigTreeCanAcceptDrop::CreateSP(this, &SModularRigModel::OnCanAcceptDrop);
+	Delegates.OnAcceptDrop = FOnModularRigTreeAcceptDrop::CreateSP(this, &SModularRigModel::OnAcceptDrop);
+	Delegates.OnMouseButtonClick = FOnModularRigTreeMouseButtonClick::CreateSP(this, &SModularRigModel::OnItemClicked);
+	Delegates.OnMouseButtonDoubleClick = FOnModularRigTreeMouseButtonClick::CreateSP(this, &SModularRigModel::OnItemDoubleClicked);
+	Delegates.OnRequestDetailsInspection = FOnModularRigTreeRequestDetailsInspection::CreateSP(this, &SModularRigModel::OnRequestDetailsInspection);
+	Delegates.OnRenameElement = FOnModularRigTreeRenameElement::CreateSP(this, &SModularRigModel::HandleRenameModule);
+	Delegates.OnVerifyModuleNameChanged = FOnModularRigTreeVerifyElementNameChanged::CreateSP(this, &SModularRigModel::HandleVerifyNameChanged);
+	
+	ChildSlot
+	[
+		SNew(SVerticalBox)
+		+SVerticalBox::Slot()
+		.Padding(0.0f, 0.0f)
+		[
+			SNew(SBorder)
+			.Padding(0.0f)
+			.ShowEffectWhenDisabled(false)
+			[
+				SNew(SBorder)
+				.Padding(2.0f)
+				.BorderImage(FAppStyle::GetBrush("SCSEditor.TreePanel"))
+				[
+					SAssignNew(TreeView, SModularRigTreeView)
+					.RigTreeDelegates(Delegates)
+					.AutoScrollEnabled(true)
+				]
+			]
+		]
+	];
+
+	RefreshTreeView();
+
+	if (ControlRigEditor.IsValid())
+	{
+		ControlRigEditor.Pin()->GetKeyDownDelegate().BindLambda([&](const FGeometry& MyGeometry, const FKeyEvent& InKeyEvent)->FReply {
+			return OnKeyDown(MyGeometry, InKeyEvent);
+		});
+		ControlRigEditor.Pin()->OnGetViewportContextMenu().BindSP(this, &SModularRigModel::GetContextMenu);
+		ControlRigEditor.Pin()->OnViewportContextMenuCommands().BindSP(this, &SModularRigModel::GetContextMenuCommands);
+		ControlRigEditor.Pin()->OnEditorClosed().AddSP(this, &SModularRigModel::OnEditorClose);
+	}
+	
+	CreateContextMenu();
+}
+
+void SModularRigModel::OnEditorClose(const FRigVMEditor* InEditor, URigVMBlueprint* InBlueprint)
+{
+	if (InEditor)
+	{
+		FControlRigEditor* Editor = (FControlRigEditor*)InEditor;  
+		Editor->OnGetViewportContextMenu().Unbind();
+		Editor->OnViewportContextMenuCommands().Unbind();
+	}
+
+	if (UControlRigBlueprint* BP = Cast<UControlRigBlueprint>(InBlueprint))
+	{
+		InBlueprint->OnRefreshEditor().RemoveAll(this);
+		InBlueprint->OnSetObjectBeingDebugged().RemoveAll(this);
+		BP->OnModularRigPreCompiled().RemoveAll(this);
+		BP->OnModularRigCompiled().RemoveAll(this);
+	}
+	
+	ControlRigEditor.Reset();
+	ControlRigBlueprint.Reset();
+}
+
+void SModularRigModel::BindCommands()
+{
+	// create new command
+	const FControlRigModularRigCommands& Commands = FControlRigModularRigCommands::Get();
+
+	CommandList->MapAction(Commands.AddModuleItem,
+		FExecuteAction::CreateSP(this, &SModularRigModel::HandleNewItem),
+		FCanExecuteAction());
+
+	CommandList->MapAction(Commands.RenameModuleItem,
+		FExecuteAction::CreateSP(this, &SModularRigModel::HandleRenameModule),
+		FCanExecuteAction());
+
+	CommandList->MapAction(Commands.DeleteModuleItem,
+		FExecuteAction::CreateSP(this, &SModularRigModel::HandleDeleteModules),
+		FCanExecuteAction());
+}
+
+FReply SModularRigModel::OnKeyDown(const FGeometry& MyGeometry, const FKeyEvent& InKeyEvent)
+{
+	if (CommandList.IsValid() && CommandList->ProcessCommandBindings(InKeyEvent))
+	{
+		return FReply::Handled();
+	}
+	return FReply::Unhandled();
+}
+
+void SModularRigModel::RefreshTreeView(bool bRebuildContent)
+{
+	bool bDummySuspensionFlag = false;
+	bool* SuspensionFlagPtr = &bDummySuspensionFlag;
+	if (ControlRigEditor.IsValid())
+	{
+		SuspensionFlagPtr = &ControlRigEditor.Pin()->GetSuspendDetailsPanelRefreshFlag();
+	}
+	TGuardValue<bool> SuspendDetailsPanelRefreshGuard(*SuspensionFlagPtr, true);
+
+	TreeView->RefreshTreeView(bRebuildContent);
+}
+
+TArray<FString> SModularRigModel::GetSelectedKeys() const
+{
+	TArray<TSharedPtr<FModularRigTreeElement>> SelectedItems = TreeView->GetSelectedItems();
+	
+	TArray<FString> SelectedKeys;
+	for (const TSharedPtr<FModularRigTreeElement>& SelectedItem : SelectedItems)
+	{
+		if (SelectedItem.IsValid())
+		{
+			if(!SelectedItem->Key.IsEmpty())
+			{
+				SelectedKeys.AddUnique(SelectedItem->Key);
+			}
+		}
+	}
+
+	return SelectedKeys;
+}
+
+
+void SModularRigModel::HandlePreCompileModularRigs(URigVMBlueprint* InBlueprint)
+{
+	ClearDetailPanel();
+}
+
+void SModularRigModel::HandlePostCompileModularRigs(URigVMBlueprint* InBlueprint)
+{
+	RefreshTreeView();
+	if (ControlRigEditor.IsValid())
+	{
+		TArray<TSharedPtr<FModularRigTreeElement>> SelectedElements;
+		Algo::Transform(ControlRigEditor.Pin()->ModulesSelected, SelectedElements, [this](const FString& Path)
+		{
+			return TreeView->FindElement(Path);
+		});
+		TreeView->SetSelection(SelectedElements);
+		ControlRigEditor.Pin()->RefreshDetailView();
+	}
+}
+
+void SModularRigModel::HandleRefreshEditorFromBlueprint(URigVMBlueprint* InBlueprint)
+{
+	RefreshTreeView();
+}
+
+void SModularRigModel::HandleSetObjectBeingDebugged(UObject* InObject)
+{
+	if(ControlRigBeingDebuggedPtr.Get() == InObject)
+	{
+		return;
+	}
+
+	ControlRigBeingDebuggedPtr.Reset();
+	
+	if(UModularRig* ControlRig = Cast<UModularRig>(InObject))
+	{
+		ControlRigBeingDebuggedPtr = ControlRig;
+	}
+
+	RefreshTreeView();
+}
+
+TSharedPtr< SWidget > SModularRigModel::CreateContextMenuWidget()
+{
+	UToolMenus* ToolMenus = UToolMenus::Get();
+
+	if (UToolMenu* Menu = GetContextMenu())
+	{
+		return ToolMenus->GenerateWidget(Menu);
+	}
+	
+	return SNullWidget::NullWidget;
+}
+
+void SModularRigModel::OnItemClicked(TSharedPtr<FModularRigTreeElement> InItem)
+{
+	UModularRig* Rig = GetModularRig();
+	check(Rig);
+
+	if (ControlRigEditor.IsValid() && InItem.IsValid())
+	{
+		ControlRigEditor.Pin()->SetDetailViewForRigModules({InItem->Key});
+	}
+}
+
+void SModularRigModel::OnItemDoubleClicked(TSharedPtr<FModularRigTreeElement> InItem)
+{
+
+}
+
+void SModularRigModel::CreateContextMenu()
+{
+	static bool bCreatedMenu = false;
+	if(bCreatedMenu)
+	{
+		return;
+	}
+	bCreatedMenu = true;
+	
+	const FName MenuName = ContextMenuName;
+
+	UToolMenus* ToolMenus = UToolMenus::Get();
+	
+	if (!ensure(ToolMenus))
+	{
+		return;
+	}
+
+	if (UToolMenu* Menu = ToolMenus->ExtendMenu(MenuName))
+	{
+		Menu->AddDynamicSection(NAME_None, FNewToolMenuDelegate::CreateLambda([](UToolMenu* InMenu)
+			{
+				UControlRigContextMenuContext* MainContext = InMenu->FindContext<UControlRigContextMenuContext>();
+				
+				if (SModularRigModel* ModelPanel = MainContext->GetModularRigModelPanel())
+				{
+					const FControlRigModularRigCommands& Commands = FControlRigModularRigCommands::Get(); 
+				
+					FToolMenuSection& ElementsSection = InMenu->AddSection(TEXT("Elements"), LOCTEXT("ElementsHeader", "Elements"));
+					ElementsSection.AddSubMenu(TEXT("New"), LOCTEXT("New", "New"), LOCTEXT("New_ToolTip", "Create New Elements"),
+						FNewToolMenuDelegate::CreateLambda([Commands, ModelPanel](UToolMenu* InSubMenu)
+						{
+							FToolMenuSection& DefaultSection = InSubMenu->AddSection(NAME_None);
+							DefaultSection.AddMenuEntry(Commands.AddModuleItem);
+						})
+					);
+					ElementsSection.AddMenuEntry(Commands.RenameModuleItem);
+					ElementsSection.AddMenuEntry(Commands.DeleteModuleItem);
+				}
+			})
+		);
+	}
+}
+
+UToolMenu* SModularRigModel::GetContextMenu()
+{
+	const FName MenuName = ContextMenuName;
+	UToolMenus* ToolMenus = UToolMenus::Get();
+
+	if(!ensure(ToolMenus))
+	{
+		return nullptr;
+	}
+
+	// individual entries in this menu can access members of this context, particularly useful for editor scripting
+	UControlRigContextMenuContext* ContextMenuContext = NewObject<UControlRigContextMenuContext>();
+	FControlRigMenuSpecificContext MenuSpecificContext;
+	MenuSpecificContext.ModularRigModelPanel = SharedThis(this);
+	ContextMenuContext->Init(ControlRigEditor, MenuSpecificContext);
+
+	FToolMenuContext MenuContext(CommandList);
+	MenuContext.AddObject(ContextMenuContext);
+
+	UToolMenu* Menu = ToolMenus->GenerateMenu(MenuName, MenuContext);
+
+	return Menu;
+}
+
+TSharedPtr<FUICommandList> SModularRigModel::GetContextMenuCommands() const
+{
+	return CommandList;
+}
+
+bool SModularRigModel::IsSingleSelected() const
+{
+	if(GetSelectedKeys().Num() == 1)
+	{
+		return true;
+	}
+	return false;
+}
+
+/** Filter class to show only RigModules. */
+class FClassViewerRigModulesFilter : public IClassViewerFilter
+{
+public:
+	FClassViewerRigModulesFilter()
+		: AssetRegistry(FModuleManager::GetModuleChecked<FAssetRegistryModule>(TEXT("AssetRegistry")).Get())
+	{}
+	
+	virtual bool IsClassAllowed(const FClassViewerInitializationOptions& InInitOptions, const UClass* InClass, TSharedRef< FClassViewerFilterFuncs > InFilterFuncs) override
+	{
+		if(InClass)
+		{
+			const bool bChildOfObjectClass = InClass->IsChildOf(UControlRig::StaticClass());
+			const bool bMatchesFlags = !InClass->HasAnyClassFlags(CLASS_Hidden | CLASS_HideDropDown | CLASS_Deprecated | CLASS_Abstract);
+			const bool bNotNative = !InClass->IsNative();
+
+			// Allow any class contained in the extra picker common classes array
+			if (InInitOptions.ExtraPickerCommonClasses.Contains(InClass))
+			{
+				return true;
+			}
+			
+			if (bChildOfObjectClass && bMatchesFlags && bNotNative)
+			{
+				const FAssetData AssetData(InClass);
+				return MatchesFilter(AssetData);
+			}
+		}
+		return false;
+	}
+
+	virtual bool IsUnloadedClassAllowed(const FClassViewerInitializationOptions& InInitOptions, const TSharedRef< const IUnloadedBlueprintData > InUnloadedClassData, TSharedRef< FClassViewerFilterFuncs > InFilterFuncs) override
+	{
+		const bool bChildOfObjectClass = InUnloadedClassData->IsChildOf(UControlRig::StaticClass());
+		const bool bMatchesFlags = !InUnloadedClassData->HasAnyClassFlags(CLASS_Hidden | CLASS_HideDropDown | CLASS_Deprecated | CLASS_Abstract);
+		if (bChildOfObjectClass && bMatchesFlags)
+		{
+			const FString GeneratedClassPathString = InUnloadedClassData->GetClassPathName().ToString();
+			const FString BlueprintPath = GeneratedClassPathString.LeftChop(2); // Chop off _C
+			const FAssetData AssetData = AssetRegistry.GetAssetByObjectPath(FSoftObjectPath(BlueprintPath));
+			return MatchesFilter(AssetData);
+
+		}
+		return false;
+	}
+
+private:
+	bool MatchesFilter(const FAssetData& AssetData)
+	{
+		static const UEnum* ControlTypeEnum = StaticEnum<EControlRigType>();
+		const FString ControlRigTypeStr = AssetData.GetTagValueRef<FString>(TEXT("ControlRigType"));
+		if (ControlRigTypeStr.IsEmpty())
+		{
+			return false;
+		}
+
+		const EControlRigType ControlRigType = (EControlRigType)(ControlTypeEnum->GetValueByName(*ControlRigTypeStr));
+		return ControlRigType == EControlRigType::RigModule;
+	}
+
+	const IAssetRegistry& AssetRegistry;
+};
+
+/** Create Item */
+void SModularRigModel::HandleNewItem()
+{
+	if(!ControlRigEditor.IsValid())
+	{
+		return;
+	}
+
+	FString ParentPath;
+	if (IsSingleSelected())
+	{
+		ParentPath = GetSelectedKeys()[0];
+	}
+	
+	FClassViewerInitializationOptions Options;
+	Options.bShowUnloadedBlueprints = true;
+	Options.NameTypeToDisplay = EClassViewerNameTypeToDisplay::DisplayName;
+
+	TSharedPtr<FClassViewerRigModulesFilter> ClassFilter = MakeShareable(new FClassViewerRigModulesFilter());
+	Options.ClassFilters.Add(ClassFilter.ToSharedRef());
+	Options.bShowNoneOption = false;
+	
+	UClass* ChosenClass;
+	const FText TitleText = LOCTEXT("ModularRigModelPickModuleClass", "Pick Rig Module Class");
+	const bool bPressedOk = SClassPickerDialog::PickClass(TitleText, Options, ChosenClass, UControlRig::StaticClass());
+	if (bPressedOk)
+	{
+		HandleNewItem(ChosenClass, ParentPath);
+	}
+}
+
+void SModularRigModel::HandleNewItem(UClass* InClass, const FString &InParentPath)
+{
+	UControlRig* ControlRig = InClass->GetDefaultObject<UControlRig>();
+	if (!ControlRig)
+	{
+		return;
+	}
+
+	FSlateApplication::Get().DismissAllMenus();
+	
+	if (ControlRigBlueprint.IsValid())
+	{
+		FString ClassName = InClass->GetName();
+		ClassName.RemoveFromEnd(TEXT("_C"));
+		FString PathName = InParentPath.IsEmpty() ? *ClassName : FString::Printf(TEXT("%s:%s"), *InParentPath, *ClassName);
+		const FName Name = CreateUniqueName(*PathName);
+		ControlRigBlueprint->GetModularRigController()->AddModule(Name, InClass, InParentPath);
+
+		FString NewPathName = InParentPath.IsEmpty() ? *Name.ToString() : FString::Printf(TEXT("%s:%s"), *InParentPath, *Name.ToString());
+		TSharedPtr<FModularRigTreeElement> Element = TreeView->FindElement(NewPathName);
+		if (Element.IsValid())
+		{
+			TreeView->SetSelection({Element});
+			TreeView->bRequestRenameSelected = true;
+		}
+	}
+}
+
+bool SModularRigModel::CanRenameModule() const
+{
+	return IsSingleSelected();
+}
+
+void SModularRigModel::HandleRenameModule()
+{
+	if(!ControlRigEditor.IsValid())
+	{
+		return;
+	}
+
+	if (!CanRenameModule())
+	{
+		return;
+	}
+
+	UModularRig* Rig = GetDefaultModularRig();
+	if (Rig)
+	{
+		FScopedTransaction Transaction(LOCTEXT("ModularRigModelRenameSelected", "Rename selected module"));
+
+		TArray<TSharedPtr<FModularRigTreeElement>> SelectedItems = TreeView->GetSelectedItems();
+		if (SelectedItems.Num() == 1)
+		{
+			SelectedItems[0]->RequestRename();
+		}
+	}
+
+	return;
+}
+
+FName SModularRigModel::HandleRenameModule(const FString& InOldPath, const FName& InNewName)
+{
+	ClearDetailPanel();
+	
+	if (ControlRigBlueprint.IsValid())
+	{
+		FScopedTransaction Transaction(LOCTEXT("ModularRigModelRename", "Rename Module"));
+
+		UModularRigController* Controller = ControlRigBlueprint->GetModularRigController();
+		check(Controller);
+
+		FName ResultingName = NAME_None;
+		if (Controller->RenameModule(InOldPath, InNewName))
+		{
+			return InNewName;
+		}
+	}
+
+	return NAME_None;
+}
+
+bool SModularRigModel::HandleVerifyNameChanged(const FString& InOldPath, const FName& InNewName, FText& OutErrorMessage)
+{
+	if (InNewName.IsNone())
+	{
+		return false;
+	}
+	
+	if (ControlRigBlueprint.IsValid())
+	{
+		UModularRigController* Controller = ControlRigBlueprint->GetModularRigController();
+		check(Controller);
+
+		return Controller->CanRenameModule(InOldPath, InNewName, OutErrorMessage);
+	}
+
+	return false;
+}
+
+void SModularRigModel::HandleDeleteModules()
+{
+	if(!ControlRigEditor.IsValid())
+	{
+		return;
+	}
+
+	UModularRig* Rig = GetDefaultModularRig();
+	if (Rig)
+	{
+		FScopedTransaction Transaction(LOCTEXT("ModularRigModelDeleteSelected", "Delete selected modules"));
+
+		TArray<TSharedPtr<FModularRigTreeElement>> SelectedItems = TreeView->GetSelectedItems();
+		TArray<FString> SelectedPaths;
+		Algo::Transform(SelectedItems, SelectedPaths, [](const TSharedPtr<FModularRigTreeElement>& Element)
+		{
+			if (Element.IsValid())
+			{
+				return Element->Key;
+			}
+			return FString();
+		});
+		HandleDeleteModules(SelectedPaths);
+	}
+
+	return;
+}
+
+void SModularRigModel::HandleDeleteModules(const TArray<FString>& InPaths)
+{
+	ClearDetailPanel();
+	
+	if (ControlRigBlueprint.IsValid())
+	{
+		FScopedTransaction Transaction(LOCTEXT("ModularRigModelDelete", "Delete Modules"));
+
+		UModularRigController* Controller = ControlRigBlueprint->GetModularRigController();
+		check(Controller);
+
+		// Make sure we delete the modules from children to root
+		TArray<FString> SortedPaths = Controller->Model->SortPaths(InPaths);
+		Algo::Reverse(SortedPaths);
+		for (const FString& Path : SortedPaths)
+		{
+			Controller->DeleteModule(Path);
+		}
+	}
+}
+
+void SModularRigModel::HandleReparentModules(const TArray<FString>& InPaths, const FString& InParentPath)
+{
+	ClearDetailPanel();
+	
+	if (ControlRigBlueprint.IsValid())
+	{
+		FScopedTransaction Transaction(LOCTEXT("ModularRigModelReparent", "Reparent Modules"));
+
+		UModularRigController* Controller = ControlRigBlueprint->GetModularRigController();
+		check(Controller);
+
+		for (const FString& Path : InPaths)
+		{
+			Controller->ReparentModule(Path, InParentPath);
+		}
+	}
+}
+
+class SModularRigModelPasteTransformsErrorPipe : public FOutputDevice
+{
+public:
+
+	int32 NumErrors;
+
+	SModularRigModelPasteTransformsErrorPipe()
+		: FOutputDevice()
+		, NumErrors(0)
+	{
+	}
+
+	virtual void Serialize(const TCHAR* V, ELogVerbosity::Type Verbosity, const class FName& Category) override
+	{
+		UE_LOG(LogControlRig, Error, TEXT("Error importing transforms to Model: %s"), V);
+		NumErrors++;
+	}
+};
+
+UModularRig* SModularRigModel::GetModularRig() const
+{
+	if (ControlRigBlueprint.IsValid())
+	{
+		if (UControlRig* DebuggedRig = ControlRigBeingDebuggedPtr.Get())
+		{
+			return Cast<UModularRig>(DebuggedRig);
+		}
+	}
+	if (ControlRigEditor.IsValid())
+	{
+		if (UControlRig* CurrentRig = ControlRigEditor.Pin()->GetControlRig())
+		{
+			return Cast<UModularRig>(CurrentRig);
+		}
+	}
+	return nullptr;
+}
+
+UModularRig* SModularRigModel::GetDefaultModularRig() const
+{
+	if (ControlRigBlueprint.IsValid())
+	{
+		if (UControlRig* DebuggedRig = ControlRigBeingDebuggedPtr.Get())
+		{
+			return Cast<UModularRig>(DebuggedRig);
+		}
+	}
+	return nullptr;
+}
+
+
+FName SModularRigModel::CreateUniqueName(const FName& InBasePath) const
+{
+	return ControlRigBlueprint->GetModularRigController()->GetSafeNewName(InBasePath.ToString());
+}
+
+void SModularRigModel::OnRequestDetailsInspection(const FString& InKey)
+{
+	if(!ControlRigEditor.IsValid())
+	{
+		return;
+	}
+	ControlRigEditor.Pin()->SetDetailViewForRigModules({InKey});
+}
+
+void SModularRigModel::ClearDetailPanel() const
+{
+	if(ControlRigEditor.IsValid())
+	{
+		ControlRigEditor.Pin()->ClearDetailObject();
+	}
+}
+
+void SModularRigModel::PostRedo(bool bSuccess) 
+{
+	if (bSuccess)
+	{
+		RefreshTreeView();
+	}
+}
+
+void SModularRigModel::PostUndo(bool bSuccess) 
+{
+	if (bSuccess)
+	{
+		RefreshTreeView();
+	}
+}
+
+FReply SModularRigModel::OnDragDetected(const FGeometry& MyGeometry, const FPointerEvent& MouseEvent)
+{
+	TArray<FString> DraggedElements = GetSelectedKeys();
+	if (MouseEvent.IsMouseButtonDown(EKeys::LeftMouseButton) && DraggedElements.Num() > 0)
+	{
+		if (ControlRigEditor.IsValid())
+		{
+			TSharedRef<FModularRigModuleDragDropOp> DragDropOp = FModularRigModuleDragDropOp::New(MoveTemp(DraggedElements));
+			return FReply::Handled().BeginDragDrop(DragDropOp);
+		}
+	}
+
+	return FReply::Unhandled();
+}
+
+TOptional<EItemDropZone> SModularRigModel::OnCanAcceptDrop(const FDragDropEvent& DragDropEvent, EItemDropZone DropZone, TSharedPtr<FModularRigTreeElement> TargetItem)
+{
+	const TOptional<EItemDropZone> InvalidDropZone;
+	TOptional<EItemDropZone> ReturnDropZone = DropZone;
+
+	TSharedPtr<FAssetDragDropOp> AssetDragDropOperation = DragDropEvent.GetOperationAs<FAssetDragDropOp>();
+	TSharedPtr<FModularRigModuleDragDropOp> ModuleDragDropOperation = DragDropEvent.GetOperationAs<FModularRigModuleDragDropOp>();
+	if (AssetDragDropOperation)
+	{
+		for (const FAssetData& AssetData : AssetDragDropOperation->GetAssets())
+		{
+			static const UEnum* ControlTypeEnum = StaticEnum<EControlRigType>();
+			const FString ControlRigTypeStr = AssetData.GetTagValueRef<FString>(TEXT("ControlRigType"));
+			if (ControlRigTypeStr.IsEmpty())
+			{
+				ReturnDropZone.Reset();
+				break;
+			}
+
+			const EControlRigType ControlRigType = (EControlRigType)(ControlTypeEnum->GetValueByName(*ControlRigTypeStr));
+			if (ControlRigType != EControlRigType::RigModule)
+			{
+				ReturnDropZone.Reset();
+				break;
+			}
+		}
+	}
+	else if(ModuleDragDropOperation)
+	{
+		// Accept this drop
+	}
+	else
+	{
+		ReturnDropZone.Reset();
+	}
+
+	return ReturnDropZone;
+}
+
+FReply SModularRigModel::OnAcceptDrop(const FDragDropEvent& DragDropEvent, EItemDropZone DropZone, TSharedPtr<FModularRigTreeElement> TargetItem)
+{
+	const TSharedPtr<FModularRigTreeElement>* ItemAtMouse = TreeView->FindItemAtPosition(DragDropEvent.GetScreenSpacePosition());
+	FString ParentPath;
+	if (ItemAtMouse && ItemAtMouse->IsValid())
+	{
+		ParentPath = ItemAtMouse->Get()->Key;
+	}
+
+	TSharedPtr<FAssetDragDropOp> AssetDragDropOperation = DragDropEvent.GetOperationAs<FAssetDragDropOp>();
+	TSharedPtr<FModularRigModuleDragDropOp> ModuleDragDropOperation = DragDropEvent.GetOperationAs<FModularRigModuleDragDropOp>();
+	if (AssetDragDropOperation)
+	{
+		for (const FAssetData& AssetData : AssetDragDropOperation->GetAssets())
+		{
+			static const UEnum* ControlTypeEnum = StaticEnum<EControlRigType>();
+			const FString ControlRigTypeStr = AssetData.GetTagValueRef<FString>(TEXT("ControlRigType"));
+			if (ControlRigTypeStr.IsEmpty())
+			{
+				continue;
+			}
+
+			const EControlRigType ControlRigType = (EControlRigType)(ControlTypeEnum->GetValueByName(*ControlRigTypeStr));
+			if (ControlRigType != EControlRigType::RigModule)
+			{
+				continue;
+			}
+
+			UClass* AssetClass = AssetData.GetClass();
+			if (!AssetClass->IsChildOf(UControlRigBlueprint::StaticClass()))
+			{
+				continue;
+			}
+
+			if(UControlRigBlueprint* AssetBlueprint = Cast<UControlRigBlueprint>(AssetData.GetAsset()))
+			{
+				HandleNewItem(AssetBlueprint->GetControlRigClass(), ParentPath);
+			}
+		}
+
+		FReply::Handled();
+	}
+	else if(ModuleDragDropOperation)
+	{
+		const TArray<FString> Paths = ModuleDragDropOperation->GetElements();
+		HandleReparentModules(Paths, ParentPath);
+	}
+	
+	return FReply::Unhandled();
+}
+
+FReply SModularRigModel::OnDrop(const FGeometry& MyGeometry, const FDragDropEvent& DragDropEvent)
+{
+	// only allow drops onto empty space of the widget (when there's no target item under the mouse)
+	// when dropped onto an item SModularRigModel::OnAcceptDrop will deal with the event
+	const TSharedPtr<FModularRigTreeElement>* ItemAtMouse = TreeView->FindItemAtPosition(DragDropEvent.GetScreenSpacePosition());
+	FString ParentPath;
+	if (ItemAtMouse && ItemAtMouse->IsValid())
+	{
+		return SCompoundWidget::OnDrop(MyGeometry, DragDropEvent);
+	}
+	
+	if (OnCanAcceptDrop(DragDropEvent, EItemDropZone::BelowItem, nullptr))
+	{
+		if (OnAcceptDrop(DragDropEvent, EItemDropZone::BelowItem, nullptr).IsEventHandled())
+		{
+			return FReply::Handled();
+		}
+	}
+	return SCompoundWidget::OnDrop(MyGeometry, DragDropEvent);
+}
+
+#undef LOCTEXT_NAMESPACE
+
diff --git a/Engine/Plugins/Animation/ControlRig/Source/ControlRigEditor/Private/Editor/SModularRigModel.h b/Engine/Plugins/Animation/ControlRig/Source/ControlRigEditor/Private/Editor/SModularRigModel.h
new file mode 100644
--- /dev/null
+++ b/Engine/Plugins/Animation/ControlRig/Source/ControlRigEditor/Private/Editor/SModularRigModel.h
@@ -0,0 +1,164 @@
+// Copyright Epic Games, Inc. All Rights Reserved.
+
+#pragma once
+
+#include "CoreMinimal.h"
+#include "EditorUndoClient.h"
+#include "Editor/SModularRigTreeView.h"
+#include "ControlRigBlueprint.h"
+#include "DragAndDrop/GraphNodeDragDropOp.h"
+#include "Editor/RigVMEditor.h"
+
+class SModularRigModel;
+class FControlRigEditor;
+class SSearchBox;
+class FUICommandList;
+class URigVMBlueprint;
+class UControlRig;
+struct FAssetData;
+class FMenuBuilder;
+class UToolMenu;
+struct FToolMenuContext;
+
+class FModularRigModuleDragDropOp : public FDragDropOperation
+{
+public:
+	DRAG_DROP_OPERATOR_TYPE(FModularRigModuleDragDropOp, FDragDropOperation)
+
+	static TSharedRef<FModularRigModuleDragDropOp> New(const TArray<FString>& InElements);
+
+	virtual TSharedPtr<SWidget> GetDefaultDecorator() const override;
+
+	/** @return true if this drag operation contains property paths */
+	bool HasElements() const
+	{
+		return Elements.Num() > 0;
+	}
+
+	/** @return The property paths from this drag operation */
+	const TArray<FString>& GetElements() const
+	{
+		return Elements;
+	}
+
+	FString GetJoinedElementNames() const;
+
+private:
+
+	/** Data for the property paths this item represents */
+	TArray<FString> Elements;
+};
+
+/** Widget allowing editing of a control rig's structure */
+class SModularRigModel : public SCompoundWidget, public FEditorUndoClient
+{
+public:
+	SLATE_BEGIN_ARGS(SModularRigModel) {}
+	SLATE_END_ARGS()
+
+	~SModularRigModel();
+
+	void Construct(const FArguments& InArgs, TSharedRef<FControlRigEditor> InControlRigEditor);
+
+	FControlRigEditor* GetControlRigEditor() const
+	{
+		if(ControlRigEditor.IsValid())
+		{
+			return ControlRigEditor.Pin().Get();
+		}
+		return nullptr;
+	}
+
+private:
+
+	void OnEditorClose(const FRigVMEditor* InEditor, URigVMBlueprint* InBlueprint);
+
+	/** Bind commands that this widget handles */
+	void BindCommands();
+
+	/** SWidget interface */
+	virtual FReply OnKeyDown(const FGeometry& MyGeometry, const FKeyEvent& InKeyEvent);
+
+	/** Rebuild the tree view */
+	void RefreshTreeView(bool bRebuildContent = true);
+
+	/** Return all selected keys */
+	TArray<FString> GetSelectedKeys() const;
+
+	/** Create a new item */
+	void HandleNewItem();
+	void HandleNewItem(UClass* InClass, const FString& InParentPath);
+
+	/** Rename item */
+	bool CanRenameModule() const;
+	void HandleRenameModule();
+	FName HandleRenameModule(const FString& InOldPath, const FName& InNewName);
+	bool HandleVerifyNameChanged(const FString& InOldPath, const FName& InNewName, FText& OutErrorMessage);
+
+	/** Delete items */
+	void HandleDeleteModules();
+	void HandleDeleteModules(const TArray<FString>& InPaths);
+
+	/** Reparent items */
+	void HandleReparentModules(const TArray<FString>& InPaths, const FString& InParentPath);
+
+	/** Set Selection Changed */
+	void OnSelectionChanged(TSharedPtr<FModularRigTreeElement> Selection, ESelectInfo::Type SelectInfo);
+
+	TSharedPtr< SWidget > CreateContextMenuWidget();
+	void OnItemClicked(TSharedPtr<FModularRigTreeElement> InItem);
+	void OnItemDoubleClicked(TSharedPtr<FModularRigTreeElement> InItem);
+	
+	// FEditorUndoClient
+	virtual void PostUndo(bool bSuccess) override;
+	virtual void PostRedo(bool bSuccess) override;
+
+	// reply to a drag operation
+	FReply OnDragDetected(const FGeometry& MyGeometry, const FPointerEvent& MouseEvent);
+
+	// reply to a drop operation on item
+	TOptional<EItemDropZone> OnCanAcceptDrop(const FDragDropEvent& DragDropEvent, EItemDropZone DropZone, TSharedPtr<FModularRigTreeElement> TargetItem);
+	FReply OnAcceptDrop(const FDragDropEvent& DragDropEvent, EItemDropZone DropZone, TSharedPtr<FModularRigTreeElement> TargetItem);
+
+	// SWidget Overrides
+	virtual FReply OnDrop(const FGeometry& MyGeometry, const FDragDropEvent& DragDropEvent) override;
+
+	static const FName ContextMenuName;
+	static void CreateContextMenu();
+	UToolMenu* GetContextMenu();
+	TSharedPtr<FUICommandList> GetContextMenuCommands() const;
+	
+	/** Our owning control rig editor */
+	TWeakPtr<FControlRigEditor> ControlRigEditor;
+
+	/** Tree view widget */
+	TSharedPtr<SModularRigTreeView> TreeView;
+
+	TWeakObjectPtr<UControlRigBlueprint> ControlRigBlueprint;
+	TWeakObjectPtr<UModularRig> ControlRigBeingDebuggedPtr;
+	
+	/** Command list we bind to */
+	TSharedPtr<FUICommandList> CommandList;
+
+	bool IsSingleSelected() const;
+	
+	UModularRig* GetModularRig() const;
+	UModularRig* GetDefaultModularRig() const;
+	const UModularRig* GetModularRigForTreeView() const { return GetModularRig(); }
+	FName CreateUniqueName(const FName& InBasePath) const;
+	void OnRequestDetailsInspection(const FString& InKey);
+	void ClearDetailPanel() const;
+
+	void HandlePreCompileModularRigs(URigVMBlueprint* InBlueprint);
+	void HandlePostCompileModularRigs(URigVMBlueprint* InBlueprint);
+	
+	void HandleRefreshEditorFromBlueprint(URigVMBlueprint* InBlueprint);
+	void HandleSetObjectBeingDebugged(UObject* InObject);
+
+public:
+
+	friend class FModularRigTreeElement;
+	friend class SModularRigModelItem;
+	friend class UControlRigBlueprintEditorLibrary;
+};
+
diff --git a/Engine/Plugins/Animation/ControlRig/Source/ControlRigEditor/Private/Editor/SModularRigTreeView.cpp b/Engine/Plugins/Animation/ControlRig/Source/ControlRigEditor/Private/Editor/SModularRigTreeView.cpp
new file mode 100644
--- /dev/null
+++ b/Engine/Plugins/Animation/ControlRig/Source/ControlRigEditor/Private/Editor/SModularRigTreeView.cpp
@@ -0,0 +1,720 @@
+// Copyright Epic Games, Inc. All Rights Reserved.
+
+#include "Editor/SModularRigTreeView.h"
+#include "Styling/AppStyle.h"
+#include "Widgets/Layout/SSpacer.h"
+#include "Widgets/Layout/SScrollBox.h"
+#include "Widgets/Images/SImage.h"
+#include "Widgets/Input/SSearchBox.h"
+#include "Widgets/Text/STextBlock.h"
+#include "PropertyCustomizationHelpers.h"
+#include "Framework/Application/SlateApplication.h"
+#include "Editor/EditorEngine.h"
+#include "HelperUtil.h"
+#include "Widgets/Text/SInlineEditableTextBlock.h"
+#include "ControlRig.h"
+#include "ControlRigEditorStyle.h"
+#include "ModularRig.h"
+#include "Blueprint/WidgetBlueprintLibrary.h"
+#include "Editor/SModularRigModel.h"
+#include "Settings/ControlRigSettings.h"
+#include "Graph/ControlRigGraphSchema.h"
+#include "Rigs/AdditiveControlRig.h"
+#include "Rigs/RigHierarchyController.h"
+#include "Styling/AppStyle.h"
+
+#define LOCTEXT_NAMESPACE "SModularRigTreeView"
+
+TMap<FSoftObjectPath, FSlateBrush> SModularRigModelItem::IconPathToBrush;
+
+//////////////////////////////////////////////////////////////
+/// FModularRigTreeElement
+///////////////////////////////////////////////////////////
+FModularRigTreeElement::FModularRigTreeElement(const FString& InKey, TWeakPtr<SModularRigTreeView> InTreeView, bool InSupportsRename)
+{
+	Key = InKey;
+	FString ShortNameStr = Key;
+	FString ParentPath;
+	ShortNameStr.Split(UModularRig::NamespaceSeparator, &ParentPath, &ShortNameStr, ESearchCase::Type::CaseSensitive, ESearchDir::FromEnd);
+	ShortName = *ShortNameStr;
+	
+	if(InTreeView.IsValid())
+	{
+		if(const UModularRig* ModularRig = InTreeView.Pin()->GetRigTreeDelegates().GetModularRig())
+		{
+			RefreshDisplaySettings(ModularRig);
+		}
+	}
+}
+
+void FModularRigTreeElement::RefreshDisplaySettings(const UModularRig* InModularRig)
+{
+	const TPair<const FSlateBrush*, FSlateColor> Result = SModularRigModelItem::GetBrushForElementType(InModularRig, Key);
+
+	IconBrush = Result.Key;
+	IconColor = Result.Value;
+	TextColor = FSlateColor::UseForeground();
+}
+
+TSharedRef<ITableRow> FModularRigTreeElement::MakeTreeRowWidget(const TSharedRef<STableViewBase>& InOwnerTable, TSharedRef<FModularRigTreeElement> InRigTreeElement, TSharedPtr<SModularRigTreeView> InTreeView, bool bPinned)
+{
+	return SNew(SModularRigModelItem, InOwnerTable, InRigTreeElement, InTreeView, bPinned);
+}
+
+void FModularRigTreeElement::RequestRename()
+{
+	OnRenameRequested.ExecuteIfBound();
+}
+
+//////////////////////////////////////////////////////////////
+/// SModularRigModelItem
+///////////////////////////////////////////////////////////
+void SModularRigModelItem::Construct(const FArguments& InArgs, const TSharedRef<STableViewBase>& OwnerTable, TSharedRef<FModularRigTreeElement> InRigTreeElement, TSharedPtr<SModularRigTreeView> InTreeView, bool bPinned)
+{
+	WeakRigTreeElement = InRigTreeElement;
+	Delegates = InTreeView->GetRigTreeDelegates();
+
+	if (InRigTreeElement->Key.IsEmpty())
+	{
+		STableRow<TSharedPtr<FModularRigTreeElement>>::Construct(
+			STableRow<TSharedPtr<FModularRigTreeElement>>::FArguments()
+			.ShowSelection(false)
+			.OnCanAcceptDrop(Delegates.OnCanAcceptDrop)
+			.OnAcceptDrop(Delegates.OnAcceptDrop)
+			.Content()
+			[
+				SNew(SVerticalBox)
+				+ SVerticalBox::Slot()
+				.FillHeight(200.f)
+				[
+					SNew(SSpacer)
+				]
+			], OwnerTable);
+		return;
+	}
+
+	TSharedPtr< SInlineEditableTextBlock > InlineWidget;
+	TSharedPtr< SHorizontalBox > HorizontalBox;
+
+	STableRow<TSharedPtr<FModularRigTreeElement>>::Construct(
+		STableRow<TSharedPtr<FModularRigTreeElement>>::FArguments()
+		.OnDragDetected(Delegates.OnDragDetected)
+		.OnCanAcceptDrop(Delegates.OnCanAcceptDrop)
+		.OnAcceptDrop(Delegates.OnAcceptDrop)
+		.ShowWires(true)
+		.Content()
+		[
+			SAssignNew(HorizontalBox, SHorizontalBox)
+			+ SHorizontalBox::Slot()
+			.MaxWidth(18)
+			.FillWidth(1.0)
+			.HAlign(HAlign_Left)
+			.VAlign(VAlign_Center)
+			.Padding(FMargin(0.f, 0.f, 3.f, 0.f))
+			[
+				SNew(SImage)
+				.Image_Lambda([this]() -> const FSlateBrush*
+				{
+					if(WeakRigTreeElement.IsValid())
+					{
+						return WeakRigTreeElement.Pin()->IconBrush;
+					}
+					return nullptr;
+				})
+				.ColorAndOpacity_Lambda([this]()
+				{
+					if(WeakRigTreeElement.IsValid())
+					{
+						return WeakRigTreeElement.Pin()->IconColor;
+					}
+					return FSlateColor::UseForeground();
+				})
+			]
+			+ SHorizontalBox::Slot()
+			.AutoWidth()
+			.VAlign(VAlign_Center)
+			[
+				SAssignNew(InlineWidget, SInlineEditableTextBlock)
+				.Text(this, &SModularRigModelItem::GetName, true)
+				.OnVerifyTextChanged(this, &SModularRigModelItem::OnVerifyNameChanged)
+				.OnTextCommitted(this, &SModularRigModelItem::OnNameCommitted)
+				.ToolTipText(this, &SModularRigModelItem::GetItemTooltip)
+				.MultiLine(false)
+				.ColorAndOpacity_Lambda([this]()
+				{
+					if(WeakRigTreeElement.IsValid())
+					{
+						return WeakRigTreeElement.Pin()->TextColor;
+					}
+					return FSlateColor::UseForeground();
+				})
+			]
+		], OwnerTable);
+
+	InRigTreeElement->OnRenameRequested.BindSP(InlineWidget.Get(), &SInlineEditableTextBlock::EnterEditingMode);
+}
+
+void SModularRigModelItem::OnNameCommitted(const FText& InText, ETextCommit::Type InCommitType) const
+{
+	// for now only allow enter
+	// because it is important to keep the unique names per pose
+	if (InCommitType == ETextCommit::OnEnter)
+	{
+		FString NewName = InText.ToString();
+		const FString OldKey = WeakRigTreeElement.Pin()->Key;
+
+		Delegates.HandleRenameElement(OldKey, *NewName);
+	}
+}
+
+bool SModularRigModelItem::OnVerifyNameChanged(const FText& InText, FText& OutErrorMessage)
+{
+	const FName NewName = *InText.ToString();
+	const FString OldPath = WeakRigTreeElement.Pin()->Key;
+	return Delegates.HandleVerifyElementNameChanged(OldPath, NewName, OutErrorMessage);
+}
+
+FText SModularRigModelItem::GetName(bool bUseShortName) const
+{
+	if(bUseShortName)
+	{
+		return (FText::FromName(WeakRigTreeElement.Pin()->ShortName));
+	}
+	return (FText::FromString(WeakRigTreeElement.Pin()->Key));
+}
+
+FText SModularRigModelItem::GetItemTooltip() const
+{
+	const FText FullName = GetName(false);
+	const FText ShortName = GetName(true);
+	if(FullName.EqualTo(ShortName))
+	{
+		return FText();
+	}
+	return FullName;
+}
+
+//////////////////////////////////////////////////////////////
+/// SModularRigTreeView
+///////////////////////////////////////////////////////////
+
+void SModularRigTreeView::Construct(const FArguments& InArgs)
+{
+	Delegates = InArgs._RigTreeDelegates;
+	bAutoScrollEnabled = InArgs._AutoScrollEnabled;
+
+	STreeView<TSharedPtr<FModularRigTreeElement>>::FArguments SuperArgs;
+	SuperArgs.TreeItemsSource(&RootElements);
+	SuperArgs.SelectionMode(ESelectionMode::Multi);
+	SuperArgs.OnGenerateRow(this, &SModularRigTreeView::MakeTableRowWidget, false);
+	SuperArgs.OnGetChildren(this, &SModularRigTreeView::HandleGetChildrenForTree);
+	SuperArgs.OnContextMenuOpening(Delegates.OnContextMenuOpening);
+	SuperArgs.HighlightParentNodesForSelection(true);
+	SuperArgs.ItemHeight(24);
+	SuperArgs.AllowInvisibleItemSelection(true);  //without this we deselect everything when we filter or we collapse
+	SuperArgs.OnMouseButtonClick(Delegates.OnMouseButtonClick);
+	SuperArgs.OnMouseButtonDoubleClick(Delegates.OnMouseButtonDoubleClick);
+	
+	SuperArgs.ShouldStackHierarchyHeaders_Lambda([]() -> bool {
+		return UControlRigEditorSettings::Get()->bShowStackedHierarchy;
+	});
+	SuperArgs.OnGeneratePinnedRow(this, &SModularRigTreeView::MakeTableRowWidget, true);
+	SuperArgs.MaxPinnedItems_Lambda([]() -> int32
+	{
+		return FMath::Max<int32>(1, UControlRigEditorSettings::Get()->MaxStackSize);
+	});
+
+	STreeView<TSharedPtr<FModularRigTreeElement>>::Construct(SuperArgs);
+
+	LastMousePosition = FVector2D::ZeroVector;
+	TimeAtMousePosition = 0.0;
+}
+
+void SModularRigTreeView::Tick(const FGeometry& AllottedGeometry, const double InCurrentTime, const float InDeltaTime)
+{
+	STreeView<TSharedPtr<FModularRigTreeElement, ESPMode::ThreadSafe>>::Tick(AllottedGeometry, InCurrentTime, InDeltaTime);
+
+	const FGeometry PaintGeometry = GetPaintSpaceGeometry();
+	const FVector2D MousePosition = FSlateApplication::Get().GetCursorPos();
+
+	if(PaintGeometry.IsUnderLocation(MousePosition))
+	{
+		const FVector2D WidgetPosition = PaintGeometry.AbsoluteToLocal(MousePosition);
+
+		static constexpr float SteadyMousePositionTolerance = 5.f;
+
+		if(LastMousePosition.Equals(MousePosition, SteadyMousePositionTolerance))
+		{
+			TimeAtMousePosition += InDeltaTime;
+		}
+		else
+		{
+			LastMousePosition = MousePosition;
+			TimeAtMousePosition = 0.0;
+		}
+
+		static constexpr float AutoScrollStartDuration = 0.5f; // in seconds
+		static constexpr float AutoScrollDistance = 24.f; // in pixels
+		static constexpr float AutoScrollSpeed = 150.f;
+
+		if(TimeAtMousePosition > AutoScrollStartDuration && FSlateApplication::Get().IsDragDropping())
+		{
+			if((WidgetPosition.Y < AutoScrollDistance) || (WidgetPosition.Y > PaintGeometry.Size.Y - AutoScrollDistance))
+			{
+				if(bAutoScrollEnabled)
+				{
+					const bool bScrollUp = (WidgetPosition.Y < AutoScrollDistance);
+
+					const float DeltaInSlateUnits = (bScrollUp ? -InDeltaTime : InDeltaTime) * AutoScrollSpeed; 
+					ScrollBy(GetCachedGeometry(), DeltaInSlateUnits, EAllowOverscroll::No);
+				}
+			}
+			else
+			{
+				const TSharedPtr<FModularRigTreeElement>* Item = FindItemAtPosition(MousePosition);
+				if(Item && Item->IsValid())
+				{
+					if(!IsItemExpanded(*Item))
+					{
+						SetItemExpansion(*Item, true);
+					}
+				}
+			}
+		}
+	}
+
+	if (bRequestRenameSelected)
+	{
+		RegisterActiveTimer(0.f, FWidgetActiveTimerDelegate::CreateLambda([this](double, float) {
+			TArray<TSharedPtr<FModularRigTreeElement>> SelectedItems = GetSelectedItems();
+			if (SelectedItems.Num() == 1)
+			{
+				SelectedItems[0]->RequestRename();
+			}
+			return EActiveTimerReturnType::Stop;
+		}));
+		bRequestRenameSelected = false;
+	}
+}
+
+TSharedPtr<FModularRigTreeElement> SModularRigTreeView::FindElement(const FString& InElementKey)
+{
+	for (TSharedPtr<FModularRigTreeElement> Root : RootElements)
+	{
+		if (TSharedPtr<FModularRigTreeElement> Found = FindElement(InElementKey, Root))
+		{
+			return Found;
+		}
+	}
+
+	return TSharedPtr<FModularRigTreeElement>();
+}
+
+TSharedPtr<FModularRigTreeElement> SModularRigTreeView::FindElement(const FString& InElementKey, TSharedPtr<FModularRigTreeElement> CurrentItem)
+{
+	if (CurrentItem->Key == InElementKey)
+	{
+		return CurrentItem;
+	}
+
+	for (int32 ChildIndex = 0; ChildIndex < CurrentItem->Children.Num(); ++ChildIndex)
+	{
+		TSharedPtr<FModularRigTreeElement> Found = FindElement(InElementKey, CurrentItem->Children[ChildIndex]);
+		if (Found.IsValid())
+		{
+			return Found;
+		}
+	}
+
+	return TSharedPtr<FModularRigTreeElement>();
+}
+
+bool SModularRigTreeView::AddElement(FString InKey, FString InParentKey)
+{
+	if(ElementMap.Contains(InKey))
+	{
+		return false;
+	}
+
+	if (!InKey.IsEmpty())
+	{
+		TSharedPtr<FModularRigTreeElement> NewItem = MakeShared<FModularRigTreeElement>(InKey, SharedThis(this), false);
+
+		if (!InKey.IsEmpty())
+		{
+			ElementMap.Add(InKey, NewItem);
+			if (!InParentKey.IsEmpty())
+			{
+				ParentMap.Add(InKey, InParentKey);
+			}
+
+			if (!InParentKey.IsEmpty())
+			{
+				TSharedPtr<FModularRigTreeElement>* FoundItem = ElementMap.Find(InParentKey);
+				check(FoundItem);
+				FoundItem->Get()->Children.Add(NewItem);
+			}
+			else
+			{
+				RootElements.Add(NewItem);
+			}
+		}
+		else
+		{
+			RootElements.Add(NewItem);
+		}
+	}
+
+	return true;
+}
+
+bool SModularRigTreeView::AddElement(const FRigModuleInstance* InElement)
+{
+	check(InElement);
+	
+	if (ElementMap.Contains(InElement->GetPath()))
+	{
+		return false;
+	}
+
+	const UModularRig* ModularRig = Delegates.GetModularRig();
+
+	if(!AddElement(InElement->GetPath(), FString()))
+	{
+		return false;
+	}
+
+	if (ElementMap.Contains(InElement->GetPath()))
+	{
+		if(ModularRig)
+		{
+			FString ParentKey = ModularRig->GetParentPath(InElement->GetPath());
+			if (!ParentKey.IsEmpty())
+			{
+				if(const FRigModuleInstance* ParentElement = ModularRig->FindModule(ParentKey))
+				{
+					AddElement(ParentElement);
+
+					if(ElementMap.Contains(ParentKey))
+					{
+						ReparentElement(InElement->GetPath(), ParentKey);
+					}
+				}
+			}
+		}
+	}
+
+	return true;
+}
+
+void SModularRigTreeView::AddSpacerElement()
+{
+	AddElement(FString(), FString());
+}
+
+bool SModularRigTreeView::ReparentElement(const FString InKey, const FString InParentKey)
+{
+	if (InKey.IsEmpty() || InKey == InParentKey)
+	{
+		return false;
+	}
+
+	TSharedPtr<FModularRigTreeElement>* FoundItem = ElementMap.Find(InKey);
+	if (FoundItem == nullptr)
+	{
+		return false;
+	}
+
+	if (const FString* ExistingParentKey = ParentMap.Find(InKey))
+	{
+		if (*ExistingParentKey == InParentKey)
+		{
+			return false;
+		}
+
+		if (TSharedPtr<FModularRigTreeElement>* ExistingParent = ElementMap.Find(*ExistingParentKey))
+		{
+			(*ExistingParent)->Children.Remove(*FoundItem);
+		}
+
+		ParentMap.Remove(InKey);
+	}
+	else
+	{
+		if (InParentKey.IsEmpty())
+		{
+			return false;
+		}
+
+		RootElements.Remove(*FoundItem);
+	}
+
+	if (!InParentKey.IsEmpty())
+	{
+		ParentMap.Add(InKey, InParentKey);
+
+		TSharedPtr<FModularRigTreeElement>* FoundParent = ElementMap.Find(InParentKey);
+		if(FoundParent)
+		{
+			FoundParent->Get()->Children.Add(*FoundItem);
+		}
+	}
+	else
+	{
+		RootElements.Add(*FoundItem);
+	}
+
+	return true;
+}
+
+void SModularRigTreeView::RefreshTreeView(bool bRebuildContent)
+{
+	TMap<FString, bool> ExpansionState;
+
+	if(bRebuildContent)
+	{
+		for (TPair<FString, TSharedPtr<FModularRigTreeElement>> Pair : ElementMap)
+		{
+			ExpansionState.FindOrAdd(Pair.Key) = IsItemExpanded(Pair.Value);
+		}
+
+		// internally save expansion states before rebuilding the tree, so the states can be restored later
+		SaveAndClearSparseItemInfos();
+
+		RootElements.Reset();
+		ElementMap.Reset();
+		ParentMap.Reset();
+	}
+
+	if(bRebuildContent)
+	{
+		const UModularRig* ModularRig = Delegates.GetModularRig();
+		if(ModularRig)
+		{
+			ModularRig->ForEachModule([&](const FRigModuleInstance* Element)
+			{
+				AddElement(Element);
+				return true;
+			});
+
+			// expand all elements upon the initial construction of the tree
+			if (ExpansionState.Num() < ElementMap.Num())
+			{
+				for (const TPair<FString, TSharedPtr<FModularRigTreeElement>>& Element : ElementMap)
+				{
+					if (!ExpansionState.Contains(Element.Key))
+					{
+						SetItemExpansion(Element.Value, true);
+					}
+				}
+			}
+
+			for (const auto& Pair : ElementMap)
+			{
+				RestoreSparseItemInfos(Pair.Value);
+			}
+
+			if (RootElements.Num() > 0)
+			{
+				AddSpacerElement();
+			}
+		}
+	}
+	else
+	{
+		if (RootElements.Num()> 0)
+		{
+			// elements may be added at the end of the list after a spacer element
+			// we need to remove the spacer element and re-add it at the end
+			RootElements.RemoveAll([](TSharedPtr<FModularRigTreeElement> InElement)
+			{
+				return InElement.Get()->Key == FString();
+			});
+			AddSpacerElement();
+		}
+	}
+
+	RequestTreeRefresh();
+	{
+		ClearSelection();
+	}
+}
+
+TSharedRef<ITableRow> SModularRigTreeView::MakeTableRowWidget(TSharedPtr<FModularRigTreeElement> InItem,
+	const TSharedRef<STableViewBase>& OwnerTable, bool bPinned)
+{
+	return InItem->MakeTreeRowWidget(OwnerTable, InItem.ToSharedRef(), SharedThis(this), bPinned);
+}
+
+void SModularRigTreeView::HandleGetChildrenForTree(TSharedPtr<FModularRigTreeElement> InItem,
+	TArray<TSharedPtr<FModularRigTreeElement>>& OutChildren)
+{
+	OutChildren = InItem.Get()->Children;
+}
+
+TArray<FString> SModularRigTreeView::GetSelectedKeys() const
+{
+	TArray<FString> Keys;
+	TArray<TSharedPtr<FModularRigTreeElement>> SelectedElements = GetSelectedItems();
+	for(const TSharedPtr<FModularRigTreeElement>& SelectedElement : SelectedElements)
+	{
+		Keys.Add(SelectedElement->Key);
+	}
+	return Keys;
+}
+
+void SModularRigTreeView::SetSelection(const TArray<TSharedPtr<FModularRigTreeElement>>& InSelection) 
+{
+	ClearSelection();
+	SetItemSelection(InSelection, true, ESelectInfo::Direct);
+}
+
+const TSharedPtr<FModularRigTreeElement>* SModularRigTreeView::FindItemAtPosition(FVector2D InScreenSpacePosition) const
+{
+	if (ItemsPanel.IsValid() && SListView<TSharedPtr<FModularRigTreeElement>>::HasValidItemsSource())
+	{
+		FArrangedChildren ArrangedChildren(EVisibility::Visible);
+		const int32 Index = FindChildUnderPosition(ArrangedChildren, InScreenSpacePosition);
+		if (ArrangedChildren.IsValidIndex(Index))
+		{
+			TSharedRef<SModularRigModelItem> ItemWidget = StaticCastSharedRef<SModularRigModelItem>(ArrangedChildren[Index].Widget);
+			if (ItemWidget->WeakRigTreeElement.IsValid())
+			{
+				const FString Key = ItemWidget->WeakRigTreeElement.Pin()->Key;
+				const TSharedPtr<FModularRigTreeElement>* ResultPtr = SListView<TSharedPtr<FModularRigTreeElement>>::GetItems().FindByPredicate([Key](const TSharedPtr<FModularRigTreeElement>& Item) -> bool
+					{
+						return Item->Key == Key;
+					});
+
+				if (ResultPtr)
+				{
+					return ResultPtr;
+				}
+			}
+		}
+	}
+	return nullptr;
+}
+
+TPair<const FSlateBrush*, FSlateColor> SModularRigModelItem::GetBrushForElementType(const UModularRig* InModularRig, const FString& InKey)
+{
+	const FSlateBrush* Brush = nullptr;
+	FSlateColor Color = FSlateColor::UseForeground();
+
+	if (const FRigModuleInstance* Module = InModularRig->FindModule(InKey))
+	{
+		if (Module->Rig.IsValid())
+		{
+			FSoftObjectPath IconPath = Module->Rig->GetRigModuleSettings().Icon;
+			Brush = IconPathToBrush.Find(IconPath);
+			if (!Brush)
+			{
+				if(UTexture2D* Icon = Cast<UTexture2D>(IconPath.TryLoad()))
+				{
+					Brush = &IconPathToBrush.Add(IconPath, UWidgetBlueprintLibrary::MakeBrushFromTexture(Icon, 16.0f, 16.0f));
+				}
+			}
+		}
+	}
+	if (!Brush)
+	{
+		Brush = FControlRigEditorStyle::Get().GetBrush("ControlRig.Tree.RigidBody");
+	}
+	return TPair<const FSlateBrush*, FSlateColor>(Brush, Color);
+}
+
+FLinearColor SModularRigModelItem::GetColorForControlType(ERigControlType InControlType, UEnum* InControlEnum)
+{
+	FEdGraphPinType PinType;
+	switch(InControlType)
+	{
+		case ERigControlType::Bool:
+		{
+			PinType = RigVMTypeUtils::PinTypeFromCPPType(RigVMTypeUtils::BoolTypeName, nullptr);
+			break;
+		}
+		case ERigControlType::Float:
+		case ERigControlType::ScaleFloat:
+		{
+			PinType = RigVMTypeUtils::PinTypeFromCPPType(RigVMTypeUtils::FloatTypeName, nullptr);
+			break;
+		}
+		case ERigControlType::Integer:
+		{
+			if(InControlEnum)
+			{
+				PinType = RigVMTypeUtils::PinTypeFromCPPType(NAME_None, InControlEnum);
+			}
+			else
+			{
+				PinType = RigVMTypeUtils::PinTypeFromCPPType(RigVMTypeUtils::Int32TypeName, nullptr);
+			}
+			break;
+		}
+		case ERigControlType::Vector2D:
+		{
+			UScriptStruct* Struct = TBaseStructure<FVector2D>::Get(); 
+			PinType = RigVMTypeUtils::PinTypeFromCPPType(*RigVMTypeUtils::GetUniqueStructTypeName(Struct), Struct);
+			break;
+		}
+		case ERigControlType::Position:
+		case ERigControlType::Scale:
+		{
+			UScriptStruct* Struct = TBaseStructure<FVector>::Get(); 
+			PinType = RigVMTypeUtils::PinTypeFromCPPType(*RigVMTypeUtils::GetUniqueStructTypeName(Struct), Struct);
+			break;
+		}
+		case ERigControlType::Rotator:
+		{
+			UScriptStruct* Struct = TBaseStructure<FRotator>::Get(); 
+			PinType = RigVMTypeUtils::PinTypeFromCPPType(*RigVMTypeUtils::GetUniqueStructTypeName(Struct), Struct);
+			break;
+		}
+		case ERigControlType::Transform:
+		case ERigControlType::TransformNoScale:
+		case ERigControlType::EulerTransform:
+		default:
+		{
+			UScriptStruct* Struct = TBaseStructure<FTransform>::Get(); 
+			PinType = RigVMTypeUtils::PinTypeFromCPPType(*RigVMTypeUtils::GetUniqueStructTypeName(Struct), Struct);
+			break;
+		}
+	}
+	const UControlRigGraphSchema* Schema = GetDefault<UControlRigGraphSchema>();
+	return Schema->GetPinTypeColor(PinType);
+}
+
+//////////////////////////////////////////////////////////////
+/// SSearchableModularRigTreeView
+///////////////////////////////////////////////////////////
+
+void SSearchableModularRigTreeView::Construct(const FArguments& InArgs)
+{
+	FModularRigTreeDelegates TreeDelegates = InArgs._RigTreeDelegates;
+	
+	ChildSlot
+	[
+		SNew(SVerticalBox)
+		+SVerticalBox::Slot()
+		.AutoHeight()
+		.VAlign(VAlign_Top)
+		.HAlign(HAlign_Fill)
+		.Padding(0.0f, 0.0f)
+		[
+			SNew(SScrollBox)
+			+ SScrollBox::Slot()
+			[
+				SNew(SBorder)
+				.Padding(2.0f)
+				.BorderImage(FAppStyle::GetBrush("SCSEditor.TreePanel"))
+				[
+					SAssignNew(TreeView, SModularRigTreeView)
+					.RigTreeDelegates(TreeDelegates)
+				]
+			]
+		]
+	];
+}
+
+
+#undef LOCTEXT_NAMESPACE
diff --git a/Engine/Plugins/Animation/ControlRig/Source/ControlRigEditor/Private/Editor/SModularRigTreeView.h b/Engine/Plugins/Animation/ControlRig/Source/ControlRigEditor/Private/Editor/SModularRigTreeView.h
new file mode 100644
--- /dev/null
+++ b/Engine/Plugins/Animation/ControlRig/Source/ControlRigEditor/Private/Editor/SModularRigTreeView.h
@@ -0,0 +1,229 @@
+// Copyright Epic Games, Inc. All Rights Reserved.
+
+#pragma once
+
+#include "CoreMinimal.h"
+#include "Widgets/Views/STreeView.h"
+#include "ModularRig.h"
+
+class SSearchBox;
+class SModularRigTreeView;
+class SModularRigModelItem;
+class FModularRigTreeElement;
+
+
+DECLARE_DELEGATE_RetVal(const UModularRig*, FOnGetModularRigTreeRig);
+DECLARE_DELEGATE_OneParam(FOnModularRigTreeRequestDetailsInspection, const FString&);
+DECLARE_DELEGATE_RetVal_TwoParams(FName, FOnModularRigTreeRenameElement, const FString& /*OldPath*/, const FName& /*NewName*/);
+DECLARE_DELEGATE_RetVal_ThreeParams(bool, FOnModularRigTreeVerifyElementNameChanged, const FString& /*OldPath*/, const FName& /*NewName*/, FText& /*OutErrorMessage*/);
+
+
+typedef STreeView<TSharedPtr<FModularRigTreeElement>>::FOnMouseButtonClick FOnModularRigTreeMouseButtonClick;
+typedef STreeView<TSharedPtr<FModularRigTreeElement>>::FOnMouseButtonDoubleClick FOnModularRigTreeMouseButtonDoubleClick;
+typedef STableRow<TSharedPtr<FModularRigTreeElement>>::FOnCanAcceptDrop FOnModularRigTreeCanAcceptDrop;
+typedef STableRow<TSharedPtr<FModularRigTreeElement>>::FOnAcceptDrop FOnModularRigTreeAcceptDrop;
+
+struct CONTROLRIGEDITOR_API FModularRigTreeDelegates
+{
+	FOnGetModularRigTreeRig OnGetModularRig;
+	FOnModularRigTreeMouseButtonClick OnMouseButtonClick;
+	FOnModularRigTreeMouseButtonDoubleClick OnMouseButtonDoubleClick;
+	FOnDragDetected OnDragDetected;
+	FOnModularRigTreeCanAcceptDrop OnCanAcceptDrop;
+	FOnModularRigTreeAcceptDrop OnAcceptDrop;
+	FOnContextMenuOpening OnContextMenuOpening;
+	FOnModularRigTreeRequestDetailsInspection OnRequestDetailsInspection;
+	FOnModularRigTreeRenameElement OnRenameElement;
+	FOnModularRigTreeVerifyElementNameChanged OnVerifyModuleNameChanged;
+	
+	FModularRigTreeDelegates()
+	{
+	}
+
+	const UModularRig* GetModularRig() const
+	{
+		if(OnGetModularRig.IsBound())
+		{
+			return OnGetModularRig.Execute();
+		}
+		return nullptr;
+	}
+
+	FName HandleRenameElement(const FString& OldPath, const FName& NewName) const
+	{
+		if(OnRenameElement.IsBound())
+		{
+			return OnRenameElement.Execute(OldPath, NewName);
+		}
+		return *OldPath;
+	}
+
+	bool HandleVerifyElementNameChanged(const FString& OldPath, const FName& NewName, FText& OutErrorMessage) const
+	{
+		if(OnVerifyModuleNameChanged.IsBound())
+		{
+			return OnVerifyModuleNameChanged.Execute(OldPath, NewName, OutErrorMessage);
+		}
+		return false;
+	}
+};
+
+
+/** An item in the tree */
+class FModularRigTreeElement : public TSharedFromThis<FModularRigTreeElement>
+{
+public:
+	FModularRigTreeElement(const FString& InKey, TWeakPtr<SModularRigTreeView> InTreeView, bool InSupportsRename);
+
+public:
+	/** Element Data to display */
+	FString Key;
+	FName ShortName;
+	TArray<TSharedPtr<FModularRigTreeElement>> Children;
+
+	TSharedRef<ITableRow> MakeTreeRowWidget(const TSharedRef<STableViewBase>& InOwnerTable, TSharedRef<FModularRigTreeElement> InRigTreeElement, TSharedPtr<SModularRigTreeView> InTreeView, bool bPinned);
+
+	void RequestRename();
+
+	void RefreshDisplaySettings(const UModularRig* InModularRig);
+
+	/** Delegate for when the context menu requests a rename */
+	DECLARE_DELEGATE(FOnRenameRequested);
+	FOnRenameRequested OnRenameRequested;
+
+	/** The brush to use when rendering an icon */
+	const FSlateBrush* IconBrush;
+
+	/** The color to use when rendering an icon */
+	FSlateColor IconColor;
+
+	/** The color to use when rendering the label text */
+	FSlateColor TextColor;
+};
+
+class SModularRigModelItem : public STableRow<TSharedPtr<FModularRigTreeElement>>
+{
+public:
+	
+	void Construct(const FArguments& InArgs, const TSharedRef<STableViewBase>& OwnerTable, TSharedRef<FModularRigTreeElement> InRigTreeElement, TSharedPtr<SModularRigTreeView> InTreeView, bool bPinned);
+
+	void OnNameCommitted(const FText& InText, ETextCommit::Type InCommitType) const;
+	bool OnVerifyNameChanged(const FText& InText, FText& OutErrorMessage);
+	static TPair<const FSlateBrush*, FSlateColor> GetBrushForElementType(const UModularRig* InModularRig, const FString& InKey);
+	static FLinearColor GetColorForControlType(ERigControlType InControlType, UEnum* InControlEnum);
+
+private:
+	TWeakPtr<FModularRigTreeElement> WeakRigTreeElement;
+ 	FModularRigTreeDelegates Delegates;
+
+	FText GetName(bool bUseShortName) const;
+	FText GetItemTooltip() const;
+
+	static TMap<FSoftObjectPath, FSlateBrush> IconPathToBrush;
+
+	friend class SModularRigTreeView; 
+};
+
+class SModularRigTreeView : public STreeView<TSharedPtr<FModularRigTreeElement>>
+{
+public:
+
+	SLATE_BEGIN_ARGS(SModularRigTreeView)
+		: _AutoScrollEnabled(false)
+	{}
+		SLATE_ARGUMENT(FModularRigTreeDelegates, RigTreeDelegates)
+		SLATE_ARGUMENT(bool, AutoScrollEnabled)
+	SLATE_END_ARGS()
+
+	void Construct(const FArguments& InArgs);
+	virtual ~SModularRigTreeView() {}
+
+	/** Performs auto scroll */
+	virtual void Tick(const FGeometry& AllottedGeometry, const double InCurrentTime, const float InDeltaTime) override;
+
+	bool bRequestRenameSelected = false;
+
+	/** Save a snapshot of the internal map that tracks item expansion before tree reconstruction */
+	void SaveAndClearSparseItemInfos()
+	{
+		// Only save the info if there is something to save (do not overwrite info with an empty map)
+		if (!SparseItemInfos.IsEmpty())
+		{
+			OldSparseItemInfos = SparseItemInfos;
+		}
+		ClearExpandedItems();
+	}
+
+	/** Restore the expansion infos map from the saved snapshot after tree reconstruction */
+	void RestoreSparseItemInfos(TSharedPtr<FModularRigTreeElement> ItemPtr)
+	{
+		for (const auto& Pair : OldSparseItemInfos)
+		{
+			if (Pair.Key->Key == ItemPtr->Key)
+			{
+				// the SparseItemInfos now reference the new element, but keep the same expansion state
+				SparseItemInfos.Add(ItemPtr, Pair.Value);
+				break;
+			}
+		}
+	}
+
+
+	TSharedPtr<FModularRigTreeElement> FindElement(const FString& InElementKey);
+	static TSharedPtr<FModularRigTreeElement> FindElement(const FString& InElementKey, TSharedPtr<FModularRigTreeElement> CurrentItem);
+	bool AddElement(FString InKey, FString InParentKey = FString());
+	bool AddElement(const FRigModuleInstance* InElement);
+	void AddSpacerElement();
+	bool ReparentElement(const FString InKey, const FString InParentKey);
+	void RefreshTreeView(bool bRebuildContent = true);
+	TSharedRef<ITableRow> MakeTableRowWidget(TSharedPtr<FModularRigTreeElement> InItem, const TSharedRef<STableViewBase>& OwnerTable, bool bPinned);
+	void HandleGetChildrenForTree(TSharedPtr<FModularRigTreeElement> InItem, TArray<TSharedPtr<FModularRigTreeElement>>& OutChildren);
+
+	TArray<FString> GetSelectedKeys() const;
+	void SetSelection(const TArray<TSharedPtr<FModularRigTreeElement>>& InSelection);
+	const TArray<TSharedPtr<FModularRigTreeElement>>& GetRootElements() const { return RootElements; }
+	FModularRigTreeDelegates& GetRigTreeDelegates() { return Delegates; }
+
+	/** Given a position, return the item under that position. If nothing is there, return null. */
+	const TSharedPtr<FModularRigTreeElement>* FindItemAtPosition(FVector2D InScreenSpacePosition) const;
+
+private:
+
+	/** A temporary snapshot of the SparseItemInfos in STreeView, used during RefreshTreeView() */
+	TSparseItemMap OldSparseItemInfos;
+
+	/** Backing array for tree view */
+	TArray<TSharedPtr<FModularRigTreeElement>> RootElements;
+	
+	/** A map for looking up items based on their key */
+	TMap<FString, TSharedPtr<FModularRigTreeElement>> ElementMap;
+
+	/** A map for looking up a parent based on their key */
+	TMap<FString, FString> ParentMap;
+
+	FModularRigTreeDelegates Delegates;
+
+	bool bAutoScrollEnabled;
+	FVector2D LastMousePosition;
+	double TimeAtMousePosition;
+
+	friend class SModularRigModel;
+};
+
+class SSearchableModularRigTreeView : public SCompoundWidget
+{
+public:
+
+	SLATE_BEGIN_ARGS(SSearchableModularRigTreeView) {}
+		SLATE_ARGUMENT(FModularRigTreeDelegates, RigTreeDelegates)
+		SLATE_ARGUMENT(FText, InitialFilterText)
+	SLATE_END_ARGS()
+
+	void Construct(const FArguments& InArgs);
+	virtual ~SSearchableModularRigTreeView() {}
+	TSharedRef<SModularRigTreeView> GetTreeView() const { return TreeView.ToSharedRef(); }
+
+private:
+
+	TSharedPtr<SModularRigTreeView> TreeView;
+};
