Control Rig: Fix layered rigs weirdness when working with driven controls
#rnx

diff --git a/Engine/Plugins/Animation/ControlRig/Source/ControlRig/Private/ControlRig.cpp b/Engine/Plugins/Animation/ControlRig/Source/ControlRig/Private/ControlRig.cpp
--- a/Engine/Plugins/Animation/ControlRig/Source/ControlRig/Private/ControlRig.cpp
+++ b/Engine/Plugins/Animation/ControlRig/Source/ControlRig/Private/ControlRig.cpp
@@ -399,7 +399,6 @@ void UControlRig::Evaluate_AnyThread()
 
 		if (PreAdditiveValuesApplicationEvent.IsBound())
 		{
-			FControlRigBracketScope BracketScope(PreAdditiveValuesApplicationBracket);
 			PreAdditiveValuesApplicationEvent.Broadcast(this, TEXT("Additive"));
 		}
 		
@@ -450,7 +449,6 @@ void UControlRig::Evaluate_AnyThread()
 		
 		// Forward solve
 		{
-			TGuardValue<bool> UpdatePreferredAngles(Hierarchy->bUpdatePreferredEulerAngleWhenSettingTransform, false);
 			Execute(FRigUnit_BeginExecution::EventName);
 		}
 	}
