RigHierarchy: Declare the object pointers on the FRigBaseElements in URigHierarchy::Element to garbage collection, so that FRigControlSettings::ControlEnum does not get garbage collected.
#rnx

diff --git a/Engine/Plugins/Animation/ControlRig/Source/ControlRig/Private/Rigs/RigHierarchy.cpp b/Engine/Plugins/Animation/ControlRig/Source/ControlRig/Private/Rigs/RigHierarchy.cpp
--- a/Engine/Plugins/Animation/ControlRig/Source/ControlRig/Private/Rigs/RigHierarchy.cpp
+++ b/Engine/Plugins/Animation/ControlRig/Source/ControlRig/Private/Rigs/RigHierarchy.cpp
@@ -163,6 +163,19 @@ void URigHierarchy::Serialize(FArchive& Ar)
 	}
 }
 
+void URigHierarchy::AddReferencedObjects(UObject* InpThis, FReferenceCollector& Collector)
+{
+	Super::AddReferencedObjects(InpThis, Collector);
+
+	URigHierarchy* pThis = static_cast<URigHierarchy*>(InpThis);
+	FReferenceCollectorArchive Ar(pThis, Collector);
+	FScopeLock Lock(&pThis->ElementsLock);
+	for (FRigBaseElement* Element : pThis->Elements)
+	{
+		Collector.AddPropertyReferencesWithStructARO(Element->GetElementStruct(), Element, pThis);
+	}
+}
+
 void URigHierarchy::Save(FArchive& Ar)
 {
 	FScopeLock Lock(&ElementsLock);
diff --git a/Engine/Plugins/Animation/ControlRig/Source/ControlRig/Private/Rigs/RigHierarchyElements.cpp b/Engine/Plugins/Animation/ControlRig/Source/ControlRig/Private/Rigs/RigHierarchyElements.cpp
--- a/Engine/Plugins/Animation/ControlRig/Source/ControlRig/Private/Rigs/RigHierarchyElements.cpp
+++ b/Engine/Plugins/Animation/ControlRig/Source/ControlRig/Private/Rigs/RigHierarchyElements.cpp
@@ -765,6 +765,10 @@ void FRigControlSettings::Save(FArchive& Ar)
 	if(ControlEnum)
 	{
 		ControlEnumPathName = ControlEnum->GetPathName();
+		if (Ar.IsObjectReferenceCollector())
+		{
+			Ar << ControlEnum;
+		}
 	}
 
 	Ar << AnimationTypeName;
diff --git a/Engine/Plugins/Animation/ControlRig/Source/ControlRig/Public/Rigs/RigHierarchy.h b/Engine/Plugins/Animation/ControlRig/Source/ControlRig/Public/Rigs/RigHierarchy.h
--- a/Engine/Plugins/Animation/ControlRig/Source/ControlRig/Public/Rigs/RigHierarchy.h
+++ b/Engine/Plugins/Animation/ControlRig/Source/ControlRig/Public/Rigs/RigHierarchy.h
@@ -175,6 +175,7 @@ public:
 	// UObject interface
 	virtual void BeginDestroy() override;
 	virtual void Serialize(FArchive& Ar) override;
+	static void AddReferencedObjects(UObject* InThis, FReferenceCollector& Collector);
 	void Save(FArchive& Ar);
 	void Load(FArchive& Ar);
 	virtual void PostLoad() override;
diff --git a/Engine/Source/Runtime/CoreUObject/Public/UObject/UObjectGlobals.h b/Engine/Source/Runtime/CoreUObject/Public/UObject/UObjectGlobals.h
--- a/Engine/Source/Runtime/CoreUObject/Public/UObject/UObjectGlobals.h
+++ b/Engine/Source/Runtime/CoreUObject/Public/UObject/UObjectGlobals.h
@@ -2236,7 +2236,7 @@ protected:
 	}
 
 public:
-	FReferenceCollectorArchive(const UObject* InSerializingObject, FReferenceCollector& InCollector);
+	COREUOBJECT_API FReferenceCollectorArchive(const UObject* InSerializingObject, FReferenceCollector& InCollector);
 
 	void SetSerializingObject(const UObject* InSerializingObject)
 	{
