Control Rig: Fix clashing LOCTEXT uses within module details

#rnx

diff --git a/Engine/Plugins/Animation/ControlRig/Source/ControlRigEditor/Private/ControlRigModuleDetails.cpp b/Engine/Plugins/Animation/ControlRig/Source/ControlRigEditor/Private/ControlRigModuleDetails.cpp
--- a/Engine/Plugins/Animation/ControlRig/Source/ControlRigEditor/Private/ControlRigModuleDetails.cpp
+++ b/Engine/Plugins/Animation/ControlRig/Source/ControlRigEditor/Private/ControlRigModuleDetails.cpp
@@ -788,7 +788,7 @@ bool FRigModuleInstanceDetails::CanRemoveBinding(FName InPropertyName) const
 
 void FRigModuleInstanceDetails::HandleRemoveBinding(FName InPropertyName) const
 {
-	FScopedTransaction Transaction(LOCTEXT("BindModuleVariableTransaction", "Remove Binding"));
+	FScopedTransaction Transaction(LOCTEXT("RemoveModuleVariableTransaction", "Remove Binding"));
 	for(const FPerModuleInfo& Info : PerModuleInfos)
 	{
 		if (UControlRigBlueprint* Blueprint = Info.GetBlueprint())
