Control Rig: Another fix layered rigs weirdness when working with driven controls
#rnx

diff --git a/Engine/Plugins/Animation/ControlRig/Source/ControlRig/Private/ControlRig.cpp b/Engine/Plugins/Animation/ControlRig/Source/ControlRig/Private/ControlRig.cpp
--- a/Engine/Plugins/Animation/ControlRig/Source/ControlRig/Private/ControlRig.cpp
+++ b/Engine/Plugins/Animation/ControlRig/Source/ControlRig/Private/ControlRig.cpp
@@ -390,7 +390,6 @@ void UControlRig::Evaluate_AnyThread()
 		
 		// Backwards solve
 		{
-			TGuardValue<bool> UpdatePreferredAngles(Hierarchy->bUpdatePreferredEulerAngleWhenSettingTransform, false);
 			Execute(FRigUnit_InverseExecution::EventName);
 		}
 
