Fixed regression from 28268803

#rnx

diff --git a/Engine/Plugins/Animation/ControlRig/Source/ControlRigEditor/Private/EditMode/ControlRigEditMode.cpp b/Engine/Plugins/Animation/ControlRig/Source/ControlRigEditor/Private/EditMode/ControlRigEditMode.cpp
--- a/Engine/Plugins/Animation/ControlRig/Source/ControlRigEditor/Private/EditMode/ControlRigEditMode.cpp
+++ b/Engine/Plugins/Animation/ControlRig/Source/ControlRigEditor/Private/EditMode/ControlRigEditMode.cpp
@@ -2360,7 +2360,6 @@ void FControlRigEditMode::UpdatePivotTransforms()
 		}
 	}
 
-
 	//If in level editor and the transforms changed we need to force hit proxy invalidate so widget hit testing 
 	//doesn't work off of it's last transform.  Similar to what sequencer does on re-evaluation but do to how edit modes and widget ticks happen
 	//it doesn't work for control rig gizmo's
@@ -2448,7 +2447,6 @@ void FControlRigEditMode::HandleSelectionChanged()
 	{
 		GEditor->RedrawLevelEditingViewports(true);
 	}
-
 }
 
 void FControlRigEditMode::BindCommands()
@@ -4189,7 +4187,7 @@ void FControlRigEditMode::SetControlShapeTransform(
 	}
 
 	static constexpr bool bNotify = true, bFixEuler = true, bUndo = true;
-	if (!AreEditingControlRigDirectly())
+	if (AreEditingControlRigDirectly())
 	{
 		// assumes it's attached to actor
 		ControlRig->SetControlGlobalTransform(
