Control Rig: Fix GetInteraction node not working for modular rigs
#rnx

diff --git a/Engine/Plugins/Animation/ControlRig/Source/ControlRig/Private/ModularRig.cpp b/Engine/Plugins/Animation/ControlRig/Source/ControlRig/Private/ModularRig.cpp
--- a/Engine/Plugins/Animation/ControlRig/Source/ControlRig/Private/ModularRig.cpp
+++ b/Engine/Plugins/Animation/ControlRig/Source/ControlRig/Private/ModularRig.cpp
@@ -8,6 +8,7 @@
 #include "ControlRigObjectBinding.h"
 #include "Rigs/RigHierarchyController.h"
 #include "ControlRigComponent.h"
+#include "RigVMCore/RigVMExecuteContext.h"
 
 #define LOCTEXT_NAMESPACE "ModularRig"
 
@@ -70,6 +71,11 @@ FString FRigModuleInstance::GetPath() const
 	return Name.ToString();
 }
 
+FString FRigModuleInstance::GetNamespace() const
+{
+	return FString::Printf(TEXT("%s:"), *GetPath());
+}
+
 void UModularRig::InitializeVMs(bool bRequestInit)
 {
 	URigVMHost::Initialize(bRequestInit);
@@ -101,10 +107,12 @@ bool UModularRig::Execute_Internal(const FName& InEventName)
 {
 	if (VM)
 	{
-		FRigVMExtendedExecuteContext& Context = GetRigVMExtendedExecuteContext();
-		URigHierarchy* Hierarchy = GetHierarchy();
+		FRigVMExtendedExecuteContext& ModularRigContext = GetRigVMExtendedExecuteContext();
+		const FControlRigExecuteContext& PublicContext = ModularRigContext.GetPublicDataSafe<FControlRigExecuteContext>();
+		const FRigUnitContext& UnitContext = PublicContext.UnitContext;
+		const URigHierarchy* Hierarchy = GetHierarchy();
 
-		ForEachModule([&InEventName, this, Hierarchy](FRigModuleInstance* Module) -> bool
+		ForEachModule([&InEventName, this, Hierarchy, UnitContext](FRigModuleInstance* Module) -> bool
 		{
 			if (Module->Rig.IsValid())
 			{
@@ -115,6 +123,22 @@ bool UModularRig::Execute_Internal(const FName& InEventName)
 					return true;
 				}
 
+				// Only emit interaction event on this module if any of the interaction elements
+				// belong to the module's namespace
+				if (InEventName == FRigUnit_InteractionExecution::EventName)
+				{
+					const FString ModuleNamespace = Module->GetNamespace();
+					const bool bIsInteracting = UnitContext.ElementsBeingInteracted.ContainsByPredicate(
+						[ModuleNamespace, Hierarchy](const FRigElementKey& InteractionElement)
+						{
+							return ModuleNamespace == Hierarchy->GetNameMetadata(InteractionElement, URigHierarchy::NameSpaceMetadataName, NAME_None);
+						});
+					if (!bIsInteracting)
+					{
+						return true;
+					}
+				}
+
 				ExecutionQueue.Add(FRigModuleExecutionElement(Module, InEventName));
 			}
 			return true;
@@ -158,6 +182,22 @@ void UModularRig::ExecuteQueue()
 			// Make sure the hierarchy has the correct execute context with the rig module namespace
 			FRigHierarchyExecuteContextBracket ExecuteContextBracket(Hierarchy, &RigExtendedExecuteContext);
 
+			FControlRigExecuteContext& PublicContext = Context.GetPublicDataSafe<FControlRigExecuteContext>();
+			FControlRigExecuteContext& RigPublicContext = RigExtendedExecuteContext.GetPublicDataSafe<FControlRigExecuteContext>();
+			FRigUnitContext& RigUnitContext = RigPublicContext.UnitContext;
+			RigUnitContext = PublicContext.UnitContext;
+
+			// Update the interaction elements to show only the ones belonging to this module
+			const FString ModuleNamespace = FString::Printf(TEXT("%s:"), *ExecutionElement.ModulePath);
+			RigUnitContext.ElementsBeingInteracted = RigUnitContext.ElementsBeingInteracted.FilterByPredicate(
+				[ModuleNamespace, Hierarchy](const FRigElementKey& Key)
+			{
+				return ModuleNamespace == Hierarchy->GetNameMetadata(Key, URigHierarchy::NameSpaceMetadataName, NAME_None);
+			});
+			RigUnitContext.InteractionType = RigUnitContext.ElementsBeingInteracted.IsEmpty() ?
+				(uint8) EControlRigInteractionType::None
+				: RigUnitContext.InteractionType;
+
 			// Make sure the module's rig has the corrct user data
 			// The rig will combine the user data of the
 			// - skeleton
@@ -167,7 +207,6 @@ void UModularRig::ExecuteQueue()
 			// - outer modular rig
 			// - external variables
 			{
-				FControlRigExecuteContext& RigPublicContext = RigExtendedExecuteContext.GetPublicDataSafe<FControlRigExecuteContext>();
 				RigPublicContext.AssetUserData.Reset();
 				if(const TArray<UAssetUserData*>* ControlRigUserDataArray = Rig->GetAssetUserDataArray())
 				{
@@ -217,6 +256,7 @@ void UModularRig::OnObjectsReplaced(const TMap<UObject*, UObject*>& OldToNewInst
 		}
 	}
 	InitializeVMs(true);
+	UpdateSupportedEvents();
 }
 
 void UModularRig::ResetModules()
@@ -241,6 +281,7 @@ void UModularRig::ResetModules()
 	
 	RootModules.Reset();
 	Modules.Reset();
+	SupportedEvents.Reset();
 }
 
 void UModularRig::UpdateCachedChildren()
@@ -269,6 +310,23 @@ void UModularRig::UpdateCachedChildren()
 	}
 }
 
+void UModularRig::UpdateSupportedEvents()
+{
+	SupportedEvents.Reset();
+	ForEachModule([this](const FRigModuleInstance* Module) -> bool
+	{
+		if (Module->Rig.IsValid())
+		{
+			const TArray<FName>& ModuleEvents = Module->Rig->GetSupportedEvents();
+			for (const FName& EventName : ModuleEvents)
+			{
+				SupportedEvents.AddUnique(EventName);
+			}
+		}
+		return true;
+	});
+}
+
 bool UModularRig::AddModuleInstance(const FName& InModuleName, TSubclassOf<UControlRig> InModuleClass, FString InParentPath,
 	const TMap<FRigElementKey, FRigElementKey>& InConnectionMap, const TMap<FName, FString>& InVariableDefaultValues, const TMap<FName, FString>& InVariableBindings )
 {
@@ -317,6 +375,7 @@ FRigModuleInstance* UModularRig::AddModuleInstance(const FName& InModuleName, TS
 		NewModule.ParentPath = InParent->GetPath();
 	}
 	UpdateCachedChildren();
+	SupportedEvents.Append(NewModule.Rig->GetSupportedEvents());
 
 	// Configure module
 	{
diff --git a/Engine/Plugins/Animation/ControlRig/Source/ControlRig/Public/ControlRig.h b/Engine/Plugins/Animation/ControlRig/Source/ControlRig/Public/ControlRig.h
--- a/Engine/Plugins/Animation/ControlRig/Source/ControlRig/Public/ControlRig.h
+++ b/Engine/Plugins/Animation/ControlRig/Source/ControlRig/Public/ControlRig.h
@@ -198,6 +198,14 @@ public:
 	virtual bool Execute_Internal(const FName& InEventName) override;
 	virtual void RequestInit() override;
 	virtual void RequestInitVMs()  { Super::RequestInit(); }
+	virtual bool SupportsEvent(const FName& InEventName) const override { return Super::SupportsEvent(InEventName); }
+	virtual const TArray<FName>& GetSupportedEvents() const override{ return Super::GetSupportedEvents(); }
+
+	template<class T>
+	bool SupportsEvent() const
+	{
+		return SupportsEvent(T::EventName);
+	}
 
 	bool AllConnectorsAreResolved(FString* OutFailureReason = nullptr, FRigElementKey* OutConnector = nullptr) const;
 
diff --git a/Engine/Plugins/Animation/ControlRig/Source/ControlRig/Public/ModularRig.h b/Engine/Plugins/Animation/ControlRig/Source/ControlRig/Public/ModularRig.h
--- a/Engine/Plugins/Animation/ControlRig/Source/ControlRig/Public/ModularRig.h
+++ b/Engine/Plugins/Animation/ControlRig/Source/ControlRig/Public/ModularRig.h
@@ -63,6 +63,7 @@ struct CONTROLRIG_API FRigModuleInstance
 	TArray<FRigModuleInstance*> CachedChildren;
 
 	FString GetPath() const;
+	FString GetNamespace() const;
 };
 
 USTRUCT(BlueprintType)
@@ -105,6 +106,8 @@ class CONTROLRIG_API UModularRig : public UControlRig
 	TArray<FRigModuleInstance> Modules;
 	TArray<FRigModuleInstance*> RootModules;
 
+	TArray<FName> SupportedEvents;
+
 public:
 
 	// BEGIN ControlRig
@@ -115,6 +118,8 @@ public:
 	virtual bool Execute_Internal(const FName& InEventName) override;
 	virtual void Evaluate_AnyThread() override;
 	virtual FRigElementKeyRedirector& GetElementKeyRedirector() override { return ElementKeyRedirector; }
+	virtual bool SupportsEvent(const FName& InEventName) const override { return SupportedEvents.Contains(InEventName); }
+	virtual const TArray<FName>& GetSupportedEvents() const override { return SupportedEvents; }
 	// END ControlRig
 
 	UPROPERTY()
@@ -132,6 +137,7 @@ public:
 	void ResetModules();
 
 	void UpdateCachedChildren();
+	void UpdateSupportedEvents();
 
 	/** Adds a module to the rig*/
 	bool AddModuleInstance(const FName& InModuleName, TSubclassOf<UControlRig> InModuleClass, FString InParentPath, const TMap<FRigElementKey, FRigElementKey>& InConnectionMap, const TMap<FName, FString>& InVariableDefaultValues, const TMap<FName, FString>& InVariableBindings);
diff --git a/Engine/Plugins/Runtime/RigVM/Source/RigVM/Public/RigVMHost.h b/Engine/Plugins/Runtime/RigVM/Source/RigVM/Public/RigVMHost.h
--- a/Engine/Plugins/Runtime/RigVM/Source/RigVM/Public/RigVMHost.h
+++ b/Engine/Plugins/Runtime/RigVM/Source/RigVM/Public/RigVMHost.h
@@ -183,10 +183,10 @@ public:
 	}
 
 	UFUNCTION(BlueprintPure, Category = "RigVM")
-	bool SupportsEvent(const FName& InEventName) const;
+	virtual bool SupportsEvent(const FName& InEventName) const;
 
 	UFUNCTION(BlueprintPure, Category = "RigVM")
-	const TArray<FName>& GetSupportedEvents() const;
+	virtual const TArray<FName>& GetSupportedEvents() const;
 
 	/** Execute a user defined event */
 	UFUNCTION(BlueprintCallable, Category = "RigVM")
