Control Rig: Only allow bindings from child to parent modules

#rnx

diff --git a/Engine/Plugins/Animation/ControlRig/Source/ControlRigDeveloper/Private/ModularRigController.cpp b/Engine/Plugins/Animation/ControlRig/Source/ControlRigDeveloper/Private/ModularRigController.cpp
--- a/Engine/Plugins/Animation/ControlRig/Source/ControlRigDeveloper/Private/ModularRigController.cpp
+++ b/Engine/Plugins/Animation/ControlRig/Source/ControlRigDeveloper/Private/ModularRigController.cpp
@@ -274,10 +274,11 @@ TArray<FString> UModularRigController::GetPossibleBindings(const FString& InModu
 	}
 
 	// Add possible module variables
-	Model->ForEachModule([this, &PossibleBindings, InModulePath, InVariableName](const FRigModuleReference* InModule) -> bool
+	const FString InvalidModulePrefix = InModulePath + UModularRig::NamespaceSeparator;
+	Model->ForEachModule([this, &PossibleBindings, InModulePath, InVariableName, InvalidModulePrefix](const FRigModuleReference* InModule) -> bool
 	{
 		const FString CurModulePath = InModule->GetPath();
-		if (InModulePath != CurModulePath)
+		if (InModulePath != CurModulePath && !CurModulePath.StartsWith(InvalidModulePrefix))
 		{
 			TArray<FRigVMExternalVariable> Variables = InModule->Class->GetDefaultObject<UControlRig>()->GetExternalVariables();
 			for (const FRigVMExternalVariable& Variable : Variables)
diff --git a/Engine/Plugins/Animation/ControlRig/Source/ControlRigEditor/Private/ControlRigModuleDetails.cpp b/Engine/Plugins/Animation/ControlRig/Source/ControlRigEditor/Private/ControlRigModuleDetails.cpp
--- a/Engine/Plugins/Animation/ControlRig/Source/ControlRigEditor/Private/ControlRigModuleDetails.cpp
+++ b/Engine/Plugins/Animation/ControlRig/Source/ControlRigEditor/Private/ControlRigModuleDetails.cpp
@@ -648,6 +648,24 @@ void FRigModuleInstanceDetails::FillBindingMenu(FMenuBuilder& MenuBuilder, const
 		}
 	}
 
+	if(CombinedBindings.IsEmpty())
+	{
+		MenuBuilder.AddMenuEntry(
+			FUIAction(FExecuteAction()),
+			SNew(SHorizontalBox)
+				+SHorizontalBox::Slot()
+				.AutoWidth()
+				.VAlign(VAlign_Center)
+				.Padding(0.0f)
+				[
+					SNew(STextBlock)
+					.Text(LOCTEXT("NoBindingAvailable", "No bindings available for this property."))
+					.ColorAndOpacity(FLinearColor::White)
+				]
+			);
+		return;
+	}
+
 	// sort lexically
 	CombinedBindings.Sort();
 
