Control Rig: Modular Rig to show tree after asset editor open

#rnx

diff --git a/Engine/Plugins/Animation/ControlRig/Source/ControlRig/Private/ModularRig.cpp b/Engine/Plugins/Animation/ControlRig/Source/ControlRig/Private/ModularRig.cpp
--- a/Engine/Plugins/Animation/ControlRig/Source/ControlRig/Private/ModularRig.cpp
+++ b/Engine/Plugins/Animation/ControlRig/Source/ControlRig/Private/ModularRig.cpp
@@ -76,6 +76,13 @@ FString FRigModuleInstance::GetNamespace() const
 	return FString::Printf(TEXT("%s:"), *GetPath());
 }
 
+void UModularRig::Initialize(bool bRequestInit)
+{
+	Super::Initialize(bRequestInit);
+
+	UpdateCachedChildren();
+}
+
 void UModularRig::InitializeVMs(bool bRequestInit)
 {
 	URigVMHost::Initialize(bRequestInit);
diff --git a/Engine/Plugins/Animation/ControlRig/Source/ControlRig/Public/ModularRig.h b/Engine/Plugins/Animation/ControlRig/Source/ControlRig/Public/ModularRig.h
--- a/Engine/Plugins/Animation/ControlRig/Source/ControlRig/Public/ModularRig.h
+++ b/Engine/Plugins/Animation/ControlRig/Source/ControlRig/Public/ModularRig.h
@@ -111,6 +111,7 @@ class CONTROLRIG_API UModularRig : public UControlRig
 public:
 
 	// BEGIN ControlRig
+	virtual void Initialize(bool bRequestInit) override;
 	virtual void InitializeVMs(bool bRequestInit = true) override;
 	virtual bool InitializeVMs(const FName& InEventName) override;
 	virtual void InitializeVMsFromCDO() override { URigVMHost::InitializeFromCDO(); }
