Control Rig: Propagate the modular rig hierarchy from BP to instances
#rnx

diff --git a/Engine/Plugins/Animation/ControlRig/Source/ControlRig/Private/ModularRig.cpp b/Engine/Plugins/Animation/ControlRig/Source/ControlRig/Private/ModularRig.cpp
--- a/Engine/Plugins/Animation/ControlRig/Source/ControlRig/Private/ModularRig.cpp
+++ b/Engine/Plugins/Animation/ControlRig/Source/ControlRig/Private/ModularRig.cpp
@@ -110,6 +110,29 @@ bool UModularRig::InitializeVMs(const FName& InEventName)
 	return true;
 }
 
+void UModularRig::InitializeFromCDO()
+{
+	Super::InitializeFromCDO();
+
+	if (!HasAnyFlags(RF_ClassDefaultObject))
+	{
+		const UModularRig* CDO = GetClass()->GetDefaultObject<UModularRig>();
+
+		// Generate the rig module tree based on the CDO
+		ResetModules();
+		CDO->ForEachModule([this](const FRigModuleInstance* CDOModule) -> bool
+		{
+			if (FRigModuleInstance* NewModule = AddModuleInstance(CDOModule))
+			{
+				NewModule->Rig->Initialize();
+			}
+			return true;
+		});
+
+		SupportedEvents = CDO->SupportedEvents;
+	}
+}
+
 bool UModularRig::Execute_Internal(const FName& InEventName)
 {
 	if (VM)
@@ -382,7 +405,10 @@ FRigModuleInstance* UModularRig::AddModuleInstance(const FName& InModuleName, TS
 		NewModule.ParentPath = InParent->GetPath();
 	}
 	UpdateCachedChildren();
-	SupportedEvents.Append(NewModule.Rig->GetSupportedEvents());
+	for (const FName& EventName : NewModule.Rig->GetSupportedEvents())
+	{
+		SupportedEvents.AddUnique(EventName);
+	}
 
 	// Configure module
 	{
@@ -404,7 +430,7 @@ FRigModuleInstance* UModularRig::AddModuleInstance(const FName& InModuleName, TS
 		for (const TPair<FName, FString>& Pair : InVariableBindings)
 		{
 			FString SourceModulePath, SourceVariableName = Pair.Value;
-			Pair.Value.Split(NamespaceSeparator, &SourceModulePath, &SourceVariableName);
+			Pair.Value.Split(NamespaceSeparator, &SourceModulePath, &SourceVariableName, ESearchCase::CaseSensitive, ESearchDir::FromEnd);
 			FRigVMExternalVariable SourceVariable;
 			if (SourceModulePath.IsEmpty())
 			{
@@ -417,6 +443,8 @@ FRigModuleInstance* UModularRig::AddModuleInstance(const FName& InModuleName, TS
 			{
 				SourceVariable = SourceModule->Rig->GetPublicVariableByName(*SourceVariableName);
 			}
+			SourceVariable.Name = *Pair.Value; // Adapt the name of the variable to contain the full path
+			check(SourceVariable.Property);
 			NewModule.VariableBindings.Add(Pair.Key, SourceVariable);
 		}
 	}
@@ -424,6 +452,39 @@ FRigModuleInstance* UModularRig::AddModuleInstance(const FName& InModuleName, TS
 	return &NewModule;
 }
 
+FRigModuleInstance* UModularRig::AddModuleInstance(const FRigModuleInstance* InOtherModule)
+{
+	if (!InOtherModule->Rig.IsValid())
+	{
+		return nullptr;
+	}
+
+	// Figure out the ConnectionMap
+	const FRigElementKeyRedirector& Redirector = InOtherModule->Rig->GetElementKeyRedirector();
+	const TMap<FRigElementKey, FRigElementKey> ConnectionMap = Redirector.ExternalKeys;
+
+	// Figure out the DefaultValues (which is the current value on the CDO)
+	TMap<FName, FString> DefaultValues;
+	const TArray<FRigVMExternalVariable> Variables = InOtherModule->Rig->GetPublicVariables();
+	for (const FRigVMExternalVariable& Variable : Variables)
+	{
+		const FString Value = InOtherModule->Rig->GetVariableAsString(Variable.Name);
+		DefaultValues.Add(Variable.Name, Value);
+	}
+
+	// Figure out the variable bindings (where the source is either a root variable or a variable on another module)
+	TMap<FName, FString> VariableBindings;
+	for (const TPair<FName, FRigVMExternalVariable>& Pair : InOtherModule->VariableBindings)
+	{
+		VariableBindings.Add(Pair.Key, Pair.Value.Name.ToString());
+	}
+
+	// Figure out the parent module
+	FRigModuleInstance* ParentModule = const_cast<FRigModuleInstance*>(FindModule(InOtherModule->ParentPath));
+	
+	return AddModuleInstance(InOtherModule->Name, InOtherModule->Rig->GetClass(), ParentModule, ConnectionMap, DefaultValues, VariableBindings);
+}
+
 const FRigModuleInstance* UModularRig::FindModule(const FString& InPath) const
 {
 	const TArray<FRigModuleInstance*>* Children = &RootModules;
diff --git a/Engine/Plugins/Animation/ControlRig/Source/ControlRig/Public/ModularRig.h b/Engine/Plugins/Animation/ControlRig/Source/ControlRig/Public/ModularRig.h
--- a/Engine/Plugins/Animation/ControlRig/Source/ControlRig/Public/ModularRig.h
+++ b/Engine/Plugins/Animation/ControlRig/Source/ControlRig/Public/ModularRig.h
@@ -115,6 +115,7 @@ public:
 	virtual void InitializeVMs(bool bRequestInit = true) override;
 	virtual bool InitializeVMs(const FName& InEventName) override;
 	virtual void InitializeVMsFromCDO() override { URigVMHost::InitializeFromCDO(); }
+	virtual void InitializeFromCDO() override;
 	virtual void RequestInitVMs() override { URigVMHost::RequestInit(); }
 	virtual bool Execute_Internal(const FName& InEventName) override;
 	virtual void Evaluate_AnyThread() override;
@@ -143,6 +144,7 @@ public:
 	/** Adds a module to the rig*/
 	bool AddModuleInstance(const FName& InModuleName, TSubclassOf<UControlRig> InModuleClass, FString InParentPath, const TMap<FRigElementKey, FRigElementKey>& InConnectionMap, const TMap<FName, FString>& InVariableDefaultValues, const TMap<FName, FString>& InVariableBindings);
 	FRigModuleInstance* AddModuleInstance(const FName& InModuleName, TSubclassOf<UControlRig> InModuleClass, FRigModuleInstance* InParent, const TMap<FRigElementKey, FRigElementKey>& InConnectionMap, const TMap<FName, FString>& InVariableDefaultValues, const TMap<FName, FString>& InVariableBindings);
+	FRigModuleInstance* AddModuleInstance(const FRigModuleInstance* InOtherModule);
 
 	const FRigModuleInstance* FindModule(const FString& InPath) const;
 	const FRigModuleInstance* FindModule(const UControlRig* InModuleInstance) const;
diff --git a/Engine/Plugins/Animation/ControlRig/Source/ControlRig/Public/Rigs/RigHierarchyCache.h b/Engine/Plugins/Animation/ControlRig/Source/ControlRig/Public/Rigs/RigHierarchyCache.h
--- a/Engine/Plugins/Animation/ControlRig/Source/ControlRig/Public/Rigs/RigHierarchyCache.h
+++ b/Engine/Plugins/Animation/ControlRig/Source/ControlRig/Public/Rigs/RigHierarchyCache.h
@@ -186,4 +186,5 @@ private:
 	friend class URigHierarchyController;
 	friend struct FRigConnectionInfo;
 	friend class FRigModuleInstanceDetails;
+	friend class UModularRig;
 };
\ No newline at end of file
diff --git a/Engine/Plugins/Animation/ControlRig/Source/ControlRigDeveloper/Private/ControlRigBlueprint.cpp b/Engine/Plugins/Animation/ControlRig/Source/ControlRigDeveloper/Private/ControlRigBlueprint.cpp
--- a/Engine/Plugins/Animation/ControlRig/Source/ControlRigDeveloper/Private/ControlRigBlueprint.cpp
+++ b/Engine/Plugins/Animation/ControlRig/Source/ControlRigDeveloper/Private/ControlRigBlueprint.cpp
@@ -1560,16 +1560,6 @@ static void AddModulesRecursively(UModularRig* Rig, const FRigModuleReference* I
 	}
 }
 
-static void BuildModularRig(UModularRig* Rig, const FModularRigModel& Model)
-{
-	Rig->ResetModules();
-	for (const FRigModuleReference* RootModule : Model.RootModules)
-	{
-		AddModulesRecursively(Rig, RootModule, nullptr);
-	}
-	Rig->InitializeVMs(true);
-}
-
 void UControlRigBlueprint::RecompileModularRig()
 {
 	OnModularRigPreCompiled().Broadcast(this);
@@ -1577,18 +1567,13 @@ void UControlRigBlueprint::RecompileModularRig()
 	{
 		if (UModularRig* DefaultObject = Cast<UModularRig>(MyControlRigClass->GetDefaultObject(false)))
 		{
-			BuildModularRig(DefaultObject, ModularRigModel);
-
-			// Initialize Archetype Instances
-			TArray<UObject*> ArchetypeInstances;
-			DefaultObject->GetArchetypeInstances(ArchetypeInstances);
-			for (UObject* Instance : ArchetypeInstances)
+			DefaultObject->ResetModules();
+			for (const FRigModuleReference* RootModule : ModularRigModel.RootModules)
 			{
-				if (UModularRig* ModularRigInstance = Cast<UModularRig>(Instance))
-				{
-					BuildModularRig(ModularRigInstance, ModularRigModel);
-				}
+				AddModulesRecursively(DefaultObject, RootModule, nullptr);
 			}
+			DefaultObject->InitializeVMs(true);
+			PropagateModuleHierarchyFromBPToInstances();
 		}
 	}
 
@@ -2187,6 +2172,25 @@ void UControlRigBlueprint::PropagatePropertyFromInstanceToBP(FRigElementKey InRi
 	FMemory::Memcpy(Dest, Source, PropertySize);
 }
 
+void UControlRigBlueprint::PropagateModuleHierarchyFromBPToInstances() const
+{
+	if (UClass* MyControlRigClass = GeneratedClass)
+	{
+		if (UControlRig* DefaultObject = Cast<UControlRig>(MyControlRigClass->GetDefaultObject(false)))
+		{
+			TArray<UObject*> ArchetypeInstances;
+			DefaultObject->GetArchetypeInstances(ArchetypeInstances);
+			for (UObject* ArchetypeInstance : ArchetypeInstances)
+			{
+				if (UControlRig* InstanceRig = Cast<UControlRig>(ArchetypeInstance))
+				{
+					// this will provoke a call to InitializeFromCDO
+					InstanceRig->Initialize(true);
+				}
+			}
+		}
+	}
+}
 
 void UControlRigBlueprint::HandleHierarchyModified(ERigHierarchyNotification InNotification, URigHierarchy* InHierarchy, const FRigBaseElement* InElement)
 {
diff --git a/Engine/Plugins/Animation/ControlRig/Source/ControlRigDeveloper/Private/ModularRigController.cpp b/Engine/Plugins/Animation/ControlRig/Source/ControlRigDeveloper/Private/ModularRigController.cpp
--- a/Engine/Plugins/Animation/ControlRig/Source/ControlRigDeveloper/Private/ModularRigController.cpp
+++ b/Engine/Plugins/Animation/ControlRig/Source/ControlRigDeveloper/Private/ModularRigController.cpp
@@ -387,7 +387,7 @@ bool UModularRigController::BindModuleVariable(const FString& InModulePath, cons
 	const FProperty* TargetProperty = Module->Class->FindPropertyByName(InVariableName);
 
 	FString SourceModulePath, SourceVariableName = InSourcePath;
-	InSourcePath.Split(UModularRig::NamespaceSeparator, &SourceModulePath, &SourceVariableName);
+	InSourcePath.Split(UModularRig::NamespaceSeparator, &SourceModulePath, &SourceVariableName, ESearchCase::CaseSensitive, ESearchDir::FromEnd);
 
 	FRigModuleReference* SourceModule = nullptr;
 	if (!SourceModulePath.IsEmpty())
@@ -509,7 +509,7 @@ bool UModularRigController::DeleteModule(const FString& InModulePath, bool bSetu
 		Reference.Bindings = Reference.Bindings.FilterByPredicate([InModulePath](const TPair<FName, FString>& Binding)
 		{
 			FString ModulePath, VariableName = Binding.Value;
-			Binding.Value.Split(UModularRig::NamespaceSeparator, &ModulePath, &VariableName);
+			Binding.Value.Split(UModularRig::NamespaceSeparator, &ModulePath, &VariableName, ESearchCase::CaseSensitive, ESearchDir::FromEnd);
 			if (ModulePath == InModulePath)
 			{
 				return false;
@@ -583,7 +583,7 @@ bool UModularRigController::RenameModule(const FString& InModulePath, const FNam
 		for (TPair<FName, FString>& Binding : Reference.Bindings)
 		{
 			FString ModulePath, VariableName = Binding.Value;
-			Binding.Value.Split(UModularRig::NamespaceSeparator, &ModulePath, &VariableName);
+			Binding.Value.Split(UModularRig::NamespaceSeparator, &ModulePath, &VariableName, ESearchCase::CaseSensitive, ESearchDir::FromEnd);
 			if (ModulePath == OldPath)
 			{
 				Binding.Value = FString::Printf(TEXT("%s:%s"), *NewPath, *VariableName);
@@ -688,7 +688,7 @@ bool UModularRigController::ReparentModule(const FString& InModulePath, const FS
 		for (TPair<FName, FString>& Binding : Reference.Bindings)
 		{
 			FString ModulePath, VariableName = Binding.Value;
-			Binding.Value.Split(UModularRig::NamespaceSeparator, &ModulePath, &VariableName);
+			Binding.Value.Split(UModularRig::NamespaceSeparator, &ModulePath, &VariableName, ESearchCase::CaseSensitive, ESearchDir::FromEnd);
 			if (ModulePath == OldPath)
 			{
 				Binding.Value = FString::Printf(TEXT("%s:%s"), *NewPath, *VariableName);
diff --git a/Engine/Plugins/Animation/ControlRig/Source/ControlRigDeveloper/Public/ControlRigBlueprint.h b/Engine/Plugins/Animation/ControlRig/Source/ControlRigDeveloper/Public/ControlRigBlueprint.h
--- a/Engine/Plugins/Animation/ControlRig/Source/ControlRigDeveloper/Public/ControlRigBlueprint.h
+++ b/Engine/Plugins/Animation/ControlRig/Source/ControlRigDeveloper/Public/ControlRigBlueprint.h
@@ -290,6 +290,7 @@ public:
 	void PropagateDrawInstructionsFromBPToInstances() const;
 	void PropagatePropertyFromBPToInstances(FRigElementKey InRigElement, const FProperty* InProperty) const;
 	void PropagatePropertyFromInstanceToBP(FRigElementKey InRigElement, const FProperty* InProperty, UControlRig* InInstance) const;
+	void PropagateModuleHierarchyFromBPToInstances() const;
 
 	/**
 	* Returns the modified event, which can be used to 
