Control Rig: Allow drag and drop module class onto a rig element in the rig hierarchy widget
#rnx

diff --git a/Engine/Plugins/Animation/ControlRig/Source/ControlRig/Private/ControlRig.cpp b/Engine/Plugins/Animation/ControlRig/Source/ControlRig/Private/ControlRig.cpp
--- a/Engine/Plugins/Animation/ControlRig/Source/ControlRig/Private/ControlRig.cpp
+++ b/Engine/Plugins/Animation/ControlRig/Source/ControlRig/Private/ControlRig.cpp
@@ -604,6 +604,11 @@ bool UControlRig::AllConnectorsAreResolved(FString* OutFailureReason, FRigElemen
 		TMap<FRigElementKey, FRigElementKey> ConnectionMap;
 		for(const FRigModuleConnector& Connector : Connectors)
 		{
+			if (Connector.Settings.bOptional)
+			{
+				continue;
+			}
+			
 			const FRigElementKey ConnectorKey(*Connector.Name, ERigElementType::Connector);
 			if(const FCachedRigElement* Cache = ElementKeyRedirector.Find(ConnectorKey))
 			{
diff --git a/Engine/Plugins/Animation/ControlRig/Source/ControlRig/Private/Rigs/RigHierarchy.cpp b/Engine/Plugins/Animation/ControlRig/Source/ControlRig/Private/Rigs/RigHierarchy.cpp
--- a/Engine/Plugins/Animation/ControlRig/Source/ControlRig/Private/Rigs/RigHierarchy.cpp
+++ b/Engine/Plugins/Animation/ControlRig/Source/ControlRig/Private/Rigs/RigHierarchy.cpp
@@ -4218,6 +4218,14 @@ void URigHierarchy::SetConnectorSettings(FRigConnectorElement* InConnectorElemen
 		return;
 	}
 
+	if (InSettings.Type == EConnectorType::Primary)
+	{
+		if (InSettings.bOptional)
+		{
+			return;
+		}
+	}
+
 	if(bSetupUndo && !HasAnyFlags(RF_Transient))
 	{
 		Modify();
diff --git a/Engine/Plugins/Animation/ControlRig/Source/ControlRig/Private/Rigs/RigHierarchyElements.cpp b/Engine/Plugins/Animation/ControlRig/Source/ControlRig/Private/Rigs/RigHierarchyElements.cpp
--- a/Engine/Plugins/Animation/ControlRig/Source/ControlRig/Private/Rigs/RigHierarchyElements.cpp
+++ b/Engine/Plugins/Animation/ControlRig/Source/ControlRig/Private/Rigs/RigHierarchyElements.cpp
@@ -1427,6 +1427,8 @@ FRigConnectionInfo::FRigConnectionInfo(const FRigElementKeyRedirector* InRedirec
 ////////////////////////////////////////////////////////////////////////////////
 
 FRigConnectorSettings::FRigConnectorSettings()
+	: Type(EConnectorType::Primary)
+	, bOptional(false)
 {
 }
 
@@ -1436,6 +1438,12 @@ void FRigConnectorSettings::Save(FArchive& Ar)
 
 	Ar << Description;
 
+	if (Ar.CustomVer(FControlRigObjectVersion::GUID) >= FControlRigObjectVersion::ConnectorsWithType)
+	{
+		Ar << Type;
+		Ar << bOptional;
+	}
+
 	int32 NumRules = Rules.Num();
 	Ar << NumRules;
 	for(int32 Index = 0; Index < NumRules; Index++)
@@ -1451,12 +1459,17 @@ void FRigConnectorSettings::Load(FArchive& Ar)
 
 	Ar << Description;
 
+	if (Ar.CustomVer(FControlRigObjectVersion::GUID) >= FControlRigObjectVersion::ConnectorsWithType)
+	{
+		Ar << Type;
+		Ar << bOptional;
+	}
+
 	int32 NumRules = 0;
 	Ar << NumRules;
 	Rules.SetNumZeroed(NumRules);
 	for(int32 Index = 0; Index < NumRules; Index++)
 	{
-
 		Rules[Index].Load(Ar);
 	}
 }
@@ -1467,6 +1480,14 @@ bool FRigConnectorSettings::operator==(const FRigConnectorSettings& InOther) con
 	{
 		return false;
 	}
+	if(Type != InOther.Type)
+	{
+		return false;
+	}
+	if(bOptional != InOther.bOptional)
+	{
+		return false;
+	}
 	if(Rules.Num() != InOther.Rules.Num())
 	{
 		return false;
@@ -1493,7 +1514,9 @@ uint32 FRigConnectorSettings::GetRulesHash() const
 
 uint32 GetTypeHash(const FRigConnectorSettings& Settings)
 {
-	return Settings.GetRulesHash();
+	uint32 Hash = HashCombine(GetTypeHash(Settings.Type), Settings.GetRulesHash());
+	Hash = HashCombine(Hash, GetTypeHash(Settings.bOptional));
+	return Hash;
 }
 
 ////////////////////////////////////////////////////////////////////////////////
diff --git a/Engine/Plugins/Animation/ControlRig/Source/ControlRig/Public/ControlRigObjectVersion.h b/Engine/Plugins/Animation/ControlRig/Source/ControlRig/Public/ControlRigObjectVersion.h
--- a/Engine/Plugins/Animation/ControlRig/Source/ControlRig/Public/ControlRigObjectVersion.h
+++ b/Engine/Plugins/Animation/ControlRig/Source/ControlRig/Public/ControlRigObjectVersion.h
@@ -95,6 +95,9 @@ struct CONTROLRIG_API FControlRigObjectVersion
 		// Element metadata is now stored on URigHierarchy, rather than FRigBaseElement
 		RigHierarchyStoresElementMetadata,
 
+		// Add type (primary, secondary) and optional bool to FRigConnectorSettings
+		ConnectorsWithType,
+
 		// -----<new versions can be added above this line>-------------------------------------------------
 		VersionPlusOne,
 		LatestVersion = VersionPlusOne - 1
diff --git a/Engine/Plugins/Animation/ControlRig/Source/ControlRig/Public/Rigs/RigHierarchyElements.h b/Engine/Plugins/Animation/ControlRig/Source/ControlRig/Public/Rigs/RigHierarchyElements.h
--- a/Engine/Plugins/Animation/ControlRig/Source/ControlRig/Public/Rigs/RigHierarchyElements.h
+++ b/Engine/Plugins/Animation/ControlRig/Source/ControlRig/Public/Rigs/RigHierarchyElements.h
@@ -1653,6 +1653,15 @@ public:
 	const URigHierarchy* TargetHierarchy; 
 };
 
+UENUM(BlueprintType)
+enum class EConnectorType : uint8
+{
+	Primary, // Single primary connector, non-optional and always visible. When dropped on another element, this connector will resolve to that element.
+	
+	Secondary, // Could be multiple, can auto-solve (visible if not solved), can be optional
+};
+
+
 USTRUCT(BlueprintType)
 struct CONTROLRIG_API FRigConnectorSettings
 {
@@ -1668,6 +1677,12 @@ struct CONTROLRIG_API FRigConnectorSettings
 	UPROPERTY(BlueprintReadWrite, EditAnywhere, Category = Connector)
 	FString Description;
 
+	UPROPERTY(BlueprintReadWrite, EditAnywhere, Category = Connector)
+	EConnectorType Type;
+
+	UPROPERTY(BlueprintReadWrite, EditAnywhere, Category = Connector)
+	bool bOptional;
+
 	UPROPERTY(BlueprintReadWrite, EditAnywhere, Category = Connector)
 	TArray<FRigConnectionRuleStash> Rules;
 
diff --git a/Engine/Plugins/Animation/ControlRig/Source/ControlRig/Public/Rigs/RigModuleDefines.h b/Engine/Plugins/Animation/ControlRig/Source/ControlRig/Public/Rigs/RigModuleDefines.h
--- a/Engine/Plugins/Animation/ControlRig/Source/ControlRig/Public/Rigs/RigModuleDefines.h
+++ b/Engine/Plugins/Animation/ControlRig/Source/ControlRig/Public/Rigs/RigModuleDefines.h
@@ -33,15 +33,11 @@ struct CONTROLRIG_API FRigModuleConnector
 	GENERATED_BODY()
 	
 	FRigModuleConnector()
-		: bIsRoot(false)
 	{}
 
 	UPROPERTY(BlueprintReadWrite, EditAnywhere, Category = Connector)
 	FString Name;
 
-	UPROPERTY(BlueprintReadWrite, EditAnywhere, Category = Connector)
-	bool bIsRoot;
-
 	UPROPERTY(BlueprintReadWrite, EditAnywhere, Category = Connector)
 	FRigConnectorSettings Settings;
 };
diff --git a/Engine/Plugins/Animation/ControlRig/Source/ControlRigDeveloper/Private/ControlRigBlueprint.cpp b/Engine/Plugins/Animation/ControlRig/Source/ControlRigDeveloper/Private/ControlRigBlueprint.cpp
--- a/Engine/Plugins/Animation/ControlRig/Source/ControlRigDeveloper/Private/ControlRigBlueprint.cpp
+++ b/Engine/Plugins/Animation/ControlRig/Source/ControlRigDeveloper/Private/ControlRigBlueprint.cpp
@@ -948,11 +948,11 @@ void UControlRigBlueprint::UpdateExposedModuleConnectors() const
 	{
 		FRigModuleConnector ExposedConnector;
 		ExposedConnector.Name = ConnectorElement->GetName();
-		ExposedConnector.bIsRoot = MutableThis->Hierarchy->GetNumberOfParents(ConnectorElement) == 0;
 		ExposedConnector.Settings = ConnectorElement->Settings;
 		MutableThis->RigModuleSettings.ExposedConnectors.Add(ExposedConnector);
 		return true;
 	});
+	PropagateHierarchyFromBPToInstances();
 }
 
 bool UControlRigBlueprint::ResolveConnector(const FRigElementKey& DraggedKey, const FRigElementKey& TargetKey, bool bSetupUndoRedo)
@@ -975,7 +975,7 @@ bool UControlRigBlueprint::ResolveConnector(const FRigElementKey& DraggedKey, co
 		// Add connection to the model
 		if (UModularRigController* Controller = GetModularRigController())
 		{
-			Controller->ConnectModuleToElement(DraggedKey, TargetKey, bSetupUndoRedo);
+			Controller->ConnectConnectorToElement(DraggedKey, TargetKey, bSetupUndoRedo);
 		}
 	}
 	else
diff --git a/Engine/Plugins/Animation/ControlRig/Source/ControlRigDeveloper/Private/ModularRigController.cpp b/Engine/Plugins/Animation/ControlRig/Source/ControlRigDeveloper/Private/ModularRigController.cpp
--- a/Engine/Plugins/Animation/ControlRig/Source/ControlRigDeveloper/Private/ModularRigController.cpp
+++ b/Engine/Plugins/Animation/ControlRig/Source/ControlRigDeveloper/Private/ModularRigController.cpp
@@ -21,19 +21,19 @@ UModularRigController::UModularRigController(const FObjectInitializer& ObjectIni
 {
 }
 
-bool UModularRigController::AddModule(const FName& InModuleName, TSubclassOf<UControlRig> InClass, const FString& InParentModulePath, bool bSetupUndo)
+FString UModularRigController::AddModule(const FName& InModuleName, TSubclassOf<UControlRig> InClass, const FString& InParentModulePath, bool bSetupUndo)
 {
 	if (!InClass)
 	{
 		UE_LOG(LogControlRig, Error, TEXT("Invalid InClass"));
-		return false;
+		return FString();
 	}
 
 	UControlRig* ClassDefaultObject = InClass->GetDefaultObject<UControlRig>();
 	if (!ClassDefaultObject->IsRigModule())
 	{
 		UE_LOG(LogControlRig, Error, TEXT("Class %s is not a rig module"), *InClass->GetClassPathName().ToString());
-		return false;
+		return FString();
 	}
 
 #if WITH_EDITOR
@@ -55,7 +55,7 @@ bool UModularRigController::AddModule(const FName& InModuleName, TSubclassOf<UCo
 		{
 			if (Module->Name.ToString() == InModuleName)
 			{
-				return false;
+				return FString();
 			}
 		}
 
@@ -68,7 +68,7 @@ bool UModularRigController::AddModule(const FName& InModuleName, TSubclassOf<UCo
 		{
 			if (Module->Name.ToString() == InModuleName)
 			{
-				return false;
+				return FString();
 			}
 		}
 
@@ -81,7 +81,7 @@ bool UModularRigController::AddModule(const FName& InModuleName, TSubclassOf<UCo
 	if (!NewModule)
 	{
 		UE_LOG(LogControlRig, Error, TEXT("Error while creating module %s"), *InModuleName.ToString());
-		return false;
+		return FString();
 	}
 
 	Notify(EModularRigNotification::ModuleAdded, NewModule);
@@ -89,8 +89,8 @@ bool UModularRigController::AddModule(const FName& InModuleName, TSubclassOf<UCo
 #if WITH_EDITOR
 	TransactionPtr.Reset();
 #endif
-	
-	return true;
+
+	return NewModule->GetPath();
 }
 
 FRigModuleReference* UModularRigController::FindModule(const FString& InPath)
@@ -125,7 +125,13 @@ FRigModuleReference* UModularRigController::FindModule(const FString& InPath)
 	return *Cur;
 }
 
-bool UModularRigController::ConnectModuleToElement(const FRigElementKey& InConnectorKey, const FRigElementKey& InTargetKey, bool bSetupUndo)
+bool UModularRigController::CanConnectConnectorToElement(const FRigModuleConnector& InConnector, const FRigElementKey& InTargetKey, FText& OutErrorMessage)
+{
+	// TODO: Check rules are satisfied
+	return true;
+}
+
+bool UModularRigController::ConnectConnectorToElement(const FRigElementKey& InConnectorKey, const FRigElementKey& InTargetKey, bool bSetupUndo)
 {
 	FString ConnectorParentPath, ConnectorName;
 	if (!InConnectorKey.Name.ToString().Split(UModularRig::NamespaceSeparator, &ConnectorParentPath, &ConnectorName, ESearchCase::CaseSensitive, ESearchDir::FromEnd))
@@ -141,6 +147,32 @@ bool UModularRigController::ConnectModuleToElement(const FRigElementKey& InConne
 		return false;
 	}
 
+	UControlRig* RigCDO = Module->Class->GetDefaultObject<UControlRig>();
+	if (!RigCDO)
+	{
+		UE_LOG(LogControlRig, Error, TEXT("Invalid rig module class %s"), *Module->Class->GetPathName());
+		return false;
+	}
+
+	const FRigModuleConnector* ModuleConnector = RigCDO->GetRigModuleSettings().ExposedConnectors.FindByPredicate(
+		[ConnectorName](FRigModuleConnector& Connector)
+		{
+			return Connector.Name == ConnectorName;
+		});
+	if (!ModuleConnector)
+	{
+		UE_LOG(LogControlRig, Error, TEXT("Could not find connector %s in class %s"), *ConnectorName, *Module->Class->GetPathName());
+		return false;
+	}
+
+	FText ErrorMessage;
+	if (!CanConnectConnectorToElement(*ModuleConnector, InTargetKey, ErrorMessage))
+	{
+		UE_LOG(LogControlRig, Error, TEXT("Cannot connect connector %s to target %s: %s"),
+			*InConnectorKey.Name.ToString(), *InTargetKey.ToString(), *ErrorMessage.ToString());
+		return false;
+	}
+
 	UControlRigBlueprint* Blueprint = Cast<UControlRigBlueprint>(GetOuter());
 	FRigConnectorElement* Connector = Cast<FRigConnectorElement>(Blueprint->Hierarchy->Find(InConnectorKey));
 	if (!Connector)
diff --git a/Engine/Plugins/Animation/ControlRig/Source/ControlRigDeveloper/Public/ModularRigController.h b/Engine/Plugins/Animation/ControlRig/Source/ControlRigDeveloper/Public/ModularRigController.h
--- a/Engine/Plugins/Animation/ControlRig/Source/ControlRigDeveloper/Public/ModularRigController.h
+++ b/Engine/Plugins/Animation/ControlRig/Source/ControlRigDeveloper/Public/ModularRigController.h
@@ -49,10 +49,13 @@ class CONTROLRIGDEVELOPER_API UModularRigController : public UObject
 
 
 	UFUNCTION(BlueprintCallable, Category = "Control Rig | Modules")
-	bool AddModule(const FName& InModuleName, TSubclassOf<UControlRig> InClass, const FString& InParentModulePath, bool bSetupUndo = true);
+	FString AddModule(const FName& InModuleName, TSubclassOf<UControlRig> InClass, const FString& InParentModulePath, bool bSetupUndo = true);
 
 	UFUNCTION(BlueprintCallable, Category = "Control Rig | Modules")
-	bool ConnectModuleToElement(const FRigElementKey& InConnectorKey, const FRigElementKey& InTargetKey, bool bSetupUndo = true);
+	bool CanConnectConnectorToElement(const FRigModuleConnector& InConnector, const FRigElementKey& InTargetKey, FText& OutErrorMessage);
+
+	UFUNCTION(BlueprintCallable, Category = "Control Rig | Modules")
+	bool ConnectConnectorToElement(const FRigElementKey& InConnectorKey, const FRigElementKey& InTargetKey, bool bSetupUndo = true);
 
 	UFUNCTION(BlueprintCallable, Category = "Control Rig | Modules")
 	bool SetConfigValueInModule(const FString& InModulePath, const FName& InVariableName, const FString& InValue, bool bSetupUndo = true);
diff --git a/Engine/Plugins/Animation/ControlRig/Source/ControlRigEditor/Private/ControlRigElementDetails.cpp b/Engine/Plugins/Animation/ControlRig/Source/ControlRigEditor/Private/ControlRigElementDetails.cpp
--- a/Engine/Plugins/Animation/ControlRig/Source/ControlRigEditor/Private/ControlRigElementDetails.cpp
+++ b/Engine/Plugins/Animation/ControlRig/Source/ControlRigEditor/Private/ControlRigElementDetails.cpp
@@ -1122,6 +1122,14 @@ bool FRigBaseElementDetails::IsAnyElementProcedural() const
 	});
 }
 
+bool FRigBaseElementDetails::IsAnyConnectorImported() const
+{
+	return ContainsElementByPredicate([](const FPerElementInfo& Info)
+	{
+		return Info.Element.GetKey().Name.ToString().Contains(UModularRig::NamespaceSeparator);
+	});
+}
+
 bool FRigBaseElementDetails::GetCommonElementType(ERigElementType& OutElementType) const
 {
 	OutElementType = ERigElementType::None;
@@ -5210,6 +5218,10 @@ void FRigConnectorElementDetails::CustomizeSettings(IDetailLayoutBuilder& Detail
 
 	IDetailCategoryBuilder& SettingsCategory = DetailBuilder.EditCategory(TEXT("Settings"), LOCTEXT("Settings", "Settings"));
 
+	SettingsCategory
+		.AddProperty(SettingsHandle->GetChildHandle(GET_MEMBER_NAME_CHECKED(FRigConnectorSettings, bOptional)))
+		.IsEnabled(!IsAnyConnectorImported());
+
 	bool bHideRules = false;
 	uint32 FirstHash = UINT32_MAX;
 	for (const FPerElementInfo& Info : PerElementInfos)
@@ -5240,7 +5252,7 @@ void FRigConnectorElementDetails::CustomizeSettings(IDetailLayoutBuilder& Detail
 	{
 		SettingsCategory
 			.AddProperty(SettingsHandle->GetChildHandle(GET_MEMBER_NAME_CHECKED(FRigConnectorSettings, Rules)))
-			.IsEnabled(!IsAnyElementProcedural());
+			.IsEnabled(!IsAnyConnectorImported());
 	}
 }
 
diff --git a/Engine/Plugins/Animation/ControlRig/Source/ControlRigEditor/Private/ControlRigElementDetails.h b/Engine/Plugins/Animation/ControlRig/Source/ControlRigEditor/Private/ControlRigElementDetails.h
--- a/Engine/Plugins/Animation/ControlRig/Source/ControlRigEditor/Private/ControlRigElementDetails.h
+++ b/Engine/Plugins/Animation/ControlRig/Source/ControlRigEditor/Private/ControlRigElementDetails.h
@@ -290,6 +290,7 @@ public:
 	bool IsAnyControlOfValueType(ERigControlType InType) const;
 	bool IsAnyControlNotOfValueType(ERigControlType InType) const;
 	bool IsAnyElementProcedural() const;
+	bool IsAnyConnectorImported() const;
 	bool GetCommonElementType(ERigElementType& OutElementType) const;
 	bool GetCommonControlType(ERigControlType& OutControlType) const;
 	bool GetCommonAnimationType(ERigControlAnimationType& OutAnimationType) const;
diff --git a/Engine/Plugins/Animation/ControlRig/Source/ControlRigEditor/Private/ControlRigModuleDetails.cpp b/Engine/Plugins/Animation/ControlRig/Source/ControlRigEditor/Private/ControlRigModuleDetails.cpp
--- a/Engine/Plugins/Animation/ControlRig/Source/ControlRigEditor/Private/ControlRigModuleDetails.cpp
+++ b/Engine/Plugins/Animation/ControlRig/Source/ControlRigEditor/Private/ControlRigModuleDetails.cpp
@@ -494,7 +494,7 @@ void FRigModuleInstanceDetails::OnElementNameChanged(TSharedPtr<FString> InItem,
 				}
 
 				FRigElementKey NamespacedConnector(*FString::Printf(TEXT("%s:%s"), *Info.GetModule()->GetPath(), *Connector.Name.ToString()), ERigElementType::Connector);
-				Controller->ConnectModuleToElement(NamespacedConnector, *TargetKey);
+				Controller->ConnectConnectorToElement(NamespacedConnector, *TargetKey);
 			}
 		}
 	}
@@ -514,7 +514,7 @@ void FRigModuleInstanceDetails::OnElementTypeChanged(ERigElementType InElementTy
 				TargetKey->Type = InElementType;
 
 				FRigElementKey NamespacedConnector(*FString::Printf(TEXT("%s:%s"), *Info.GetModule()->GetPath(), *Connector.Name.ToString()), ERigElementType::Connector);
-				Controller->ConnectModuleToElement(NamespacedConnector, *TargetKey);
+				Controller->ConnectConnectorToElement(NamespacedConnector, *TargetKey);
 			}
 		}
 	}
@@ -537,7 +537,7 @@ FReply FRigModuleInstanceDetails::OnGetSelectedClicked(FRigElementKey Connector)
 		if (Selected.Num() > 0)
 		{
 			FRigElementKey NamespacedConnector(*FString::Printf(TEXT("%s:%s"), *PerModuleInfos[0].GetModule()->GetPath(), *Connector.Name.ToString()), ERigElementType::Connector);
-			Blueprint->GetModularRigController()->ConnectModuleToElement(NamespacedConnector, Selected[0]);
+			Blueprint->GetModularRigController()->ConnectConnectorToElement(NamespacedConnector, Selected[0]);
 		}
 	}
 	return FReply::Handled();
diff --git a/Engine/Plugins/Animation/ControlRig/Source/ControlRigEditor/Private/Editor/ControlRigEditor.cpp b/Engine/Plugins/Animation/ControlRig/Source/ControlRigEditor/Private/Editor/ControlRigEditor.cpp
--- a/Engine/Plugins/Animation/ControlRig/Source/ControlRigEditor/Private/Editor/ControlRigEditor.cpp
+++ b/Engine/Plugins/Animation/ControlRig/Source/ControlRigEditor/Private/Editor/ControlRigEditor.cpp
@@ -3477,12 +3477,21 @@ void FControlRigEditor::OnHierarchyModified_AnyThread(ERigHierarchyNotification
 							}
 						}
 					}
-				}						
+				}
 				break;
 			}
 			case ERigHierarchyNotification::ElementAdded:
 			case ERigHierarchyNotification::ElementRemoved:
 			case ERigHierarchyNotification::ElementRenamed:
+			{
+				if (Key.IsValid() && Key.Type == ERigElementType::Connector)
+				{
+					UControlRigBlueprint* RigBlueprint = GetControlRigBlueprint();
+					check(RigBlueprint);
+					RigBlueprint->UpdateExposedModuleConnectors();
+				}
+				// Fallthrough to next case
+			}
 			case ERigHierarchyNotification::ParentChanged:
             case ERigHierarchyNotification::HierarchyReset:
 			{
@@ -3530,6 +3539,14 @@ void FControlRigEditor::OnHierarchyModified_AnyThread(ERigHierarchyNotification
 				}
 				break;
 			}
+			case ERigHierarchyNotification::ConnectorSettingChanged:
+			{
+				UControlRigBlueprint* RigBlueprint = GetControlRigBlueprint();
+				check(RigBlueprint);
+				RigBlueprint->UpdateExposedModuleConnectors();
+				RigBlueprint->RecompileModularRig();
+				break;
+			}
 			default:
 			{
 				break;
diff --git a/Engine/Plugins/Animation/ControlRig/Source/ControlRigEditor/Private/Editor/SRigHierarchy.cpp b/Engine/Plugins/Animation/ControlRig/Source/ControlRigEditor/Private/Editor/SRigHierarchy.cpp
--- a/Engine/Plugins/Animation/ControlRig/Source/ControlRigEditor/Private/Editor/SRigHierarchy.cpp
+++ b/Engine/Plugins/Animation/ControlRig/Source/ControlRigEditor/Private/Editor/SRigHierarchy.cpp
@@ -48,6 +48,7 @@
 #include "Sequencer/ControlRigLayerInstance.h"
 #include "Algo/MinElement.h"
 #include "Algo/MaxElement.h"
+#include "DragAndDrop/AssetDragDropOp.h"
 #include "RigVMFunctions/Math/RigVMMathLibrary.h"
 #include "Preferences/PersonaOptions.h"
 
@@ -2030,7 +2031,7 @@ void SRigHierarchy::HandleNewItem(ERigElementType InElementType, bool bIsAnimati
 						Settings.MaximumValue = ValueToSet;
 
 						NewItemKey = Controller->AddControl(NewElementName, ParentKey, Settings, Settings.GetIdentityValue(), FTransform::Identity, FTransform::Identity, true, true);
-					}						
+					}
 					break;
 				}
 				case ERigElementType::Null:
@@ -2063,7 +2064,9 @@ void SRigHierarchy::HandleNewItem(ERigElementType InElementType, bool bIsAnimati
 						}
 					}
 
-					NewItemKey = Controller->AddConnector(NewElementName, FRigConnectorSettings(), true);
+					FRigConnectorSettings Settings;
+					Settings.Type = Hierarchy->GetConnectorKeys(false).Num() == 0 ? EConnectorType::Primary : EConnectorType::Secondary;
+					NewItemKey = Controller->AddConnector(NewElementName, Settings, true);
 					(void)ResolveConnector(NewItemKey, ParentKey);
 					break;
 				}
@@ -2635,6 +2638,61 @@ TOptional<EItemDropZone> SRigHierarchy::OnCanAcceptDrop(const FDragDropEvent& Dr
 		}
 	}
 
+	TSharedPtr<FAssetDragDropOp> AssetDragDropOp = DragDropEvent.GetOperationAs<FAssetDragDropOp>();
+	if (AssetDragDropOp.IsValid())
+	{
+		for (const FAssetData& AssetData : AssetDragDropOp->GetAssets())
+		{
+			static const UEnum* ControlTypeEnum = StaticEnum<EControlRigType>();
+			const FString ControlRigTypeStr = AssetData.GetTagValueRef<FString>(TEXT("ControlRigType"));
+			if (ControlRigTypeStr.IsEmpty())
+			{
+				return InvalidDropZone;
+			}
+
+			const EControlRigType ControlRigType = (EControlRigType)(ControlTypeEnum->GetValueByName(*ControlRigTypeStr));
+			if (ControlRigType != EControlRigType::RigModule)
+			{
+				return InvalidDropZone;
+			}
+
+			if(UControlRigBlueprint* AssetBlueprint = Cast<UControlRigBlueprint>(AssetData.GetAsset()))
+			{
+				if (UModularRigController* Controller = ControlRigBlueprint->GetModularRigController())
+				{
+					FRigModuleConnector* PrimaryConnector = nullptr;
+					for (FRigModuleConnector& Connector : AssetBlueprint->RigModuleSettings.ExposedConnectors)
+					{
+						if (Connector.Settings.Type == EConnectorType::Primary)
+						{
+							PrimaryConnector = &Connector;
+							break;
+						}
+					}
+					if (!PrimaryConnector)
+					{
+						return InvalidDropZone;
+					}
+
+					FRigElementKey TargetKey;
+					if (TargetItem.IsValid())
+					{
+						TargetKey = TargetItem->Key;
+					}
+					FText ErrorMessage;
+					if (Controller->CanConnectConnectorToElement(*PrimaryConnector, TargetKey, ErrorMessage))
+					{
+						ReturnDropZone = DropZone;
+					}
+					else
+					{
+						return InvalidDropZone;
+					}
+				}
+			}
+		}
+	}
+
 	return ReturnDropZone;
 }
 
@@ -2720,6 +2778,48 @@ FReply SRigHierarchy::OnAcceptDrop(const FDragDropEvent& DragDropEvent, EItemDro
 		}
 	}
 
+	TSharedPtr<FAssetDragDropOp> AssetDragDropOp = DragDropEvent.GetOperationAs<FAssetDragDropOp>();
+	if (AssetDragDropOp.IsValid())
+	{
+		for (const FAssetData& AssetData : AssetDragDropOp->GetAssets())
+		{
+			UClass* AssetClass = AssetData.GetClass();
+			if (!AssetClass->IsChildOf(UControlRigBlueprint::StaticClass()))
+			{
+				continue;
+			}
+
+			if(UControlRigBlueprint* AssetBlueprint = Cast<UControlRigBlueprint>(AssetData.GetAsset()))
+			{
+				if (UModularRigController* Controller = ControlRigBlueprint->GetModularRigController())
+				{
+					const FName ModuleName = Controller->GetSafeNewName(AssetBlueprint->RigModuleSettings.Identifier.Name);
+					const FString ModulePath = Controller->AddModule(ModuleName, AssetBlueprint->GetControlRigClass(), FString());
+					if(TargetItem.IsValid() && !ModulePath.IsEmpty())
+					{
+						FRigElementKey PrimaryConnectorKey;
+						TArray<FRigConnectorElement*> Connectors = GetHierarchy()->GetElementsOfType<FRigConnectorElement>();
+						for (FRigConnectorElement* Connector : Connectors)
+						{
+							if (Connector->Settings.Type == EConnectorType::Primary)
+							{
+								FString Path, Name;
+								Connector->GetName().Split(UModularRig::NamespaceSeparator, &Path, &Name, ESearchCase::CaseSensitive, ESearchDir::FromEnd);
+								if (Path == ModulePath)
+								{
+									PrimaryConnectorKey = Connector->GetKey();
+									break;
+								}
+							}
+						}
+						return ResolveConnector(PrimaryConnectorKey, TargetItem->Key);
+					}
+				}
+				return FReply::Handled();
+			}
+		}
+	}
+
 	return FReply::Unhandled();
 }
 
