Control Rig: Fix renaming/reparenting modules not properly adjusting connectors
#rnx

diff --git a/Engine/Plugins/Animation/ControlRig/Source/ControlRigDeveloper/Private/ControlRigBlueprint.cpp b/Engine/Plugins/Animation/ControlRig/Source/ControlRigDeveloper/Private/ControlRigBlueprint.cpp
--- a/Engine/Plugins/Animation/ControlRig/Source/ControlRigDeveloper/Private/ControlRigBlueprint.cpp
+++ b/Engine/Plugins/Animation/ControlRig/Source/ControlRigDeveloper/Private/ControlRigBlueprint.cpp
@@ -2353,7 +2353,7 @@ void UControlRigBlueprint::HandleRigModulesModified(EModularRigNotification InNo
 							FString OldConnectorNamespace, ShortName;
 							OldConnectorName.Split(UModularRig::NamespaceSeparator, &OldConnectorNamespace, &ShortName, ESearchCase::CaseSensitive, ESearchDir::FromEnd);
 							OldConnectorNamespace.Append(UModularRig::NamespaceSeparator);
-							if (OldConnectorNamespace == OldNamespace)
+							if (OldConnectorNamespace.StartsWith(OldNamespace))
 							{
 								ConnectionInfo& Info = RenamedConnectors.FindOrAdd(Connector);
 								Info.NewPath = OldConnectorName.Replace(*OldNamespace, *NewNamespace);
diff --git a/Engine/Plugins/Animation/ControlRig/Source/ControlRigEditor/Private/Editor/ControlRigEditorMode.cpp b/Engine/Plugins/Animation/ControlRig/Source/ControlRigEditor/Private/Editor/ControlRigEditorMode.cpp
--- a/Engine/Plugins/Animation/ControlRig/Source/ControlRigEditor/Private/Editor/ControlRigEditorMode.cpp
+++ b/Engine/Plugins/Animation/ControlRig/Source/ControlRigEditor/Private/Editor/ControlRigEditorMode.cpp
@@ -148,7 +148,7 @@ FModularRigEditorMode::FModularRigEditorMode(const TSharedRef<FControlRigEditor>
 	TabFactories.RegisterFactory(MakeShared<FModularRigModelTabSummoner>(InControlRigEditor));
 	TabFactories.RegisterFactory(MakeShared<FRigModuleAssetBrowserTabSummoner>(InControlRigEditor));
 	
-	TabLayout = FTabManager::NewLayout("Standalone_ModularRigEditMode_Layout_v1.2")
+	TabLayout = FTabManager::NewLayout("Standalone_ModularRigEditMode_Layout_v1.3")
 		->AddArea
 		(
 			// Main application area
