RigHierarchy: Report the FRigControlSettings::ControlEnum to the cooker as a SoftObjectPath rather than as a hard import; a hard import will cause a dependent load at runtime. The dependent load is not necessarily a problem, but we have had problems with hard imports in the past so we are being cautious.
#rnx

diff --git a/Engine/Plugins/Animation/ControlRig/Source/ControlRig/Private/Rigs/RigHierarchyElements.cpp b/Engine/Plugins/Animation/ControlRig/Source/ControlRig/Private/Rigs/RigHierarchyElements.cpp
--- a/Engine/Plugins/Animation/ControlRig/Source/ControlRig/Private/Rigs/RigHierarchyElements.cpp
+++ b/Engine/Plugins/Animation/ControlRig/Source/ControlRig/Private/Rigs/RigHierarchyElements.cpp
@@ -767,7 +767,8 @@ void FRigControlSettings::Save(FArchive& Ar)
 		ControlEnumPathName = ControlEnum->GetPathName();
 		if (Ar.IsObjectReferenceCollector())
 		{
-			Ar << ControlEnum;
+			FSoftObjectPath DeclareControlEnumToCooker(ControlEnumPathName);
+			Ar << DeclareControlEnumToCooker;
 		}
 	}
 
