Control Rig: Reset local transform of controls to initial before running backwards solve in layered rigs
#rnx

diff --git a/Engine/Plugins/Animation/ControlRig/Source/ControlRig/Private/ControlRig.cpp b/Engine/Plugins/Animation/ControlRig/Source/ControlRig/Private/ControlRig.cpp
--- a/Engine/Plugins/Animation/ControlRig/Source/ControlRig/Private/ControlRig.cpp
+++ b/Engine/Plugins/Animation/ControlRig/Source/ControlRig/Private/ControlRig.cpp
@@ -376,6 +376,19 @@ void UControlRig::Evaluate_AnyThread()
 
 		URigHierarchy* Hierarchy = GetHierarchy();
 
+		// Reset all control local transforms to initial
+		Hierarchy->ForEach([Hierarchy](FRigBaseElement* Element) -> bool
+		{
+			if (FRigControlElement* Control = Cast<FRigControlElement>(Element))
+			{
+				if (Control->Settings.ControlType != ERigControlType::Bool)
+				{
+					Hierarchy->SetTransform(Control, Hierarchy->GetTransform(Control, ERigTransformType::InitialLocal), ERigTransformType::CurrentLocal, true);
+				}
+			}
+			return true;
+		});
+
 		if (PoseBeforeBackwardsSolve.Num() == 0)
 		{
 			// If the pose is empty, this is an indication that a new pose is coming in
