Fix to get blending of control rig poses working when a sequence has multiple control rigs

diff --git a/Engine/Plugins/Animation/ControlRig/Source/ControlRigEditor/Private/EditMode/SControlRigControlViews.cpp b/Engine/Plugins/Animation/ControlRig/Source/ControlRigEditor/Private/EditMode/SControlRigControlViews.cpp
--- a/Engine/Plugins/Animation/ControlRig/Source/ControlRigEditor/Private/EditMode/SControlRigControlViews.cpp
+++ b/Engine/Plugins/Animation/ControlRig/Source/ControlRigEditor/Private/EditMode/SControlRigControlViews.cpp
@@ -568,12 +568,25 @@ FReply SControlRigPoseView::OnSelectControls()
 
 void SControlRigPoseView::OnPoseBlendChanged(float ChangedVal)
 {
+	auto ControlRigContainsControls = [this] (const UControlRig* InControlRig)
+	{
+		TArray<FName> ControlNames = PoseAsset->GetControlNames();
+		for (FName ControlName : ControlNames)
+		{
+			if (InControlRig->FindControl(ControlName))
+			{
+				return true;
+			}
+		}
+		return false;
+	};
+	
 	if (PoseAsset.IsValid())
 	{
 		TArray<UControlRig*> ControlRigs = GetControlRigs();
 		for (UControlRig* ControlRig : ControlRigs)
 		{
-			if (ControlRig)
+			if (ControlRig && ControlRigContainsControls(ControlRig))
 			{
 				PoseBlendValue = ChangedVal;
 				if (!bIsBlending)
