Constraint: improved evaluation/notification
    - added ability to pre-tick skeletal meshes to ensure controls and bones transformations are up to date before propagating changes
    - notify any control changes to ensure that updates are fully propagated thru the rig
    - avoid reentrant animation evaluation when ticking skeletal meshes

TODO:
    - as the propagation of delegates works in the opposite direction to their registration, find a way to notify from hierarchical parent to children to ensure that transformations are propagated in the correct way (in the case of multiple constraints from arm to fingers, fingers should always be evaluated last)

diff --git a/Engine/Plugins/Animation/ControlRig/Source/ControlRig/Private/Constraints/ControlRigTransformableHandle.cpp b/Engine/Plugins/Animation/ControlRig/Source/ControlRig/Private/Constraints/ControlRigTransformableHandle.cpp
--- a/Engine/Plugins/Animation/ControlRig/Source/ControlRig/Private/Constraints/ControlRigTransformableHandle.cpp
+++ b/Engine/Plugins/Animation/ControlRig/Source/ControlRig/Private/Constraints/ControlRigTransformableHandle.cpp
@@ -8,6 +8,7 @@
 #include "ControlRigComponent.h"
 #include "ControlRig.h"
 #include "IControlRigObjectBinding.h"
+#include "TransformableHandleUtils.h"
 #include "Rigs/RigHierarchyElements.h"
 #include "Sequencer/MovieSceneControlRigParameterSection.h"
 #include "Sections/MovieScene3DTransformSection.h"
@@ -56,30 +57,11 @@ bool UTransformableControlHandle::IsValid(const bool bDeepCheck) const
 	return true;
 }
 
-void UTransformableControlHandle::TickForBaking() const
+void UTransformableControlHandle::TickTarget() const
 {
 	if (const USkeletalMeshComponent* SkeletalMeshComponent = GetSkeletalMesh())
 	{
-		const AActor* Parent = SkeletalMeshComponent->GetOwner();
-		while (Parent)
-		{
-			TArray<USkeletalMeshComponent*> MeshComps;
-			Parent->GetComponents(MeshComps, true);
-
-			for (USkeletalMeshComponent* MeshComp : MeshComps)
-			{
-				MeshComp->TickAnimation(0.03f, false);
-				MeshComp->RefreshBoneTransforms();
-				MeshComp->RefreshFollowerComponents();
-				MeshComp->UpdateComponentToWorld();
-				MeshComp->FinalizeBoneTransform();
-				MeshComp->MarkRenderTransformDirty();
-				MeshComp->MarkRenderDynamicDataDirty();
-			}
-
-			Parent = Parent->GetAttachParentActor();
-		}
-		return;
+		return TransformableHandleUtils::TickDependantComponents(SkeletalMeshComponent);
 	}
 
 	if (UControlRigComponent* ControlRigComponent = GetControlRigComponent())
@@ -420,31 +402,32 @@ void UTransformableControlHandle::OnControlModified(
 		return;
 	}
 
+	if (bNotifying)
+	{
+		return;
+	}
+
 	if (!ControlRig.IsValid() || ControlName == NAME_None)
 	{
 		return;
 	}
 
-	if (OnHandleModified.IsBound() && (ControlRig == InControlRig))
+	if (HandleModified().IsBound() && (ControlRig == InControlRig))
 	{
 		const EHandleEvent Event = InContext.bConstraintUpdate ?
 			EHandleEvent::GlobalTransformUpdated : EHandleEvent::LocalTransformUpdated;
 
 		if (InControl->GetFName() == ControlName)
-		{	// if that handle is wrapping InControl  
-			OnHandleModified.Broadcast(this, Event);
+		{	// if that handle is wrapping InControl
+			Notify(Event);
 		}
 		else if (Event == EHandleEvent::GlobalTransformUpdated)
 		{
+			// the control being modified is not the one wrapped by this handle 
 			if (const FRigControlElement* Control = ControlRig->FindControl(ControlName))
-			{	// if that handle control's transform has been dirtied
-				const bool bIsTransformDirty =
-					Control->Pose.IsDirty(ERigTransformType::CurrentLocal) ||
-					Control->Pose.IsDirty(ERigTransformType::CurrentGlobal);
-				if (bIsTransformDirty)
-				{
-					OnHandleModified.Broadcast(this, EHandleEvent::UpperDependencyUpdated);
-				}
+			{
+				static constexpr  bool bPreTick = true;
+				Notify(EHandleEvent::UpperDependencyUpdated, bPreTick);
 			}
 		}
 	}
@@ -485,7 +468,7 @@ void UTransformableControlHandle::OnObjectBoundToControlRig(UObject* InObject)
 		const UWorld* ThisWorld = GetWorld();
 		if (ThisWorld && InObject->GetWorld() == ThisWorld)
 		{
-			OnHandleModified.Broadcast(this, EHandleEvent::ComponentUpdated);
+			Notify(EHandleEvent::ComponentUpdated);
 		}
 	}
 }
diff --git a/Engine/Plugins/Animation/ControlRig/Source/ControlRig/Public/Constraints/ControlRigTransformableHandle.h b/Engine/Plugins/Animation/ControlRig/Source/ControlRig/Public/Constraints/ControlRigTransformableHandle.h
--- a/Engine/Plugins/Animation/ControlRig/Source/ControlRig/Public/Constraints/ControlRigTransformableHandle.h
+++ b/Engine/Plugins/Animation/ControlRig/Source/ControlRig/Public/Constraints/ControlRigTransformableHandle.h
@@ -71,8 +71,8 @@ public:
 	/** Make a duplicate of myself with this outer*/
 	virtual UTransformableHandle* Duplicate(UObject* NewOuter) const override;
 
-	/** tick things when baking*/ 
-	virtual void TickForBaking() const override;
+	/** Tick any skeletal mesh related to the bound component. */ 
+	virtual void TickTarget() const override;
 
 	/** Registers/Unregisters useful delegates to track changes in the control's transform. */
 	void UnregisterDelegates() const;
diff --git a/Engine/Source/Runtime/Experimental/Animation/Constraints/Private/TransformConstraint.cpp b/Engine/Source/Runtime/Experimental/Animation/Constraints/Private/TransformConstraint.cpp
--- a/Engine/Source/Runtime/Experimental/Animation/Constraints/Private/TransformConstraint.cpp
+++ b/Engine/Source/Runtime/Experimental/Animation/Constraints/Private/TransformConstraint.cpp
@@ -511,14 +511,14 @@ void UTickableTransformConstraint::Evaluate(bool bTickHandlesAlso) const
 	{
 		if (ParentTRSHandle)
 		{
-			ParentTRSHandle->TickForBaking();
+			ParentTRSHandle->TickTarget();
 		}
 
 		Super::Evaluate();
 		//todo test this more may be able to remove it
 		if (ChildTRSHandle)
 		{
-			ChildTRSHandle->TickForBaking();
+			ChildTRSHandle->TickTarget();
 		}
 	}
 	else
@@ -601,7 +601,7 @@ void UTickableTransformConstraint::OnHandleModified(UTransformableHandle* InHand
 			const FConstraintTickFunction* ConstraintTick = ConstraintTicks.Find(World->GetCurrentLevel());
 			if (IsFullyActive() && ConstraintTick->IsTickFunctionRegistered() && ConstraintTick->IsTickFunctionEnabled())
 			{
-				Evaluate();
+				Evaluate(true);
 			}
 			return;
 		}
@@ -754,13 +754,16 @@ FConstraintTickFunction::ConstraintFunction UTickableTranslationConstraint::GetF
 		}
 
 		AxisFilter.FilterVector(NewTranslation, ChildTranslation);
-		
-		Transform.SetLocation(NewTranslation);
-			
-		SetChildGlobalTransform(Transform);
-		if (ChildTRSHandle)
+
+		if (!NewTranslation.Equals(ChildTranslation))
 		{
-			ChildTRSHandle->TickForBaking();
+			Transform.SetLocation(NewTranslation);
+			
+			SetChildGlobalTransform(Transform);
+			if (ChildTRSHandle)
+			{
+				ChildTRSHandle->TickTarget();
+			}
 		}
 	};
 }
@@ -895,13 +898,17 @@ FConstraintTickFunction::ConstraintFunction UTickableRotationConstraint::GetFunc
 		}
 
 		AxisFilter.FilterQuat(NewRotation, ChildRotation);
-		
-		Transform.SetRotation(NewRotation);
-		
-		SetChildGlobalTransform(Transform);
-		if (ChildTRSHandle)
+
+		if (!NewRotation.Equals(ChildRotation))
 		{
-			ChildTRSHandle->TickForBaking();
+			Transform.SetRotation(NewRotation);
+		
+			SetChildGlobalTransform(Transform);
+
+			if (ChildTRSHandle)
+			{
+				ChildTRSHandle->TickTarget();
+			}
 		}
 	};
 }
@@ -1040,13 +1047,16 @@ FConstraintTickFunction::ConstraintFunction UTickableScaleConstraint::GetFunctio
 		}
 
 		AxisFilter.FilterVector(NewScale, ChildScale);
-		
-		Transform.SetScale3D(NewScale);
-			
-		SetChildGlobalTransform(Transform);
-		if (ChildTRSHandle)
+
+		if (!NewScale.Equals(ChildScale))
 		{
-			ChildTRSHandle->TickForBaking();
+			Transform.SetScale3D(NewScale);
+			
+			SetChildGlobalTransform(Transform);
+			if (ChildTRSHandle)
+			{
+				ChildTRSHandle->TickTarget();
+			}
 		}
 	};
 }
@@ -1206,11 +1216,14 @@ FConstraintTickFunction::ConstraintFunction UTickableParentConstraint::GetFuncti
 		{
 			TargetTransform.SetScale3D(ChildGlobalTransform.GetScale3D());
 		}
-		
-		SetChildGlobalTransform(TargetTransform);
-		if (ChildTRSHandle)
+
+		if (!ChildGlobalTransform.Equals(TargetTransform))
 		{
-			ChildTRSHandle->TickForBaking();
+			SetChildGlobalTransform(TargetTransform);
+			if (ChildTRSHandle)
+			{
+				ChildTRSHandle->TickTarget();
+			}
 		}
 	};
 }
@@ -1362,7 +1375,7 @@ FConstraintTickFunction::ConstraintFunction UTickableLookAtConstraint::GetFuncti
 				SetChildGlobalTransform(Transform);
 				if (ChildTRSHandle)
 				{
-					ChildTRSHandle->TickForBaking();
+					ChildTRSHandle->TickTarget();
 				}
 			}
 		}
diff --git a/Engine/Source/Runtime/Experimental/Animation/Constraints/Private/TransformableHandle.cpp b/Engine/Source/Runtime/Experimental/Animation/Constraints/Private/TransformableHandle.cpp
--- a/Engine/Source/Runtime/Experimental/Animation/Constraints/Private/TransformableHandle.cpp
+++ b/Engine/Source/Runtime/Experimental/Animation/Constraints/Private/TransformableHandle.cpp
@@ -7,6 +7,7 @@
 #include "GameFramework/Actor.h"
 #include "Engine/Engine.h"
 #include "MovieSceneSection.h"
+#include "TransformableHandleUtils.h"
 #include "Channels/MovieSceneFloatChannel.h"
 #include "Channels/MovieSceneDoubleChannel.h"
 #include "Channels/MovieSceneChannelProxy.h"
@@ -28,6 +29,20 @@ UTransformableHandle::FHandleModifiedEvent& UTransformableHandle::HandleModified
 	return OnHandleModified;
 }
 
+void UTransformableHandle::Notify(EHandleEvent InEvent, const bool bPreTickTarget) const
+{
+	if (!bNotifying && OnHandleModified.IsBound())
+	{
+		TGuardValue<bool> ReentrantGuardSelf(bNotifying, true);
+		if (bPreTickTarget)
+		{
+			TickTarget();
+		}
+		
+		OnHandleModified.Broadcast(const_cast<UTransformableHandle*>(this), InEvent);
+	}
+}
+
 bool UTransformableHandle::HasBoundObjects() const
 {
 	return ConstraintBindingID.IsValid();
@@ -67,33 +82,15 @@ bool UTransformableComponentHandle::IsValid(const bool bDeepCheck) const
 }
 
 //need to tick any skelmesh component, sibling or parent
-void UTransformableComponentHandle::TickForBaking() const
+void UTransformableComponentHandle::TickTarget() const
 {
 	if (!Component.IsValid())
 	{
 		return;
 	}
-
-	const AActor* Parent = Component->GetOwner();
-	while (Parent)
-	{
-		TArray<USkeletalMeshComponent*> MeshComps;
-		Parent->GetComponents(MeshComps, true);
-
-		for (USkeletalMeshComponent* MeshComp : MeshComps)
-		{
-			MeshComp->TickAnimation(0.03f, false);
-			MeshComp->RefreshBoneTransforms();
-			MeshComp->RefreshFollowerComponents();
-			MeshComp->UpdateComponentToWorld();
-			MeshComp->FinalizeBoneTransform();
-			MeshComp->MarkRenderTransformDirty();
-			MeshComp->MarkRenderDynamicDataDirty();
-		}
-
-		Parent = Parent->GetAttachParentActor();
-	}
+	TransformableHandleUtils::TickDependantComponents(Component.Get());
 }
+
 void UTransformableComponentHandle::SetGlobalTransform(const FTransform& InGlobal) const
 {
 	if (Component.IsValid())
@@ -308,10 +305,7 @@ void UTransformableComponentHandle::OnActorMoving(AActor* InActor)
 		return;
 	}
 
-	if(OnHandleModified.IsBound())
-	{
-		OnHandleModified.Broadcast(this, EHandleEvent::GlobalTransformUpdated);
-	}
+	Notify(EHandleEvent::GlobalTransformUpdated);
 }
 
 void UTransformableComponentHandle::OnPostPropertyChanged(
@@ -358,10 +352,7 @@ void UTransformableComponentHandle::OnPostPropertyChanged(
 		return;
 	}
 
-	if(OnHandleModified.IsBound())
-	{
-		OnHandleModified.Broadcast(this, EHandleEvent::GlobalTransformUpdated);
-	}
+	Notify(EHandleEvent::GlobalTransformUpdated);
 }
 
 void UTransformableComponentHandle::OnObjectsReplaced(const TMap<UObject*, UObject*>& InOldToNewInstances)
diff --git a/Engine/Source/Runtime/Experimental/Animation/Constraints/Private/TransformableHandleUtils.cpp b/Engine/Source/Runtime/Experimental/Animation/Constraints/Private/TransformableHandleUtils.cpp
new file mode 100644
--- /dev/null
+++ b/Engine/Source/Runtime/Experimental/Animation/Constraints/Private/TransformableHandleUtils.cpp
@@ -0,0 +1,54 @@
+﻿// Copyright Epic Games, Inc. All Rights Reserved.
+
+#include "TransformableHandleUtils.h"
+
+#include "GameFramework/Actor.h"
+#include "Components/SceneComponent.h"
+#include "Components/SkeletalMeshComponent.h"
+
+namespace TransformableHandleUtils
+{
+	
+void TickDependantComponents(const USceneComponent* InComponent)
+{
+	if (!InComponent)
+	{
+		return;
+	}
+
+	static constexpr bool bIncludeFromChildActors = true;
+
+	const AActor* Parent = InComponent->GetOwner();
+	while (Parent)
+	{
+		Parent->ForEachComponent<USkeletalMeshComponent>(bIncludeFromChildActors, &TickSkeletalMeshComponent);
+		Parent = Parent->GetAttachParentActor();
+	}
+}
+
+void TickSkeletalMeshComponent(USkeletalMeshComponent* InSkeletalMeshComponent)
+{
+	if (!InSkeletalMeshComponent)
+	{
+		return;
+	}
+
+	// avoid re-entrant animation evaluation
+	if (InSkeletalMeshComponent->IsPostEvaluatingAnimation())
+	{
+		return;
+	}
+
+	static constexpr float DeltaTime = 0.03f;
+	static constexpr bool bNeedsValidRootMotion = false;
+	
+	InSkeletalMeshComponent->TickAnimation(DeltaTime, bNeedsValidRootMotion);
+	InSkeletalMeshComponent->RefreshBoneTransforms();
+	InSkeletalMeshComponent->RefreshFollowerComponents();
+	InSkeletalMeshComponent->UpdateComponentToWorld();
+	InSkeletalMeshComponent->FinalizeBoneTransform();
+	InSkeletalMeshComponent->MarkRenderTransformDirty();
+	InSkeletalMeshComponent->MarkRenderDynamicDataDirty();
+}
+	
+}
\ No newline at end of file
diff --git a/Engine/Source/Runtime/Experimental/Animation/Constraints/Public/ConstraintsManager.h b/Engine/Source/Runtime/Experimental/Animation/Constraints/Public/ConstraintsManager.h
--- a/Engine/Source/Runtime/Experimental/Animation/Constraints/Public/ConstraintsManager.h
+++ b/Engine/Source/Runtime/Experimental/Animation/Constraints/Public/ConstraintsManager.h
@@ -131,7 +131,7 @@ public:
 	bool Active = true;
 
 	UPROPERTY()
-	bool bValid = true;
+	bool bValid = false;
 
 
 	/** Unique id to be used to fine constraint in a particular world, may be shared*/
diff --git a/Engine/Source/Runtime/Experimental/Animation/Constraints/Public/TransformableHandle.h b/Engine/Source/Runtime/Experimental/Animation/Constraints/Public/TransformableHandle.h
--- a/Engine/Source/Runtime/Experimental/Animation/Constraints/Public/TransformableHandle.h
+++ b/Engine/Source/Runtime/Experimental/Animation/Constraints/Public/TransformableHandle.h
@@ -78,7 +78,8 @@ public:
 	CONSTRAINTS_API void OnBindingIDsUpdated(const TMap<UE::MovieScene::FFixedObjectBindingID, UE::MovieScene::FFixedObjectBindingID>& OldFixedToNewFixedMap, FMovieSceneSequenceID LocalSequenceID, const FMovieSceneSequenceHierarchy* Hierarchy, IMovieScenePlayer& Player);
 	
 	/** Perform any special ticking needed for this handle, by default it does nothing, todo need to see if we need to tick control rig also*/
-	virtual void TickForBaking() const {};
+	virtual void TickTarget() const {};
+	
 	/** Get the array of float channels for the specified section*/
 	CONSTRAINTS_API virtual TArrayView<FMovieSceneFloatChannel*>  GetFloatChannels(const UMovieSceneSection* InSection) const PURE_VIRTUAL(GetFloatChannels, return TArrayView<FMovieSceneFloatChannel*>(); );
 	/** Get the array of double channels for the specified section*/
@@ -115,6 +116,8 @@ public:
 	CONSTRAINTS_API virtual FTickPrerequisite GetPrimaryPrerequisite() const PURE_VIRTUAL(GetPrimaryPrerequisite, return FTickPrerequisite(););
 	
 	CONSTRAINTS_API FHandleModifiedEvent& HandleModified();
+	CONSTRAINTS_API void Notify(EHandleEvent InEvent, const bool bPreTickTarget = false) const;
+	mutable bool bNotifying = false;
 
 #if WITH_EDITOR
 	CONSTRAINTS_API virtual FString GetLabel() const PURE_VIRTUAL(GetLabel, return FString(););
@@ -124,7 +127,8 @@ public:
 	//possible bindingID
 	UPROPERTY(BlueprintReadOnly, Category = "Binding")
 	FMovieSceneObjectBindingID ConstraintBindingID;
-protected:
+
+private:
 	FHandleModifiedEvent OnHandleModified;
 };
 
@@ -154,8 +158,8 @@ public:
 	CONSTRAINTS_API virtual FTransform GetGlobalTransform() const override;
 	/** Gets the local transform of Component in it's attachment. */
 	CONSTRAINTS_API virtual FTransform GetLocalTransform() const override;
-	/** Tick the component*/
-	CONSTRAINTS_API virtual void TickForBaking() const override;
+	/** Tick any skeletal mesh related to the component. */
+	CONSTRAINTS_API virtual void TickTarget() const override;
 	/** Returns the target object containing the tick function (e.i. Component). */
 	CONSTRAINTS_API virtual UObject* GetPrerequisiteObject() const override;
 	/** Returns Component's tick function. */
diff --git a/Engine/Source/Runtime/Experimental/Animation/Constraints/Public/TransformableHandleUtils.h b/Engine/Source/Runtime/Experimental/Animation/Constraints/Public/TransformableHandleUtils.h
new file mode 100644
--- /dev/null
+++ b/Engine/Source/Runtime/Experimental/Animation/Constraints/Public/TransformableHandleUtils.h
@@ -0,0 +1,14 @@
+﻿// Copyright Epic Games, Inc. All Rights Reserved.
+
+#pragma once
+
+#include "HAL/Platform.h"
+
+class USceneComponent;
+class USkeletalMeshComponent;
+
+namespace TransformableHandleUtils
+{
+	CONSTRAINTS_API void TickDependantComponents(const USceneComponent* InComponent);
+	CONSTRAINTS_API void TickSkeletalMeshComponent(USkeletalMeshComponent* InSkeletalMeshComponent);
+}
