Control RIg: Allow switching parent while working with Layered rigs
#rnx

diff --git a/Engine/Plugins/Animation/ControlRig/Source/ControlRig/Private/ControlRig.cpp b/Engine/Plugins/Animation/ControlRig/Source/ControlRig/Private/ControlRig.cpp
--- a/Engine/Plugins/Animation/ControlRig/Source/ControlRig/Private/ControlRig.cpp
+++ b/Engine/Plugins/Animation/ControlRig/Source/ControlRig/Private/ControlRig.cpp
@@ -376,13 +376,17 @@ void UControlRig::Evaluate_AnyThread()
 
 		URigHierarchy* Hierarchy = GetHierarchy();
 
-		// Reset all control local transforms to initial
+		// Reset all control local transforms to initial, and switch to default parents
 		Hierarchy->ForEach([Hierarchy](FRigBaseElement* Element) -> bool
 		{
 			if (FRigControlElement* Control = Cast<FRigControlElement>(Element))
 			{
 				if (Control->Settings.ControlType != ERigControlType::Bool)
 				{
+					if (Hierarchy->GetActiveParent(Control->GetKey()) != URigHierarchy::GetDefaultParentKey())
+					{
+						Hierarchy->SwitchToDefaultParent(Control);
+					}
 					Hierarchy->SetTransform(Control, Hierarchy->GetTransform(Control, ERigTransformType::InitialLocal), ERigTransformType::CurrentLocal, true);
 				}
 			}
@@ -406,7 +410,25 @@ void UControlRig::Evaluate_AnyThread()
 			Execute(FRigUnit_InverseExecution::EventName);
 		}
 
+		// Switch parents
+		for (TPair<FRigElementKey, FRigSwitchParentInfo>& Pair : SwitchParentValues)
+		{
+			if (FRigBaseElement* Element = Hierarchy->Find(Pair.Key))
+			{
+				if (FRigBaseElement* NewParent = Hierarchy->Find(Pair.Value.NewParent))
+				{
+#if WITH_EDITOR
+					URigHierarchy::TElementDependencyMap DependencyMap = Hierarchy->GetDependenciesForVM(GetVM());
+					Hierarchy->SwitchToParent(Element, NewParent, Pair.Value.bInitial, Pair.Value.bAffectChildren, DependencyMap, nullptr);
+#else
+					Hierarchy->SwitchToParent(Element, NewParent, Pair.Value.bInitial, Pair.Value.bAffectChildren);
+#endif
+				}
+			}
+		}
+
 		// Store control pose after backwards solve to figure out additive local transforms based on animation
+		// This needs to happen after the switch parents
 		ControlsAfterBackwardsSolve = Hierarchy->GetPose(false, ERigElementType::Control, TArrayView<const FRigElementKey>());
 
 		if (PreAdditiveValuesApplicationEvent.IsBound())
@@ -1956,6 +1978,37 @@ void UControlRig::SetControlValueImpl(const FName& InControlName, const FRigCont
 	}
 }
 
+void UControlRig::SwitchToParent(const FRigElementKey& InElementKey, const FRigElementKey& InNewParentKey, bool bInitial, bool bAffectChildren)
+{
+	FRigBaseElement* Element = DynamicHierarchy->Find<FRigBaseElement>(InElementKey);
+	if(Element == nullptr)
+	{
+		return;
+	}
+	FRigBaseElement* Parent = DynamicHierarchy->Find<FRigBaseElement>(InNewParentKey);
+	if(Parent == nullptr)
+	{
+		return;
+	}
+	if (bIsAdditive)
+	{
+		FRigSwitchParentInfo Info;
+		Info.NewParent = InNewParentKey;
+		Info.bInitial = bInitial;
+		Info.bAffectChildren = bAffectChildren;
+		SwitchParentValues.Add(InElementKey, Info);
+	}
+	else
+	{
+#if WITH_EDITOR
+		URigHierarchy::TElementDependencyMap Dependencies = DynamicHierarchy->GetDependenciesForVM(GetVM());
+		DynamicHierarchy->SwitchToParent(InElementKey, InNewParentKey, bInitial, bAffectChildren, Dependencies, nullptr);
+#else
+		DynamicHierarchy->SwitchToParent(InElementKey, InNewParentKey, bInitial, bAffectChildren);
+#endif
+	}
+}
+
 bool UControlRig::SetControlGlobalTransform(const FName& InControlName, const FTransform& InGlobalTransform, bool bNotify, const FRigControlModifiedContext& Context, bool bSetupUndo, bool bPrintPythonCommands, bool bFixEulerFlips)
 {
 	FTransform GlobalTransform = InGlobalTransform;
diff --git a/Engine/Plugins/Animation/ControlRig/Source/ControlRig/Private/Sequencer/MovieSceneControlRigParameterTemplate.cpp b/Engine/Plugins/Animation/ControlRig/Source/ControlRig/Private/Sequencer/MovieSceneControlRigParameterTemplate.cpp
--- a/Engine/Plugins/Animation/ControlRig/Source/ControlRig/Private/Sequencer/MovieSceneControlRigParameterTemplate.cpp
+++ b/Engine/Plugins/Animation/ControlRig/Source/ControlRig/Private/Sequencer/MovieSceneControlRigParameterTemplate.cpp
@@ -795,8 +795,7 @@ struct FControlRigParameterPreAnimatedTokenProducer : IMovieScenePreAnimatedToke
 								case EMovieSceneControlRigSpaceType::ControlRig:
 								{
 #if WITH_EDITOR
-									URigHierarchy::TElementDependencyMap Dependencies = RigHierarchy->GetDependenciesForVM(ControlRig->GetVM());
-									RigHierarchy->SwitchToParent(ControlKey, SpaceNameAndValue.Value.ControlRigElement, false, true, Dependencies, nullptr);
+									ControlRig->SwitchToParent(ControlKey, SpaceNameAndValue.Value.ControlRigElement, false, true);
 #else
 									RigHierarchy->SwitchToParent(ControlKey, SpaceNameAndValue.Value.ControlRigElement);
 #endif
@@ -1327,18 +1326,17 @@ struct FControlRigParameterExecutionToken : IMovieSceneExecutionToken
 								switch (SpaceNameAndValue.Value.SpaceType)
 								{
 									case EMovieSceneControlRigSpaceType::Parent:
-										RigHierarchy->SwitchToDefaultParent(ControlKey);
+										ControlRig->SwitchToParent(ControlKey, RigHierarchy->GetDefaultParent(ControlKey), false, true);
 										break;
 									case EMovieSceneControlRigSpaceType::World:
-										RigHierarchy->SwitchToWorldSpace(ControlKey);
+										ControlRig->SwitchToParent(ControlKey, RigHierarchy->GetWorldSpaceReferenceKey(), false, true);
 										break;
 									case EMovieSceneControlRigSpaceType::ControlRig:
 										{
 #if WITH_EDITOR
-											URigHierarchy::TElementDependencyMap Dependencies = RigHierarchy->GetDependenciesForVM(ControlRig->GetVM());
-											RigHierarchy->SwitchToParent(ControlKey, SpaceNameAndValue.Value.ControlRigElement, false, true, Dependencies, nullptr);
+											ControlRig->SwitchToParent(ControlKey, SpaceNameAndValue.Value.ControlRigElement, false, true);
 #else
-											RigHierarchy->SwitchToParent(ControlKey, SpaceNameAndValue.Value.ControlRigElement);
+											ControlRig->SwitchToParent(ControlKey, SpaceNameAndValue.Value.ControlRigElement, false, true);
 #endif	
 										}
 										break;
diff --git a/Engine/Plugins/Animation/ControlRig/Source/ControlRig/Public/ControlRig.h b/Engine/Plugins/Animation/ControlRig/Source/ControlRig/Public/ControlRig.h
--- a/Engine/Plugins/Animation/ControlRig/Source/ControlRig/Public/ControlRig.h
+++ b/Engine/Plugins/Animation/ControlRig/Source/ControlRig/Public/ControlRig.h
@@ -262,6 +262,8 @@ public:
 	virtual void SetControlValueImpl(const FName& InControlName, const FRigControlValue& InValue, bool bNotify = true,
 		const FRigControlModifiedContext& Context = FRigControlModifiedContext(), bool bSetupUndo = true, bool bPrintPythonCommnds = false, bool bFixEulerFlips = false);
 
+	void SwitchToParent(const FRigElementKey& InElementKey, const FRigElementKey& InNewParentKey, bool bInitial, bool bAffectChildren);
+
 	FTransform GetInitialLocalTransform(const FRigElementKey &InKey)
 	{
 		if (bIsAdditive)
@@ -660,9 +662,16 @@ protected:
 		bool bPrintPythonCommnds;
 		bool bFixEulerFlips;
 	};
+	struct FRigSwitchParentInfo
+	{
+		FRigElementKey NewParent;
+		bool bInitial;
+		bool bAffectChildren;
+	};
 	FRigPose PoseBeforeBackwardsSolve;
 	FRigPose ControlsAfterBackwardsSolve;
-	TMap<FRigElementKey, FRigSetControlValueInfo> ControlValues; // Additive values in local space (to add after backwards solve)
+	TMap<FRigElementKey, FRigSetControlValueInfo> ControlValues; // Layered Rigs: Additive values in local space (to add after backwards solve)
+	TMap<FRigElementKey, FRigSwitchParentInfo> SwitchParentValues; // Layered Rigs: Parent switching values to perform after backwards solve
 
 
 private:
diff --git a/Engine/Plugins/Animation/ControlRig/Source/ControlRigEditor/Private/EditMode/ControlRigEditMode.cpp b/Engine/Plugins/Animation/ControlRig/Source/ControlRigEditor/Private/EditMode/ControlRigEditMode.cpp
--- a/Engine/Plugins/Animation/ControlRig/Source/ControlRigEditor/Private/EditMode/ControlRigEditMode.cpp
+++ b/Engine/Plugins/Animation/ControlRig/Source/ControlRigEditor/Private/EditMode/ControlRigEditMode.cpp
@@ -2834,10 +2834,22 @@ void FControlRigEditMode::OpenSpacePickerWidget()
 		}
 		else
 		{
-			const FTransform Transform = InHierarchy->GetGlobalTransform(InControlKey);
-			URigHierarchy::TElementDependencyMap Dependencies = InHierarchy->GetDependenciesForVM(RuntimeRig->GetVM());
-			InHierarchy->SwitchToParent(InControlKey, InSpaceKey, false, true, Dependencies, nullptr);
-			InHierarchy->SetGlobalTransform(InControlKey, Transform);
+			if (RuntimeRig->IsAdditive())
+			{
+				const FTransform Transform = RuntimeRig->GetControlGlobalTransform(InControlKey.Name);
+			   RuntimeRig->SwitchToParent(InControlKey, InSpaceKey, false, true);
+			   RuntimeRig->Evaluate_AnyThread();
+			   FRigControlValue ControlValue = RuntimeRig->GetControlValueFromGlobalTransform(InControlKey.Name, Transform, ERigTransformType::CurrentGlobal);
+			   RuntimeRig->SetControlValue(InControlKey.Name, ControlValue);
+			   RuntimeRig->Evaluate_AnyThread();
+			}
+			else
+			{
+				const FTransform Transform = InHierarchy->GetGlobalTransform(InControlKey);
+				URigHierarchy::TElementDependencyMap Dependencies = InHierarchy->GetDependenciesForVM(RuntimeRig->GetVM());
+				InHierarchy->SwitchToParent(InControlKey, InSpaceKey, false, true, Dependencies, nullptr);
+				InHierarchy->SetGlobalTransform(InControlKey, Transform);
+			}
 		}
 		
 	})
diff --git a/Engine/Plugins/Animation/ControlRig/Source/ControlRigEditor/Private/Sequencer/ControlRigParameterTrackEditor.cpp b/Engine/Plugins/Animation/ControlRig/Source/ControlRigEditor/Private/Sequencer/ControlRigParameterTrackEditor.cpp
--- a/Engine/Plugins/Animation/ControlRig/Source/ControlRigEditor/Private/Sequencer/ControlRigParameterTrackEditor.cpp
+++ b/Engine/Plugins/Animation/ControlRig/Source/ControlRigEditor/Private/Sequencer/ControlRigParameterTrackEditor.cpp
@@ -1431,9 +1431,12 @@ void FControlRigParameterTrackEditor::HandleAddControlRigSubMenu(FMenuBuilder& M
 		{
 			if (UControlRig* ControlRig = Track->GetControlRig())
 			{
-				if (ControlRig->GetObjectBinding()->GetBoundObject() == SkeletalMeshComponent)
+				if (TSharedPtr<IControlRigObjectBinding> ObjectBinding = ControlRig->GetObjectBinding())
 				{
-					ExistingRigs.Add(ControlRig->GetClass());
+					if (ObjectBinding.IsValid() && ObjectBinding->GetBoundObject() == SkeletalMeshComponent)
+					{
+						ExistingRigs.Add(ControlRig->GetClass());
+					}
 				}
 			}
 			return true;
@@ -1667,8 +1670,7 @@ static void EvaluateThisControl(UMovieSceneControlRigParameterSection* Section,
 					RigHierarchy->SwitchToWorldSpace(ControlElement->GetKey());
 					break;
 				case EMovieSceneControlRigSpaceType::ControlRig:
-					URigHierarchy::TElementDependencyMap Dependencies = RigHierarchy->GetDependenciesForVM(ControlRig->GetVM());
-					RigHierarchy->SwitchToParent(ControlElement->GetKey(), SpaceKey.GetValue().ControlRigElement, false, true, Dependencies, nullptr);
+					ControlRig->SwitchToParent(ControlElement->GetKey(), SpaceKey.GetValue().ControlRigElement, false, true);
 					break;
 				}
 			}
diff --git a/Engine/Plugins/Animation/ControlRig/Source/ControlRigEditor/Private/Sequencer/ControlRigSpaceChannelEditors.cpp b/Engine/Plugins/Animation/ControlRig/Source/ControlRigEditor/Private/Sequencer/ControlRigSpaceChannelEditors.cpp
--- a/Engine/Plugins/Animation/ControlRig/Source/ControlRigEditor/Private/Sequencer/ControlRigSpaceChannelEditors.cpp
+++ b/Engine/Plugins/Animation/ControlRig/Source/ControlRigEditor/Private/Sequencer/ControlRigSpaceChannelEditors.cpp
@@ -453,6 +453,11 @@ FKeyHandle FControlRigSpaceChannelHelpers::SequencerKeyControlRigSpaceChannel(UC
 
 			Context.LocalTime = TickResolution.AsSeconds(FFrameTime(Time - 1));
 			ControlRig->SetControlGlobalTransform(ControlKey.Name, ControlWorldTransforms[0], true, Context, false /*undo*/, false /*bPrintPython*/, true/* bFixEulerFlips*/);
+			if (ControlRig->IsAdditive())
+			{
+				// We need to evaluate in order to trigger a notification
+				ControlRig->Evaluate_AnyThread();
+			}
 
 			//need to do this after eval
 			FChannelMapInfo* pChannelIndex = nullptr;
@@ -469,8 +474,7 @@ FKeyHandle FControlRigSpaceChannelHelpers::SequencerKeyControlRigSpaceChannel(UC
 		}
 
 		// effectively switch to new space
-		URigHierarchy::TElementDependencyMap Dependencies = ControlRig->GetHierarchy()->GetDependenciesForVM(ControlRig->GetVM());
-		ControlRig->GetHierarchy()->SwitchToParent(ControlKey, SpaceKey, false, true, Dependencies, nullptr);
+		ControlRig->SwitchToParent(ControlKey, SpaceKey, false, true);
 		// add new keys in the new space context
 		int32 FramesIndex = 0;
 		for (const FFrameNumber& Frame : Frames)
@@ -484,7 +488,12 @@ FKeyHandle FControlRigSpaceChannelHelpers::SequencerKeyControlRigSpaceChannel(UC
 			ControlRig->Evaluate_AnyThread();
 			Context.LocalTime = TickResolution.AsSeconds(FFrameTime(Frame));
 			ControlRig->SetControlGlobalTransform(ControlKey.Name, ControlWorldTransforms[FramesIndex], true, Context, false /*undo*/, false /*bPrintPython*/, true/* bFixEulerFlips*/);
-
+			if (ControlRig->IsAdditive())
+			{
+				// We need to evaluate in order to trigger a notification
+				ControlRig->Evaluate_AnyThread();
+			}
+			
 			//need to do this after eval
 			FChannelMapInfo* pChannelIndex = nullptr;
 			FRigControlElement* ControlElement = nullptr;
@@ -700,8 +709,7 @@ void  FControlRigSpaceChannelHelpers::SequencerSpaceChannelKeyDeleted(UControlRi
 				RigHierarchy->SwitchToWorldSpace(ControlKey);
 				break;
 			case EMovieSceneControlRigSpaceType::ControlRig:
-				URigHierarchy::TElementDependencyMap Dependencies = RigHierarchy->GetDependenciesForVM(ControlRig->GetVM());
-				RigHierarchy->SwitchToParent(ControlKey, PreviousValue.ControlRigElement, false, true, Dependencies, nullptr);
+				ControlRig->SwitchToParent(ControlKey, PreviousValue.ControlRigElement, false, true);
 				break;
 			}
 			ControlRig->Evaluate_AnyThread();
@@ -1073,8 +1081,7 @@ void FControlRigSpaceChannelHelpers::SequencerBakeControlInSpace(UControlRig* Co
 			FRigControlModifiedContext Context;
 			Context.SetKey = EControlRigSetKey::Always;
 			Context.KeyMask = (uint32)EControlRigContextChannelToKey::AllTransform;
-			URigHierarchy::TElementDependencyMap Dependencies = RigHierarchy->GetDependenciesForVM(ControlRig->GetVM());
-			RigHierarchy->SwitchToParent(ControlKey, Settings.TargetSpace, false, true, Dependencies, nullptr);
+			ControlRig->SwitchToParent(ControlKey, Settings.TargetSpace, false, true);
 			ControlRig->Evaluate_AnyThread();
 
 			FMovieSceneSequenceTransform RootToLocalTransform = Sequencer->GetFocusedMovieSceneSequenceTransform();
@@ -1120,7 +1127,7 @@ void FControlRigSpaceChannelHelpers::SequencerBakeControlInSpace(UControlRig* Co
 					RigHierarchy->SwitchToWorldSpace(ControlKey);
 					break;
 				case EMovieSceneControlRigSpaceType::ControlRig:
-					RigHierarchy->SwitchToParent(ControlKey, EndFrameValue.ControlRigElement, false, true, Dependencies, nullptr);
+					ControlRig->SwitchToParent(ControlKey, EndFrameValue.ControlRigElement, false, true);
 					break;
 				}
 
@@ -1337,8 +1344,7 @@ void FControlRigSpaceChannelHelpers::CompensateIfNeeded(
 									RigHierarchy->SwitchToWorldSpace(Control->GetKey());
 									break;
 								case EMovieSceneControlRigSpaceType::ControlRig:
-									URigHierarchy::TElementDependencyMap Dependencies = RigHierarchy->GetDependenciesForVM(ControlRig->GetVM());
-									RigHierarchy->SwitchToParent(Control->GetKey(), PreviousValue.ControlRigElement, false, true, Dependencies, nullptr);
+									ControlRig->SwitchToParent(Control->GetKey(), PreviousValue.ControlRigElement, false, true);
 									break;
 							}
 							
