Potential fix for crashing using MRC

diff --git a/Engine/Plugins/Animation/ControlRig/Source/ControlRigEditor/Private/EditMode/SControlRigBaseListWidget.cpp b/Engine/Plugins/Animation/ControlRig/Source/ControlRigEditor/Private/EditMode/SControlRigBaseListWidget.cpp
--- a/Engine/Plugins/Animation/ControlRig/Source/ControlRigEditor/Private/EditMode/SControlRigBaseListWidget.cpp
+++ b/Engine/Plugins/Animation/ControlRig/Source/ControlRigEditor/Private/EditMode/SControlRigBaseListWidget.cpp
@@ -1508,12 +1508,14 @@ TSharedRef<class SControlRigPoseView> SControlRigBaseListWidget::CreatePoseView(
 
 TSharedRef<class SControlRigPoseView> SControlRigBaseListWidget::CreateAnimationView(UObject* InObject)
 {
-	return SNew(SControlRigPoseView);
+	return SNew(SControlRigPoseView)
+		.OwningWidget(SharedThis(this));
 }
 
 TSharedRef<class SControlRigPoseView> SControlRigBaseListWidget::CreateSelectionSetView(UObject* InObject)
 {
-	return SNew(SControlRigPoseView);
+	return SNew(SControlRigPoseView)
+		.OwningWidget(SharedThis(this));
 }
 
 #undef LOCTEXT_NAMESPACE
diff --git a/Engine/Plugins/Animation/ControlRig/Source/ControlRigEditor/Private/EditMode/SControlRigControlViews.cpp b/Engine/Plugins/Animation/ControlRig/Source/ControlRigEditor/Private/EditMode/SControlRigControlViews.cpp
--- a/Engine/Plugins/Animation/ControlRig/Source/ControlRigEditor/Private/EditMode/SControlRigControlViews.cpp
+++ b/Engine/Plugins/Animation/ControlRig/Source/ControlRigEditor/Private/EditMode/SControlRigControlViews.cpp
@@ -677,25 +677,28 @@ TSharedRef<SWidget> SControlRigPoseView::GetThumbnailWidget()
 
 TArray<UControlRig*> SControlRigPoseView::GetControlRigs()
 {
-	FControlRigEditMode* EditMode = OwningWidget->GetEditMode();
 	TArray<UControlRig*> NewControlRigs;
-	if (EditMode)
+	if (OwningWidget.IsValid())
 	{
-		NewControlRigs =  EditMode->GetControlRigsArray(false /*bIsVisible*/);
-	}
-	for (TWeakObjectPtr<UControlRig> ControlRigPtr : CurrentControlRigs)
-	{
-		if (ControlRigPtr.IsValid())
+		FControlRigEditMode* EditMode = OwningWidget->GetEditMode();
+		if (EditMode)
 		{
-			if (NewControlRigs.Contains(ControlRigPtr.Get()) == false)
+			NewControlRigs =  EditMode->GetControlRigsArray(false /*bIsVisible*/);
+		}
+		for (TWeakObjectPtr<UControlRig> ControlRigPtr : CurrentControlRigs)
+		{
+			if (ControlRigPtr.IsValid())
 			{
-				(ControlRigPtr.Get())->ControlSelected().RemoveAll(this);
+				if (NewControlRigs.Contains(ControlRigPtr.Get()) == false)
+				{
+					(ControlRigPtr.Get())->ControlSelected().RemoveAll(this);
+				}
 			}
 		}
-	}
-	if (EditMode)
-	{
-		CurrentControlRigs = EditMode->GetControlRigs();
+		if (EditMode)
+		{
+			CurrentControlRigs = EditMode->GetControlRigs();
+		}
 	}
 	
 	return NewControlRigs;
