Sequencer: Pose Library:  Blend float/vector2d controls.

diff --git a/Engine/Plugins/Animation/ControlRig/Source/ControlRig/Private/Tools/ControlRigPose.cpp b/Engine/Plugins/Animation/ControlRig/Source/ControlRig/Private/Tools/ControlRigPose.cpp
--- a/Engine/Plugins/Animation/ControlRig/Source/ControlRig/Private/Tools/ControlRigPose.cpp
+++ b/Engine/Plugins/Animation/ControlRig/Source/ControlRig/Private/Tools/ControlRigPose.cpp
@@ -289,6 +289,21 @@ void FControlRigControlPose::BlendWithInitialPoses(FControlRigControlPose& Initi
 						ControlRig->SetControlGlobalTransform(ControlElement->GetFName(), GlobalTransform, true, Context, bSetupUndo, false /*bPrintPython*/, true/* bFixEulerFlips*/);
 					}
 				}
+				else if(CopyRigControl->ControlType == ERigControlType::Float || 
+						CopyRigControl->ControlType == ERigControlType::ScaleFloat)
+				{
+					float InitialVal = InitialFound->Value.Get<float>();
+					float Val = CopyRigControl->Value.Get<float>();
+					Val = FMath::Lerp(InitialVal, Val, BlendValue);
+					ControlRig->SetControlValue<float>(ControlElement->GetFName(), Val, true, Context, bSetupUndo);
+				}
+				else if (CopyRigControl->ControlType == ERigControlType::Vector2D)
+				{
+					FVector3f InitialVal = InitialFound->Value.Get<FVector3f>();
+					FVector3f Val = CopyRigControl->Value.Get<FVector3f>();
+					Val = FMath::Lerp(InitialVal, Val, BlendValue);
+					ControlRig->SetControlValue<FVector3f>(ControlElement->GetFName(), Val, true, Context, bSetupUndo);
+				}
 			}
 		}
 	}
