Sequencer: Space Switching: Make sure to evaluate at current frame before compensating

diff --git a/Engine/Plugins/Animation/ControlRig/Source/ControlRigEditor/Private/Sequencer/ControlRigSpaceChannelEditors.cpp b/Engine/Plugins/Animation/ControlRig/Source/ControlRigEditor/Private/Sequencer/ControlRigSpaceChannelEditors.cpp
--- a/Engine/Plugins/Animation/ControlRig/Source/ControlRigEditor/Private/Sequencer/ControlRigSpaceChannelEditors.cpp
+++ b/Engine/Plugins/Animation/ControlRig/Source/ControlRigEditor/Private/Sequencer/ControlRigSpaceChannelEditors.cpp
@@ -334,6 +334,17 @@ FKeyHandle FControlRigSpaceChannelHelpers::SequencerKeyControlRigSpaceChannel(UC
 	//we only key if the value is different.
 	if (Value != ExistingValue)
 	{
+		FFrameRate TickResolution = Sequencer->GetFocusedTickResolution();
+		FMovieSceneSequenceTransform RootToLocalTransform = Sequencer->GetFocusedMovieSceneSequenceTransform();
+
+		//make sure to evaluate first frame
+		FFrameTime CurrentTime(Time);
+		CurrentTime = CurrentTime * RootToLocalTransform.InverseNoLooping();
+
+		FMovieSceneContext SceneContext = FMovieSceneContext(FMovieSceneEvaluationRange(CurrentTime, TickResolution), Sequencer->GetPlaybackStatus()).SetHasJumped(true);
+		Sequencer->GetEvaluationTemplate().EvaluateSynchronousBlocking(SceneContext);
+		ControlRig->Evaluate_AnyThread();
+
 		TArray<FFrameNumber> Frames;
 		Frames.Add(Time);
 		
@@ -439,8 +450,6 @@ FKeyHandle FControlRigSpaceChannelHelpers::SequencerKeyControlRigSpaceChannel(UC
 		//do any compensation or previous key adding
 		FRigControlModifiedContext Context;
 		Context.SetKey = EControlRigSetKey::Always;
-		FFrameRate TickResolution = Sequencer->GetFocusedTickResolution();
-		FMovieSceneSequenceTransform RootToLocalTransform = Sequencer->GetFocusedMovieSceneSequenceTransform();
 
 		// set previous key if we're going to switch space
 		if (bSetPreviousKey)
@@ -448,7 +457,7 @@ FKeyHandle FControlRigSpaceChannelHelpers::SequencerKeyControlRigSpaceChannel(UC
 			FFrameTime GlobalTime(Time - 1);
 			GlobalTime = GlobalTime * RootToLocalTransform.InverseNoLooping();
 
-			FMovieSceneContext SceneContext = FMovieSceneContext(FMovieSceneEvaluationRange(GlobalTime, TickResolution), Sequencer->GetPlaybackStatus()).SetHasJumped(true);
+			SceneContext = FMovieSceneContext(FMovieSceneEvaluationRange(GlobalTime, TickResolution), Sequencer->GetPlaybackStatus()).SetHasJumped(true);
 			Sequencer->GetEvaluationTemplate().EvaluateSynchronousBlocking(SceneContext);
 
 			Context.LocalTime = TickResolution.AsSeconds(FFrameTime(Time - 1));
@@ -482,7 +491,7 @@ FKeyHandle FControlRigSpaceChannelHelpers::SequencerKeyControlRigSpaceChannel(UC
 			FFrameTime GlobalTime(Frame);
 			GlobalTime = GlobalTime * RootToLocalTransform.InverseNoLooping();
 
-			FMovieSceneContext SceneContext = FMovieSceneContext(FMovieSceneEvaluationRange(GlobalTime, TickResolution), Sequencer->GetPlaybackStatus()).SetHasJumped(true);
+			SceneContext = FMovieSceneContext(FMovieSceneEvaluationRange(GlobalTime, TickResolution), Sequencer->GetPlaybackStatus()).SetHasJumped(true);
 			Sequencer->GetEvaluationTemplate().EvaluateSynchronousBlocking(SceneContext);
 
 			ControlRig->Evaluate_AnyThread();
