Control Rig: Modular Rigs - Support for reinstancing + asset editor startup

#rnx

diff --git a/Engine/Plugins/Animation/ControlRig/Source/ControlRig/Private/ControlRig.cpp b/Engine/Plugins/Animation/ControlRig/Source/ControlRig/Private/ControlRig.cpp
--- a/Engine/Plugins/Animation/ControlRig/Source/ControlRig/Private/ControlRig.cpp
+++ b/Engine/Plugins/Animation/ControlRig/Source/ControlRig/Private/ControlRig.cpp
@@ -187,16 +187,22 @@ void UControlRig::Initialize(bool bRequestInit)
 	GetDataSourceRegistry();
 
 	// Create the Hierarchy Controller here to avoid UObject creation from Non-Game Threads
-	GetHierarchy()->GetController(true);
+	if(!IsRigModuleInstance())
+	{
+		GetHierarchy()->GetController(true);
+	}
 	
 	// should refresh mapping 
 	RequestConstruction();
-	
-	GetHierarchy()->OnModified().RemoveAll(this);
-	GetHierarchy()->OnModified().AddUObject(this, &UControlRig::HandleHierarchyModified);
-	GetHierarchy()->OnEventReceived().RemoveAll(this);
-	GetHierarchy()->OnEventReceived().AddUObject(this, &UControlRig::HandleHierarchyEvent);
-	GetHierarchy()->UpdateVisibilityOnProxyControls();
+
+	if(!IsRigModuleInstance())
+	{
+		GetHierarchy()->OnModified().RemoveAll(this);
+		GetHierarchy()->OnModified().AddUObject(this, &UControlRig::HandleHierarchyModified);
+		GetHierarchy()->OnEventReceived().RemoveAll(this);
+		GetHierarchy()->OnEventReceived().AddUObject(this, &UControlRig::HandleHierarchyEvent);
+		GetHierarchy()->UpdateVisibilityOnProxyControls();
+	}
 }
 
 void UControlRig::OnAddShapeLibrary(const FControlRigExecuteContext* InContext, const FString& InLibraryName, UControlRigShapeLibrary* InShapeLibrary, bool bReplaceExisting, bool bLogResults)
@@ -1439,6 +1445,20 @@ bool UControlRig::Execute_Internal(const FName& InEventName)
 	{
 		return false;
 	}
+
+	// make sure to initialize here as well - just in case this gets
+	// called without a call to ::Execute. This is already tackled by
+	// the UModularRig::ExecuteQueue, but added it here as well nevertheless
+	// to avoid crashes in the future.
+	if(bRequiresInitExecution)
+	{
+		const TGuardValue<float> AbsoluteTimeGuard(AbsoluteTime, AbsoluteTime);
+		const TGuardValue<float> DeltaTimeGuard(DeltaTime, DeltaTime);
+		if(!InitializeVM(InEventName))
+		{
+			return false;
+		}
+	}
 	
 	if(IsRigModule())
 	{
@@ -3276,12 +3296,18 @@ void UControlRig::PostInitInstance(URigVMHost* InCDO)
 			}
 		}
 
-		DynamicHierarchy = NewObject<URigHierarchy>(this, TEXT("DynamicHierarchy"), SubObjectFlags);
+		if(!IsRigModuleInstance())
+		{
+			DynamicHierarchy = NewObject<URigHierarchy>(this, TEXT("DynamicHierarchy"), SubObjectFlags);
+		}
 	}
 
 #if WITH_EDITOR
+	if(!IsRigModuleInstance())
+	{
 		const TWeakObjectPtr<UControlRig> WeakThis = this;
 		DynamicHierarchy->OnUndoRedo().AddStatic(&UControlRig::OnHierarchyTransformUndoRedoWeak, WeakThis);
+	}
 #endif
 
 	if(!HasAnyFlags(RF_ClassDefaultObject))
@@ -3294,7 +3320,10 @@ void UControlRig::PostInitInstance(URigVMHost* InCDO)
 				VM = InCDO->GetVM();
 			}
 
-			DynamicHierarchy->CopyHierarchy(CastChecked<UControlRig>(InCDO)->GetHierarchy());
+			if(!IsRigModuleInstance())
+			{
+				DynamicHierarchy->CopyHierarchy(CastChecked<UControlRig>(InCDO)->GetHierarchy());
+			}
 		}
 	}
 	else // we are the CDO
@@ -3317,7 +3346,11 @@ void UControlRig::PostInitInstance(URigVMHost* InCDO)
 		if(GetClass()->IsNative())
 		{
 			VM->AddToRoot();
-			DynamicHierarchy->AddToRoot();
+
+			if(!IsRigModuleInstance())
+			{
+				DynamicHierarchy->AddToRoot();
+			}
 		}
 	}
 
diff --git a/Engine/Plugins/Animation/ControlRig/Source/ControlRig/Private/ModularRig.cpp b/Engine/Plugins/Animation/ControlRig/Source/ControlRig/Private/ModularRig.cpp
--- a/Engine/Plugins/Animation/ControlRig/Source/ControlRig/Private/ModularRig.cpp
+++ b/Engine/Plugins/Animation/ControlRig/Source/ControlRig/Private/ModularRig.cpp
@@ -2,8 +2,6 @@
 
 #include "ModularRig.h"
 #include "Units/Execution/RigUnit_BeginExecution.h"
-#include "Units/Execution/RigUnit_InverseExecution.h"
-#include "Units/Execution/RigUnit_PrepareForExecution.h"
 #include "Units/Execution/RigUnit_InteractionExecution.h"
 #include "ControlRigObjectBinding.h"
 #include "Rigs/RigHierarchyController.h"
@@ -18,13 +16,13 @@ const FString UModularRig::NamespaceSeparator = TEXT(":");
 // FModuleInstanceHandle
 ////////////////////////////////////////////////////////////////////////////////
 
-FModuleInstanceHandle::FModuleInstanceHandle(UModularRig* InModularRig, const FString& InPath)
+FModuleInstanceHandle::FModuleInstanceHandle(const UModularRig* InModularRig, const FString& InPath)
 : ModularRig(InModularRig)
 , Path(InPath)
 {
 }
 
-FModuleInstanceHandle::FModuleInstanceHandle(UModularRig* InModularRig, const FRigModuleInstance* InModule)
+FModuleInstanceHandle::FModuleInstanceHandle(const UModularRig* InModularRig, const FRigModuleInstance* InModule)
 : ModularRig(InModularRig)
 , Path(InModule->GetPath())
 {
@@ -32,9 +30,9 @@ FModuleInstanceHandle::FModuleInstanceHandle(UModularRig* InModularRig, const FR
 
 const FRigModuleInstance* FModuleInstanceHandle::Get() const
 {
-	if(ModularRig.IsValid())
+	if(const UModularRig* ResolvedRig = ModularRig.Get())
 	{
-		return ModularRig->FindModule(Path);
+		return ResolvedRig->FindModule(Path);
 	}
 	return nullptr;
 }
@@ -47,7 +45,7 @@ UModularRig::UModularRig(const FObjectInitializer& ObjectInitializer)
 	: Super(ObjectInitializer)
 {
 #if WITH_EDITOR
-	FCoreUObjectDelegates::OnObjectsReplaced.AddUObject(this, &UModularRig::OnObjectsReplaced);
+	FCoreUObjectDelegates::OnObjectsReinstanced.AddUObject(this, &UModularRig::OnObjectsReplaced);
 #endif
 }
 
@@ -58,7 +56,7 @@ void UModularRig::BeginDestroy()
 	Super::BeginDestroy();
 	
 #if WITH_EDITOR
-	FCoreUObjectDelegates::OnObjectsReplaced.RemoveAll(this);
+	FCoreUObjectDelegates::OnObjectsReinstanced.RemoveAll(this);
 #endif
 }
 
@@ -76,21 +74,56 @@ FString FRigModuleInstance::GetNamespace() const
 	return FString::Printf(TEXT("%s:"), *GetPath());
 }
 
-void UModularRig::Initialize(bool bRequestInit)
+UControlRig* FRigModuleInstance::GetRig() const
 {
-	Super::Initialize(bRequestInit);
+	// prefer the cached UObject over the soft object ptrs
+	// since the blueprint reinstancer does not reinstance soft object ptrs.
+	if(IsValid(RigPtr))
+	{
+		return RigPtr;
+	}
 
-	UpdateCachedChildren();
+	// reset the cache if it is not valid
+	return RigPtr = nullptr;
+}
+
+void FRigModuleInstance::SetRig(UControlRig* InRig)
+{
+	UControlRig* PreviousRig = GetRig();
+	if(PreviousRig && (PreviousRig != InRig))
+	{
+		// rename the previous rig.
+		// GC will pick it up eventually - since we won't have any
+		// owning pointers to it anymore.
+		PreviousRig->Rename(nullptr, GetTransientPackage(), REN_ForceNoResetLoaders | REN_DoNotDirty | REN_DontCreateRedirectors | REN_NonTransactional);
+		PreviousRig->MarkAsGarbage();
+	}
+
+	// update the cache
+	RigPtr = InRig;
+}
+
+bool FRigModuleInstance::ContainsRig(const UControlRig* InRig) const
+{
+	if(InRig == nullptr)
+	{
+		return false;
+	}
+	if(RigPtr == InRig)
+	{
+		return true;
+	}
+	return false;
 }
 
 void UModularRig::InitializeVMs(bool bRequestInit)
 {
 	URigVMHost::Initialize(bRequestInit);
-	ForEachModule([bRequestInit](FRigModuleInstance* Module) -> bool
+	ForEachModule([bRequestInit](const FRigModuleInstance* Module) -> bool
 	{
-		if (Module->Rig.IsValid())
+		if (UControlRig* ModuleRig = Module->GetRig())
 		{
-			Module->Rig->InitializeVMs(bRequestInit);
+			ModuleRig->InitializeVMs(bRequestInit);
 		}
 		return true;
 	});
@@ -99,11 +132,13 @@ void UModularRig::InitializeVMs(bool bRequestInit)
 bool UModularRig::InitializeVMs(const FName& InEventName)
 {
 	URigVMHost::InitializeVM(InEventName);
-	ForEachModule([InEventName](FRigModuleInstance* Module) -> bool
+	UpdateModuleHierarchyFromCDO();
+
+	ForEachModule([InEventName](const FRigModuleInstance* Module) -> bool
 	{
-		if (Module->Rig.IsValid())
+		if (UControlRig* ModuleRig = Module->GetRig())
 		{
-			Module->Rig->InitializeVMs(InEventName);
+			ModuleRig->InitializeVMs(InEventName);
 		}
 		return true;
 	});
@@ -113,23 +148,64 @@ bool UModularRig::InitializeVMs(const FName& InEventName)
 void UModularRig::InitializeFromCDO()
 {
 	Super::InitializeFromCDO();
+	UpdateModuleHierarchyFromCDO();
+}
 
-	if (!HasAnyFlags(RF_ClassDefaultObject))
+static void AddModulesRecursively(UModularRig* Rig, const FRigModuleReference* InModule, FRigModuleInstance* InParent)
+{
+	// Make sure the inner is loaded and valid
 	{
-		const UModularRig* CDO = GetClass()->GetDefaultObject<UModularRig>();
+		if (!InModule->Class.IsValid())
+		{
+			(void)InModule->Class.LoadSynchronous();
+		}
+		if (!InModule->Class.IsValid())
+		{
+			return;
+		}
+	}
+	
+	FRigModuleInstance* NewModule = Rig->AddModuleInstance(InModule->Name, InModule->Class.Get(), InParent, InModule->Connections, InModule->ConfigValues, InModule->Bindings);
+	if(NewModule)
+	{
+		for (const FRigModuleReference* ChildModule : InModule->CachedChildren)
+		{
+			AddModulesRecursively(Rig, ChildModule, NewModule);
+		}
+	}
+	else
+	{
+		// todooo
+	}
+}
 
-		// Generate the rig module tree based on the CDO
+void UModularRig::UpdateModuleHierarchyFromCDO()
+{
+	if (!HasAnyFlags(RF_ClassDefaultObject))
+	{
 		ResetModules();
-		CDO->ForEachModule([this](const FRigModuleInstance* CDOModule) -> bool
+
+		// the CDO owns the model - when we ask for the model we'll always
+		// get the model from the CDO. we'll now add UObject module instances
+		// for each module (data only) reference in the model.
+		// Note: The CDO does not contain any UObject module instances itself.
+		const FModularRigModel& Model = GetModularRigModel();
+		for (const FRigModuleReference* RootModule : Model.RootModules)
+		{
+			AddModulesRecursively(this, RootModule, nullptr);
+		}
+
+		ForEachModule([this](const FRigModuleInstance* Module) -> bool
 		{
-			if (FRigModuleInstance* NewModule = AddModuleInstance(CDOModule))
+			if(UControlRig* ModuleRig = Module->GetRig())
 			{
-				NewModule->Rig->Initialize();
+				ModuleRig->Initialize();
 			}
 			return true;
 		});
 
-		SupportedEvents = CDO->SupportedEvents;
+		UpdateCachedChildren();
+		UpdateSupportedEvents();
 	}
 }
 
@@ -144,11 +220,9 @@ bool UModularRig::Execute_Internal(const FName& InEventName)
 
 		ForEachModule([&InEventName, this, Hierarchy, UnitContext](FRigModuleInstance* Module) -> bool
 		{
-			if (Module->Rig.IsValid())
+			if (const UControlRig* ModuleRig = Module->GetRig())
 			{
-				UControlRig* Rig = Module->Rig.Get();
-
-				if (!Rig->SupportsEvent(InEventName))
+				if (!ModuleRig->SupportsEvent(InEventName))
 				{
 					return true;
 				}
@@ -194,20 +268,18 @@ void UModularRig::ExecuteQueue()
 	while(ExecutionQueue.IsValidIndex(ExecutionQueueFront))
 	{
 		FRigModuleExecutionElement& ExecutionElement = ExecutionQueue[ExecutionQueueFront];
-		if (ExecutionElement.ModuleInstance->Rig.IsValid())
+		if (UControlRig* ModuleRig = ExecutionElement.ModuleInstance->GetRig())
 		{
-			UControlRig* Rig = ExecutionElement.ModuleInstance->Rig.Get();
-
-			if (!Rig->SupportsEvent(ExecutionElement.EventName))
+			if (!ModuleRig->SupportsEvent(ExecutionElement.EventName))
 			{
 				ExecutionQueueFront++;
 				continue;
 			}
-				
+
 			// Make sure the hierarchy has the correct element redirector from this module rig
-			FRigHierarchyRedirectorGuard ElementRedirectorGuard(Rig);
+			FRigHierarchyRedirectorGuard ElementRedirectorGuard(ModuleRig);
 
-			FRigVMExtendedExecuteContext& RigExtendedExecuteContext= Rig->GetRigVMExtendedExecuteContext();
+			FRigVMExtendedExecuteContext& RigExtendedExecuteContext= ModuleRig->GetRigVMExtendedExecuteContext();
 
 			// Make sure the hierarchy has the correct execute context with the rig module namespace
 			FRigHierarchyExecuteContextBracket ExecuteContextBracket(Hierarchy, &RigExtendedExecuteContext);
@@ -217,6 +289,29 @@ void UModularRig::ExecuteQueue()
 			FRigUnitContext& RigUnitContext = RigPublicContext.UnitContext;
 			RigUnitContext = PublicContext.UnitContext;
 
+			// re-initialize the module in case only the VM side got recompiled.
+			// this happens when the user relies on auto recompilation when editing the
+			// module (dependency) graph - by changing a value, add / remove nodes or links.
+			if(ModuleRig->IsInitRequired())
+			{
+				const TGuardValue<float> AbsoluteTimeGuard(ModuleRig->AbsoluteTime, ModuleRig->AbsoluteTime);
+				const TGuardValue<float> DeltaTimeGuard(ModuleRig->DeltaTime, ModuleRig->DeltaTime);
+				if(!ModuleRig->InitializeVM(ExecutionElement.EventName))
+				{
+					ExecutionQueueFront++;
+					continue;
+				}
+
+				// put the variable defaults back
+				if(const FRigModuleReference* ModuleReference = GetModularRigModel().FindModule(ExecutionElement.ModulePath))
+				{
+					for (const TPair<FName, FString>& Variable : ModuleReference->ConfigValues)
+					{
+						ModuleRig->SetVariableFromString(Variable.Key, Variable.Value);
+					}
+				}
+			}
+
 			// Update the interaction elements to show only the ones belonging to this module
 			const FString ModuleNamespace = FString::Printf(TEXT("%s:"), *ExecutionElement.ModulePath);
 			RigUnitContext.ElementsBeingInteracted = RigUnitContext.ElementsBeingInteracted.FilterByPredicate(
@@ -238,7 +333,7 @@ void UModularRig::ExecuteQueue()
 			// - external variables
 			{
 				RigPublicContext.AssetUserData.Reset();
-				if(const TArray<UAssetUserData*>* ControlRigUserDataArray = Rig->GetAssetUserDataArray())
+				if(const TArray<UAssetUserData*>* ControlRigUserDataArray = ModuleRig->GetAssetUserDataArray())
 				{
 					for(const UAssetUserData* ControlRigUserData : *ControlRigUserDataArray)
 					{
@@ -251,7 +346,7 @@ void UModularRig::ExecuteQueue()
 			// Copy variable bindings
 			for (TPair<FName, FRigVMExternalVariable>& Pair : ExecutionElement.ModuleInstance->VariableBindings)
 			{
-				const FRigVMExternalVariable TargetVariable = ExecutionElement.ModuleInstance->Rig->GetPublicVariableByName(Pair.Key);
+				const FRigVMExternalVariable TargetVariable = ExecutionElement.ModuleInstance->GetRig()->GetPublicVariableByName(Pair.Key);
 				if(ensure(TargetVariable.Property))
 				{
 					if (RigVMTypeUtils::AreCompatible(Pair.Value.Property, TargetVariable.Property))
@@ -261,7 +356,7 @@ void UModularRig::ExecuteQueue()
 				}
 			}
 			
-			Rig->Execute_Internal(ExecutionElement.EventName);
+			ModuleRig->Execute_Internal(ExecutionElement.EventName);
 			ExecutionElement.bExecuted = true;
 		}
 		
@@ -275,18 +370,44 @@ void UModularRig::ResetExecutionQueue()
 	ExecutionQueueFront = 0;
 }
 
-
 void UModularRig::OnObjectsReplaced(const TMap<UObject*, UObject*>& OldToNewInstanceMap)
 {
-	for (FRigModuleInstance& ModuleInstance : Modules)
+	if(Modules.IsEmpty())
 	{
-		if (UObject*const * NewObject = OldToNewInstanceMap.Find(ModuleInstance.Rig.Get()))
+		return;
+	}
+	
+	bool bPerformedChange = false;
+	for(const TPair<UObject*, UObject*>& Pair : OldToNewInstanceMap)
+	{
+		UObject* NewObject = Pair.Value;
+		if((NewObject == nullptr) || (NewObject->GetOuter() != this) || !NewObject->IsA<UControlRig>())
+		{
+			continue;
+		}
+
+		UControlRig* NewRig = CastChecked<UControlRig>(NewObject);
+
+		// relying on GetFName since URigVMHost is overloading GetName()
+		const FString& Path = NewRig->GetFName().ToString();
+
+		// if we find a matching module update it.
+		// FRigModuleInstance::SetRig takes care of disregarding the previous module instance.
+		if(FRigModuleInstance* Module = const_cast<FRigModuleInstance*>(FindModule(Path)))
 		{
-			ModuleInstance.Rig = Cast<UControlRig>(*NewObject);
+			Module->SetRig(NewRig);
+			NewRig->bCopyHierarchyBeforeConstruction = false;
+			NewRig->SetDynamicHierarchy(GetHierarchy());
+			NewRig->Initialize(true);
+			bPerformedChange = true;
 		}
 	}
-	InitializeVMs(true);
-	UpdateSupportedEvents();
+
+	if(bPerformedChange)
+	{
+		UpdateSupportedEvents();
+		RequestInit();
+	}
 }
 
 void UModularRig::ResetModules()
@@ -294,18 +415,12 @@ void UModularRig::ResetModules()
 	for (FRigModuleInstance& Module : Modules)
 	{
 		Module.CachedChildren.Reset();
-		TSoftObjectPtr<UControlRig> Rig = Module.Rig;
-		if (Rig && Rig.IsValid())
+
+		if (const UControlRig* ModuleRig = Module.GetRig())
 		{
-			static int32 ObjectIndexToBeDestroyed = 0;
-			static constexpr TCHAR ObjectNameFormat[] = TEXT("FRigModuleInstance_ObjectToBeDestroyed_%d");
-			const FString NewObjectName = FString::Printf(ObjectNameFormat, ObjectIndexToBeDestroyed++);
-			Rig->Rename(*NewObjectName, GetTransientPackage(), REN_ForceNoResetLoaders | REN_DoNotDirty | REN_DontCreateRedirectors | REN_NonTransactional);
-			if(!Rig->IsRooted())
-			{
-				Rig->MarkAsGarbage();
-			}
-			Rig = nullptr;
+			check(ModuleRig->GetOuter() == this);
+			// takes care of renaming / moving the rig to the transient package
+			Module.SetRig(nullptr);
 		}
 	}
 	
@@ -314,6 +429,16 @@ void UModularRig::ResetModules()
 	SupportedEvents.Reset();
 }
 
+const FModularRigModel& UModularRig::GetModularRigModel() const
+{
+	if (!HasAnyFlags(RF_ClassDefaultObject))
+	{
+		const UModularRig* CDO = GetClass()->GetDefaultObject<UModularRig>();
+		return CDO->GetModularRigModel();
+	}
+	return ModularRigModel;
+}
+
 void UModularRig::UpdateCachedChildren()
 {
 	TMap<FString, FRigModuleInstance*> PathToModule;
@@ -345,9 +470,9 @@ void UModularRig::UpdateSupportedEvents()
 	SupportedEvents.Reset();
 	ForEachModule([this](const FRigModuleInstance* Module) -> bool
 	{
-		if (Module->Rig.IsValid())
+		if (const UControlRig* ModuleRig = Module->GetRig())
 		{
-			const TArray<FName>& ModuleEvents = Module->Rig->GetSupportedEvents();
+			const TArray<FName>& ModuleEvents = ModuleRig->GetSupportedEvents();
 			for (const FName& EventName : ModuleEvents)
 			{
 				SupportedEvents.AddUnique(EventName);
@@ -370,7 +495,7 @@ FRigModuleInstance* UModularRig::AddModuleInstance(const FName& InModuleName, TS
 	// Make sure there are no name clashes
 	if (InParent)
 	{
-		for (FRigModuleInstance* Child : InParent->CachedChildren)
+		for (const FRigModuleInstance* Child : InParent->CachedChildren)
 		{
 			if (Child->Name == InModuleName)
 			{
@@ -380,7 +505,7 @@ FRigModuleInstance* UModularRig::AddModuleInstance(const FName& InModuleName, TS
 	}
 	else
 	{
-		for (FRigModuleInstance* RootModule : RootModules)
+		for (const FRigModuleInstance* RootModule : RootModules)
 		{
 			if (RootModule->Name == InModuleName)
 			{
@@ -395,17 +520,19 @@ FRigModuleInstance* UModularRig::AddModuleInstance(const FName& InModuleName, TS
 		return nullptr;
 	}
 
-	FString Name = (InParent) ? InParent->Name.ToString() + NamespaceSeparator + InModuleName.ToString() : InModuleName.ToString();
 	FRigModuleInstance& NewModule = Modules.Add_GetRef(FRigModuleInstance());
-	NewModule.Rig = NewObject<UControlRig>(this, InModuleClass, *Name);
 	NewModule.Name = InModuleName;
-
 	if (InParent)
 	{
 		NewModule.ParentPath = InParent->GetPath();
 	}
+	const FString Name = NewModule.GetPath();
+	
+	UControlRig* NewModuleRig = NewObject<UControlRig>(this, InModuleClass, *Name);
+	NewModule.SetRig(NewModuleRig);
+
 	UpdateCachedChildren();
-	for (const FName& EventName : NewModule.Rig->GetSupportedEvents())
+	for (const FName& EventName : NewModule.GetRig()->GetSupportedEvents())
 	{
 		SupportedEvents.AddUnique(EventName);
 	}
@@ -413,18 +540,18 @@ FRigModuleInstance* UModularRig::AddModuleInstance(const FName& InModuleName, TS
 	// Configure module
 	{
 		URigHierarchy* Hierarchy = GetHierarchy();
-		FRigVMExtendedExecuteContext& ModuleContext = NewModule.Rig->GetRigVMExtendedExecuteContext();
+		FRigVMExtendedExecuteContext& ModuleContext = NewModuleRig->GetRigVMExtendedExecuteContext();
 		FControlRigExecuteContext& ModulePublicContext = ModuleContext.GetPublicDataSafe<FControlRigExecuteContext>();
-		NewModule.Rig->bCopyHierarchyBeforeConstruction = false;
-		NewModule.Rig->SetDynamicHierarchy(Hierarchy);
+		NewModuleRig->bCopyHierarchyBeforeConstruction = false;
+		NewModuleRig->SetDynamicHierarchy(Hierarchy);
 		ModulePublicContext.Hierarchy = Hierarchy;
-		ModulePublicContext.RigModuleNameSpace = NewModule.Rig->GetRigModuleNameSpace();
+		ModulePublicContext.RigModuleNameSpace = NewModuleRig->GetRigModuleNameSpace();
 		ModulePublicContext.RigModuleNameSpaceHash = GetTypeHash(ModulePublicContext.RigModuleNameSpace);
-		NewModule.Rig->SetElementKeyRedirector(FRigElementKeyRedirector(InConnectionMap, Hierarchy));
+		NewModuleRig->SetElementKeyRedirector(FRigElementKeyRedirector(InConnectionMap, Hierarchy));
 
 		for (const TPair<FName, FString>& Variable : InVariableDefaultValues )
 		{
-			NewModule.Rig->SetVariableFromString(Variable.Key, Variable.Value);
+			NewModule.GetRig()->SetVariableFromString(Variable.Key, Variable.Value);
 		}
 
 		for (const TPair<FName, FString>& Pair : InVariableBindings)
@@ -434,14 +561,14 @@ FRigModuleInstance* UModularRig::AddModuleInstance(const FName& InModuleName, TS
 			FRigVMExternalVariable SourceVariable;
 			if (SourceModulePath.IsEmpty())
 			{
-				if (FProperty* Property = GetClass()->FindPropertyByName(*SourceVariableName))
+				if (const FProperty* Property = GetClass()->FindPropertyByName(*SourceVariableName))
 				{
 					SourceVariable = FRigVMExternalVariable::Make(Property, (UObject*)this);
 				}
 			}
 			else if(const FRigModuleInstance* SourceModule = FindModule(SourceModulePath))
 			{
-				SourceVariable = SourceModule->Rig->GetPublicVariableByName(*SourceVariableName);
+				SourceVariable = SourceModule->GetRig()->GetPublicVariableByName(*SourceVariableName);
 			}
 			SourceVariable.Name = *Pair.Value; // Adapt the name of the variable to contain the full path
 			check(SourceVariable.Property);
@@ -452,66 +579,12 @@ FRigModuleInstance* UModularRig::AddModuleInstance(const FName& InModuleName, TS
 	return &NewModule;
 }
 
-FRigModuleInstance* UModularRig::AddModuleInstance(const FRigModuleInstance* InOtherModule)
-{
-	if (!InOtherModule->Rig.IsValid())
-	{
-		return nullptr;
-	}
-
-	// Figure out the ConnectionMap
-	const FRigElementKeyRedirector& Redirector = InOtherModule->Rig->GetElementKeyRedirector();
-	const TMap<FRigElementKey, FRigElementKey> ConnectionMap = Redirector.ExternalKeys;
-
-	// Figure out the DefaultValues (which is the current value on the CDO)
-	TMap<FName, FString> DefaultValues;
-	const TArray<FRigVMExternalVariable> Variables = InOtherModule->Rig->GetPublicVariables();
-	for (const FRigVMExternalVariable& Variable : Variables)
-	{
-		const FString Value = InOtherModule->Rig->GetVariableAsString(Variable.Name);
-		DefaultValues.Add(Variable.Name, Value);
-	}
-
-	// Figure out the variable bindings (where the source is either a root variable or a variable on another module)
-	TMap<FName, FString> VariableBindings;
-	for (const TPair<FName, FRigVMExternalVariable>& Pair : InOtherModule->VariableBindings)
-	{
-		VariableBindings.Add(Pair.Key, Pair.Value.Name.ToString());
-	}
-
-	// Figure out the parent module
-	FRigModuleInstance* ParentModule = const_cast<FRigModuleInstance*>(FindModule(InOtherModule->ParentPath));
-	
-	return AddModuleInstance(InOtherModule->Name, InOtherModule->Rig->GetClass(), ParentModule, ConnectionMap, DefaultValues, VariableBindings);
-}
-
 const FRigModuleInstance* UModularRig::FindModule(const FString& InPath) const
 {
-	const TArray<FRigModuleInstance*>* Children = &RootModules;
-	FString Left = InPath, Right;
-	while (Left.Split(NamespaceSeparator, &Left, &Right))
+	return Modules.FindByPredicate([InPath](const FRigModuleInstance& Module)
 	{
-		FRigModuleInstance* const * Cur = Children->FindByPredicate([Left](FRigModuleInstance* Module)
-		{
-			return Module->Name == Left;
-		});
-		if (!Cur)
-		{
-			return nullptr;
-		}
-		Children = &(*Cur)->CachedChildren;
-		Left = Right;
-	}
-
-	FRigModuleInstance* const * Cur = Children->FindByPredicate([Left](FRigModuleInstance* Module)
-		{
-			return Module->Name == Left;
-		});
-	if (!Cur)
-	{
-		return nullptr;
-	}
-	return *Cur;
+		return Module.GetPath() == InPath;
+	});
 }
 
 const FRigModuleInstance* UModularRig::FindModule(const UControlRig* InModuleInstance) const
@@ -519,14 +592,11 @@ const FRigModuleInstance* UModularRig::FindModule(const UControlRig* InModuleIns
 	const FRigModuleInstance* FoundModule = nullptr;
 	ForEachModule([InModuleInstance, &FoundModule](const FRigModuleInstance* Module) -> bool
 	{
-		if (Module->Rig.IsValid())
+		if(Module->GetRig() == InModuleInstance)
 		{
-			if(Module->Rig.Get() == InModuleInstance)
-			{
-				FoundModule = Module;
-				// don't continue
-				return false;
-			}
+			FoundModule = Module;
+			// don't continue ForEachModule
+			return false;
 		}
 		return true;
 	});
diff --git a/Engine/Plugins/Animation/ControlRig/Source/ControlRig/Private/ModularRigController.cpp b/Engine/Plugins/Animation/ControlRig/Source/ControlRig/Private/ModularRigController.cpp
new file mode 100644
--- /dev/null
+++ b/Engine/Plugins/Animation/ControlRig/Source/ControlRig/Private/ModularRigController.cpp
@@ -0,0 +1,791 @@
+// Copyright Epic Games, Inc. All Rights Reserved.
+
+#include "ModularRigController.h"
+
+#include "ControlRig.h"
+#include "ModularRig.h"
+#include "ModularRigModel.h"
+#include "Misc/DefaultValueHelper.h"
+#include "Rigs/RigHierarchyController.h"
+
+#include UE_INLINE_GENERATED_CPP_BY_NAME(ModularRigController)
+
+#if WITH_EDITOR
+#include "ScopedTransaction.h"
+#include "Kismet2/BlueprintEditorUtils.h"
+#endif
+
+UModularRigController::UModularRigController(const FObjectInitializer& ObjectInitializer)
+	: Super(ObjectInitializer)
+	, Model(nullptr)
+	, bSuspendNotifications(false)
+{
+}
+
+FString UModularRigController::AddModule(const FName& InModuleName, TSubclassOf<UControlRig> InClass, const FString& InParentModulePath, bool bSetupUndo)
+{
+	if (!InClass)
+	{
+		UE_LOG(LogControlRig, Error, TEXT("Invalid InClass"));
+		return FString();
+	}
+
+	UControlRig* ClassDefaultObject = InClass->GetDefaultObject<UControlRig>();
+	if (!ClassDefaultObject->IsRigModule())
+	{
+		UE_LOG(LogControlRig, Error, TEXT("Class %s is not a rig module"), *InClass->GetClassPathName().ToString());
+		return FString();
+	}
+
+#if WITH_EDITOR
+	TSharedPtr<FScopedTransaction> TransactionPtr;
+	if (bSetupUndo)
+	{
+		TransactionPtr = MakeShared<FScopedTransaction>(NSLOCTEXT("ModularRigController", "AddModuleTransaction", "Add Module"));
+		if(UBlueprint* Blueprint = Cast<UBlueprint>(GetOuter()))
+		{
+			Blueprint->Modify();
+		}
+	}
+#endif 
+
+	FRigModuleReference* NewModule = nullptr;
+	if (InParentModulePath.IsEmpty())
+	{
+		for (FRigModuleReference* Module : Model->RootModules)
+		{
+			if (Module->Name.ToString() == InModuleName)
+			{
+				return FString();
+			}
+		}
+
+		Model->Modules.Add(FRigModuleReference(InModuleName, InClass, FString()));
+		NewModule = &Model->Modules.Last();
+	}
+	else if (FRigModuleReference* ParentModule = FindModule(InParentModulePath))
+	{
+		for (FRigModuleReference* Module : ParentModule->CachedChildren)
+		{
+			if (Module->Name.ToString() == InModuleName)
+			{
+				return FString();
+			}
+		}
+
+		Model->Modules.Add(FRigModuleReference(InModuleName, InClass, ParentModule->GetPath()));
+		NewModule = &Model->Modules.Last();
+	}
+
+	Model->UpdateCachedChildren();
+
+	if (!NewModule)
+	{
+		UE_LOG(LogControlRig, Error, TEXT("Error while creating module %s"), *InModuleName.ToString());
+		return FString();
+	}
+
+	Notify(EModularRigNotification::ModuleAdded, NewModule);
+
+#if WITH_EDITOR
+	TransactionPtr.Reset();
+#endif
+
+	return NewModule->GetPath();
+}
+
+FRigModuleReference* UModularRigController::FindModule(const FString& InPath)
+{
+	FString Path = InPath;
+	Path.RemoveFromEnd(UModularRig::NamespaceSeparator);
+	
+	TArray<FRigModuleReference*>* Children = &Model->RootModules;
+	FString Left = Path, Right;
+	while (Left.Split(UModularRig::NamespaceSeparator, &Left, &Right))
+	{
+		FRigModuleReference** Cur = Children->FindByPredicate([Left](FRigModuleReference* Module)
+		{
+			return Module->Name == Left;
+		});
+		if (!Cur)
+		{
+			return nullptr;
+		}
+		Children = &(*Cur)->CachedChildren;
+		Left = Right;
+	}
+
+	FRigModuleReference** Cur = Children->FindByPredicate([Left](FRigModuleReference* Module)
+		{
+			return Module->Name == Left;
+		});
+	if (!Cur)
+	{
+		return nullptr;
+	}
+	return *Cur;
+}
+
+bool UModularRigController::CanConnectConnectorToElement(const FRigModuleConnector& InConnector, const FRigElementKey& InTargetKey, FText& OutErrorMessage)
+{
+	// TODO: Check rules are satisfied
+	return true;
+}
+
+bool UModularRigController::ConnectConnectorToElement(const FRigElementKey& InConnectorKey, const FRigElementKey& InTargetKey, bool bSetupUndo)
+{
+	FString ConnectorParentPath, ConnectorName;
+	if (!InConnectorKey.Name.ToString().Split(UModularRig::NamespaceSeparator, &ConnectorParentPath, &ConnectorName, ESearchCase::CaseSensitive, ESearchDir::FromEnd))
+	{
+		UE_LOG(LogControlRig, Error, TEXT("Connector %s does not contain a namespace"), *InConnectorKey.ToString());
+		return false;
+	}
+	
+	FRigModuleReference* Module = FindModule(ConnectorParentPath);
+	if (!Module)
+	{
+		UE_LOG(LogControlRig, Error, TEXT("Could not find module %s"), *ConnectorParentPath);
+		return false;
+	}
+
+	UControlRig* RigCDO = Module->Class->GetDefaultObject<UControlRig>();
+	if (!RigCDO)
+	{
+		UE_LOG(LogControlRig, Error, TEXT("Invalid rig module class %s"), *Module->Class->GetPathName());
+		return false;
+	}
+
+	const FRigModuleConnector* ModuleConnector = RigCDO->GetRigModuleSettings().ExposedConnectors.FindByPredicate(
+		[ConnectorName](FRigModuleConnector& Connector)
+		{
+			return Connector.Name == ConnectorName;
+		});
+	if (!ModuleConnector)
+	{
+		UE_LOG(LogControlRig, Error, TEXT("Could not find connector %s in class %s"), *ConnectorName, *Module->Class->GetPathName());
+		return false;
+	}
+
+	FText ErrorMessage;
+	if (!CanConnectConnectorToElement(*ModuleConnector, InTargetKey, ErrorMessage))
+	{
+		UE_LOG(LogControlRig, Error, TEXT("Cannot connect connector %s to target %s: %s"),
+			*InConnectorKey.Name.ToString(), *InTargetKey.ToString(), *ErrorMessage.ToString());
+		return false;
+	}
+
+	UBlueprint* Blueprint = Cast<UBlueprint>(GetOuter());
+	const IRigHierarchyProvider* HierarchyProvider = CastChecked<IRigHierarchyProvider>(Blueprint);
+	const FRigConnectorElement* Connector = Cast<FRigConnectorElement>(HierarchyProvider->GetHierarchy()->Find(InConnectorKey));
+	if (!Connector)
+	{
+		UE_LOG(LogControlRig, Error, TEXT("Could not find connector %s"), *InConnectorKey.ToString());
+		return false;
+	}
+
+#if WITH_EDITOR
+	TSharedPtr<FScopedTransaction> TransactionPtr;
+	if (bSetupUndo)
+	{
+		TransactionPtr = MakeShared<FScopedTransaction>(NSLOCTEXT("ModularRigController", "ConnectModuleToElementTransaction", "Connect to Element"));
+		Blueprint->Modify();
+	}
+#endif 
+
+	const FRigElementKey ConnectorKey(*ConnectorName, ERigElementType::Connector);
+	FRigElementKey& TargetKey = Module->Connections.FindOrAdd(ConnectorKey);
+	TargetKey = InTargetKey;
+
+	Notify(EModularRigNotification::ConnectionChanged, Module);
+
+#if WITH_EDITOR
+	TransactionPtr.Reset();
+#endif
+	
+	return true;
+}
+
+bool UModularRigController::SetConfigValueInModule(const FString& InModulePath, const FName& InVariableName, const FString& InValue, bool bSetupUndo)
+{
+	FRigModuleReference* Module = FindModule(InModulePath);
+	if (!Module)
+	{
+		UE_LOG(LogControlRig, Error, TEXT("Could not find module %s"), *InModulePath);
+		return false;
+	}
+
+	if (!Module->Class.IsValid())
+	{
+		UE_LOG(LogControlRig, Error, TEXT("Class defined in module %s is not valid"), *InModulePath);
+		return false;
+	}
+
+	const FProperty* Property = Module->Class->FindPropertyByName(InVariableName);
+	if (!Property)
+	{
+		UE_LOG(LogControlRig, Error, TEXT("Could not find variable %s in module %s"), *InVariableName.ToString(), *InModulePath);
+		return false;
+	}
+
+	if (Property->HasAllPropertyFlags(CPF_BlueprintReadOnly))
+	{
+		UE_LOG(LogControlRig, Error, TEXT("The target variable %s in module %s is read only"), *InVariableName.ToString(), *InModulePath);
+		return false;
+	}
+
+#if WITH_EDITOR
+	TArray<uint8, TAlignedHeapAllocator<16>> TempStorage;
+	TempStorage.AddZeroed(Property->GetSize());
+	uint8* TempMemory = TempStorage.GetData();
+	Property->InitializeValue(TempMemory);
+
+	if (!FBlueprintEditorUtils::PropertyValueFromString_Direct(Property, InValue, TempMemory))
+	{
+		UE_LOG(LogControlRig, Error, TEXT("Value %s for variable %s in module %s is not valid"), *InValue, *InVariableName.ToString(), *InModulePath);
+		return false;
+	}
+
+	TSharedPtr<FScopedTransaction> TransactionPtr;
+	if (bSetupUndo)
+	{
+		TransactionPtr = MakeShared<FScopedTransaction>(NSLOCTEXT("ModularRigController", "ConfigureModuleValueTransaction", "Configure Module Value"));
+		if(UBlueprint* Blueprint = Cast<UBlueprint>(GetOuter()))
+		{
+			Blueprint->Modify();
+		}
+	}
+#endif 
+
+	Module->ConfigValues.FindOrAdd(InVariableName) = InValue;
+
+	Notify(EModularRigNotification::ModuleConfigValueChanged, Module);
+
+#if WITH_EDITOR
+	TransactionPtr.Reset();
+#endif
+
+	return true;
+}
+
+TArray<FString> UModularRigController::GetPossibleBindings(const FString& InModulePath, const FName& InVariableName)
+{
+	TArray<FString> PossibleBindings;
+	const FRigModuleReference* Module = FindModule(InModulePath);
+	if (!Module)
+	{
+		return PossibleBindings;
+	}
+
+	if (!Module->Class.IsValid())
+	{
+		return PossibleBindings;
+	}
+
+	const FProperty* TargetProperty = Module->Class->FindPropertyByName(InVariableName);
+	if (!TargetProperty)
+	{
+		return PossibleBindings;
+	}
+
+	if (TargetProperty->HasAnyPropertyFlags(CPF_BlueprintReadOnly | CPF_DisableEditOnInstance))
+	{
+		return PossibleBindings;
+	}
+
+	// Add possible blueprint variables
+	if(UBlueprint* Blueprint = Cast<UBlueprint>(GetOuter()))
+	{
+		TArray<FRigVMExternalVariable> Variables = Blueprint->GeneratedClass->GetDefaultObject<UControlRig>()->GetExternalVariables();
+		for (const FRigVMExternalVariable& Variable : Variables)
+		{
+			FText ErrorMessage;
+			const FString VariableName = Variable.Name.ToString();
+			if (CanBindModuleVariable(InModulePath, InVariableName, VariableName, ErrorMessage))
+			{
+				PossibleBindings.Add(VariableName);
+			}
+		}
+	}
+
+	// Add possible module variables
+	const FString InvalidModulePrefix = InModulePath + UModularRig::NamespaceSeparator;
+	Model->ForEachModule([this, &PossibleBindings, InModulePath, InVariableName, InvalidModulePrefix](const FRigModuleReference* InModule) -> bool
+	{
+		const FString CurModulePath = InModule->GetPath();
+		if (InModulePath != CurModulePath && !CurModulePath.StartsWith(InvalidModulePrefix))
+		{
+			TArray<FRigVMExternalVariable> Variables = InModule->Class->GetDefaultObject<UControlRig>()->GetExternalVariables();
+			for (const FRigVMExternalVariable& Variable : Variables)
+			{
+				FText ErrorMessage;
+				const FString SourceVariablePath = FString::Printf(TEXT("%s:%s"), *CurModulePath, *Variable.Name.ToString());
+				if (CanBindModuleVariable(InModulePath, InVariableName, SourceVariablePath, ErrorMessage))
+				{
+					PossibleBindings.Add(SourceVariablePath);
+				}
+			}
+		}		
+		return true;
+	});
+
+	return PossibleBindings;
+}
+
+bool UModularRigController::CanBindModuleVariable(const FString& InModulePath, const FName& InVariableName, const FString& InSourcePath, FText& OutErrorMessage)
+{
+	FRigModuleReference* Module = FindModule(InModulePath);
+	if (!Module)
+	{
+		OutErrorMessage = FText::FromString(FString::Printf(TEXT("Could not find module %s"), *InModulePath));
+		return false;
+	}
+
+	if (!Module->Class.IsValid())
+	{
+		OutErrorMessage = FText::FromString(FString::Printf(TEXT("Class defined in module %s is not valid"), *InModulePath));
+		return false;
+	}
+
+	const FProperty* TargetProperty = Module->Class->FindPropertyByName(InVariableName);
+	if (!TargetProperty)
+	{
+		OutErrorMessage = FText::FromString(FString::Printf(TEXT("Could not find variable %s in module %s"), *InVariableName.ToString(), *InModulePath));
+		return false;
+	}
+
+	if (TargetProperty->HasAnyPropertyFlags(CPF_BlueprintReadOnly | CPF_DisableEditOnInstance))
+	{
+		OutErrorMessage = FText::FromString(FString::Printf(TEXT("The target variable %s in module %s is read only"), *InVariableName.ToString(), *InModulePath));
+		return false;
+	}
+
+	FString SourceModulePath, SourceVariableName = InSourcePath;
+	InSourcePath.Split(UModularRig::NamespaceSeparator, &SourceModulePath, &SourceVariableName, ESearchCase::CaseSensitive, ESearchDir::FromEnd);
+
+	FRigModuleReference* SourceModule = nullptr;
+	if (!SourceModulePath.IsEmpty())
+	{
+		SourceModule = FindModule(SourceModulePath);
+		if (!SourceModule)
+		{
+			OutErrorMessage = FText::FromString(FString::Printf(TEXT("Could not find source module %s"), *SourceModulePath));
+			return false;
+		}
+
+		if (SourceModulePath.StartsWith(InModulePath))
+		{
+			OutErrorMessage = FText::FromString(FString::Printf(TEXT("Cannot bind variable of module %s to a variable of module %s because the source module is a child of the target module"), *InModulePath, *SourceModulePath));
+			return false;
+		}
+	}
+
+	const FProperty* SourceProperty = nullptr;
+	if (SourceModule)
+	{
+		SourceProperty = SourceModule->Class->FindPropertyByName(*SourceVariableName);
+	}
+	else
+	{
+		if(const UBlueprint* Blueprint = Cast<UBlueprint>(GetOuter()))
+		{
+			SourceProperty = Blueprint->GeneratedClass->FindPropertyByName(*SourceVariableName);
+		}
+	}
+	if (!SourceProperty)
+	{
+		OutErrorMessage = FText::FromString(FString::Printf(TEXT("Could not find source variable %s"), *InSourcePath));
+		return false;
+	}
+
+	FString SourcePath = (SourceModulePath.IsEmpty()) ? SourceVariableName : FString::Printf(TEXT("%s:%s"), *SourceModulePath, *SourceVariableName);
+	if (!RigVMTypeUtils::AreCompatible(SourceProperty, TargetProperty))
+	{
+		FString TargetPath = FString::Printf(TEXT("%s.%s"), *InModulePath, *InVariableName.ToString());
+		OutErrorMessage = FText::FromString(FString::Printf(TEXT("Property %s of type %s and %s of type %s are not compatible"), *SourcePath, *SourceProperty->GetCPPType(), *TargetPath, *TargetProperty->GetCPPType()));
+		return false;
+	}
+
+	return true;
+}
+
+bool UModularRigController::BindModuleVariable(const FString& InModulePath, const FName& InVariableName, const FString& InSourcePath, bool bSetupUndo)
+{
+	FText ErrorMessage;
+	if (!CanBindModuleVariable(InModulePath, InVariableName, InSourcePath, ErrorMessage))
+	{
+		UE_LOG(LogControlRig, Error, TEXT("Could not bind module variable %s:%s : %s"), *InModulePath, *InVariableName.ToString(), *ErrorMessage.ToString());
+		return false;
+	}
+	
+	FRigModuleReference* Module = FindModule(InModulePath);
+	const FProperty* TargetProperty = Module->Class->FindPropertyByName(InVariableName);
+
+	FString SourceModulePath, SourceVariableName = InSourcePath;
+	InSourcePath.Split(UModularRig::NamespaceSeparator, &SourceModulePath, &SourceVariableName, ESearchCase::CaseSensitive, ESearchDir::FromEnd);
+
+	FRigModuleReference* SourceModule = nullptr;
+	if (!SourceModulePath.IsEmpty())
+	{
+		SourceModule = FindModule(SourceModulePath);
+	}
+
+	const FProperty* SourceProperty = nullptr;
+	if (SourceModule)
+	{
+		SourceProperty = SourceModule->Class->FindPropertyByName(*SourceVariableName);
+	}
+	else
+	{
+		if(const UBlueprint* Blueprint = Cast<UBlueprint>(GetOuter()))
+		{
+			SourceProperty = Blueprint->GeneratedClass->FindPropertyByName(*SourceVariableName);
+		}
+	}
+
+	FString SourcePath = (SourceModulePath.IsEmpty()) ? SourceVariableName : FString::Printf(TEXT("%s:%s"), *SourceModulePath, *SourceVariableName);
+
+#if WITH_EDITOR
+	TSharedPtr<FScopedTransaction> TransactionPtr;
+	if (bSetupUndo)
+	{
+		TransactionPtr = MakeShared<FScopedTransaction>(NSLOCTEXT("ModularRigController", "BindModuleVariableTransaction", "Bind Module Variable"));
+		if(UBlueprint* Blueprint = Cast<UBlueprint>(GetOuter()))
+		{
+			Blueprint->Modify();
+		}
+	}
+#endif
+
+	FString& SourceStr = Module->Bindings.FindOrAdd(InVariableName);
+	SourceStr = SourcePath;
+
+	Notify(EModularRigNotification::ModuleConfigValueChanged, Module);
+
+#if WITH_EDITOR
+	TransactionPtr.Reset();
+#endif
+
+	return true;
+}
+
+bool UModularRigController::UnBindModuleVariable(const FString& InModulePath, const FName& InVariableName, bool bSetupUndo)
+{
+	FRigModuleReference* Module = FindModule(InModulePath);
+	if (!Module)
+	{
+		UE_LOG(LogControlRig, Error, TEXT("Could not find module %s"), *InModulePath);
+		return false;
+	}
+
+	if (!Module->Bindings.Contains(InVariableName))
+	{
+		UE_LOG(LogControlRig, Error, TEXT("Variable %s in module %s is not bound"), *InVariableName.ToString(), *InModulePath);
+		return false;
+	}
+
+#if WITH_EDITOR
+	TSharedPtr<FScopedTransaction> TransactionPtr;
+	if (bSetupUndo)
+	{
+		TransactionPtr = MakeShared<FScopedTransaction>(NSLOCTEXT("ModularRigController", "BindModuleVariableTransaction", "Bind Module Variable"));
+		if(UBlueprint* Blueprint = Cast<UBlueprint>(GetOuter()))
+		{
+			Blueprint->Modify();
+		}
+	}
+#endif
+
+	Module->Bindings.Remove(InVariableName);
+
+	Notify(EModularRigNotification::ModuleConfigValueChanged, Module);
+
+#if WITH_EDITOR
+	TransactionPtr.Reset();
+#endif
+
+	return true;
+}
+
+bool UModularRigController::DeleteModule(const FString& InModulePath, bool bSetupUndo)
+{
+	FRigModuleReference* Module = FindModule(InModulePath);
+	if (!Module)
+	{
+		UE_LOG(LogControlRig, Error, TEXT("Could not find module %s"), *InModulePath);
+		return false;
+	}
+
+#if WITH_EDITOR
+	TSharedPtr<FScopedTransaction> TransactionPtr;
+	if (bSetupUndo)
+	{
+		TransactionPtr = MakeShared<FScopedTransaction>(NSLOCTEXT("ModularRigController", "RenameModuleTransaction", "Rename Module"));
+		if(UBlueprint* Blueprint = Cast<UBlueprint>(GetOuter()))
+		{
+			Blueprint->Modify();
+		}
+	}
+#endif
+
+	// Unparent children (add them to root)
+	for (FRigModuleReference* Child : Module->CachedChildren)
+	{
+		ReparentModule(Child->GetPath(), FString(), bSetupUndo);
+	}
+
+	Model->DeletedModules.Add(*Module);
+	Model->Modules.RemoveSingle(*Module);
+	Model->UpdateCachedChildren();
+
+	// Fix bindings
+	for (FRigModuleReference& Reference : Model->Modules)
+	{
+		Reference.Bindings = Reference.Bindings.FilterByPredicate([InModulePath](const TPair<FName, FString>& Binding)
+		{
+			FString ModulePath, VariableName = Binding.Value;
+			Binding.Value.Split(UModularRig::NamespaceSeparator, &ModulePath, &VariableName, ESearchCase::CaseSensitive, ESearchDir::FromEnd);
+			if (ModulePath == InModulePath)
+			{
+				return false;
+			}
+			return true;
+		});
+	}
+
+	Notify(EModularRigNotification::ModuleRemoved, &Model->DeletedModules.Last());
+
+	Model->DeletedModules.Reset();
+
+#if WITH_EDITOR
+	TransactionPtr.Reset();
+#endif
+	return false;
+}
+
+bool UModularRigController::RenameModule(const FString& InModulePath, const FName& InNewName, bool bSetupUndo)
+{
+	FRigModuleReference* Module = FindModule(InModulePath);
+	if (!Module)
+	{
+		UE_LOG(LogControlRig, Error, TEXT("Could not find module %s"), *InModulePath);
+		return false;
+	}
+
+	FText ErrorMessage;
+	if (!CanRenameModule(InModulePath, InNewName, ErrorMessage))
+	{
+		UE_LOG(LogControlRig, Error, TEXT("Could not rename module %s: %s"), *InModulePath, *ErrorMessage.ToString());
+		return false;
+	}
+
+	const FString OldName = Module->Name.ToString();
+	const FString NewName = InNewName.ToString();
+	if (OldName.Equals(NewName))
+	{
+		return true;
+	}
+	
+#if WITH_EDITOR
+	TSharedPtr<FScopedTransaction> TransactionPtr;
+	if (bSetupUndo)
+	{
+		TransactionPtr = MakeShared<FScopedTransaction>(NSLOCTEXT("ModularRigController", "RenameModuleTransaction", "Rename Module"));
+		if(UBlueprint* Blueprint = Cast<UBlueprint>(GetOuter()))
+		{
+			Blueprint->Modify();
+		}
+	}
+#endif
+	
+	const FString OldPath = (Module->ParentPath.IsEmpty()) ? OldName : FString::Printf(TEXT("%s:%s"), *Module->ParentPath, *OldName);
+	const FString NewPath = (Module->ParentPath.IsEmpty()) ? *NewName :  FString::Printf(TEXT("%s:%s"), *Module->ParentPath, *NewName);
+	Module->PreviousName = Module->Name;
+	Module->Name = InNewName;
+	TArray<FRigModuleReference*> Children;
+	Children.Append(Module->CachedChildren);
+	for (int32 i=0; i<Children.Num(); ++i)
+	{
+		FRigModuleReference* Child = Children[i];
+		Child->ParentPath.ReplaceInline(*OldPath, *NewPath);
+
+		Children.Append(Child->CachedChildren);
+	}
+
+	// Fix bindings
+	for (FRigModuleReference& Reference : Model->Modules)
+	{
+		for (TPair<FName, FString>& Binding : Reference.Bindings)
+		{
+			FString ModulePath, VariableName = Binding.Value;
+			Binding.Value.Split(UModularRig::NamespaceSeparator, &ModulePath, &VariableName, ESearchCase::CaseSensitive, ESearchDir::FromEnd);
+			if (ModulePath == OldPath)
+			{
+				Binding.Value = FString::Printf(TEXT("%s:%s"), *NewPath, *VariableName);
+			}
+		};
+	}
+
+	Notify(EModularRigNotification::ModuleRenamed, Module);
+
+#if WITH_EDITOR
+	TransactionPtr.Reset();
+#endif
+	
+	return true;
+}
+
+bool UModularRigController::CanRenameModule(const FString& InModulePath, const FName& InNewName, FText& OutErrorMessage)
+{
+	if (InNewName.IsNone() || InNewName.ToString().IsEmpty())
+	{
+		OutErrorMessage = FText::FromString(TEXT("Name is empty."));
+		return false;
+	}
+
+	FRigModuleReference* Module = FindModule(InModulePath);
+	if (!Module)
+	{
+		OutErrorMessage = FText::FromString(FString::Printf(TEXT("Module %s not found."), *InModulePath));
+		return false;
+	}
+
+	const FString NewNameStr = InNewName.ToString();
+	const FString NewPath = (Module->ParentPath.IsEmpty()) ? NewNameStr : FString::Printf(TEXT("%s:%s"), *Module->ParentPath, *NewNameStr);
+	if (InModulePath == NewPath)
+	{
+		return true;
+	}
+
+	// Check there are no siblings with the same name
+	{
+		FRigModuleReference* Parent = FindModule(Module->ParentPath);
+		TArray<FRigModuleReference*>* Siblings = &Model->RootModules;
+		if (Parent)
+		{
+			Siblings = &Parent->CachedChildren;
+		}
+		for (FRigModuleReference* Sibling : *Siblings)
+		{
+			if (Sibling->Name.ToString() == NewNameStr)
+			{
+				OutErrorMessage = FText::FromString(FString::Printf(TEXT("Sibling with path %s already exists"), *Sibling->GetPath()));
+				return false;
+			}
+		}
+	}
+	return true;
+}
+
+bool UModularRigController::ReparentModule(const FString& InModulePath, const FString& InNewParentModulePath, bool bSetupUndo)
+{
+	FRigModuleReference* Module = FindModule(InModulePath);
+	if (!Module)
+	{
+		UE_LOG(LogControlRig, Error, TEXT("Could not find module %s"), *InModulePath);
+		return false;
+	}
+
+#if WITH_EDITOR
+	TSharedPtr<FScopedTransaction> TransactionPtr;
+	if (bSetupUndo)
+	{
+		TransactionPtr = MakeShared<FScopedTransaction>(NSLOCTEXT("ModularRigController", "RenameModuleTransaction", "Rename Module"));
+		if(UBlueprint* Blueprint = Cast<UBlueprint>(GetOuter()))
+		{
+			Blueprint->Modify();
+		}
+	}
+#endif
+
+	// Reparent or unparent children
+	FRigModuleReference* NewParentModule = FindModule(InNewParentModulePath);
+	const FString OldPath = Module->GetPath();
+	Module->PreviousParentPath = Module->ParentPath;
+	Module->ParentPath = (NewParentModule) ? NewParentModule->GetPath() : FString();
+	Module->Name = GetSafeNewName(Module->GetPath());
+	const FString NewPath = Module->GetPath();
+
+	// Fix all the subtree namespaces
+	TArray<FRigModuleReference*> SubTree = Module->CachedChildren;
+	for (int32 Index=0; Index<SubTree.Num(); ++Index)
+	{
+		SubTree[Index]->ParentPath.ReplaceInline(*OldPath, *NewPath);
+		SubTree.Append(SubTree[Index]->CachedChildren);
+	}
+
+
+	Model->UpdateCachedChildren();
+
+	// Fix bindings
+	for (FRigModuleReference& Reference : Model->Modules)
+	{
+		const FString ReferencePath = Reference.GetPath();
+		for (TPair<FName, FString>& Binding : Reference.Bindings)
+		{
+			FString ModulePath, VariableName = Binding.Value;
+			Binding.Value.Split(UModularRig::NamespaceSeparator, &ModulePath, &VariableName, ESearchCase::CaseSensitive, ESearchDir::FromEnd);
+			if (ModulePath == OldPath)
+			{
+				Binding.Value = FString::Printf(TEXT("%s:%s"), *NewPath, *VariableName);
+				ModulePath = NewPath;
+			}
+
+			// Remove any child dependency
+			if (ModulePath.Contains(ReferencePath))
+			{
+				UE_LOG(LogControlRig, Warning, TEXT("Binding lost due to source %s contained in child module of %s"), *Binding.Value, *ReferencePath);
+				Binding.Value.Reset();
+			}
+		};
+
+		Reference.Bindings = Reference.Bindings.FilterByPredicate([](const TPair<FName, FString>& Binding)
+		{
+			return !Binding.Value.IsEmpty();
+		});
+	}
+
+	Notify(EModularRigNotification::ModuleReparented, Module);
+	
+	return false;
+}
+
+FName UModularRigController::GetSafeNewName(const FString& InModuleDesiredPath)
+{
+	FString ParentPath, DesiredName = InModuleDesiredPath;
+	InModuleDesiredPath.Split(UModularRig::NamespaceSeparator, &ParentPath, &DesiredName, ESearchCase::CaseSensitive, ESearchDir::FromEnd);
+
+	TArray<FRigModuleReference*>* Children = &Model->RootModules;
+	if (!ParentPath.IsEmpty())
+	{
+		if (FRigModuleReference* Parent = FindModule(ParentPath))
+		{
+			Children = &Parent->CachedChildren;
+		}
+	}
+
+	bool bSafeToUse = false;
+	FString NewName = DesiredName;
+	int32 Index = 0;
+	while (!bSafeToUse)
+	{
+		bSafeToUse = true;
+		for (FRigModuleReference* Child : *Children)
+		{
+			if (Child->Name == *NewName)
+			{
+				bSafeToUse = false;
+				NewName = FString::Printf(TEXT("%s%d"), *DesiredName, ++Index);
+				break;
+			}
+		}
+	}
+	return *NewName;
+}
+
+void UModularRigController::Notify(const EModularRigNotification& InNotification, const FRigModuleReference* InElement)
+{
+	if(!bSuspendNotifications)
+	{
+		ModifiedEvent.Broadcast(InNotification, InElement);
+	}
+}
diff --git a/Engine/Plugins/Animation/ControlRig/Source/ControlRig/Private/ModularRigModel.cpp b/Engine/Plugins/Animation/ControlRig/Source/ControlRig/Private/ModularRigModel.cpp
new file mode 100644
--- /dev/null
+++ b/Engine/Plugins/Animation/ControlRig/Source/ControlRig/Private/ModularRigModel.cpp
@@ -0,0 +1,134 @@
+// Copyright Epic Games, Inc. All Rights Reserved.
+
+#include "ModularRigModel.h"
+#include "ModularRigController.h"
+#include "ModularRig.h"
+
+#include UE_INLINE_GENERATED_CPP_BY_NAME(ModularRigModel)
+
+FString FRigModuleReference::GetPath() const
+{
+	if (ParentPath.IsEmpty())
+	{
+		return Name.ToString();
+	}
+	return ParentPath + UModularRig::NamespaceSeparator + Name.ToString();
+}
+
+FString FRigModuleReference::GetNamespace() const
+{
+	return GetPath() + UModularRig::NamespaceSeparator;
+}
+
+UModularRigController* FModularRigModel::GetController(bool bCreateIfNeeded)
+{
+	if (bCreateIfNeeded && Controller == nullptr)
+	{
+		const FName SafeControllerName = *FString::Printf(TEXT("%s_ModularRig_Controller"), *GetOuter()->GetPathName());
+		UModularRigController* NewController = NewObject<UModularRigController>(GetOuter(), UModularRigController::StaticClass(), SafeControllerName);
+		NewController->SetModel(this);
+		Controller = NewController;
+	}
+	return Cast<UModularRigController>(Controller);
+}
+
+void FModularRigModel::SetOuterClientHost(UObject* InOuterClientHost)
+{
+	OuterClientHost = InOuterClientHost;
+}
+
+void FModularRigModel::UpdateCachedChildren()
+{
+	TMap<FString, FRigModuleReference*> PathToModule;
+	for (FRigModuleReference& Module : Modules)
+	{
+		Module.CachedChildren.Reset();
+		PathToModule.Add(Module.GetPath(), &Module);
+	}
+	
+	RootModules.Reset();
+	for (FRigModuleReference& Module : Modules)
+	{
+		if (Module.ParentPath.IsEmpty())
+		{
+			RootModules.Add(&Module);
+		}
+		else
+		{
+			if (FRigModuleReference** ParentModule = PathToModule.Find(Module.ParentPath))
+			{
+				(*ParentModule)->CachedChildren.Add(&Module);
+			}
+		}
+	}
+}
+
+FRigModuleReference* FModularRigModel::FindModule(const FString InPath) const
+{
+	const TArray<FRigModuleReference*>* Children = &RootModules;
+
+	FString Left = InPath, Right;
+	while (Left.Split(UModularRig::NamespaceSeparator, &Left, &Right))
+	{
+		FRigModuleReference* const * Child = Children->FindByPredicate([Left](FRigModuleReference* Module)
+		{
+			return Module->Name.ToString() == Left;
+		});
+		
+		if (!Child)
+		{
+			return nullptr;
+		}
+		
+		Children = &(*Child)->CachedChildren;
+		Left = Right;
+	}
+
+	FRigModuleReference* const * Child = Children->FindByPredicate([Left](FRigModuleReference* Module)
+		{
+			return Module->Name.ToString() == Left;
+		});
+
+	if (!Child)
+	{
+		return nullptr;
+	}
+	return *Child;
+}
+
+FString FModularRigModel::FindParentPath(const FString InPath) const
+{
+	if (FRigModuleReference* Element = FindModule(InPath))
+	{
+		return Element->ParentPath;
+	}
+	return FString();
+}
+
+void FModularRigModel::ForEachModule(TFunction<bool(const FRigModuleReference*)> PerModule) const
+{
+	TArray<FRigModuleReference*> ModuleInstances = RootModules;
+	for (int32 Index=0; Index < ModuleInstances.Num(); ++Index)
+	{
+		if (!PerModule(ModuleInstances[Index]))
+		{
+			break;
+		}
+		ModuleInstances.Append(ModuleInstances[Index]->CachedChildren);
+	}
+}
+
+TArray<FString> FModularRigModel::SortPaths(const TArray<FString>& InPaths) const
+{
+	TArray<FString> Result;
+	ForEachModule([InPaths, &Result](const FRigModuleReference* Element) -> bool
+	{
+		const FString& Path = Element->GetPath();
+		if(InPaths.Contains(Path))
+		{
+			Result.AddUnique(Path);
+		}
+		return true;
+	});
+	return Result;
+}
diff --git a/Engine/Plugins/Animation/ControlRig/Source/ControlRig/Public/ControlRig.h b/Engine/Plugins/Animation/ControlRig/Source/ControlRig/Public/ControlRig.h
--- a/Engine/Plugins/Animation/ControlRig/Source/ControlRig/Public/ControlRig.h
+++ b/Engine/Plugins/Animation/ControlRig/Source/ControlRig/Public/ControlRig.h
@@ -47,7 +47,7 @@ CONTROLRIG_API DECLARE_LOG_CATEGORY_EXTERN(LogControlRig, Log, All);
 
 /** Runs logic for mapping input data to transforms (the "Rig") */
 UCLASS(Blueprintable, Abstract, editinlinenew)
-class CONTROLRIG_API UControlRig : public URigVMHost, public INodeMappingProviderInterface
+class CONTROLRIG_API UControlRig : public URigVMHost, public INodeMappingProviderInterface, public IRigHierarchyProvider
 {
 	GENERATED_UCLASS_BODY()
 
@@ -169,7 +169,7 @@ public:
 		return DynamicHierarchy;
 	}
 	
-	URigHierarchy* GetHierarchy() const
+	virtual URigHierarchy* GetHierarchy() const override
 	{
 		return DynamicHierarchy;
 	}
@@ -832,7 +832,6 @@ private:
 	UPROPERTY()
 	TSubclassOf<UControlRig> InteractionRigClass_DEPRECATED;
 #endif
-	
 
 	friend class FControlRigBlueprintCompilerContext;
 	friend struct FRigHierarchyRef;
diff --git a/Engine/Plugins/Animation/ControlRig/Source/ControlRig/Public/ModularRig.h b/Engine/Plugins/Animation/ControlRig/Source/ControlRig/Public/ModularRig.h
--- a/Engine/Plugins/Animation/ControlRig/Source/ControlRig/Public/ModularRig.h
+++ b/Engine/Plugins/Animation/ControlRig/Source/ControlRig/Public/ModularRig.h
@@ -2,7 +2,7 @@
 
 #pragma once
 #include "ControlRig.h"
-
+#include "ModularRigModel.h"
 #include "ModularRig.generated.h"
 
 struct FRigModuleInstance;
@@ -16,8 +16,8 @@ public:
 		, Path()
 	{}
 
-	FModuleInstanceHandle(UModularRig* InModularRig, const FString& InPath);
-	FModuleInstanceHandle(UModularRig* InModularRig, const FRigModuleInstance* InElement);
+	FModuleInstanceHandle(const UModularRig* InModularRig, const FString& InPath);
+	FModuleInstanceHandle(const UModularRig* InModularRig, const FRigModuleInstance* InElement);
 
 	bool IsValid() const { return Get() != nullptr; }
 	operator bool() const { return IsValid(); }
@@ -41,9 +41,11 @@ struct CONTROLRIG_API FRigModuleInstance
 {
 	GENERATED_USTRUCT_BODY()
 
+public:
+	
 	FRigModuleInstance()
 	: Name(NAME_None)
-	, Rig(nullptr)
+	, RigPtr(nullptr)
 	, ParentPath(FString())
 	{
 	}
@@ -51,9 +53,13 @@ struct CONTROLRIG_API FRigModuleInstance
 	UPROPERTY()
 	FName Name;
 
-	UPROPERTY()
-	TSoftObjectPtr<UControlRig> Rig;
+private:
 
+	UPROPERTY(transient)
+	mutable UControlRig* RigPtr;
+
+public:
+	
 	UPROPERTY()
 	FString ParentPath;
 
@@ -64,6 +70,9 @@ struct CONTROLRIG_API FRigModuleInstance
 
 	FString GetPath() const;
 	FString GetNamespace() const;
+	UControlRig* GetRig() const;
+	void SetRig(UControlRig* InRig);
+	bool ContainsRig(const UControlRig* InRig) const;
 };
 
 USTRUCT(BlueprintType)
@@ -111,7 +120,6 @@ class CONTROLRIG_API UModularRig : public UControlRig
 public:
 
 	// BEGIN ControlRig
-	virtual void Initialize(bool bRequestInit) override;
 	virtual void InitializeVMs(bool bRequestInit = true) override;
 	virtual bool InitializeVMs(const FName& InEventName) override;
 	virtual void InitializeVMsFromCDO() override { URigVMHost::InitializeFromCDO(); }
@@ -124,6 +132,9 @@ public:
 	virtual const TArray<FName>& GetSupportedEvents() const override { return SupportedEvents; }
 	// END ControlRig
 
+	UPROPERTY()
+	FModularRigModel ModularRigModel;
+
 	UPROPERTY()
 	TArray<FRigModuleExecutionElement> ExecutionQueue;
 	int32 ExecutionQueueFront = 0;
@@ -138,13 +149,14 @@ public:
 
 	void ResetModules();
 
+	const FModularRigModel& GetModularRigModel() const;
+	void UpdateModuleHierarchyFromCDO();
 	void UpdateCachedChildren();
 	void UpdateSupportedEvents();
 
 	/** Adds a module to the rig*/
 	bool AddModuleInstance(const FName& InModuleName, TSubclassOf<UControlRig> InModuleClass, FString InParentPath, const TMap<FRigElementKey, FRigElementKey>& InConnectionMap, const TMap<FName, FString>& InVariableDefaultValues, const TMap<FName, FString>& InVariableBindings);
 	FRigModuleInstance* AddModuleInstance(const FName& InModuleName, TSubclassOf<UControlRig> InModuleClass, FRigModuleInstance* InParent, const TMap<FRigElementKey, FRigElementKey>& InConnectionMap, const TMap<FName, FString>& InVariableDefaultValues, const TMap<FName, FString>& InVariableBindings);
-	FRigModuleInstance* AddModuleInstance(const FRigModuleInstance* InOtherModule);
 
 	const FRigModuleInstance* FindModule(const FString& InPath) const;
 	const FRigModuleInstance* FindModule(const UControlRig* InModuleInstance) const;
@@ -154,8 +166,8 @@ public:
 	void ForEachModule(TFunctionRef<bool(const FRigModuleInstance*)> PerModuleFunction) const;
 
 	/**
-	 * Returns a handle to an existing element
-	 * @param InKey The key of the handle to retrieve.
+	 * Returns a handle to an existing module
+	 * @param InPath The path of the module to retrieve a handle for.
 	 * @return The retrieved handle (may be invalid)
 	 */
 	FModuleInstanceHandle GetHandle(const FString& InPath) const
diff --git a/Engine/Plugins/Animation/ControlRig/Source/ControlRig/Public/ModularRigController.h b/Engine/Plugins/Animation/ControlRig/Source/ControlRig/Public/ModularRigController.h
new file mode 100644
--- /dev/null
+++ b/Engine/Plugins/Animation/ControlRig/Source/ControlRig/Public/ModularRigController.h
@@ -0,0 +1,62 @@
+// Copyright Epic Games, Inc. All Rights Reserved.
+
+#pragma once
+
+#include "ModularRigModel.h"
+#include "ModularRigController.generated.h"
+
+struct FRigModuleReference;
+
+DECLARE_MULTICAST_DELEGATE_TwoParams(FModularRigModifiedEvent, EModularRigNotification /* type */, const FRigModuleReference* /* element */);
+
+UCLASS(BlueprintType)
+class CONTROLRIG_API UModularRigController : public UObject
+{
+	GENERATED_UCLASS_BODY()
+
+	UModularRigController()
+		: Model(nullptr)
+		, bSuspendNotifications(false)
+	{
+	}
+
+	FModularRigModel* Model;
+	FModularRigModifiedEvent ModifiedEvent;
+	bool bSuspendNotifications;
+
+	UFUNCTION(BlueprintCallable, Category = "Control Rig | Modules")
+	FString AddModule(const FName& InModuleName, TSubclassOf<UControlRig> InClass, const FString& InParentModulePath, bool bSetupUndo = true);
+
+	UFUNCTION(BlueprintCallable, Category = "Control Rig | Modules")
+	bool CanConnectConnectorToElement(const FRigModuleConnector& InConnector, const FRigElementKey& InTargetKey, FText& OutErrorMessage);
+
+	UFUNCTION(BlueprintCallable, Category = "Control Rig | Modules")
+	bool ConnectConnectorToElement(const FRigElementKey& InConnectorKey, const FRigElementKey& InTargetKey, bool bSetupUndo = true);
+
+	UFUNCTION(BlueprintCallable, Category = "Control Rig | Modules")
+	bool SetConfigValueInModule(const FString& InModulePath, const FName& InVariableName, const FString& InValue, bool bSetupUndo = true);
+
+	UFUNCTION(BlueprintCallable, Category = "ControlRig | Modules")
+	bool BindModuleVariable(const FString& InModulePath, const FName& InVariableName, const FString& InSourcePath, bool bSetupUndo = true);
+	bool CanBindModuleVariable(const FString& InModulePath, const FName& InVariableName, const FString& InSourcePath, FText& OutErrorMessage);
+	TArray<FString> GetPossibleBindings(const FString& InModulePath, const FName& InVariableName);
+
+	UFUNCTION(BlueprintCallable, Category = "ControlRig | Modules")
+	bool UnBindModuleVariable(const FString& InModulePath, const FName& InVariableName, bool bSetupUndo = true);
+
+	UFUNCTION(BlueprintCallable, Category = "ControlRig | Modules")
+	bool DeleteModule(const FString& InModulePath, bool bSetupUndo = true);
+
+	UFUNCTION(BlueprintCallable, Category = "ControlRig | Modules")
+	bool RenameModule(const FString& InModulePath, const FName& InNewName, bool bSetupUndo = true);
+	bool CanRenameModule(const FString& InModulePath, const FName& InNewName, FText& OutErrorMessage);
+
+	UFUNCTION(BlueprintCallable, Category = "ControlRig | Modules")
+	bool ReparentModule(const FString& InModulePath, const FString& InNewParentModulePath, bool bSetupUndo = true);
+	
+	FName GetSafeNewName(const FString& InModuleDesiredPath);
+	void SetModel(FModularRigModel* InModel) { Model = InModel; }
+	FRigModuleReference* FindModule(const FString& InPath);
+	FModularRigModifiedEvent& OnModified() { return ModifiedEvent; }
+	void Notify(const EModularRigNotification& InNotification, const FRigModuleReference* InElement);
+};
\ No newline at end of file
diff --git a/Engine/Plugins/Animation/ControlRig/Source/ControlRig/Public/ModularRigModel.h b/Engine/Plugins/Animation/ControlRig/Source/ControlRig/Public/ModularRigModel.h
new file mode 100644
--- /dev/null
+++ b/Engine/Plugins/Animation/ControlRig/Source/ControlRig/Public/ModularRigModel.h
@@ -0,0 +1,122 @@
+// Copyright Epic Games, Inc. All Rights Reserved.
+
+#pragma once
+
+#include "ControlRig.h"
+#include "ModularRigModel.generated.h"
+
+class UModularRigController;
+
+UENUM()
+enum class EModularRigNotification : uint8
+{
+	ModuleAdded,
+
+	ModuleRenamed,
+
+	ModuleRemoved,
+
+	ModuleReparented,
+
+	ConnectionChanged,
+
+	ModuleConfigValueChanged,
+
+	/** MAX - invalid */
+	Max UMETA(Hidden),
+};
+
+USTRUCT(BlueprintType)
+struct CONTROLRIG_API FRigModuleReference
+{
+	GENERATED_BODY()
+
+	FRigModuleReference()
+		: Name(NAME_None)
+		, ParentPath(FString())
+		, Class(nullptr)
+	{}
+	
+	FRigModuleReference(const FName& InName, TSubclassOf<UControlRig> InClass, const FString& InParentPath)
+		: Name(InName)
+		, ParentPath(InParentPath)
+		, Class(InClass)
+	{}
+
+	UPROPERTY()
+	FName Name;
+
+	UPROPERTY()
+	FString ParentPath;
+
+	UPROPERTY()
+	TSoftClassPtr<UControlRig> Class;
+
+	UPROPERTY()
+	TMap<FRigElementKey, FRigElementKey> Connections; // Connectors to Connection element
+
+	UPROPERTY()
+	TMap<FName, FString> ConfigValues;
+
+	UPROPERTY()
+	TMap<FName, FString> Bindings; // ExternalVariableName (current module) -> SourceExternalVariableNamespacedPath (root rig or other module)
+
+	UPROPERTY(transient)
+	FName PreviousName;
+
+	UPROPERTY(transient)
+	FString PreviousParentPath;
+
+	TArray<FRigModuleReference*> CachedChildren;
+
+	FString GetPath() const;
+
+	FString GetNamespace() const;
+	
+	friend bool operator==(const FRigModuleReference& A, const FRigModuleReference& B)
+	{
+		return A.ParentPath == B.ParentPath &&
+			A.Name == B.Name;
+	}
+	
+	friend class UModularRigController;
+};
+
+// A management struct containing all modules in the rig
+USTRUCT(BlueprintType)
+struct CONTROLRIG_API FModularRigModel
+{
+public:
+
+	GENERATED_BODY()
+
+	UPROPERTY()
+	TArray<FRigModuleReference> Modules;
+	TArray<FRigModuleReference*> RootModules;
+	TArray<FRigModuleReference> DeletedModules;
+
+	UPROPERTY(transient)
+	TObjectPtr<UObject> Controller;
+
+	UModularRigController* GetController(bool bCreateIfNeeded = true);
+
+	UObject* GetOuter() const { return OuterClientHost.IsValid() ? OuterClientHost.Get() : nullptr; }
+
+	void SetOuterClientHost(UObject* InOuterClientHost);
+
+	void UpdateCachedChildren();
+
+	FRigModuleReference* FindModule(const FString InPath) const;
+
+	FString FindParentPath(const FString InPath) const;
+
+	void ForEachModule(TFunction<bool(const FRigModuleReference*)> PerModule) const;
+
+	TArray<FString> SortPaths(const TArray<FString>& InPaths) const;
+
+private:
+	TWeakObjectPtr<UObject> OuterClientHost;
+
+	friend class UModularRigController;
+	friend struct FRigModuleReference;
+};
\ No newline at end of file
diff --git a/Engine/Plugins/Animation/ControlRig/Source/ControlRig/Public/Rigs/RigHierarchy.h b/Engine/Plugins/Animation/ControlRig/Source/ControlRig/Public/Rigs/RigHierarchy.h
--- a/Engine/Plugins/Animation/ControlRig/Source/ControlRig/Public/Rigs/RigHierarchy.h
+++ b/Engine/Plugins/Animation/ControlRig/Source/ControlRig/Public/Rigs/RigHierarchy.h
@@ -4941,4 +4941,19 @@ private:
 	TArray<bool> CurrentFlags;
 };
 
-#endif
\ No newline at end of file
+#endif
+
+UINTERFACE()
+class CONTROLRIG_API URigHierarchyProvider : public UInterface
+{
+	GENERATED_BODY()
+};
+
+class CONTROLRIG_API IRigHierarchyProvider
+{
+	GENERATED_BODY()
+
+public:
+
+	virtual URigHierarchy* GetHierarchy() const = 0;
+};
diff --git a/Engine/Plugins/Animation/ControlRig/Source/ControlRigDeveloper/Private/ControlRigBlueprint.cpp b/Engine/Plugins/Animation/ControlRig/Source/ControlRigDeveloper/Private/ControlRigBlueprint.cpp
--- a/Engine/Plugins/Animation/ControlRig/Source/ControlRigDeveloper/Private/ControlRigBlueprint.cpp
+++ b/Engine/Plugins/Animation/ControlRig/Source/ControlRigDeveloper/Private/ControlRigBlueprint.cpp
@@ -15,6 +15,7 @@
 #include "ControlRigObjectVersion.h"
 #include "BlueprintCompilationManager.h"
 #include "ModularRig.h"
+#include "ModularRigController.h"
 #include "RigVMCompiler/RigVMCompiler.h"
 #include "RigVMCore/RigVMRegistry.h"
 #include "Units/Execution/RigUnit_BeginExecution.h"
@@ -1117,6 +1118,8 @@ void UControlRigBlueprint::PostLoad()
 			}
 		}
 	}
+
+	UpdateModularDependencyDelegates();
 }
 
 #if WITH_EDITOR
@@ -1151,7 +1154,17 @@ void UControlRigBlueprint::HandlePackageDone()
 	PropagateHierarchyFromBPToInstances();
 
 	Super::HandlePackageDone();
-	RecompileModularRig();
+	
+	if(IsModularRig())
+	{
+		// force load all dependencies
+		ModularRigModel.ForEachModule([](const FRigModuleReference* Element) -> bool
+		{
+			(void)Element->Class.LoadSynchronous();
+			return true;
+		});
+		RecompileModularRig();
+	}
 }
 
 void UControlRigBlueprint::HandleConfigureRigVMController(const FRigVMClient* InClient, URigVMController* InControllerToConfigure)
@@ -1538,45 +1551,25 @@ UModularRigController* UControlRigBlueprint::GetModularRigController()
 	return ModularRigModel.GetController();
 }
 
-static void AddModulesRecursively(UModularRig* Rig, const FRigModuleReference* InModule, FRigModuleInstance* InParent)
-{
-	// Make sure the inner is loaded and valid
-	{
-		if (!InModule->Class.IsValid())
-		{
-			InModule->Class.LoadSynchronous();
-		}
-		if (!InModule->Class.IsValid())
-		{
-			return;
-		}
-	}
-	
-	FRigModuleInstance* NewModule = Rig->AddModuleInstance(InModule->Name, InModule->Class.Get(), InParent, InModule->Connections, InModule->ConfigValues, InModule->Bindings);
-	
-	for (const FRigModuleReference* ChildModule : InModule->CachedChildren)
-	{
-		AddModulesRecursively(Rig, ChildModule, NewModule);
-	}
-}
-
 void UControlRigBlueprint::RecompileModularRig()
 {
 	OnModularRigPreCompiled().Broadcast(this);
-	if (UClass* MyControlRigClass = GeneratedClass)
+	if (const UClass* MyControlRigClass = GeneratedClass)
 	{
 		if (UModularRig* DefaultObject = Cast<UModularRig>(MyControlRigClass->GetDefaultObject(false)))
 		{
 			DefaultObject->ResetModules();
-			for (const FRigModuleReference* RootModule : ModularRigModel.RootModules)
-			{
-				AddModulesRecursively(DefaultObject, RootModule, nullptr);
-			}
-			DefaultObject->InitializeVMs(true);
+
+			// copy the model over to the CDO.
+			// non-CDO instances are going to instantiate the model into a
+			// UObject module instance tree. CDO's are data only to avoid bugs / 
+			// behaviors in the blueprint re-instancer - which is disregarding any
+			// object under a CDO.
+			DefaultObject->ModularRigModel = ModularRigModel;
 			PropagateModuleHierarchyFromBPToInstances();
 		}
 	}
-
+	UpdateModularDependencyDelegates();
 	OnModularRigCompiled().Broadcast(this);
 }
 
@@ -2055,7 +2048,10 @@ void UControlRigBlueprint::PropagatePoseFromBPToInstances() const
 				if (UControlRig* InstanceRig = Cast<UControlRig>(ArchetypeInstance))
 				{
 					InstanceRig->PostInitInstanceIfRequired();
-					InstanceRig->GetHierarchy()->CopyPose(Hierarchy, true, true, true);
+					if(!InstanceRig->IsRigModuleInstance())
+					{
+						InstanceRig->GetHierarchy()->CopyPose(Hierarchy, true, true, true);
+					}
 				}
 			}
 		}
@@ -2174,7 +2170,7 @@ void UControlRigBlueprint::PropagatePropertyFromInstanceToBP(FRigElementKey InRi
 
 void UControlRigBlueprint::PropagateModuleHierarchyFromBPToInstances() const
 {
-	if (UClass* MyControlRigClass = GeneratedClass)
+	if (const UClass* MyControlRigClass = GeneratedClass)
 	{
 		if (UControlRig* DefaultObject = Cast<UControlRig>(MyControlRigClass->GetDefaultObject(false)))
 		{
@@ -2192,6 +2188,60 @@ void UControlRigBlueprint::PropagateModuleHierarchyFromBPToInstances() const
 	}
 }
 
+void UControlRigBlueprint::UpdateModularDependencyDelegates()
+{
+	TArray<const UBlueprint*> VisitList;
+	ModularRigModel.ForEachModule([&VisitList, this](const FRigModuleReference* Element) -> bool
+	{
+		if(const UClass* Class = Element->Class.Get())
+		{
+			if(UControlRigBlueprint* Blueprint = Cast<UControlRigBlueprint>(Class->ClassGeneratedBy))
+			{
+				if(!VisitList.Contains(Blueprint))
+				{
+					Blueprint->OnVMCompiled().RemoveAll(this);
+					Blueprint->OnVMCompiled().AddUObject(this, &UControlRigBlueprint::OnModularDependencyVMCompiled);
+					Blueprint->OnModularRigCompiled().AddUObject(this, &UControlRigBlueprint::OnModularDependencyChanged);
+					VisitList.Add(Blueprint);
+				}
+			}
+		}
+		return true;
+	});
+}
+
+void UControlRigBlueprint::OnModularDependencyVMCompiled(UObject* InBlueprint, URigVM* InVM, FRigVMExtendedExecuteContext& InExecuteContext)
+{
+	if(URigVMBlueprint* RigVMBlueprint = Cast<URigVMBlueprint>(InBlueprint))
+	{
+		OnModularDependencyChanged(RigVMBlueprint);
+	}
+}
+
+void UControlRigBlueprint::OnModularDependencyChanged(URigVMBlueprint* InBlueprint)
+{
+	// the rig will perform initialize itself - but we should request construction
+	check(IsModularRig());
+	
+	const URigVMBlueprintGeneratedClass* RigClass = GetRigVMBlueprintGeneratedClass();
+	check(RigClass);
+	
+	UControlRig* CDO = Cast<UControlRig>(RigClass->GetDefaultObject(true /* create if needed */));
+
+	TArray<UObject*> ArchetypeInstances;
+	CDO->GetArchetypeInstances(ArchetypeInstances);
+
+	// visit all or our instances and request construction
+	for (UObject* Instance : ArchetypeInstances)
+	{
+		if (UModularRig* InstanceRig = Cast<UModularRig>(Instance))
+		{
+			InstanceRig->RequestConstruction();
+		}
+	}
+
+}
+
 void UControlRigBlueprint::HandleHierarchyModified(ERigHierarchyNotification InNotification, URigHierarchy* InHierarchy, const FRigBaseElement* InElement)
 {
 #if WITH_EDITOR
@@ -2312,6 +2362,7 @@ void UControlRigBlueprint::HandleRigModulesModified(EModularRigNotification InNo
 						PropagateHierarchyFromBPToInstances();
 					}
 				}
+				UpdateModularDependencyDelegates();
 			}
 			break;
 		}
@@ -2422,6 +2473,7 @@ void UControlRigBlueprint::HandleRigModulesModified(EModularRigNotification InNo
 						}
 					}
 				}
+				UpdateModularDependencyDelegates();
 			}
 			break;
 		}
diff --git a/Engine/Plugins/Animation/ControlRig/Source/ControlRigDeveloper/Private/ModularRigController.cpp b/Engine/Plugins/Animation/ControlRig/Source/ControlRigDeveloper/Private/ModularRigController.cpp
deleted file mode 100644
diff --git a/Engine/Plugins/Animation/ControlRig/Source/ControlRigDeveloper/Private/ModularRigModel.cpp b/Engine/Plugins/Animation/ControlRig/Source/ControlRigDeveloper/Private/ModularRigModel.cpp
deleted file mode 100644
diff --git a/Engine/Plugins/Animation/ControlRig/Source/ControlRigDeveloper/Public/ControlRigBlueprint.h b/Engine/Plugins/Animation/ControlRig/Source/ControlRigDeveloper/Public/ControlRigBlueprint.h
--- a/Engine/Plugins/Animation/ControlRig/Source/ControlRigDeveloper/Public/ControlRigBlueprint.h
+++ b/Engine/Plugins/Animation/ControlRig/Source/ControlRigDeveloper/Public/ControlRigBlueprint.h
@@ -41,7 +41,7 @@ enum class EControlRigType : uint8
 
 
 UCLASS(BlueprintType, meta=(IgnoreClassThumbnail))
-class CONTROLRIGDEVELOPER_API UControlRigBlueprint : public URigVMBlueprint, public IInterface_PreviewMeshProvider
+class CONTROLRIGDEVELOPER_API UControlRigBlueprint : public URigVMBlueprint, public IInterface_PreviewMeshProvider, public IRigHierarchyProvider
 {
 	GENERATED_UCLASS_BODY()
 
@@ -146,6 +146,12 @@ public:
 
 	FOnRigTypeChanged& OnRigTypeChanged() { return OnRigTypeChangedDelegate; }
 
+	/// IRigHierarchyProvider interface
+	virtual URigHierarchy* GetHierarchy() const override
+	{
+		return Hierarchy;
+	}
+
 	UPROPERTY(EditAnywhere, Category = "Hierarchy")
 	FRigHierarchySettings HierarchySettings;
 
@@ -291,6 +297,9 @@ public:
 	void PropagatePropertyFromBPToInstances(FRigElementKey InRigElement, const FProperty* InProperty) const;
 	void PropagatePropertyFromInstanceToBP(FRigElementKey InRigElement, const FProperty* InProperty, UControlRig* InInstance) const;
 	void PropagateModuleHierarchyFromBPToInstances() const;
+	void UpdateModularDependencyDelegates();
+	void OnModularDependencyVMCompiled(UObject* InBlueprint, URigVM* InVM, FRigVMExtendedExecuteContext& InExecuteContext);
+	void OnModularDependencyChanged(URigVMBlueprint* InBlueprint);
 
 	/**
 	* Returns the modified event, which can be used to 
diff --git a/Engine/Plugins/Animation/ControlRig/Source/ControlRigDeveloper/Public/ModularRigController.h b/Engine/Plugins/Animation/ControlRig/Source/ControlRigDeveloper/Public/ModularRigController.h
deleted file mode 100644
diff --git a/Engine/Plugins/Animation/ControlRig/Source/ControlRigDeveloper/Public/ModularRigModel.h b/Engine/Plugins/Animation/ControlRig/Source/ControlRigDeveloper/Public/ModularRigModel.h
deleted file mode 100644
diff --git a/Engine/Plugins/Animation/ControlRig/Source/ControlRigEditor/Private/ControlRigModuleDetails.cpp b/Engine/Plugins/Animation/ControlRig/Source/ControlRigEditor/Private/ControlRigModuleDetails.cpp
--- a/Engine/Plugins/Animation/ControlRig/Source/ControlRigEditor/Private/ControlRigModuleDetails.cpp
+++ b/Engine/Plugins/Animation/ControlRig/Source/ControlRigEditor/Private/ControlRigModuleDetails.cpp
@@ -9,6 +9,7 @@
 #include "Widgets/Input/SVectorInputBox.h"
 #include "Widgets/Input/SCheckBox.h"
 #include "Widgets/Input/SButton.h"
+#include "ModularRigController.h"
 #include "ControlRigBlueprint.h"
 #include "ControlRigElementDetails.h"
 #include "Graph/ControlRigGraph.h"
@@ -357,7 +358,7 @@ FText FRigModuleInstanceDetails::GetRigClassPath() const
 		bool bSame = true;
 		for (int32 i=1; i<PerModuleInfos.Num(); ++i)
 		{
-			if (PerModuleInfos[i].GetModule()->Rig->GetClass() !=  PerModuleInfos[0].GetModule()->Rig->GetClass())
+			if (PerModuleInfos[i].GetModule()->GetRig()->GetClass() !=  PerModuleInfos[0].GetModule()->GetRig()->GetClass())
 			{
 				bSame = false;
 				break;
@@ -369,14 +370,11 @@ FText FRigModuleInstanceDetails::GetRigClassPath() const
 		}
 	}
 
-	if (FRigModuleInstance* Module = PerModuleInfos[0].GetModule())
+	if (const FRigModuleInstance* Module = PerModuleInfos[0].GetModule())
 	{
-		if (TSoftObjectPtr<UControlRig> Rig = Module->Rig)
+		if (const UControlRig* ModuleRig = Module->GetRig())
 		{
-			if (Rig.IsValid())
-			{
-				return FText::FromString(Module->Rig->GetClass()->GetClassPathName().ToString());
-			}
+			return FText::FromString(ModuleRig->GetClass()->GetClassPathName().ToString());
 		}
 	}
 
@@ -390,14 +388,11 @@ TArray<FRigModuleConnector> FRigModuleInstanceDetails::GetConnectors() const
 		return TArray<FRigModuleConnector>();
 	}
 
-	if (FRigModuleInstance* Module = PerModuleInfos[0].GetModule())
+	if (const FRigModuleInstance* Module = PerModuleInfos[0].GetModule())
 	{
-		if (TSoftObjectPtr<UControlRig> Rig = Module->Rig)
+		if (const UControlRig* ModuleRig = Module->GetRig())
 		{
-			if (Rig.IsValid())
-			{
-				return Rig->GetRigModuleSettings().ExposedConnectors;
-			}
+			return ModuleRig->GetRigModuleSettings().ExposedConnectors;
 		}
 	}
 
@@ -411,14 +406,11 @@ FRigElementKeyRedirector FRigModuleInstanceDetails::GetConnections() const
 		return FRigElementKeyRedirector();
 	}
 
-	if (FRigModuleInstance* Module = PerModuleInfos[0].GetModule())
+	if (const FRigModuleInstance* Module = PerModuleInfos[0].GetModule())
 	{
-		if (TSoftObjectPtr<UControlRig> Rig = Module->Rig)
+		if (UControlRig* ModuleRig = Module->GetRig())
 		{
-			if (Rig.IsValid())
-			{
-				return Rig->GetElementKeyRedirector();
-			}
+			return ModuleRig->GetElementKeyRedirector();
 		}
 	}
 
@@ -431,10 +423,9 @@ void FRigModuleInstanceDetails::OnConfigValueChanged(const FName InVariableName)
 	{
 		if (const FRigModuleInstance* ModuleInstance = Info.GetModule())
 		{
-			TSoftObjectPtr<UControlRig> Rig = ModuleInstance->Rig;
-			if (Rig.IsValid())
+			if (const UControlRig* ModuleRig = ModuleInstance->GetRig())
 			{
-				FString ValueStr = Rig->GetVariableAsString(InVariableName);
+				FString ValueStr = ModuleRig->GetVariableAsString(InVariableName);
 				if (UControlRigBlueprint* Blueprint = Info.GetBlueprint())
 				{
 					UModularRigController* Controller = Blueprint->GetModularRigController();
diff --git a/Engine/Plugins/Animation/ControlRig/Source/ControlRigEditor/Private/Editor/ControlRigEditor.cpp b/Engine/Plugins/Animation/ControlRig/Source/ControlRigEditor/Private/Editor/ControlRigEditor.cpp
--- a/Engine/Plugins/Animation/ControlRig/Source/ControlRigEditor/Private/Editor/ControlRigEditor.cpp
+++ b/Engine/Plugins/Animation/ControlRig/Source/ControlRigEditor/Private/Editor/ControlRigEditor.cpp
@@ -1202,7 +1202,7 @@ void FControlRigEditor::SetDetailViewForRigModules(const TArray<FString> InKeys)
 			continue;
 		}
 
-		if(UControlRig* ModuleInstance = Element->Rig.Get())
+		if(UControlRig* ModuleInstance = Element->GetRig())
 		{
 			Objects.Add(ModuleInstance);
 		}
diff --git a/Engine/Plugins/Animation/ControlRig/Source/ControlRigEditor/Private/Editor/SModularRigTreeView.cpp b/Engine/Plugins/Animation/ControlRig/Source/ControlRigEditor/Private/Editor/SModularRigTreeView.cpp
--- a/Engine/Plugins/Animation/ControlRig/Source/ControlRigEditor/Private/Editor/SModularRigTreeView.cpp
+++ b/Engine/Plugins/Animation/ControlRig/Source/ControlRigEditor/Private/Editor/SModularRigTreeView.cpp
@@ -603,9 +603,9 @@ TPair<const FSlateBrush*, FSlateColor> SModularRigModelItem::GetBrushForElementT
 
 	if (const FRigModuleInstance* Module = InModularRig->FindModule(InKey))
 	{
-		if (Module->Rig.IsValid())
+		if (UControlRig* ModuleRig = Module->GetRig())
 		{
-			FSoftObjectPath IconPath = Module->Rig->GetRigModuleSettings().Icon;
+			FSoftObjectPath IconPath = ModuleRig->GetRigModuleSettings().Icon;
 			Brush = IconPathToBrush.Find(IconPath);
 			if (!Brush)
 			{
diff --git a/Engine/Plugins/Animation/ControlRig/Source/ControlRigEditor/Private/Editor/SRigHierarchy.cpp b/Engine/Plugins/Animation/ControlRig/Source/ControlRigEditor/Private/Editor/SRigHierarchy.cpp
--- a/Engine/Plugins/Animation/ControlRig/Source/ControlRigEditor/Private/Editor/SRigHierarchy.cpp
+++ b/Engine/Plugins/Animation/ControlRig/Source/ControlRigEditor/Private/Editor/SRigHierarchy.cpp
@@ -20,6 +20,7 @@
 #include "Graph/ControlRigGraph.h"
 #include "Graph/ControlRigGraphNode.h"
 #include "Graph/ControlRigGraphSchema.h"
+#include "ModularRig.h"
 #include "GraphEditorModule.h"
 #include "Framework/MultiBox/MultiBoxBuilder.h"
 #include "AnimationRuntime.h"
diff --git a/Engine/Plugins/Runtime/RigVM/Source/RigVM/Public/RigVMHost.h b/Engine/Plugins/Runtime/RigVM/Source/RigVM/Public/RigVMHost.h
--- a/Engine/Plugins/Runtime/RigVM/Source/RigVM/Public/RigVMHost.h
+++ b/Engine/Plugins/Runtime/RigVM/Source/RigVM/Public/RigVMHost.h
@@ -196,6 +196,13 @@ public:
 	UFUNCTION(BlueprintCallable, Category = "RigVM")
 	virtual void RequestInit();
 
+	/** Returns true if this host requires the VM memory to be initialized */
+	UFUNCTION(BlueprintPure, Category = "RigVM")
+	virtual bool IsInitRequired() const
+	{
+		return bRequiresInitExecution;
+	} 
+
 	/** Requests to run an event once */
 	UFUNCTION(BlueprintCallable, Category = "RigVM")
 	void RequestRunOnceEvent(const FName& InEventName, int32 InEventIndex = -1);
