ControlRig: Follow up to 29967154

diff --git a/Engine/Plugins/Animation/ControlRig/Source/ControlRig/Private/ModularRig.cpp b/Engine/Plugins/Animation/ControlRig/Source/ControlRig/Private/ModularRig.cpp
--- a/Engine/Plugins/Animation/ControlRig/Source/ControlRig/Private/ModularRig.cpp
+++ b/Engine/Plugins/Animation/ControlRig/Source/ControlRig/Private/ModularRig.cpp
@@ -76,8 +76,6 @@ FString FRigModuleInstance::GetNamespace() const
 
 UControlRig* FRigModuleInstance::GetRig() const
 {
-	// prefer the cached UObject over the soft object ptrs
-	// since the blueprint reinstancer does not reinstance soft object ptrs.
 	if(IsValid(RigPtr))
 	{
 		return RigPtr;
diff --git a/Engine/Plugins/Animation/ControlRig/Source/ControlRig/Public/ModularRig.h b/Engine/Plugins/Animation/ControlRig/Source/ControlRig/Public/ModularRig.h
--- a/Engine/Plugins/Animation/ControlRig/Source/ControlRig/Public/ModularRig.h
+++ b/Engine/Plugins/Animation/ControlRig/Source/ControlRig/Public/ModularRig.h
@@ -56,7 +56,7 @@ public:
 private:
 
 	UPROPERTY(transient)
-	mutable UControlRig* RigPtr;
+	mutable TObjectPtr<UControlRig> RigPtr;
 
 public:
 	
