Control Rig: Guard against recursion when destroying hierarchy elements

#rnx

diff --git a/Engine/Plugins/Animation/ControlRig/Source/ControlRig/Private/Rigs/RigHierarchy.cpp b/Engine/Plugins/Animation/ControlRig/Source/ControlRig/Private/Rigs/RigHierarchy.cpp
--- a/Engine/Plugins/Animation/ControlRig/Source/ControlRig/Private/Rigs/RigHierarchy.cpp
+++ b/Engine/Plugins/Animation/ControlRig/Source/ControlRig/Private/Rigs/RigHierarchy.cpp
@@ -126,6 +126,7 @@ URigHierarchy::URigHierarchy()
 , bRecordTransformsAtRuntime(true)
 #endif
 , ElementKeyRedirector(nullptr)
+, ElementBeingDestroyed(nullptr)
 {
 	Reset();
 #if WITH_EDITOR
@@ -4956,6 +4957,7 @@ void URigHierarchy::DestroyElement(FRigBaseElement*& InElement)
 			FRigBoneElement* ExistingElements = Cast<FRigBoneElement>(InElement);
 			for(int32 Index=0;Index<Count;Index++)
 			{
+				TGuardValue<const FRigBaseElement*> DestroyGuard(ElementBeingDestroyed, &ExistingElements[Index]);
 				ExistingElements[Index].~FRigBoneElement(); 
 			}
 			break;
@@ -4965,6 +4967,7 @@ void URigHierarchy::DestroyElement(FRigBaseElement*& InElement)
 			FRigNullElement* ExistingElements = Cast<FRigNullElement>(InElement);
 			for(int32 Index=0;Index<Count;Index++)
 			{
+				TGuardValue<const FRigBaseElement*> DestroyGuard(ElementBeingDestroyed, &ExistingElements[Index]);
 				ExistingElements[Index].~FRigNullElement(); 
 			}
 			break;
@@ -4974,6 +4977,7 @@ void URigHierarchy::DestroyElement(FRigBaseElement*& InElement)
 			FRigControlElement* ExistingElements = Cast<FRigControlElement>(InElement);
 			for(int32 Index=0;Index<Count;Index++)
 			{
+				TGuardValue<const FRigBaseElement*> DestroyGuard(ElementBeingDestroyed, &ExistingElements[Index]);
 				ExistingElements[Index].~FRigControlElement(); 
 			}
 			break;
@@ -4983,6 +4987,7 @@ void URigHierarchy::DestroyElement(FRigBaseElement*& InElement)
 			FRigCurveElement* ExistingElements = Cast<FRigCurveElement>(InElement);
 			for(int32 Index=0;Index<Count;Index++)
 			{
+				TGuardValue<const FRigBaseElement*> DestroyGuard(ElementBeingDestroyed, &ExistingElements[Index]);
 				ExistingElements[Index].~FRigCurveElement(); 
 			}
 			break;
@@ -4992,6 +4997,7 @@ void URigHierarchy::DestroyElement(FRigBaseElement*& InElement)
 			FRigRigidBodyElement* ExistingElements = Cast<FRigRigidBodyElement>(InElement);
 			for(int32 Index=0;Index<Count;Index++)
 			{
+				TGuardValue<const FRigBaseElement*> DestroyGuard(ElementBeingDestroyed, &ExistingElements[Index]);
 				ExistingElements[Index].~FRigRigidBodyElement(); 
 			}
 			break;
@@ -5001,6 +5007,7 @@ void URigHierarchy::DestroyElement(FRigBaseElement*& InElement)
 			FRigReferenceElement* ExistingElements = Cast<FRigReferenceElement>(InElement);
 			for(int32 Index=0;Index<Count;Index++)
 			{
+				TGuardValue<const FRigBaseElement*> DestroyGuard(ElementBeingDestroyed, &ExistingElements[Index]);
 				ExistingElements[Index].~FRigReferenceElement(); 
 			}
 			break;
@@ -5010,6 +5017,7 @@ void URigHierarchy::DestroyElement(FRigBaseElement*& InElement)
 			FRigConnectorElement* ExistingElements = Cast<FRigConnectorElement>(InElement);
 			for(int32 Index=0;Index<Count;Index++)
 			{
+				TGuardValue<const FRigBaseElement*> DestroyGuard(ElementBeingDestroyed, &ExistingElements[Index]);
 				ExistingElements[Index].~FRigConnectorElement(); 
 			}
 			break;
@@ -5019,6 +5027,7 @@ void URigHierarchy::DestroyElement(FRigBaseElement*& InElement)
 			FRigSocketElement* ExistingElements = Cast<FRigSocketElement>(InElement);
 			for(int32 Index=0;Index<Count;Index++)
 			{
+				TGuardValue<const FRigBaseElement*> DestroyGuard(ElementBeingDestroyed, &ExistingElements[Index]);
 				ExistingElements[Index].~FRigSocketElement(); 
 			}
 			break;
@@ -5511,8 +5520,11 @@ bool URigHierarchy::RemoveMetadataForElement(FRigBaseElement* InElement, const F
 	{
 		Storage.LastAccessMetadata = nullptr;
 	}
-	
-	OnMetadataChanged(InElement->Key, InName);
+
+	if(ElementBeingDestroyed != InElement)
+	{
+		OnMetadataChanged(InElement->Key, InName);
+	}
 	return true;
 }
 
@@ -5534,10 +5546,13 @@ bool URigHierarchy::RemoveAllMetadataForElement(FRigBaseElement* InElement)
 	
 	ElementMetadataFreeList.Push(InElement->MetadataStorageIndex);
 	InElement->MetadataStorageIndex = INDEX_NONE;
-	
-	for (FName Name: Names)
+
+	if(ElementBeingDestroyed != InElement)
 	{
-		OnMetadataChanged(InElement->Key, Name);
+		for (FName Name: Names)
+		{
+			OnMetadataChanged(InElement->Key, Name);
+		}
 	}
 	
 	return true; 
diff --git a/Engine/Plugins/Animation/ControlRig/Source/ControlRig/Public/Rigs/RigHierarchy.h b/Engine/Plugins/Animation/ControlRig/Source/ControlRig/Public/Rigs/RigHierarchy.h
--- a/Engine/Plugins/Animation/ControlRig/Source/ControlRig/Public/Rigs/RigHierarchy.h
+++ b/Engine/Plugins/Animation/ControlRig/Source/ControlRig/Public/Rigs/RigHierarchy.h
@@ -4686,6 +4686,9 @@ private:
 	void QueueNotification(ERigHierarchyNotification InNotification, const FRigBaseElement* InElement);
 	void SendQueuedNotifications();
 	void Reset_Impl(bool bResetElements);
+
+	// the currently destroyed element - used to avoid notification storms
+	const FRigBaseElement* ElementBeingDestroyed; 
 	
 	friend class URigHierarchyController;
 	friend class UControlRig;
