Control Rig: Avoid destruction of module instances if possible

diff --git a/Engine/Plugins/Animation/ControlRig/Source/ControlRig/Private/ModularRig.cpp b/Engine/Plugins/Animation/ControlRig/Source/ControlRig/Private/ModularRig.cpp
--- a/Engine/Plugins/Animation/ControlRig/Source/ControlRig/Private/ModularRig.cpp
+++ b/Engine/Plugins/Animation/ControlRig/Source/ControlRig/Private/ModularRig.cpp
@@ -88,11 +88,7 @@ void FRigModuleInstance::SetRig(UControlRig* InRig)
 	UControlRig* PreviousRig = GetRig();
 	if(PreviousRig && (PreviousRig != InRig))
 	{
-		// rename the previous rig.
-		// GC will pick it up eventually - since we won't have any
-		// owning pointers to it anymore.
-		PreviousRig->Rename(nullptr, GetTransientPackage(), REN_ForceNoResetLoaders | REN_DoNotDirty | REN_DontCreateRedirectors | REN_NonTransactional);
-		PreviousRig->MarkAsGarbage();
+		UModularRig::DiscardModuleRig(PreviousRig);
 	}
 
 	// update the cache
@@ -147,50 +143,61 @@ void UModularRig::InitializeFromCDO()
 	UpdateModuleHierarchyFromCDO();
 }
 
-static void AddModulesRecursively(UModularRig* Rig, const FRigModuleReference* InModule, FRigModuleInstance* InParent)
-{
-	// Make sure the inner is loaded and valid
-	{
-		if (!InModule->Class.IsValid())
-		{
-			(void)InModule->Class.LoadSynchronous();
-		}
-		if (!InModule->Class.IsValid())
-		{
-			return;
-		}
-	}
-	
-	FRigModuleInstance* NewModule = Rig->AddModuleInstance(InModule->Name, InModule->Class.Get(), InParent, InModule->Connections, InModule->ConfigValues);
-	if(NewModule)
-	{
-		for (const FRigModuleReference* ChildModule : InModule->CachedChildren)
-		{
-			AddModulesRecursively(Rig, ChildModule, NewModule);
-		}
-	}
-	else
-	{
-		// todooo
-	}
-}
-
 void UModularRig::UpdateModuleHierarchyFromCDO()
 {
 	if (!HasAnyFlags(RF_ClassDefaultObject))
 	{
-		ResetModules();
+		// keep the previous rigs around
+		check(PreviousModuleRigs.IsEmpty());
+		for (const FRigModuleInstance& Module : Modules)
+        {
+			if(UControlRig* ModuleRig = Module.GetRig())
+			{
+				if(IsValid(ModuleRig))
+				{
+					PreviousModuleRigs.Add(Module.GetPath(), ModuleRig);
+				}
+			}
+        }
+
+		// don't destroy the rigs when resetting
+		ResetModules(false);
 
 		// the CDO owns the model - when we ask for the model we'll always
 		// get the model from the CDO. we'll now add UObject module instances
 		// for each module (data only) reference in the model.
 		// Note: The CDO does not contain any UObject module instances itself.
 		const FModularRigModel& Model = GetModularRigModel();
-		for (const FRigModuleReference* RootModule : Model.RootModules)
+		Model.ForEachModule([this](const FRigModuleReference* InModuleReference) -> bool
 		{
-			AddModulesRecursively(this, RootModule, nullptr);
+			check(InModuleReference);
+			if (!InModuleReference->Class.IsValid())
+			{
+				(void)InModuleReference->Class.LoadSynchronous();
+			}
+			if (InModuleReference->Class.IsValid())
+			{
+				(void)AddModuleInstance(
+					InModuleReference->Name,
+					InModuleReference->Class.Get(),
+					FindModule(InModuleReference->ParentPath),
+					InModuleReference->Connections,
+					InModuleReference->ConfigValues);
+			}
+
+			// continue to the next module
+			return true;
+		});
+
+		// discard any remaining rigs
+		for(const TPair<FString, UControlRig*>& Pair : PreviousModuleRigs)
+		{
+			DiscardModuleRig(Pair.Value);
 		}
+		PreviousModuleRigs.Reset();
 
+		// update the module variable bindings now - since for this all
+		// modules have to exist first
 		ForEachModule([this, Model](const FRigModuleInstance* Module) -> bool
 		{
 			if(const FRigModuleReference* ModuleReference = Model.FindModule(Module->GetPath()))
@@ -410,17 +417,20 @@ void UModularRig::OnObjectsReplaced(const TMap<UObject*, UObject*>& OldToNewInst
 	}
 }
 
-void UModularRig::ResetModules()
+void UModularRig::ResetModules(bool bDestroyModuleRigs)
 {
 	for (FRigModuleInstance& Module : Modules)
 	{
 		Module.CachedChildren.Reset();
 
-		if (const UControlRig* ModuleRig = Module.GetRig())
+		if(bDestroyModuleRigs)
 		{
-			check(ModuleRig->GetOuter() == this);
-			// takes care of renaming / moving the rig to the transient package
-			Module.SetRig(nullptr);
+			if (const UControlRig* ModuleRig = Module.GetRig())
+			{
+				check(ModuleRig->GetOuter() == this);
+				// takes care of renaming / moving the rig to the transient package
+				Module.SetRig(nullptr);
+			}
 		}
 	}
 	
@@ -482,8 +492,8 @@ void UModularRig::UpdateSupportedEvents()
 	});
 }
 
-FRigModuleInstance* UModularRig::AddModuleInstance(const FName& InModuleName, TSubclassOf<UControlRig> InModuleClass, FRigModuleInstance* InParent,
-	const TMap<FRigElementKey, FRigElementKey>& InConnectionMap, const TMap<FName, FString>& InVariableDefaultValues ) 
+FRigModuleInstance* UModularRig::AddModuleInstance(const FName& InModuleName, TSubclassOf<UControlRig> InModuleClass, const FRigModuleInstance* InParent,
+	const TMap<FRigElementKey, FRigElementKey>& InConnectionMap, const TMap<FName, FString>& InVariableDefaultValues) 
 {
 	// Make sure there are no name clashes
 	if (InParent)
@@ -520,8 +530,32 @@ FRigModuleInstance* UModularRig::AddModuleInstance(const FName& InModuleName, TS
 		NewModule.ParentPath = InParent->GetPath();
 	}
 	const FString Name = NewModule.GetPath();
+
+	UControlRig* NewModuleRig = nullptr;
+
+	// reuse existing module rig instances first
+	if(UControlRig** ExistingModuleRigPtr = PreviousModuleRigs.Find(Name))
+	{
+		if(UControlRig* ExistingModuleRig = *ExistingModuleRigPtr)
+		{
+			// again relying on GetFName since RigVMHost overloads GetName
+			if(ExistingModuleRig->GetFName().ToString().Equals(Name) && ExistingModuleRig->GetClass() == InModuleClass)
+			{
+				NewModuleRig = ExistingModuleRig;
+			}
+			else
+			{
+				DiscardModuleRig(ExistingModuleRig);
+			}
+			PreviousModuleRigs.Remove(Name);
+		}
+	}
+
+	if(NewModuleRig == nullptr)
+	{
+		NewModuleRig = NewObject<UControlRig>(this, InModuleClass, *Name);
+	}
 	
-	UControlRig* NewModuleRig = NewObject<UControlRig>(this, InModuleClass, *Name);
 	NewModule.SetRig(NewModuleRig);
 
 	UpdateCachedChildren();
@@ -535,6 +569,7 @@ FRigModuleInstance* UModularRig::AddModuleInstance(const FName& InModuleName, TS
 		URigHierarchy* Hierarchy = GetHierarchy();
 		FRigVMExtendedExecuteContext& ModuleContext = NewModuleRig->GetRigVMExtendedExecuteContext();
 		FControlRigExecuteContext& ModulePublicContext = ModuleContext.GetPublicDataSafe<FControlRigExecuteContext>();
+		NewModuleRig->RequestInit();
 		NewModuleRig->bCopyHierarchyBeforeConstruction = false;
 		NewModuleRig->SetDynamicHierarchy(Hierarchy);
 		ModulePublicContext.Hierarchy = Hierarchy;
@@ -588,6 +623,18 @@ bool UModularRig::SetModuleVariableBindings(const FString& InModulePath, const T
 	return false;
 }
 
+void UModularRig::DiscardModuleRig(UControlRig* InControlRig)
+{
+	if(InControlRig)
+	{
+		// rename the previous rig.
+		// GC will pick it up eventually - since we won't have any
+		// owning pointers to it anymore.
+		InControlRig->Rename(nullptr, GetTransientPackage(), REN_ForceNoResetLoaders | REN_DoNotDirty | REN_DontCreateRedirectors | REN_NonTransactional);
+		InControlRig->MarkAsGarbage();
+	}
+}
+
 const FRigModuleInstance* UModularRig::FindModule(const FString& InPath) const
 {
 	return Modules.FindByPredicate([InPath](const FRigModuleInstance& Module)
diff --git a/Engine/Plugins/Animation/ControlRig/Source/ControlRig/Public/ModularRig.h b/Engine/Plugins/Animation/ControlRig/Source/ControlRig/Public/ModularRig.h
--- a/Engine/Plugins/Animation/ControlRig/Source/ControlRig/Public/ModularRig.h
+++ b/Engine/Plugins/Animation/ControlRig/Source/ControlRig/Public/ModularRig.h
@@ -147,17 +147,13 @@ public:
 	
 	void OnObjectsReplaced(const TMap<UObject*, UObject*>& OldToNewInstanceMap);
 
-	void ResetModules();
+	void ResetModules(bool bDestroyModuleRigs = true);
 
 	const FModularRigModel& GetModularRigModel() const;
 	void UpdateModuleHierarchyFromCDO();
 	void UpdateCachedChildren();
 	void UpdateSupportedEvents();
 
-	/** Adds a module to the rig*/
-	FRigModuleInstance* AddModuleInstance(const FName& InModuleName, TSubclassOf<UControlRig> InModuleClass, FRigModuleInstance* InParent, const TMap<FRigElementKey, FRigElementKey>& InConnectionMap, const TMap<FName, FString>& InVariableDefaultValues);
-	bool SetModuleVariableBindings(const FString& InModulePath, const TMap<FName, FString>& InVariableBindings);
-	
 	const FRigModuleInstance* FindModule(const FString& InPath) const;
 	const FRigModuleInstance* FindModule(const UControlRig* InModuleInstance) const;
 	FString GetParentPath(const FString& InPath) const;
@@ -180,4 +176,19 @@ public:
 	}
 
 	static const FString NamespaceSeparator;
+
+private:
+
+	/** Adds a module to the rig*/
+	FRigModuleInstance* AddModuleInstance(const FName& InModuleName, TSubclassOf<UControlRig> InModuleClass, const FRigModuleInstance* InParent, const TMap<FRigElementKey, FRigElementKey>& InConnectionMap, const TMap<FName, FString>& InVariableDefaultValues);
+
+	/** Updates the module's variable bindings */
+	bool SetModuleVariableBindings(const FString& InModulePath, const TMap<FName, FString>& InVariableBindings);
+
+	/** Destroys / discards a module instance rig */
+	static void DiscardModuleRig(UControlRig* InControlRig);
+
+	TMap<FString, UControlRig*> PreviousModuleRigs;
+
+	friend struct FRigModuleInstance;
 };
diff --git a/Engine/Plugins/Animation/ControlRig/Source/ControlRigEditor/Private/Editor/SModularRigModel.cpp b/Engine/Plugins/Animation/ControlRig/Source/ControlRigEditor/Private/Editor/SModularRigModel.cpp
--- a/Engine/Plugins/Animation/ControlRig/Source/ControlRigEditor/Private/Editor/SModularRigModel.cpp
+++ b/Engine/Plugins/Animation/ControlRig/Source/ControlRigEditor/Private/Editor/SModularRigModel.cpp
@@ -255,7 +255,6 @@ TArray<FString> SModularRigModel::GetSelectedKeys() const
 
 void SModularRigModel::HandlePreCompileModularRigs(URigVMBlueprint* InBlueprint)
 {
-	ClearDetailPanel();
 }
 
 void SModularRigModel::HandlePostCompileModularRigs(URigVMBlueprint* InBlueprint)
@@ -561,8 +560,6 @@ void SModularRigModel::HandleRenameModule()
 
 FName SModularRigModel::HandleRenameModule(const FString& InOldPath, const FName& InNewName)
 {
-	ClearDetailPanel();
-	
 	if (ControlRigBlueprint.IsValid())
 	{
 		FScopedTransaction Transaction(LOCTEXT("ModularRigModelRename", "Rename Module"));
@@ -628,8 +625,6 @@ void SModularRigModel::HandleDeleteModules()
 
 void SModularRigModel::HandleDeleteModules(const TArray<FString>& InPaths)
 {
-	ClearDetailPanel();
-	
 	if (ControlRigBlueprint.IsValid())
 	{
 		FScopedTransaction Transaction(LOCTEXT("ModularRigModelDelete", "Delete Modules"));
@@ -649,8 +644,6 @@ void SModularRigModel::HandleDeleteModules(const TArray<FString>& InPaths)
 
 void SModularRigModel::HandleReparentModules(const TArray<FString>& InPaths, const FString& InParentPath)
 {
-	ClearDetailPanel();
-	
 	if (ControlRigBlueprint.IsValid())
 	{
 		FScopedTransaction Transaction(LOCTEXT("ModularRigModelReparent", "Reparent Modules"));
@@ -730,14 +723,6 @@ void SModularRigModel::OnRequestDetailsInspection(const FString& InKey)
 	ControlRigEditor.Pin()->SetDetailViewForRigModules({InKey});
 }
 
-void SModularRigModel::ClearDetailPanel() const
-{
-	if(ControlRigEditor.IsValid())
-	{
-		ControlRigEditor.Pin()->ClearDetailObject();
-	}
-}
-
 void SModularRigModel::PostRedo(bool bSuccess) 
 {
 	if (bSuccess)
diff --git a/Engine/Plugins/Animation/ControlRig/Source/ControlRigEditor/Private/Editor/SModularRigModel.h b/Engine/Plugins/Animation/ControlRig/Source/ControlRigEditor/Private/Editor/SModularRigModel.h
--- a/Engine/Plugins/Animation/ControlRig/Source/ControlRigEditor/Private/Editor/SModularRigModel.h
+++ b/Engine/Plugins/Animation/ControlRig/Source/ControlRigEditor/Private/Editor/SModularRigModel.h
@@ -147,8 +147,6 @@ private:
 	const UModularRig* GetModularRigForTreeView() const { return GetModularRig(); }
 	FName CreateUniqueName(const FName& InBasePath) const;
 	void OnRequestDetailsInspection(const FString& InKey);
-	void ClearDetailPanel() const;
-
 	void HandlePreCompileModularRigs(URigVMBlueprint* InBlueprint);
 	void HandlePostCompileModularRigs(URigVMBlueprint* InBlueprint);
 	
