Control Rig: Disable renaming in asset picker + support double click

#rnx

diff --git a/Engine/Plugins/Animation/ControlRig/Source/ControlRigEditor/Private/Editor/SRigModuleAssetBrowser.cpp b/Engine/Plugins/Animation/ControlRig/Source/ControlRigEditor/Private/Editor/SRigModuleAssetBrowser.cpp
--- a/Engine/Plugins/Animation/ControlRig/Source/ControlRigEditor/Private/Editor/SRigModuleAssetBrowser.cpp
+++ b/Engine/Plugins/Animation/ControlRig/Source/ControlRigEditor/Private/Editor/SRigModuleAssetBrowser.cpp
@@ -5,6 +5,7 @@
 #include "ContentBrowserModule.h"
 #include "ControlRigBlueprint.h"
 #include "IContentBrowserSingleton.h"
+#include "Subsystems/AssetEditorSubsystem.h"
 
 #include "ControlRigEditor.h"
 #include "Widgets/Layout/SUniformGridPanel.h"
@@ -45,6 +46,10 @@ void SRigModuleAssetBrowser::RefreshView()
 	AssetPickerConfig.GetCurrentSelectionDelegates.Add(&GetCurrentSelectionDelegate);
 	AssetPickerConfig.bAllowNullSelection = false;
 	AssetPickerConfig.bFocusSearchBoxWhenOpened = false;
+	AssetPickerConfig.bAllowDragging = true;
+	AssetPickerConfig.bAllowRename = false;
+	AssetPickerConfig.bForceShowPluginContent = true;
+	AssetPickerConfig.bForceShowEngineContent = true;
 
 	// hide all asset registry columns by default (we only really want the name and path)
 	TArray<UObject::FAssetRegistryTag> AssetRegistryTags;
@@ -58,6 +63,9 @@ void SRigModuleAssetBrowser::RefreshView()
 	AssetPickerConfig.HiddenColumnNames.Add(TEXT("Class"));
 	AssetPickerConfig.HiddenColumnNames.Add(TEXT("Has Virtualized Data"));
 
+	// allow to open the rigs directly on double click
+	AssetPickerConfig.OnAssetDoubleClicked = FOnAssetDoubleClicked::CreateSP(this, &SRigModuleAssetBrowser::OnAssetDoubleClicked);
+
 	const FContentBrowserModule& ContentBrowserModule = FModuleManager::Get().LoadModuleChecked<FContentBrowserModule>(TEXT("ContentBrowser"));
 	AssetBrowserBox->SetContent(ContentBrowserModule.Get().CreateAssetPicker(AssetPickerConfig));
 }
@@ -121,4 +129,12 @@ bool SRigModuleAssetBrowser::OnShouldFilterAsset(const struct FAssetData& AssetD
 	return ControlRigType != EControlRigType::RigModule;
 }
 
+void SRigModuleAssetBrowser::OnAssetDoubleClicked(const FAssetData& AssetData)
+{
+	if (UAssetEditorSubsystem* EditorSubsystem = GEditor->GetEditorSubsystem<UAssetEditorSubsystem>())
+	{
+		EditorSubsystem->OpenEditorForAsset(AssetData.ToSoftObjectPath());
+	}
+}
+
 #undef LOCTEXT_NAMESPACE
diff --git a/Engine/Plugins/Animation/ControlRig/Source/ControlRigEditor/Private/Editor/SRigModuleAssetBrowser.h b/Engine/Plugins/Animation/ControlRig/Source/ControlRigEditor/Private/Editor/SRigModuleAssetBrowser.h
--- a/Engine/Plugins/Animation/ControlRig/Source/ControlRigEditor/Private/Editor/SRigModuleAssetBrowser.h
+++ b/Engine/Plugins/Animation/ControlRig/Source/ControlRigEditor/Private/Editor/SRigModuleAssetBrowser.h
@@ -26,6 +26,7 @@ private:
 	
 	bool OnShouldFilterAsset(const struct FAssetData& AssetData);
 	TSharedPtr<SWidget> OnGetAssetContextMenu(const TArray<FAssetData>& SelectedAssets) const;
+	void OnAssetDoubleClicked(const FAssetData& AssetData);
 
 	/** Used to get the currently selected assets */
 	FGetCurrentSelectionDelegate GetCurrentSelectionDelegate;
diff --git a/Engine/Source/Editor/ContentBrowser/Private/SAssetPicker.cpp b/Engine/Source/Editor/ContentBrowser/Private/SAssetPicker.cpp
--- a/Engine/Source/Editor/ContentBrowser/Private/SAssetPicker.cpp
+++ b/Engine/Source/Editor/ContentBrowser/Private/SAssetPicker.cpp
@@ -85,6 +85,8 @@ void SAssetPicker::Construct( const FArguments& InArgs )
 	CurrentSourcesData = FSourcesData(InArgs._AssetPickerConfig.Filter.PackagePaths, InArgs._AssetPickerConfig.Collections);
 	CurrentBackendFilter = InArgs._AssetPickerConfig.Filter;
 	CurrentBackendFilter.PackagePaths.Reset();
+
+	bAllowRename = InArgs._AssetPickerConfig.bAllowRename;
 	
 	FOnGetContentBrowserItemContextMenu OnGetItemContextMenu;
 	if (OnGetAssetContextMenu.IsBound() || OnGetFolderContextMenu.IsBound())
@@ -698,6 +700,10 @@ void SAssetPicker::OnRenameRequested() const
 
 bool SAssetPicker::CanExecuteRenameRequested()
 {
+	if(!bAllowRename)
+	{
+		return false;
+	}
 	return ContentBrowserUtils::CanRenameFromAssetView(AssetViewPtr);
 }
 
diff --git a/Engine/Source/Editor/ContentBrowser/Private/SAssetPicker.h b/Engine/Source/Editor/ContentBrowser/Private/SAssetPicker.h
--- a/Engine/Source/Editor/ContentBrowser/Private/SAssetPicker.h
+++ b/Engine/Source/Editor/ContentBrowser/Private/SAssetPicker.h
@@ -203,4 +203,7 @@ private:
 
 	/** If set, view settings will be saved and loaded for the asset view using this name in ini files */
 	FString SaveSettingsName;
+
+	/** Wether to allow renaming of items */
+	bool bAllowRename;
 };
diff --git a/Engine/Source/Editor/ContentBrowser/Public/IContentBrowserSingleton.h b/Engine/Source/Editor/ContentBrowser/Public/IContentBrowserSingleton.h
--- a/Engine/Source/Editor/ContentBrowser/Public/IContentBrowserSingleton.h
+++ b/Engine/Source/Editor/ContentBrowser/Public/IContentBrowserSingleton.h
@@ -266,6 +266,9 @@ struct FAssetPickerConfig
 	/** Whether to allow dragging of items */
 	bool bAllowDragging;
 
+	/** Whether to allow renaming of items */
+	bool bAllowRename;
+
 	/** Indicates if this view is allowed to show classes */
 	bool bCanShowClasses;
 
@@ -316,6 +319,7 @@ struct FAssetPickerConfig
 		, bShowBottomToolbar(true)
 		, bAutohideSearchBar(false)
 		, bAllowDragging(true)
+		, bAllowRename(true)
 		, bCanShowClasses(true)
 		, bCanShowFolders(false)
 		, bCanShowReadOnlyFolders(true)
