Control Rig: Rely on fixed index loop for listening hierarchy to avoid crashes in recursive callstacks

#rnx

diff --git a/Engine/Plugins/Animation/ControlRig/Source/ControlRig/Private/Rigs/RigHierarchy.cpp b/Engine/Plugins/Animation/ControlRig/Source/ControlRig/Private/Rigs/RigHierarchy.cpp
--- a/Engine/Plugins/Animation/ControlRig/Source/ControlRig/Private/Rigs/RigHierarchy.cpp
+++ b/Engine/Plugins/Animation/ControlRig/Source/ControlRig/Private/Rigs/RigHierarchy.cpp
@@ -448,6 +448,18 @@ void URigHierarchy::Reset_Impl(bool bResetElements)
 	}
 }
 
+#if WITH_EDITOR
+
+void URigHierarchy::ForEachListeningHierarchy(TFunctionRef<void(const FRigHierarchyListener&)> PerListeningHierarchyFunction)
+{
+	for(int32 Index = 0; Index < ListeningHierarchies.Num(); Index++)
+	{
+		PerListeningHierarchyFunction(ListeningHierarchies[Index]);
+	}
+}
+
+#endif
+
 void URigHierarchy::ResetToDefault()
 {
 	FScopeLock Lock(&ElementsLock);
@@ -1992,23 +2004,22 @@ bool URigHierarchy::SetParentWeight(FRigBaseElement* InChild, int32 InParentInde
 			if (!bPropagatingChange)
 			{
 				TGuardValue<bool> bPropagatingChangeGuardValue(bPropagatingChange, true);
-			
-				for(FRigHierarchyListener& Listener : ListeningHierarchies)
+
+				ForEachListeningHierarchy([this, LocalType, InChild, InParentIndex, InWeight, bInitial, bAffectChildren](const FRigHierarchyListener& Listener)
 				{
 					if(!bForcePropagation && !Listener.ShouldReactToChange(LocalType))
 					{
-						continue;
+						return;
 					}
 
-					URigHierarchy* ListeningHierarchy = Listener.Hierarchy.Get();
-					if (ListeningHierarchy)
+					if (URigHierarchy* ListeningHierarchy = Listener.Hierarchy.Get())
 					{
 						if(FRigBaseElement* ListeningElement = ListeningHierarchy->Find(InChild->GetKey()))
 						{
 							ListeningHierarchy->SetParentWeight(ListeningElement, InParentIndex, InWeight, bInitial, bAffectChildren);
 						}
 					}
-				}	
+				});
 			}
 #endif
 
@@ -2117,22 +2128,21 @@ bool URigHierarchy::SetParentWeightArray(FRigBaseElement* InChild,  const TArray
 			{
 				TGuardValue<bool> bPropagatingChangeGuardValue(bPropagatingChange, true);
 			
-				for(FRigHierarchyListener& Listener : ListeningHierarchies)
-				{
+				ForEachListeningHierarchy([this, LocalType, InChild, InWeights, bInitial, bAffectChildren](const FRigHierarchyListener& Listener)
+     			{
 					if(!bForcePropagation && !Listener.ShouldReactToChange(LocalType))
 					{
-						continue;
+						return;
 					}
 
-					URigHierarchy* ListeningHierarchy = Listener.Hierarchy.Get();
-					if (ListeningHierarchy)
+					if (URigHierarchy* ListeningHierarchy = Listener.Hierarchy.Get())
 					{
 						if(FRigBaseElement* ListeningElement = ListeningHierarchy->Find(InChild->GetKey()))
 						{
 							ListeningHierarchy->SetParentWeightArray(ListeningElement, InWeights, bInitial, bAffectChildren);
 						}
 					}
-				}	
+				});
 			}
 #endif
 
@@ -2975,17 +2985,16 @@ void URigHierarchy::Notify(ERigHierarchyNotification InNotifType, const FRigBase
 		{
 			if (ensure(InElement != nullptr))
 			{
-				for(FRigHierarchyListener& Listener : ListeningHierarchies)
+				ForEachListeningHierarchy([this, InNotifType, InElement](const FRigHierarchyListener& Listener)
 				{
-					URigHierarchy* ListeningHierarchy = Listener.Hierarchy.Get();
-					if (ListeningHierarchy)
+					if (URigHierarchy* ListeningHierarchy = Listener.Hierarchy.Get())
 					{			
 						if(const FRigBaseElement* ListeningElement = ListeningHierarchy->Find( InElement->GetKey()))
 						{
 							ListeningHierarchy->Notify(InNotifType, ListeningElement);
 						}
 					}
-				}
+				});
 			}
 			break;
 		}
@@ -3420,16 +3429,15 @@ void URigHierarchy::SetTransform(FRigTransformElement* InTransformElement, const
 	if (!bPropagatingChange)
 	{
 		TGuardValue<bool> bPropagatingChangeGuardValue(bPropagatingChange, true);
-			
-		for(FRigHierarchyListener& Listener : ListeningHierarchies)
+
+		ForEachListeningHierarchy([this, InTransformElement, InTransform, InTransformType, bAffectChildren, bForce](const FRigHierarchyListener& Listener)
 		{
 			if(!bForcePropagation && !Listener.ShouldReactToChange(InTransformType))
 			{
-				continue;
+				return;
 			}
 
-			URigHierarchy* ListeningHierarchy = Listener.Hierarchy.Get();
-			if (ListeningHierarchy)
+			if (URigHierarchy* ListeningHierarchy = Listener.Hierarchy.Get())
 			{			
 				if(FRigTransformElement* ListeningElement = Cast<FRigTransformElement>(ListeningHierarchy->Find(InTransformElement->GetKey())))
 				{
@@ -3437,7 +3445,7 @@ void URigHierarchy::SetTransform(FRigTransformElement* InTransformElement, const
 					ListeningHierarchy->SetTransform(ListeningElement, InTransform, InTransformType, bAffectChildren, false, bForce);
 				}
 			}
-		}
+		});
 	}
 
 	if (bPrintPythonCommands)
@@ -3650,10 +3658,9 @@ void URigHierarchy::SetControlOffsetTransform(FRigControlElement* InControlEleme
 	{
 		TGuardValue<bool> bPropagatingChangeGuardValue(bPropagatingChange, true);
 			
-		for(FRigHierarchyListener& Listener : ListeningHierarchies)
+		ForEachListeningHierarchy([this, InControlElement, InTransform, InTransformType, bAffectChildren, bForce](const FRigHierarchyListener& Listener)
 		{
-			URigHierarchy* ListeningHierarchy = Listener.Hierarchy.Get();
-			if (ListeningHierarchy)
+			if (URigHierarchy* ListeningHierarchy = Listener.Hierarchy.Get())
 			{	
 				if(FRigControlElement* ListeningElement = Cast<FRigControlElement>(ListeningHierarchy->Find(InControlElement->GetKey())))
 				{
@@ -3661,7 +3668,7 @@ void URigHierarchy::SetControlOffsetTransform(FRigControlElement* InControlEleme
 					ListeningHierarchy->SetControlOffsetTransform(ListeningElement, InTransform, InTransformType, bAffectChildren, false, bForce);
 				}
 			}
-		}
+		});
 	}
 
 	if (bPrintPythonCommands)
@@ -3782,10 +3789,9 @@ void URigHierarchy::SetControlShapeTransform(FRigControlElement* InControlElemen
 	{
 		TGuardValue<bool> bPropagatingChangeGuardValue(bPropagatingChange, true);
 			
-		for(FRigHierarchyListener& Listener : ListeningHierarchies)
+		ForEachListeningHierarchy([this, InControlElement, InTransform, InTransformType, bForce](const FRigHierarchyListener& Listener)
 		{
-			URigHierarchy* ListeningHierarchy = Listener.Hierarchy.Get();
-			if (ListeningHierarchy)
+			if (URigHierarchy* ListeningHierarchy = Listener.Hierarchy.Get())
 			{	
 				if(FRigControlElement* ListeningElement = Cast<FRigControlElement>(ListeningHierarchy->Find(InControlElement->GetKey())))
 				{
@@ -3793,7 +3799,7 @@ void URigHierarchy::SetControlShapeTransform(FRigControlElement* InControlElemen
 					ListeningHierarchy->SetControlShapeTransform(ListeningElement, InTransform, InTransformType, false, bForce);
 				}
 			}
-		}
+		});
 	}
 
 	if (bPrintPythonCommands)
@@ -3850,10 +3856,9 @@ void URigHierarchy::SetControlSettings(FRigControlElement* InControlElement, FRi
 	{
 		TGuardValue<bool> bPropagatingChangeGuardValue(bPropagatingChange, true);
 			
-		for(FRigHierarchyListener& Listener : ListeningHierarchies)
+		ForEachListeningHierarchy([this, InControlElement, InSettings, bForce](const FRigHierarchyListener& Listener)
 		{
-			URigHierarchy* ListeningHierarchy = Listener.Hierarchy.Get();
-			if (ListeningHierarchy)
+			if (URigHierarchy* ListeningHierarchy = Listener.Hierarchy.Get())
 			{	
 				if(FRigControlElement* ListeningElement = Cast<FRigControlElement>(ListeningHierarchy->Find(InControlElement->GetKey())))
 				{
@@ -3861,7 +3866,7 @@ void URigHierarchy::SetControlSettings(FRigControlElement* InControlElement, FRi
 					ListeningHierarchy->SetControlSettings(ListeningElement, InSettings, false, bForce);
 				}
 			}
-		}
+		});
 	}
 
 	if (bPrintPythonCommands)
@@ -4128,18 +4133,16 @@ void URigHierarchy::SetControlValue(FRigControlElement* InControlElement, const
 				{
 					TGuardValue<bool> bPropagatingChangeGuardValue(bPropagatingChange, true);
 			
-					for(FRigHierarchyListener& Listener : ListeningHierarchies)
+					ForEachListeningHierarchy([this, InControlElement, InValue, InValueType, bForce](const FRigHierarchyListener& Listener)
 					{
-						URigHierarchy* ListeningHierarchy = Listener.Hierarchy.Get();
-					
-						if (ListeningHierarchy)
+						if (URigHierarchy* ListeningHierarchy = Listener.Hierarchy.Get())
 						{
 							if(FRigControlElement* ListeningElement = Cast<FRigControlElement>(ListeningHierarchy->Find(InControlElement->GetKey())))
 							{
 								ListeningHierarchy->SetControlValue(ListeningElement, InValue, InValueType, false, bForce);
 							}
 						}
-					}
+					});
 				}
 
 				if (bPrintPythonCommands)
@@ -4189,17 +4192,16 @@ void URigHierarchy::SetControlVisibility(FRigControlElement* InControlElement, b
 	{
 		TGuardValue<bool> bPropagatingChangeGuardValue(bPropagatingChange, true);
 			
-		for(FRigHierarchyListener& Listener : ListeningHierarchies)
+		ForEachListeningHierarchy([this, InControlElement, bVisibility](const FRigHierarchyListener& Listener)
 		{
-			URigHierarchy* ListeningHierarchy = Listener.Hierarchy.Get();
-			if (ListeningHierarchy)
+			if (URigHierarchy* ListeningHierarchy = Listener.Hierarchy.Get())
 			{
 				if(FRigControlElement* ListeningElement = Cast<FRigControlElement>(ListeningHierarchy->Find(InControlElement->GetKey())))
 				{
 					ListeningHierarchy->SetControlVisibility(ListeningElement, bVisibility);
 				}
 			}
-		}
+		});
 	}
 #endif
 }
@@ -4240,10 +4242,9 @@ void URigHierarchy::SetConnectorSettings(FRigConnectorElement* InConnectorElemen
 	{
 		TGuardValue<bool> bPropagatingChangeGuardValue(bPropagatingChange, true);
 			
-		for(FRigHierarchyListener& Listener : ListeningHierarchies)
+		ForEachListeningHierarchy([this, InConnectorElement, InSettings, bForce](const FRigHierarchyListener& Listener)
 		{
-			URigHierarchy* ListeningHierarchy = Listener.Hierarchy.Get();
-			if (ListeningHierarchy)
+			if (URigHierarchy* ListeningHierarchy = Listener.Hierarchy.Get())
 			{	
 				if(FRigConnectorElement* ListeningElement = Cast<FRigConnectorElement>(ListeningHierarchy->Find(InConnectorElement->GetKey())))
 				{
@@ -4251,7 +4252,7 @@ void URigHierarchy::SetConnectorSettings(FRigConnectorElement* InConnectorElemen
 					ListeningHierarchy->SetConnectorSettings(ListeningElement, InSettings, false, bForce);
 				}
 			}
-		}
+		});
 	}
 
 	if (bPrintPythonCommands)
@@ -4334,11 +4335,11 @@ void URigHierarchy::SetCurveValue(FRigCurveElement* InCurveElement, float InValu
 	{
 		TGuardValue<bool> bPropagatingChangeGuardValue(bPropagatingChange, true);
 			
-		for(FRigHierarchyListener& Listener : ListeningHierarchies)
+		ForEachListeningHierarchy([this, InCurveElement, InValue, bForce](const FRigHierarchyListener& Listener)
 		{
 			if(!Listener.Hierarchy.IsValid())
 			{
-				continue;
+				return;
 			}
 
 			if (URigHierarchy* ListeningHierarchy = Listener.Hierarchy.Get())
@@ -4349,7 +4350,7 @@ void URigHierarchy::SetCurveValue(FRigCurveElement* InCurveElement, float InValu
 					ListeningHierarchy->SetCurveValue(ListeningElement, InValue, false, bForce);
 				}
 			}
-		}
+		});
 	}
 #endif
 }
@@ -4381,11 +4382,11 @@ void URigHierarchy::UnsetCurveValue(FRigCurveElement* InCurveElement, bool bSetu
 	{
 		TGuardValue<bool> bPropagatingChangeGuardValue(bPropagatingChange, true);
 			
-		for(FRigHierarchyListener& Listener : ListeningHierarchies)
+		ForEachListeningHierarchy([this, InCurveElement, bForce](const FRigHierarchyListener& Listener)
 		{
 			if(!Listener.Hierarchy.IsValid())
 			{
-				continue;
+				return;
 			}
 
 			if (URigHierarchy* ListeningHierarchy = Listener.Hierarchy.Get())
@@ -4396,7 +4397,7 @@ void URigHierarchy::UnsetCurveValue(FRigCurveElement* InCurveElement, bool bSetu
 					ListeningHierarchy->UnsetCurveValue(ListeningElement, false, bForce);
 				}
 			}
-		}
+		});
 	}
 #endif
 }
diff --git a/Engine/Plugins/Animation/ControlRig/Source/ControlRig/Public/Rigs/RigHierarchy.h b/Engine/Plugins/Animation/ControlRig/Source/ControlRig/Public/Rigs/RigHierarchy.h
--- a/Engine/Plugins/Animation/ControlRig/Source/ControlRig/Public/Rigs/RigHierarchy.h
+++ b/Engine/Plugins/Animation/ControlRig/Source/ControlRig/Public/Rigs/RigHierarchy.h
@@ -4687,6 +4687,10 @@ private:
 	void SendQueuedNotifications();
 	void Reset_Impl(bool bResetElements);
 
+#if WITH_EDITOR
+	void ForEachListeningHierarchy(TFunctionRef<void(const FRigHierarchyListener&)> PerListeningHierarchyFunction);
+#endif
+
 	// the currently destroyed element - used to avoid notification storms
 	const FRigBaseElement* ElementBeingDestroyed; 
 	
