Control Rig: Enable debug drawing in modular rig

#rnx

diff --git a/Engine/Plugins/Animation/ControlRig/Source/ControlRig/Private/ModularRig.cpp b/Engine/Plugins/Animation/ControlRig/Source/ControlRig/Private/ModularRig.cpp
--- a/Engine/Plugins/Animation/ControlRig/Source/ControlRig/Private/ModularRig.cpp
+++ b/Engine/Plugins/Animation/ControlRig/Source/ControlRig/Private/ModularRig.cpp
@@ -296,6 +296,10 @@ void UModularRig::ExecuteQueue()
 			FRigUnitContext& RigUnitContext = RigPublicContext.UnitContext;
 			RigUnitContext = PublicContext.UnitContext;
 
+			// forward the draw interface to each module
+			RigPublicContext.SetDrawInterface(PublicContext.GetDrawInterface());
+			RigPublicContext.SetDrawContainer(PublicContext.GetDrawContainer());
+
 			// re-initialize the module in case only the VM side got recompiled.
 			// this happens when the user relies on auto recompilation when editing the
 			// module (dependency) graph - by changing a value, add / remove nodes or links.
