Control Rig: Disable preferred euler angles for Layered rigs
#rnx

diff --git a/Engine/Plugins/Animation/ControlRig/Source/ControlRig/Private/ControlRig.cpp b/Engine/Plugins/Animation/ControlRig/Source/ControlRig/Private/ControlRig.cpp
--- a/Engine/Plugins/Animation/ControlRig/Source/ControlRig/Private/ControlRig.cpp
+++ b/Engine/Plugins/Animation/ControlRig/Source/ControlRig/Private/ControlRig.cpp
@@ -475,7 +475,6 @@ void UControlRig::Evaluate_AnyThread()
 
 						const bool bSetupUndo = false; // Rely on the sequencer track to handle undo/redo
 						Hierarchy->SetControlValue(Control, FinalValue, ERigControlValueType::Current, bSetupUndo, false, Info.bPrintPythonCommnds, false);
-						Hierarchy->SetPreferredEulerAnglesFromValue(Control, AdditiveValue, ERigControlValueType::Current, Info.bFixEulerFlips);
 					}
 
 					if (Info.bNotify && OnControlModified.IsBound())
@@ -1757,6 +1756,8 @@ void UControlRig::PostLoad()
 		NewInfluences.FindOrAdd(EventName).Merge(Map, true);
 	}
 	Influences = NewInfluences;
+
+	GetHierarchy()->bUsePreferredEulerAngles = !bIsAdditive;
 }
 
 const FRigModuleSettings& UControlRig::GetRigModuleSettings() const
@@ -1960,7 +1961,7 @@ FRigControlValue UControlRig::GetControlValue(FRigControlElement* InControl, con
 			
 			// return local space control value (the one to be added after backwards solve)
 			const FRigPoseElement& AnimPose = ControlsAfterBackwardsSolve[ControlIndex];
-			const FRigControlValue& CurrentValue = GetHierarchy()->GetControlValue(InControl, InValueType);
+			const FRigControlValue& CurrentValue = GetHierarchy()->GetControlValue(InControl, InValueType, false);
 			const FTransform FinalTransform = CurrentValue.GetAsTransform(InControl->Settings.ControlType, InControl->Settings.PrimaryAxis);
 			const FTransform AdditiveTransform = FinalTransform * AnimPose.LocalTransform.Inverse();
 			FRigControlValue AdditiveValue;
diff --git a/Engine/Plugins/Animation/ControlRig/Source/ControlRig/Private/Rigs/RigHierarchy.cpp b/Engine/Plugins/Animation/ControlRig/Source/ControlRig/Private/Rigs/RigHierarchy.cpp
--- a/Engine/Plugins/Animation/ControlRig/Source/ControlRig/Private/Rigs/RigHierarchy.cpp
+++ b/Engine/Plugins/Animation/ControlRig/Source/ControlRig/Private/Rigs/RigHierarchy.cpp
@@ -119,7 +119,7 @@ URigHierarchy::URigHierarchy()
 #endif
 , bEnableCacheValidityCheck(bEnableValidityCheckbyDefault)
 , HierarchyForCacheValidation()
-, bUpdatePreferredEulerAngleWhenSettingTransform(true)
+, bUsePreferredEulerAngles(true)
 , bAllowNameSpaceWhenSanitizingName(false)
 , ExecuteContext(nullptr)
 #if WITH_EDITOR
@@ -3404,7 +3404,7 @@ void URigHierarchy::SetTransform(FRigTransformElement* InTransformElement, const
 	{
 		ControlElement->Shape.MarkDirty(MakeGlobal(InTransformType));
 
-		if(bUpdatePreferredEulerAngleWhenSettingTransform && ERigTransformType::IsLocal(InTransformType))
+		if(bUsePreferredEulerAngles && ERigTransformType::IsLocal(InTransformType))
 		{
 			const bool bInitial = ERigTransformType::IsInitial(InTransformType);
 			ControlElement->PreferredEulerAngles.SetRotator(InTransform.Rotator(), bInitial, true);
@@ -4065,7 +4065,7 @@ void URigHierarchy::SetControlValue(FRigControlElement* InControlElement, const
 				FRigControlValue Value = InValue;
 				InControlElement->Settings.ApplyLimits(Value);
 
-				TGuardValue<bool> DontSetPreferredEulerAngle(bUpdatePreferredEulerAngleWhenSettingTransform, false);
+				TGuardValue<bool> DontSetPreferredEulerAngle(bUsePreferredEulerAngles, false);
 				SetTransform(
 					InControlElement,
 					Value.GetAsTransform(
@@ -4089,7 +4089,7 @@ void URigHierarchy::SetControlValue(FRigControlElement* InControlElement, const
 				FRigControlValue Value = InValue;
 				InControlElement->Settings.ApplyLimits(Value);
 
-				TGuardValue<bool> DontSetPreferredEulerAngle(bUpdatePreferredEulerAngleWhenSettingTransform, false);
+				TGuardValue<bool> DontSetPreferredEulerAngle(bUsePreferredEulerAngles, false);
 				SetTransform(
 					InControlElement,
 					Value.GetAsTransform(
diff --git a/Engine/Plugins/Animation/ControlRig/Source/ControlRig/Public/ControlRig.h b/Engine/Plugins/Animation/ControlRig/Source/ControlRig/Public/ControlRig.h
--- a/Engine/Plugins/Animation/ControlRig/Source/ControlRig/Public/ControlRig.h
+++ b/Engine/Plugins/Animation/ControlRig/Source/ControlRig/Public/ControlRig.h
@@ -691,6 +691,7 @@ public:
 	void SetIsAdditive(const bool bInIsAdditive)
 	{
 		bIsAdditive = bInIsAdditive;
+		GetHierarchy()->bUsePreferredEulerAngles = !bIsAdditive;
 	}
 
 private:
diff --git a/Engine/Plugins/Animation/ControlRig/Source/ControlRig/Public/Rigs/RigHierarchy.h b/Engine/Plugins/Animation/ControlRig/Source/ControlRig/Public/Rigs/RigHierarchy.h
--- a/Engine/Plugins/Animation/ControlRig/Source/ControlRig/Public/Rigs/RigHierarchy.h
+++ b/Engine/Plugins/Animation/ControlRig/Source/ControlRig/Public/Rigs/RigHierarchy.h
@@ -2050,7 +2050,7 @@ public:
 		{
 			if(FRigControlElement* ControlElement = Cast<FRigControlElement>(Elements[InElementIndex]))
 			{
-				return GetControlValue(ControlElement, InValueType);
+				return GetControlValue(ControlElement, InValueType, bUsePreferredEulerAngles);
 			}
 		}
 		return FRigControlValue();
@@ -2141,7 +2141,12 @@ public:
 	{
 		if(InControlElement)
 		{
-			return InControlElement->PreferredEulerAngles.GetRotator(bInitial);
+			if (bUsePreferredEulerAngles)
+			{
+				return InControlElement->PreferredEulerAngles.GetRotator(bInitial);
+			}
+			const ERigTransformType::Type Type = bInitial ? ERigTransformType::InitialLocal : ERigTransformType::CurrentLocal;
+			return GetControlValue(InControlElement->GetKey()).GetAsTransform(InControlElement->Settings.ControlType, InControlElement->Settings.PrimaryAxis).Rotator();
 		}
 		return FRotator::ZeroRotator;
 	}
@@ -4621,7 +4626,7 @@ protected:
 
 	mutable TMap<FRigElementKey, FRigElementKey> DefaultParentPerElement;
 
-	bool bUpdatePreferredEulerAngleWhenSettingTransform;
+	bool bUsePreferredEulerAngles;
 	mutable bool bAllowNameSpaceWhenSanitizingName;
 
 private:
