Control Rig: FIx crash in control rig post load
#rnx

diff --git a/Engine/Plugins/Animation/ControlRig/Source/ControlRig/Private/ControlRig.cpp b/Engine/Plugins/Animation/ControlRig/Source/ControlRig/Private/ControlRig.cpp
--- a/Engine/Plugins/Animation/ControlRig/Source/ControlRig/Private/ControlRig.cpp
+++ b/Engine/Plugins/Animation/ControlRig/Source/ControlRig/Private/ControlRig.cpp
@@ -1757,7 +1757,10 @@ void UControlRig::PostLoad()
 	}
 	Influences = NewInfluences;
 
-	GetHierarchy()->bUsePreferredEulerAngles = !bIsAdditive;
+	if(DynamicHierarchy)
+	{
+		DynamicHierarchy->bUsePreferredEulerAngles = !bIsAdditive;
+	}
 }
 
 const FRigModuleSettings& UControlRig::GetRigModuleSettings() const
