Control Rig: Preparing socket for becoming the sole connection point

#rnx

diff --git a/Engine/Plugins/Animation/ControlRig/Source/ControlRig/Private/Rigs/RigHierarchy.cpp b/Engine/Plugins/Animation/ControlRig/Source/ControlRig/Private/Rigs/RigHierarchy.cpp
--- a/Engine/Plugins/Animation/ControlRig/Source/ControlRig/Private/Rigs/RigHierarchy.cpp
+++ b/Engine/Plugins/Animation/ControlRig/Source/ControlRig/Private/Rigs/RigHierarchy.cpp
@@ -1050,33 +1050,71 @@ bool URigHierarchy::IsProcedural(const FRigBaseElement* InElement) const
 	return InElement->IsProcedural();
 }
 
-TArray<FRigConnectorInfo> URigHierarchy::GetConnectorInfos() const
+TArray<FRigSocketState> URigHierarchy::GetSocketStates() const
+{
+	const TArray<FRigElementKey> Keys = GetSocketKeys(true);
+	TArray<FRigSocketState> States;
+	States.Reserve(Keys.Num());
+	for(const FRigElementKey& Key : Keys)
+	{
+		const FRigSocketElement* Socket = FindChecked<FRigSocketElement>(Key);
+		States.Add(Socket->GetSocketState(this));
+	}
+	return States;
+}
+
+TArray<FRigElementKey> URigHierarchy::RestoreSocketsFromStates(TArray<FRigSocketState> InStates, bool bSetupUndoRedo)
+{
+	TArray<FRigElementKey> Keys;
+	for(const FRigSocketState& State : InStates)
+	{
+		FRigElementKey Key(State.Name, ERigElementType::Socket);
+
+		if(FRigSocketElement* Socket = Find<FRigSocketElement>(Key))
+		{
+			(void)GetController()->SetParent(Key, State.Parent);
+			Socket->SetColor(State.Color, this);
+			Socket->SetDescription(State.Description, this);
+			SetInitialLocalTransform(Key, State.InitialLocalTransform);
+			SetLocalTransform(Key, State.InitialLocalTransform);
+		}
+		else
+		{
+			Key = GetController()->AddSocket(State.Name, State.Parent, State.InitialLocalTransform, false, State.Color, State.Description, bSetupUndoRedo, false);
+		}
+
+		Keys.Add(Key);
+	}
+	return Keys;
+}
+
+TArray<FRigConnectorState> URigHierarchy::GetConnectorStates() const
 {
 	const TArray<FRigElementKey> Keys = GetConnectorKeys(true);
-	TArray<FRigConnectorInfo> Infos;
-	Infos.Reserve(Keys.Num());
+	TArray<FRigConnectorState> States;
+	States.Reserve(Keys.Num());
 	for(const FRigElementKey& Key : Keys)
 	{
 		const FRigConnectorElement* Connector = FindChecked<FRigConnectorElement>(Key);
-		Infos.Add(Connector->GetConnectorInfo(this));
+		States.Add(Connector->GetConnectorState(this));
 	}
-	return Infos;
+	return States;
 }
 
-TArray<FRigElementKey> URigHierarchy::RestoreConnectorsFromInfos(TArray<FRigConnectorInfo> InInfos, bool bSetupUndoRedo)
+TArray<FRigElementKey> URigHierarchy::RestoreConnectorsFromStates(TArray<FRigConnectorState> InStates, bool bSetupUndoRedo)
 {
 	TArray<FRigElementKey> Keys;
-	for(const FRigConnectorInfo& Info : InInfos)
+	for(const FRigConnectorState& State : InStates)
 	{
-		FRigElementKey Key(Info.Name, ERigElementType::Connector);
+		FRigElementKey Key(State.Name, ERigElementType::Connector);
 
 		if(const FRigConnectorElement* Connector = Find<FRigConnectorElement>(Key))
 		{
-			SetConnectorSettings(Key, Info.Settings, bSetupUndoRedo, false, false);
+			SetConnectorSettings(Key, State.Settings, bSetupUndoRedo, false, false);
 		}
 		else
 		{
-			Key = GetController()->AddConnector(Info.Name, Info.Settings, bSetupUndoRedo, false);
+			Key = GetController()->AddConnector(State.Name, State.Settings, bSetupUndoRedo, false);
 		}
 
 		Keys.Add(Key);
diff --git a/Engine/Plugins/Animation/ControlRig/Source/ControlRig/Private/Rigs/RigHierarchyController.cpp b/Engine/Plugins/Animation/ControlRig/Source/ControlRig/Private/Rigs/RigHierarchyController.cpp
--- a/Engine/Plugins/Animation/ControlRig/Source/ControlRig/Private/Rigs/RigHierarchyController.cpp
+++ b/Engine/Plugins/Animation/ControlRig/Source/ControlRig/Private/Rigs/RigHierarchyController.cpp
@@ -578,6 +578,28 @@ FRigElementKey URigHierarchyController::AddConnector(FName InName, FRigConnector
 		return FRigElementKey();
 	}
 
+	// only allow to add one primary connector
+	if(InSettings.Type == EConnectorType::Primary)
+	{
+		const TArray<FRigConnectorElement*>& Connectors = Hierarchy->GetConnectors();
+		for(const FRigConnectorElement* Connector : Connectors)
+		{
+			if(Connector->Settings.Type == EConnectorType::Primary)
+			{
+				static constexpr TCHAR Format[] = TEXT("Cannot add connector '%s' - there already is a primary connector.");
+				ReportAndNotifyErrorf(Format, *InName.ToString());
+				return FRigElementKey();
+			}
+		}
+
+		if(InSettings.bOptional)
+		{
+			static constexpr TCHAR Format[] = TEXT("Cannot add connector '%s' - primary connectors cannot be optional.");
+			ReportAndNotifyErrorf(Format, *InName.ToString());
+			return FRigElementKey();
+		}
+	}
+
 #if WITH_EDITOR
 	TSharedPtr<FScopedTransaction> TransactionPtr;
 	if(bSetupUndo)
@@ -620,41 +642,47 @@ FRigElementKey URigHierarchyController::AddConnector(FName InName, FRigConnector
 }
 
 FRigElementKey URigHierarchyController::AddSocket(FName InName, FRigElementKey InParent, FTransform InTransform,
-	bool bTransformInGlobal, bool bSetupUndo, bool bPrintPythonCommand)
+	bool bTransformInGlobal, const FLinearColor& InColor, const FString& InDescription, bool bSetupUndo, bool bPrintPythonCommand)
 {
 	if(!IsValid())
 	{
 		return FRigElementKey();
 	}
 
+	URigHierarchy* CurrentHierarchy = GetHierarchy();
+
 #if WITH_EDITOR
 	TSharedPtr<FScopedTransaction> TransactionPtr;
 	if(bSetupUndo)
 	{
 		TransactionPtr = MakeShared<FScopedTransaction>(NSLOCTEXT("RigHierarchyController", "Add Socket", "Add Socket"));
-		Hierarchy->Modify();
+		CurrentHierarchy->Modify();
 	}
 #endif
 
 	FRigSocketElement* NewElement = MakeElement<FRigSocketElement>();
 	{
-		TGuardValue<bool> DisableCacheValidityChecks(Hierarchy->bEnableCacheValidityCheck, false);
+		TGuardValue<bool> DisableCacheValidityChecks(CurrentHierarchy->bEnableCacheValidityCheck, false);
 		NewElement->Key.Type = ERigElementType::Socket;
 		NewElement->Key.Name = GetSafeNewName(InName, NewElement->Key.Type);
-		AddElement(NewElement, Hierarchy->Get(Hierarchy->GetIndex(InParent)), true, InName);
+		AddElement(NewElement, CurrentHierarchy->Get(CurrentHierarchy->GetIndex(InParent)), true, InName);
 
 		if(bTransformInGlobal)
 		{
-			Hierarchy->SetTransform(NewElement, InTransform, ERigTransformType::InitialGlobal, true, false);
-			Hierarchy->SetTransform(NewElement, InTransform, ERigTransformType::CurrentGlobal, true, false);
+			CurrentHierarchy->SetTransform(NewElement, InTransform, ERigTransformType::InitialGlobal, true, false);
+			CurrentHierarchy->SetTransform(NewElement, InTransform, ERigTransformType::CurrentGlobal, true, false);
 		}
 		else
 		{
-			Hierarchy->SetTransform(NewElement, InTransform, ERigTransformType::InitialLocal, true, false);
-			Hierarchy->SetTransform(NewElement, InTransform, ERigTransformType::CurrentLocal, true, false);
+			CurrentHierarchy->SetTransform(NewElement, InTransform, ERigTransformType::InitialLocal, true, false);
+			CurrentHierarchy->SetTransform(NewElement, InTransform, ERigTransformType::CurrentLocal, true, false);
 		}
 
 		NewElement->Pose.Current = NewElement->Pose.Initial;
+
+		NewElement->SetColor(InColor, CurrentHierarchy);
+		NewElement->SetDescription(InDescription, CurrentHierarchy);
+		CurrentHierarchy->SetRigElementKeyMetadata(NewElement->Key, FRigSocketElement::DesiredParentMetaName, InParent);
 	}
 
 #if WITH_EDITOR
@@ -675,7 +703,7 @@ FRigElementKey URigHierarchyController::AddSocket(FName InName, FRigElementKey I
 	}
 #endif
 
-	Hierarchy->EnsureCacheValidity();
+	CurrentHierarchy->EnsureCacheValidity();
 		
 	return NewElement->Key;
 }
@@ -1810,11 +1838,15 @@ TArray<FString> URigHierarchyController::GetAddSocketPythonCommands(FRigSocketEl
 		ParentKeyStr = Socket->ParentElement->GetKey().ToPythonString();
 	}
 
-	// AddSocket(FName InName, FRigElementKey InParent, FTransform InTransform, bool bTransformInGlobal = true, bool bSetupUndo = false);
-	Commands.Add(FString::Printf(TEXT("hierarchy_controller.add_socket('%s', %s, %s, False)"),
+	const URigHierarchy* CurrentHierarchy = GetHierarchy();
+
+	// AddSocket(FName InName, FRigElementKey InParent, FTransform InTransform, bool bTransformInGlobal = true, FLinearColor Color, FString Description, bool bSetupUndo = false);
+	Commands.Add(FString::Printf(TEXT("hierarchy_controller.add_socket('%s', %s, %s, False, %s, '%s')"),
 		*Socket->GetName(),
 		*ParentKeyStr,
-		*TransformStr
+		*TransformStr,
+		*RigVMPythonUtils::LinearColorToPythonString(Socket->GetColor(CurrentHierarchy)),
+		*Socket->GetDescription(CurrentHierarchy)
 	));
 
 	return Commands;
@@ -2660,6 +2692,12 @@ bool URigHierarchyController::AddParent(FRigBaseElement* InChild, FRigBaseElemen
 		AddElementToDirty(Constraint.ParentElement, SingleParentElement);
 		SingleParentElement->ParentElement = Constraint.ParentElement;
 
+		if(Cast<FRigSocketElement>(SingleParentElement))
+		{
+			Hierarchy->SetRigElementKeyMetadata(SingleParentElement->GetKey(), FRigSocketElement::DesiredParentMetaName, Constraint.ParentElement->GetKey());
+			Hierarchy->Notify(ERigHierarchyNotification::SocketDesiredParentChanged, SingleParentElement);
+		}
+
 		Hierarchy->IncrementTopologyVersion();
 
 		if(!bMaintainGlobalTransform)
@@ -2818,6 +2856,13 @@ bool URigHierarchyController::RemoveParent(FRigBaseElement* InChild, FRigBaseEle
 			
 			// remove the previous parent
 			SingleParentElement->ParentElement = nullptr;
+
+			if(Cast<FRigSocketElement>(SingleParentElement))
+			{
+				Hierarchy->RemoveMetadata(SingleParentElement->GetKey(), FRigSocketElement::DesiredParentMetaName);
+				Hierarchy->Notify(ERigHierarchyNotification::SocketDesiredParentChanged, SingleParentElement);
+			}
+
 			RemoveElementToDirty(InParent, SingleParentElement); 
 			Hierarchy->IncrementTopologyVersion();
 
@@ -3003,6 +3048,12 @@ bool URigHierarchyController::SetParent(FRigElementKey InChild, FRigElementKey I
 	FRigBaseElement* Parent = Hierarchy->Find(InParent);
 	if(Parent == nullptr)
 	{
+		if(InChild.Type == ERigElementType::Socket)
+		{
+			Hierarchy->SetRigElementKeyMetadata(InChild, FRigSocketElement::DesiredParentMetaName, InParent);
+			Hierarchy->Notify(ERigHierarchyNotification::SocketDesiredParentChanged, Child);
+			return true;
+		}
 		ReportWarningf(TEXT("Cannot Set Parent, Parent '%s' not found."), *InParent.ToString());
 		return false;
 	}
diff --git a/Engine/Plugins/Animation/ControlRig/Source/ControlRig/Private/Rigs/RigHierarchyElements.cpp b/Engine/Plugins/Animation/ControlRig/Source/ControlRig/Private/Rigs/RigHierarchyElements.cpp
--- a/Engine/Plugins/Animation/ControlRig/Source/ControlRig/Private/Rigs/RigHierarchyElements.cpp
+++ b/Engine/Plugins/Animation/ControlRig/Source/ControlRig/Private/Rigs/RigHierarchyElements.cpp
@@ -1433,6 +1433,13 @@ FRigConnectorSettings::FRigConnectorSettings()
 {
 }
 
+FRigConnectorSettings FRigConnectorSettings::DefaultSettings()
+{
+	FRigConnectorSettings Settings;
+	Settings.AddRule(FRigTypeConnectionRule(ERigElementType::Socket));
+	return Settings;
+}
+
 void FRigConnectorSettings::Save(FArchive& Ar)
 {
 	Ar.UsingCustomVersion(FControlRigObjectVersion::GUID);
@@ -1637,13 +1644,13 @@ bool FRigConnectorElement::CanConnect(const FRigConnectionInfo* InConnectionInfo
 	return true;
 }
 
-FRigConnectorInfo FRigConnectorElement::GetConnectorInfo(const URigHierarchy* InHierarchy) const
+FRigConnectorState FRigConnectorElement::GetConnectorState(const URigHierarchy* InHierarchy) const
 {
-	FRigConnectorInfo Info;
-	Info.Name = Key.Name;
-	Info.ResolvedTarget = InHierarchy->GetResolvedTarget(Key);
-	Info.Settings = Settings;
-	return Info;
+	FRigConnectorState State;
+	State.Name = Key.Name;
+	State.ResolvedTarget = InHierarchy->GetResolvedTarget(Key);
+	State.Settings = Settings;
+	return State;
 }
 
 void FRigConnectorElement::CopyFrom(const FRigBaseElement* InOther)
@@ -1658,7 +1665,18 @@ void FRigConnectorElement::CopyFrom(const FRigBaseElement* InOther)
 // FRigSocketElement
 ////////////////////////////////////////////////////////////////////////////////
 
+FRigSocketState::FRigSocketState()
+: Name(NAME_None)
+, InitialLocalTransform(FTransform::Identity)
+, Color(FRigSocketElement::SocketDefaultColor)
+{
+}
+
 const FRigBaseElement::EElementIndex FRigSocketElement::ElementTypeIndex = SocketElement;
+const FName FRigSocketElement::ColorMetaName = TEXT("SocketColor");
+const FName FRigSocketElement::DescriptionMetaName = TEXT("SocketDescription");
+const FName FRigSocketElement::DesiredParentMetaName = TEXT("SocketDesiredParent");
+const FLinearColor FRigSocketElement::SocketDefaultColor = FLinearColor::White;
 
 void FRigSocketElement::Save(FArchive& Ar, ESerializationPhase SerializationPhase)
 {
@@ -1670,6 +1688,63 @@ void FRigSocketElement::Load(FArchive& Ar, ESerializationPhase SerializationPhas
 	Super::Load(Ar, SerializationPhase);
 }
 
+FRigSocketState FRigSocketElement::GetSocketState(const URigHierarchy* InHierarchy) const
+{
+	FRigSocketState State;
+	State.Name = GetFName();
+	State.Parent = InHierarchy->GetRigElementKeyMetadata(GetKey(), DesiredParentMetaName, FRigElementKey());
+	if(!State.Parent.IsValid())
+	{
+		State.Parent = InHierarchy->GetFirstParent(GetKey());
+	}
+	State.InitialLocalTransform = InHierarchy->GetInitialLocalTransform(GetIndex());
+	State.Color = GetColor(InHierarchy);
+	State.Description = GetDescription(InHierarchy);
+	return State;
+}
+
+FLinearColor FRigSocketElement::GetColor(const URigHierarchy* InHierarchy) const
+{
+	return InHierarchy->GetLinearColorMetadata(GetKey(), ColorMetaName, SocketDefaultColor);
+}
+
+void FRigSocketElement::SetColor(const FLinearColor& InColor, URigHierarchy* InHierarchy, bool bNotify)
+{
+	if(InHierarchy->GetLinearColorMetadata(GetKey(), ColorMetaName, SocketDefaultColor).Equals(InColor))
+	{
+		return;
+	}
+	InHierarchy->SetLinearColorMetadata(GetKey(), ColorMetaName, InColor);
+	if(bNotify)
+	{
+		InHierarchy->Notify(ERigHierarchyNotification::SocketColorChanged, this);
+	}
+}
+
+FString FRigSocketElement::GetDescription(const URigHierarchy* InHierarchy) const
+{
+	const FName Description = InHierarchy->GetNameMetadata(GetKey(), DescriptionMetaName, NAME_None);
+	if(Description.IsNone())
+	{
+		return FString();
+	}
+	return Description.ToString();
+}
+
+void FRigSocketElement::SetDescription(const FString& InDescription, URigHierarchy* InHierarchy, bool bNotify)
+{
+	const FName Description = InDescription.IsEmpty() ? FName(NAME_None) : *InDescription;
+	if(InHierarchy->GetNameMetadata(GetKey(), DescriptionMetaName, NAME_None).IsEqual(Description, ENameCase::CaseSensitive))
+	{
+		return;
+	}
+	InHierarchy->SetNameMetadata(GetKey(), DescriptionMetaName, *InDescription);
+	if(bNotify)
+	{
+		InHierarchy->Notify(ERigHierarchyNotification::SocketDescriptionChanged, this);
+	}
+}
+
 void FRigSocketElement::CopyFrom(const FRigBaseElement* InOther)
 {
 	Super::CopyFrom(InOther);
diff --git a/Engine/Plugins/Animation/ControlRig/Source/ControlRig/Private/Units/Execution/RigUnit_DynamicHierarchy.cpp b/Engine/Plugins/Animation/ControlRig/Source/ControlRig/Private/Units/Execution/RigUnit_DynamicHierarchy.cpp
--- a/Engine/Plugins/Animation/ControlRig/Source/ControlRig/Private/Units/Execution/RigUnit_DynamicHierarchy.cpp
+++ b/Engine/Plugins/Animation/ControlRig/Source/ControlRig/Private/Units/Execution/RigUnit_DynamicHierarchy.cpp
@@ -1289,6 +1289,6 @@ FRigUnit_HierarchyAddSocket_Execute()
 	if(URigHierarchyController* Controller = ExecuteContext.Hierarchy->GetController(true))
 	{
 		FRigHierarchyControllerInstructionBracket InstructionBracket(Controller, ExecuteContext.GetInstructionIndex());
-		Item = Controller->AddSocket(Name, Parent, Transform, Space == ERigVMTransformSpace::GlobalSpace, false, false);
+		Item = Controller->AddSocket(Name, Parent, Transform, Space == ERigVMTransformSpace::GlobalSpace, Color, Description, false, false);
 	}
 }
diff --git a/Engine/Plugins/Animation/ControlRig/Source/ControlRig/Public/Rigs/RigHierarchy.h b/Engine/Plugins/Animation/ControlRig/Source/ControlRig/Public/Rigs/RigHierarchy.h
--- a/Engine/Plugins/Animation/ControlRig/Source/ControlRig/Public/Rigs/RigHierarchy.h
+++ b/Engine/Plugins/Animation/ControlRig/Source/ControlRig/Public/Rigs/RigHierarchy.h
@@ -977,16 +977,28 @@ public:
 	}
 
 	/**
-	 * Returns all of the connectors' infos
+	 * Returns all of the sockets' state
 	 */
 	UFUNCTION(BlueprintPure, Category = URigHierarchy)
-	TArray<FRigConnectorInfo> GetConnectorInfos() const;
+	TArray<FRigSocketState> GetSocketStates() const;
 
 	/**
-	 * Try to restore the connectors from the info structs
+	 * Try to restore the sockets from the state structs
 	 */
 	UFUNCTION(BlueprintCallable, Category = URigHierarchy)
-	TArray<FRigElementKey> RestoreConnectorsFromInfos(TArray<FRigConnectorInfo> InInfos, bool bSetupUndoRedo = false);
+	TArray<FRigElementKey> RestoreSocketsFromStates(TArray<FRigSocketState> InStates, bool bSetupUndoRedo = false);
+
+	/**
+	 * Returns all of the connectors' state
+	 */
+	UFUNCTION(BlueprintPure, Category = URigHierarchy)
+	TArray<FRigConnectorState> GetConnectorStates() const;
+
+	/**
+	 * Try to restore the connectors from the state structs
+	 */
+	UFUNCTION(BlueprintCallable, Category = URigHierarchy)
+	TArray<FRigElementKey> RestoreConnectorsFromStates(TArray<FRigConnectorState> InStates, bool bSetupUndoRedo = false);
 
 	/**
 	 * Returns all Socket elements
diff --git a/Engine/Plugins/Animation/ControlRig/Source/ControlRig/Public/Rigs/RigHierarchyController.h b/Engine/Plugins/Animation/ControlRig/Source/ControlRig/Public/Rigs/RigHierarchyController.h
--- a/Engine/Plugins/Animation/ControlRig/Source/ControlRig/Public/Rigs/RigHierarchyController.h
+++ b/Engine/Plugins/Animation/ControlRig/Source/ControlRig/Public/Rigs/RigHierarchyController.h
@@ -254,12 +254,14 @@ public:
 	 * @param InParent The (optional) parent of the new null. If you don't need a parent, pass FRigElementKey()
 	 * @param InTransform The transform for the new socket - either in local or global space, based on bTransformInGlobal
 	 * @param bTransformInGlobal Set this to true if the Transform passed is expressed in global space, false for local space.
+	 * @param InColor The color of the socket
+	 * @param InDescription The description of the socket
 	 * @param bSetupUndo If set to true the stack will record the change for undo / redo
 	 * @param bPrintPythonCommand If set to true a python command equivalent to this call will be printed out
 	 * @return The key for the newly created bone.
 	 */
 	UFUNCTION(BlueprintCallable, Category = URigHierarchyController)
-	FRigElementKey AddSocket(FName InName, FRigElementKey InParent, FTransform InTransform, bool bTransformInGlobal = true, bool bSetupUndo = false, bool bPrintPythonCommand = false);
+	FRigElementKey AddSocket(FName InName, FRigElementKey InParent, FTransform InTransform, bool bTransformInGlobal = true, const FLinearColor& InColor = FLinearColor::White, const FString& InDescription = TEXT(""), bool bSetupUndo = false, bool bPrintPythonCommand = false);
 
 	/**
 	 * Returns the control settings of a given control
diff --git a/Engine/Plugins/Animation/ControlRig/Source/ControlRig/Public/Rigs/RigHierarchyDefines.h b/Engine/Plugins/Animation/ControlRig/Source/ControlRig/Public/Rigs/RigHierarchyDefines.h
--- a/Engine/Plugins/Animation/ControlRig/Source/ControlRig/Public/Rigs/RigHierarchyDefines.h
+++ b/Engine/Plugins/Animation/ControlRig/Source/ControlRig/Public/Rigs/RigHierarchyDefines.h
@@ -91,6 +91,9 @@ enum class ERigHierarchyNotification : uint8
 	InteractionBracketClosed,
 	ElementReordered,
 	ConnectorSettingChanged,
+	SocketColorChanged,
+	SocketDescriptionChanged,
+	SocketDesiredParentChanged,
 
 	/** MAX - invalid */
 	Max UMETA(Hidden),
diff --git a/Engine/Plugins/Animation/ControlRig/Source/ControlRig/Public/Rigs/RigHierarchyElements.h b/Engine/Plugins/Animation/ControlRig/Source/ControlRig/Public/Rigs/RigHierarchyElements.h
--- a/Engine/Plugins/Animation/ControlRig/Source/ControlRig/Public/Rigs/RigHierarchyElements.h
+++ b/Engine/Plugins/Animation/ControlRig/Source/ControlRig/Public/Rigs/RigHierarchyElements.h
@@ -1657,7 +1657,6 @@ UENUM(BlueprintType)
 enum class EConnectorType : uint8
 {
 	Primary, // Single primary connector, non-optional and always visible. When dropped on another element, this connector will resolve to that element.
-	
 	Secondary, // Could be multiple, can auto-solve (visible if not solved), can be optional
 };
 
@@ -1668,6 +1667,7 @@ struct CONTROLRIG_API FRigConnectorSettings
 	GENERATED_BODY()
 
 	FRigConnectorSettings();
+	static FRigConnectorSettings DefaultSettings();
 
 	void Save(FArchive& Ar);
 	void Load(FArchive& Ar);
@@ -1703,9 +1703,13 @@ struct CONTROLRIG_API FRigConnectorSettings
 };
 
 USTRUCT(BlueprintType)
-struct CONTROLRIG_API FRigConnectorInfo
+struct CONTROLRIG_API FRigConnectorState
 {
 	GENERATED_BODY()
+
+	FRigConnectorState()
+		: Name(NAME_None)
+	{}
 	
 	UPROPERTY(BlueprintReadWrite, EditAnywhere, Category = Connector)
 	FName Name;
@@ -1749,7 +1753,7 @@ struct CONTROLRIG_API FRigConnectorElement final : public FRigBaseElement
 
 	bool CanConnect(const FRigConnectionInfo* InConnectionInfo, FString* OutFailureReason) const;
 
-	FRigConnectorInfo GetConnectorInfo(const URigHierarchy* InHierarchy) const;
+	FRigConnectorState GetConnectorState(const URigHierarchy* InHierarchy) const;
 
 private:
 	explicit FRigConnectorElement(URigHierarchy* InOwner)
@@ -1767,6 +1771,28 @@ private:
 	friend struct FRigBaseElement;
 };
 
+USTRUCT(BlueprintType)
+struct CONTROLRIG_API FRigSocketState
+{
+	GENERATED_BODY()
+	
+	FRigSocketState();
+	
+	UPROPERTY(BlueprintReadWrite, EditAnywhere, Category = Connector)
+	FName Name;
+
+	UPROPERTY(BlueprintReadWrite, EditAnywhere, Category = Connector)
+	FRigElementKey Parent;
+
+	UPROPERTY(BlueprintReadWrite, EditAnywhere, Category = Connector)
+	FTransform InitialLocalTransform;
+
+	UPROPERTY(BlueprintReadWrite, EditAnywhere, Category = Connector)
+	FLinearColor Color;
+
+	UPROPERTY(BlueprintReadWrite, EditAnywhere, Category = Connector)
+	FString Description;
+};
 
 USTRUCT(BlueprintType)
 struct CONTROLRIG_API FRigSocketElement final : public FRigSingleParentElement
@@ -1775,6 +1801,10 @@ struct CONTROLRIG_API FRigSocketElement final : public FRigSingleParentElement
 	DECLARE_RIG_ELEMENT_METHODS(FRigSocketElement)
 
 	static const EElementIndex ElementTypeIndex;
+	static const FName ColorMetaName;
+	static const FName DescriptionMetaName;
+	static const FName DesiredParentMetaName;
+	static const FLinearColor SocketDefaultColor;
 
 	FRigSocketElement()
 		: FRigSocketElement(nullptr)
@@ -1785,6 +1815,14 @@ struct CONTROLRIG_API FRigSocketElement final : public FRigSingleParentElement
 	virtual void Save(FArchive& A, ESerializationPhase SerializationPhase) override;
 	virtual void Load(FArchive& Ar, ESerializationPhase SerializationPhase) override;
 
+	FRigSocketState GetSocketState(const URigHierarchy* InHierarchy) const;
+
+	FLinearColor GetColor(const URigHierarchy* InHierarchy) const;
+	void SetColor(const FLinearColor& InColor, URigHierarchy* InHierarchy, bool bNotify = true);
+
+	FString GetDescription(const URigHierarchy* InHierarchy) const;
+	void SetDescription(const FString& InDescription, URigHierarchy* InHierarchy, bool bNotify = true);
+
 private:
 	explicit FRigSocketElement(URigHierarchy* InOwner)
 		: FRigSingleParentElement(InOwner, ERigElementType::Socket)
diff --git a/Engine/Plugins/Animation/ControlRig/Source/ControlRig/Public/Units/Execution/RigUnit_DynamicHierarchy.h b/Engine/Plugins/Animation/ControlRig/Source/ControlRig/Public/Units/Execution/RigUnit_DynamicHierarchy.h
--- a/Engine/Plugins/Animation/ControlRig/Source/ControlRig/Public/Units/Execution/RigUnit_DynamicHierarchy.h
+++ b/Engine/Plugins/Animation/ControlRig/Source/ControlRig/Public/Units/Execution/RigUnit_DynamicHierarchy.h
@@ -1647,6 +1647,8 @@ struct CONTROLRIG_API FRigUnit_HierarchyAddSocket : public FRigUnit_HierarchyAdd
 		Name = TEXT("NewSocket");
 		Transform = FTransform::Identity;
 		Space = ERigVMTransformSpace::LocalSpace;
+		Color = FLinearColor::White;
+		Description = FString();
 	}
 
 	virtual ERigElementType GetElementTypeToSpawn() const override { return ERigElementType::Socket; }
@@ -1663,6 +1665,18 @@ struct CONTROLRIG_API FRigUnit_HierarchyAddSocket : public FRigUnit_HierarchyAdd
 	UPROPERTY(meta = (Input))
 	ERigVMTransformSpace Space;
 
+	/*
+	 * The color of the socket
+	 */
+	UPROPERTY(meta = (Input))
+	FLinearColor Color;
+
+	/*
+	 * The (optional) description of the socket
+	 */
+	UPROPERTY(meta = (Input))
+	FString Description;
+
 	RIGVM_METHOD()
 	virtual void Execute() override;
 };
diff --git a/Engine/Plugins/Animation/ControlRig/Source/ControlRigEditor/Private/ControlRigElementDetails.cpp b/Engine/Plugins/Animation/ControlRig/Source/ControlRigEditor/Private/ControlRigElementDetails.cpp
--- a/Engine/Plugins/Animation/ControlRig/Source/ControlRigEditor/Private/ControlRigElementDetails.cpp
+++ b/Engine/Plugins/Animation/ControlRig/Source/ControlRigEditor/Private/ControlRigElementDetails.cpp
@@ -9,6 +9,8 @@
 #include "Widgets/Input/SVectorInputBox.h"
 #include "Widgets/Input/SCheckBox.h"
 #include "Widgets/Input/SButton.h"
+#include "Widgets/Colors/SColorBlock.h"
+#include "Widgets/Colors/SColorPicker.h"
 #include "ControlRigBlueprint.h"
 #include "ModularRig.h"
 #include "Graph/ControlRigGraph.h"
@@ -884,7 +886,8 @@ void FRigBaseElementDetails::CustomizeDetails(IDetailLayoutBuilder& DetailBuilde
 	if(!IsAnyElementOfType(ERigElementType::Bone) &&
 		!IsAnyElementOfType(ERigElementType::Control) &&
 		!IsAnyElementOfType(ERigElementType::Null) &&
-		!IsAnyElementOfType(ERigElementType::Connector))
+		!IsAnyElementOfType(ERigElementType::Connector) &&
+		!IsAnyElementOfType(ERigElementType::Socket))
 	{
 		CustomizeMetadata(DetailBuilder);
 	}
@@ -1131,6 +1134,18 @@ bool FRigBaseElementDetails::IsAnyConnectorImported() const
 	});
 }
 
+bool FRigBaseElementDetails::IsAnyConnectorPrimary() const
+{
+	return ContainsElementByPredicate([](const FPerElementInfo& Info)
+	{
+		if(const FRigConnectorElement* Connector = Info.Element.Get<FRigConnectorElement>())
+		{
+			return Connector->Settings.Type == EConnectorType::Primary;
+		}
+		return false;
+	});
+}
+
 bool FRigBaseElementDetails::GetCommonElementType(ERigElementType& OutElementType) const
 {
 	OutElementType = ERigElementType::None;
@@ -5197,8 +5212,8 @@ void FRigNullElementDetails::CustomizeDetails(IDetailLayoutBuilder& DetailBuilde
 void FRigConnectorElementDetails::CustomizeDetails(IDetailLayoutBuilder& DetailBuilder)
 {
 	FRigTransformElementDetails::CustomizeDetails(DetailBuilder);
-	CustomizeTransform(DetailBuilder);
 	CustomizeSettings(DetailBuilder);
+	//CustomizeTransform(DetailBuilder);
 	//CustomizeMetadata(DetailBuilder);
 }
 
@@ -5219,8 +5234,13 @@ void FRigConnectorElementDetails::CustomizeSettings(IDetailLayoutBuilder& Detail
 
 	IDetailCategoryBuilder& SettingsCategory = DetailBuilder.EditCategory(TEXT("Settings"), LOCTEXT("Settings", "Settings"));
 
+	SettingsCategory
+		.AddProperty(SettingsHandle->GetChildHandle(GET_MEMBER_NAME_CHECKED(FRigConnectorSettings, Type)))
+		.IsEnabled(false);
+
 	SettingsCategory
 		.AddProperty(SettingsHandle->GetChildHandle(GET_MEMBER_NAME_CHECKED(FRigConnectorSettings, bOptional)))
+		.Visibility(IsAnyConnectorPrimary() ? EVisibility::Collapsed : EVisibility::Visible)
 		.IsEnabled(!IsAnyConnectorImported());
 
 	bool bHideRules = false;
@@ -5257,6 +5277,124 @@ void FRigConnectorElementDetails::CustomizeSettings(IDetailLayoutBuilder& Detail
 	}
 }
 
+void FRigSocketElementDetails::CustomizeDetails(IDetailLayoutBuilder& DetailBuilder)
+{
+	FRigTransformElementDetails::CustomizeDetails(DetailBuilder);
+	CustomizeSettings(DetailBuilder);
+	CustomizeTransform(DetailBuilder);
+	CustomizeMetadata(DetailBuilder);
+}
+
+void FRigSocketElementDetails::CustomizeSettings(IDetailLayoutBuilder& DetailBuilder)
+{
+	if(PerElementInfos.IsEmpty())
+	{
+		return;
+	}
+
+	if(IsAnyElementNotOfType(ERigElementType::Socket))
+	{
+		return;
+	}
+
+	IDetailCategoryBuilder& SettingsCategory = DetailBuilder.EditCategory(TEXT("Settings"), LOCTEXT("Settings", "Settings"));
+
+	SettingsCategory.AddCustomRow(FText::FromString(TEXT("Color")))
+	.NameContent()
+	[
+		SNew(STextBlock)
+		.Text(LOCTEXT("Color", "Color"))
+		.Font(IDetailLayoutBuilder::GetDetailFont())
+	]
+	.ValueContent()
+	[
+		SNew(SColorBlock)
+		//.Size(FVector2D(6.0, 38.0))
+		.Color(this, &FRigSocketElementDetails::GetSocketColor) 
+		.OnMouseButtonDown(this, &FRigSocketElementDetails::SetSocketColor)
+	];
+
+	SettingsCategory.AddCustomRow(FText::FromString(TEXT("Description")))
+	.NameContent()
+	[
+		SNew(STextBlock)
+		.Text(LOCTEXT("Description", "Description"))
+		.Font(IDetailLayoutBuilder::GetDetailFont())
+	]
+	.ValueContent()
+	[
+		SNew(SEditableText)
+		.Font(IDetailLayoutBuilder::GetDetailFont())
+		.Text(this, &FRigSocketElementDetails::GetSocketDescription)
+		.OnTextCommitted(this, &FRigSocketElementDetails::SetSocketDescription)
+	];
+}
+
+FReply FRigSocketElementDetails::SetSocketColor(const FGeometry& MyGeometry, const FPointerEvent& MouseEvent)
+{
+	FColorPickerArgs PickerArgs;
+	PickerArgs.bUseAlpha = false;
+	PickerArgs.DisplayGamma = TAttribute<float>::Create(TAttribute<float>::FGetter::CreateUObject(GEngine, &UEngine::GetDisplayGamma));
+	PickerArgs.InitialColor = GetSocketColor();
+	PickerArgs.OnColorCommitted = FOnLinearColorValueChanged::CreateSP(this, &FRigSocketElementDetails::OnSocketColorPicked);
+	OpenColorPicker(PickerArgs);
+	return FReply::Handled();
+}
+
+FLinearColor FRigSocketElementDetails::GetSocketColor() const 
+{
+	if(PerElementInfos.Num() > 1)
+	{
+		return FRigSocketElement::SocketDefaultColor;
+	}
+	URigHierarchy* Hierarchy = PerElementInfos[0].GetDefaultHierarchy();
+	const FRigSocketElement* Socket = PerElementInfos[0].GetDefaultElement<FRigSocketElement>();
+	return Socket->GetColor(Hierarchy);
+}
+
+void FRigSocketElementDetails::OnSocketColorPicked(FLinearColor NewColor)
+{
+	FScopedTransaction Transaction(LOCTEXT("SocketColorChanged", "Socket Color Changed"));
+	for(FPerElementInfo& Info : PerElementInfos)
+	{
+		URigHierarchy* Hierarchy = Info.GetDefaultHierarchy();
+		Hierarchy->Modify();
+		FRigSocketElement* Socket = Info.GetDefaultElement<FRigSocketElement>();
+		Socket->SetColor(NewColor, Hierarchy);
+	}
+}
+
+void FRigSocketElementDetails::SetSocketDescription(const FText& InDescription, ETextCommit::Type InCommitType)
+{
+	const FString Description = InDescription.ToString();
+	for(FPerElementInfo& Info : PerElementInfos)
+	{
+		URigHierarchy* Hierarchy = Info.GetDefaultHierarchy();
+		Hierarchy->Modify();
+		FRigSocketElement* Socket = Info.GetDefaultElement<FRigSocketElement>();
+		Socket->SetDescription(Description, Hierarchy);
+	}
+}
+
+FText FRigSocketElementDetails::GetSocketDescription() const
+{
+	FString FirstValue;
+	for(int32 Index = 0; Index < PerElementInfos.Num(); Index++)
+	{
+		URigHierarchy* Hierarchy = PerElementInfos[Index].GetDefaultHierarchy();
+		const FRigSocketElement* Socket = PerElementInfos[Index].GetDefaultElement<FRigSocketElement>();
+		const FString Description = Socket->GetDescription(Hierarchy);
+		if(Index == 0)
+		{
+			FirstValue = Description;
+		}
+		else if(!FirstValue.Equals(Description, ESearchCase::CaseSensitive))
+		{
+			return ControlRigDetailsMultipleValues;
+		}
+	}
+	return FText::FromString(FirstValue);
+}
 
 void FRigConnectionRuleDetails::CustomizeHeader(TSharedRef<IPropertyHandle> InStructPropertyHandle, FDetailWidgetRow& HeaderRow, IPropertyTypeCustomizationUtils& StructCustomizationUtils)
 {
diff --git a/Engine/Plugins/Animation/ControlRig/Source/ControlRigEditor/Private/ControlRigElementDetails.h b/Engine/Plugins/Animation/ControlRig/Source/ControlRigEditor/Private/ControlRigElementDetails.h
--- a/Engine/Plugins/Animation/ControlRig/Source/ControlRigEditor/Private/ControlRigElementDetails.h
+++ b/Engine/Plugins/Animation/ControlRig/Source/ControlRigEditor/Private/ControlRigElementDetails.h
@@ -291,6 +291,7 @@ public:
 	bool IsAnyControlNotOfValueType(ERigControlType InType) const;
 	bool IsAnyElementProcedural() const;
 	bool IsAnyConnectorImported() const;
+	bool IsAnyConnectorPrimary() const;
 	bool GetCommonElementType(ERigElementType& OutElementType) const;
 	bool GetCommonControlType(ERigControlType& OutControlType) const;
 	bool GetCommonAnimationType(ERigControlAnimationType& OutAnimationType) const;
@@ -824,6 +825,19 @@ public:
 	{
 		return MakeShareable(new FRigSocketElementDetails);
 	}
+
+	/** IDetailCustomization interface */
+	virtual void CustomizeDetails(IDetailLayoutBuilder& DetailBuilder) override;
+
+	void CustomizeSettings(IDetailLayoutBuilder& DetailBuilder);
+
+private:
+
+	FReply SetSocketColor(const FGeometry& MyGeometry, const FPointerEvent& MouseEvent);
+	FLinearColor GetSocketColor() const; 
+	void OnSocketColorPicked(FLinearColor NewColor);
+	void SetSocketDescription(const FText& InDescription, ETextCommit::Type InCommitType);
+	FText GetSocketDescription() const; 
 };
 
 class FRigConnectionRuleDetails : public IPropertyTypeCustomization
diff --git a/Engine/Plugins/Animation/ControlRig/Source/ControlRigEditor/Private/ControlRigHierarchyCommands.cpp b/Engine/Plugins/Animation/ControlRig/Source/ControlRigEditor/Private/ControlRigHierarchyCommands.cpp
--- a/Engine/Plugins/Animation/ControlRig/Source/ControlRigEditor/Private/ControlRigHierarchyCommands.cpp
+++ b/Engine/Plugins/Animation/ControlRig/Source/ControlRigEditor/Private/ControlRigHierarchyCommands.cpp
@@ -10,8 +10,8 @@ void FControlRigHierarchyCommands::RegisterCommands()
 	UI_COMMAND(AddControlItem, "New Control", "Add new control at the origin (0, 0, 0) to the hierarchy.", EUserInterfaceActionType::Button, FInputChord(EKeys::N, EModifierKey::Control));
 	UI_COMMAND(AddAnimationChannelItem, "New Animation Channel", "Add new animation channel.", EUserInterfaceActionType::Button, FInputChord());
 	UI_COMMAND(AddNullItem, "New Null", "Add new null at the origin (0, 0, 0) to the hierarchy.", EUserInterfaceActionType::Button, FInputChord());
-	UI_COMMAND(AddConnectorItem, "New Connector", "Add new connector at the origin (0, 0, 0) to the hierarchy.", EUserInterfaceActionType::Button, FInputChord());
-	UI_COMMAND(AddSocketItem, "New Socket", "Add new socket at the origin (0, 0, 0) to the hierarchy.", EUserInterfaceActionType::Button, FInputChord());
+	UI_COMMAND(AddConnectorItem, "New Connector", "Add new connector at the origin (0, 0, 0) to the hierarchy.\nNote: This is only available for Rig Modules.", EUserInterfaceActionType::Button, FInputChord());
+	UI_COMMAND(AddSocketItem, "New Socket", "Add new socket at the origin (0, 0, 0) to the hierarchy.\nNote: This is only available for Rig Modules.", EUserInterfaceActionType::Button, FInputChord());
 	UI_COMMAND(DuplicateItem, "Duplicate", "Duplicate the selected items in the hierarchy.", EUserInterfaceActionType::Button, FInputChord(EKeys::D, EModifierKey::Control));
 	UI_COMMAND(MirrorItem, "Mirror", "Mirror the selected items in the hierarchy.", EUserInterfaceActionType::Button, FInputChord());
 	UI_COMMAND(DeleteItem, "Delete", "Deletes the selected items from the hierarchy.", EUserInterfaceActionType::Button, FInputChord(EKeys::Delete));
diff --git a/Engine/Plugins/Animation/ControlRig/Source/ControlRigEditor/Private/Editor/ControlRigEditor.cpp b/Engine/Plugins/Animation/ControlRig/Source/ControlRigEditor/Private/Editor/ControlRigEditor.cpp
--- a/Engine/Plugins/Animation/ControlRig/Source/ControlRigEditor/Private/Editor/ControlRigEditor.cpp
+++ b/Engine/Plugins/Animation/ControlRig/Source/ControlRigEditor/Private/Editor/ControlRigEditor.cpp
@@ -4804,7 +4804,8 @@ void FControlRigEditor::OnPreConstructionForUI_AnyThread(UControlRig* InRig, con
 	{
 		if(RigBlueprint->IsControlRigModule())
 		{
-			ConnectorInfos = RigBlueprint->Hierarchy->GetConnectorInfos();
+			SocketStates = InRig->GetHierarchy()->GetSocketStates();
+			ConnectorStates = RigBlueprint->Hierarchy->GetConnectorStates();
 		}
 	}
 }
@@ -4867,7 +4868,9 @@ void FControlRigEditor::OnPostConstruction_AnyThread(UControlRig* InRig, const F
 	{
 		if(RigBlueprint->IsControlRigModule())
 		{
-			RigBlueprint->Hierarchy->RestoreConnectorsFromInfos(ConnectorInfos);
+			RigBlueprint->Hierarchy->RestoreSocketsFromStates(SocketStates);
+			InRig->GetHierarchy()->RestoreSocketsFromStates(SocketStates);
+			RigBlueprint->Hierarchy->RestoreConnectorsFromStates(ConnectorStates);
 		}
 	}
 	
diff --git a/Engine/Plugins/Animation/ControlRig/Source/ControlRigEditor/Private/Editor/ControlRigEditor.h b/Engine/Plugins/Animation/ControlRig/Source/ControlRigEditor/Private/Editor/ControlRigEditor.h
--- a/Engine/Plugins/Animation/ControlRig/Source/ControlRigEditor/Private/Editor/ControlRigEditor.h
+++ b/Engine/Plugins/Animation/ControlRig/Source/ControlRigEditor/Private/Editor/ControlRigEditor.h
@@ -327,7 +327,8 @@ protected:
 	void OnPreConstruction_AnyThread(UControlRig* InRig, const FName& InEventName);
 	void OnPostConstruction_AnyThread(UControlRig* InRig, const FName& InEventName);
 	FRigPose PreConstructionPose;
-	TArray<FRigConnectorInfo> ConnectorInfos;
+	TArray<FRigSocketState> SocketStates;
+	TArray<FRigConnectorState> ConnectorStates;
 
 	bool bIsConstructionEventRunning;
 	uint32 LastHierarchyHash;
diff --git a/Engine/Plugins/Animation/ControlRig/Source/ControlRigEditor/Private/Editor/SRigHierarchy.cpp b/Engine/Plugins/Animation/ControlRig/Source/ControlRigEditor/Private/Editor/SRigHierarchy.cpp
--- a/Engine/Plugins/Animation/ControlRig/Source/ControlRigEditor/Private/Editor/SRigHierarchy.cpp
+++ b/Engine/Plugins/Animation/ControlRig/Source/ControlRigEditor/Private/Editor/SRigHierarchy.cpp
@@ -99,6 +99,15 @@ bool FRigElementHierarchyDragDropOp::IsDraggingSingleConnector() const
 	return false;
 }
 
+bool FRigElementHierarchyDragDropOp::IsDraggingSingleSocket() const
+{
+	if(Elements.Num() == 1)
+	{
+		return Elements[0].Type == ERigElementType::Socket;
+	}
+	return false;
+}
+
 ///////////////////////////////////////////////////////////
 
 const FName SRigHierarchy::ContextMenuName = TEXT("ControlRigEditor.RigHierarchy.ContextMenu");
@@ -1845,12 +1854,11 @@ bool SRigHierarchy::IsNonProceduralElementSelected() const
 
 bool SRigHierarchy::CanAddElement(const ERigElementType ElementType) const
 {
-	// Always allow connectors 
-	if (ElementType == ERigElementType::Connector)
+	if (ElementType == ERigElementType::Connector ||
+		ElementType == ERigElementType::Socket)
 	{
-		return true;
+		return ControlRigBlueprint->IsControlRigModule();
 	}
-	
 	return !ControlRigBlueprint->IsControlRigModule();
 }
 
@@ -2073,7 +2081,7 @@ void SRigHierarchy::HandleNewItem(ERigElementType InElementType, bool bIsAnimati
 				}
 				case ERigElementType::Socket:
 				{
-					NewItemKey = Controller->AddSocket(NewElementName, ParentKey, ParentTransform, true, true, true);
+					NewItemKey = Controller->AddSocket(NewElementName, ParentKey, ParentTransform, true, FRigSocketElement::SocketDefaultColor, FString(), true, true);
 					break;
 				}
 				default:
@@ -2523,7 +2531,7 @@ TOptional<EItemDropZone> SRigHierarchy::OnCanAcceptDrop(const FDragDropEvent& Dr
 					return InvalidDropZone;
 				}
 
-				if(RigDragDropOp->IsDraggingSingleConnector())
+				if(RigDragDropOp->IsDraggingSingleConnector() || RigDragDropOp->IsDraggingSingleSocket())
 				{
 					if(DropZone != EItemDropZone::OntoItem)
 					{
@@ -2540,8 +2548,9 @@ TOptional<EItemDropZone> SRigHierarchy::OnCanAcceptDrop(const FDragDropEvent& Dr
 			}
 		}
 
-		// don't allow dragging onto procedural items
-		if(TargetKey.IsValid() && !GetDefaultHierarchy()->Contains(TargetKey) && !RigDragDropOp->IsDraggingSingleConnector())
+		// don't allow dragging onto procedural items (except for connectors + sockets)
+		if(TargetKey.IsValid() && !GetDefaultHierarchy()->Contains(TargetKey) &&
+			!(RigDragDropOp->IsDraggingSingleConnector() || RigDragDropOp->IsDraggingSingleSocket()))
 		{
 			return InvalidDropZone;
 		}
diff --git a/Engine/Plugins/Animation/ControlRig/Source/ControlRigEditor/Private/Editor/SRigHierarchy.h b/Engine/Plugins/Animation/ControlRig/Source/ControlRigEditor/Private/Editor/SRigHierarchy.h
--- a/Engine/Plugins/Animation/ControlRig/Source/ControlRigEditor/Private/Editor/SRigHierarchy.h
+++ b/Engine/Plugins/Animation/ControlRig/Source/ControlRigEditor/Private/Editor/SRigHierarchy.h
@@ -47,6 +47,7 @@ public:
 	FString GetJoinedElementNames() const;
 
 	bool IsDraggingSingleConnector() const;
+	bool IsDraggingSingleSocket() const;
 
 private:
 
