Control Rig: Controls not refreshing correctly when shape libraries change

#rnx

diff --git a/Engine/Plugins/Animation/ControlRig/Source/ControlRig/Private/ControlRig.cpp b/Engine/Plugins/Animation/ControlRig/Source/ControlRig/Private/ControlRig.cpp
--- a/Engine/Plugins/Animation/ControlRig/Source/ControlRig/Private/ControlRig.cpp
+++ b/Engine/Plugins/Animation/ControlRig/Source/ControlRig/Private/ControlRig.cpp
@@ -205,6 +205,18 @@ void UControlRig::Initialize(bool bRequestInit)
 	}
 }
 
+void UControlRig::RestoreShapeLibrariesFromCDO()
+{
+	if (!HasAnyFlags(RF_ClassDefaultObject))
+	{
+		if(UControlRig* CDO = GetClass()->GetDefaultObject<UControlRig>())
+		{
+			ShapeLibraryNameMap.Reset();
+			ShapeLibraries = CDO->ShapeLibraries;
+		}
+	}
+}
+
 void UControlRig::OnAddShapeLibrary(const FControlRigExecuteContext* InContext, const FString& InLibraryName, UControlRigShapeLibrary* InShapeLibrary, bool bReplaceExisting, bool bLogResults)
 {
 	// don't ever change the CDO
@@ -230,8 +242,7 @@ void UControlRig::OnAddShapeLibrary(const FControlRigExecuteContext* InContext,
 	// if we've removed all shape libraries - let's add the ones from the CDO back
 	if (ShapeLibraries.IsEmpty())
 	{
-		UControlRig* CDO = GetClass()->GetDefaultObject<UControlRig>();
-		ShapeLibraries = CDO->ShapeLibraries;
+		RestoreShapeLibrariesFromCDO();
 	}
 
 	// if we are supposed to replace the library and the library name is empty
@@ -249,34 +260,38 @@ void UControlRig::OnAddShapeLibrary(const FControlRigExecuteContext* InContext,
 		LibraryName = InShapeLibrary->GetName();
 	}
 
-	if(LibraryName != InShapeLibrary->GetName())
-	{
-		ShapeLibraryNameMap.FindOrAdd(InShapeLibrary->GetName()) = LibraryName;
-	}
-
+	// if we need to replace all existing shape libraries - we'll remove any shape library matching the library
+	// name which doesn't map to something else
 	if(bReplaceExisting)
 	{
-		for(int32 Index = 0; Index < ShapeLibraries.Num(); Index++)
+		TArray<FString> KeysToRemove;
+		for(const TPair<FString, FString>& Pair : ShapeLibraryNameMap)
 		{
-			const TSoftObjectPtr<UControlRigShapeLibrary>& ExistingShapeLibrary = ShapeLibraries[Index];
-			if(ExistingShapeLibrary.IsNull())
+			if(Pair.Value.Equals(LibraryName, ESearchCase::CaseSensitive))
 			{
-				continue;
+				KeysToRemove.Add(Pair.Key);
 			}
-			FString ExistingName = ExistingShapeLibrary->GetName();
-			if (FString* MapName = ShapeLibraryNameMap.Find(ExistingName))
-			{
-				ExistingName = *MapName;
-			}
-			if(ExistingName.Equals(LibraryName, ESearchCase::IgnoreCase))
+		}
+		
+		if(KeysToRemove.IsEmpty())
+		{
+			KeysToRemove.Add(InShapeLibrary->GetName());
+		}
+		
+		for(const FString& KeyToRemove : KeysToRemove)
+		{
+			ShapeLibraryNameMap.Remove(KeyToRemove);
+			ShapeLibraries.RemoveAll([KeyToRemove](const TSoftObjectPtr<UControlRigShapeLibrary>& ShapeLibrary) -> bool
 			{
-				ShapeLibraryNameMap.Remove(ExistingShapeLibrary->GetName());
-				ShapeLibraries[Index] = InShapeLibrary;
-				break;
-			}
+				return ShapeLibrary->GetName().Equals(KeyToRemove, ESearchCase::CaseSensitive);
+			});
 		}
 	}
 
+	if(LibraryName != InShapeLibrary->GetName())
+	{
+		ShapeLibraryNameMap.FindOrAdd(InShapeLibrary->GetName()) = LibraryName;
+	}
 	ShapeLibraries.AddUnique(InShapeLibrary);
 
 #if WITH_EDITOR
@@ -1021,12 +1036,7 @@ bool UControlRig::Execute(const FName& InEventName)
 						GetHierarchy()->ResetPoseToInitial(ERigElementType::All);
 					}
 
-					// clone the shape libraries again from the CDO 
-					if (!HasAnyFlags(RF_ClassDefaultObject) && ShapeLibraries.IsEmpty())
-					{
-						UControlRig* CDO = GetClass()->GetDefaultObject<UControlRig>();
-						ShapeLibraries = CDO->ShapeLibraries;
-					}
+					RestoreShapeLibrariesFromCDO();
 
 					if (PreConstructionEvent.IsBound())
 					{
@@ -3482,6 +3492,20 @@ UControlRig::FTransientControlScope::~FTransientControlScope()
 	}
 }
 
+uint32 UControlRig::GetShapeLibraryHash() const
+{
+	uint32 Hash = 0;
+	for(const TPair<FString, FString>& Pair : ShapeLibraryNameMap)
+	{
+		Hash = HashCombine(Hash, HashCombine(GetTypeHash(Pair.Key), GetTypeHash(Pair.Value)));
+	}
+	for(const TSoftObjectPtr<UControlRigShapeLibrary>& ShapeLibrary : ShapeLibraries)
+	{
+		Hash = HashCombine(Hash, GetTypeHash(ShapeLibrary.GetUniqueID()));
+	}
+	return Hash;
+}
+
 #endif
  
 #undef LOCTEXT_NAMESPACE
diff --git a/Engine/Plugins/Animation/ControlRig/Source/ControlRig/Private/ControlRigGizmoLibrary.cpp b/Engine/Plugins/Animation/ControlRig/Source/ControlRig/Private/ControlRigGizmoLibrary.cpp
--- a/Engine/Plugins/Animation/ControlRig/Source/ControlRig/Private/ControlRigGizmoLibrary.cpp
+++ b/Engine/Plugins/Animation/ControlRig/Source/ControlRig/Private/ControlRigGizmoLibrary.cpp
@@ -126,26 +126,37 @@ const FControlRigShapeDefinition* UControlRigShapeLibrary::GetShapeByName(const
 		RightName = *Right;
 	}
 
-	for(const TSoftObjectPtr<UControlRigShapeLibrary>& ShapeLibrary : InShapeLibraries)
+	// we'll do this in two passes - first with the whole namespace
+	// and then without - to allow to find a shape also if it has been specified
+	// without a namespace all-together.
+	for(int32 Pass = 0; Pass < 2; Pass++)
 	{
-		if(!ShapeLibrary.IsValid())
+		// we need to walk backwards since shape libraries added later need to take precedence
+		for(int32 Index = InShapeLibraries.Num() - 1; Index >= 0; Index--)
 		{
-			continue;
-		}
+			const TSoftObjectPtr<UControlRigShapeLibrary>& ShapeLibrary = InShapeLibraries[Index];
+			if(!ShapeLibrary.IsValid())
+			{
+				continue;
+			}
 
-		FString ShapeLibraryName = ShapeLibrary->GetName();
-		if(const FString* RemappedName = InLibraryNameMap.Find(ShapeLibraryName))
-		{
-			ShapeLibraryName = *RemappedName;
-		}
+			FString ShapeLibraryName = ShapeLibrary->GetName();
+			if(const FString* RemappedName = InLibraryNameMap.Find(ShapeLibraryName))
+			{
+				ShapeLibraryName = *RemappedName;
+			}
 
-		if(ShapeLibraryName.Equals(Left) || Left.IsEmpty())
-		{
-			if(const FControlRigShapeDefinition* Shape = ShapeLibrary->GetShapeByName(RightName))
+			if(ShapeLibraryName.Equals(Left) || Left.IsEmpty())
 			{
-				return Shape;
+				if(const FControlRigShapeDefinition* Shape = ShapeLibrary->GetShapeByName(RightName))
+				{
+					return Shape;
+				}
 			}
 		}
+
+		// remove the namespace for the next pass
+		Left.Reset();
 	}
 
 	return nullptr;
diff --git a/Engine/Plugins/Animation/ControlRig/Source/ControlRig/Public/ControlRig.h b/Engine/Plugins/Animation/ControlRig/Source/ControlRig/Public/ControlRig.h
--- a/Engine/Plugins/Animation/ControlRig/Source/ControlRig/Public/ControlRig.h
+++ b/Engine/Plugins/Animation/ControlRig/Source/ControlRig/Public/ControlRig.h
@@ -520,6 +520,7 @@ protected:
 	void HandleHierarchyEvent(URigHierarchy* InHierarchy, const FRigEventContext& InEvent);
 	FRigEventDelegate RigEventDelegate;
 
+	void RestoreShapeLibrariesFromCDO();
 	void OnAddShapeLibrary(const FControlRigExecuteContext* InContext, const FString& InLibraryName, UControlRigShapeLibrary* InShapeLibrary, bool bReplaceExisting, bool bLogResults);
 	bool OnShapeExists(const FName& InShapeName) const;
 	virtual void InitializeVMsFromCDO() { Super::InitializeFromCDO(); }
@@ -824,6 +825,8 @@ public:
 	UE_DEPRECATED(5.4, "InteractionRig is no longer used")
 	UFUNCTION(BlueprintSetter, meta = (DeprecatedFunction, DeprecationMessage = "InteractionRig is no longer used"))
 	void SetInteractionRigClass(TSubclassOf<UControlRig> InInteractionRigClass) {}
+
+	uint32 GetShapeLibraryHash() const;
 	
 private:
 #if WITH_EDITORONLY_DATA
diff --git a/Engine/Plugins/Animation/ControlRig/Source/ControlRigEditor/Private/EditMode/ControlRigEditMode.cpp b/Engine/Plugins/Animation/ControlRig/Source/ControlRigEditor/Private/EditMode/ControlRigEditMode.cpp
--- a/Engine/Plugins/Animation/ControlRig/Source/ControlRigEditor/Private/EditMode/ControlRigEditMode.cpp
+++ b/Engine/Plugins/Animation/ControlRig/Source/ControlRigEditor/Private/EditMode/ControlRigEditMode.cpp
@@ -4087,16 +4087,19 @@ void FControlRigEditMode::OnPostConstruction_AnyThread(UControlRig* InRig, const
 	bIsConstructionEventRunning = false;
 
 	const int32 RigIndex = RuntimeControlRigs.Find(InRig);
-	if(!LastHierarchyHash.IsValidIndex(RigIndex))
+	if(!LastHierarchyHash.IsValidIndex(RigIndex) || !LastShapeLibraryHash.IsValidIndex(RigIndex))
 	{
 		return;
 	}
 	
 	const int32 HierarchyHash = InRig->GetHierarchy()->GetTopologyHash(false, true);
-	if(LastHierarchyHash[RigIndex] != HierarchyHash)
+	const int32 ShapeLibraryHash = InRig->GetShapeLibraryHash();
+	if((LastHierarchyHash[RigIndex] != HierarchyHash) ||
+		(LastShapeLibraryHash[RigIndex] != ShapeLibraryHash))
 	{
 		LastHierarchyHash[RigIndex] = HierarchyHash;
-		
+		LastShapeLibraryHash[RigIndex] = ShapeLibraryHash;
+
 		auto Task = [this, InRig]()
 		{
 			RequestToRecreateControlShapeActors(InRig);
@@ -4644,6 +4647,7 @@ void FControlRigEditMode::AddControlRigInternal(UControlRig* InControlRig)
 {
 	RuntimeControlRigs.AddUnique(InControlRig);
 	LastHierarchyHash.Add(INDEX_NONE);
+	LastShapeLibraryHash.Add(INDEX_NONE);
 
 	InControlRig->SetControlsVisible(true);
 	InControlRig->PostInitInstanceIfRequired();
@@ -4726,6 +4730,10 @@ void FControlRigEditMode::RemoveControlRig(UControlRig* InControlRig)
 	{
 		LastHierarchyHash.RemoveAt(Index);
 	}
+	if (LastShapeLibraryHash.IsValidIndex(Index))
+	{
+		LastShapeLibraryHash.RemoveAt(Index);
+	}
 
 	//needed for the control rig track editor delegates to get removed
 	if (WeakSequencer.IsValid())
diff --git a/Engine/Plugins/Animation/ControlRig/Source/ControlRigEditor/Public/EditMode/ControlRigEditMode.h b/Engine/Plugins/Animation/ControlRig/Source/ControlRigEditor/Public/EditMode/ControlRigEditMode.h
--- a/Engine/Plugins/Animation/ControlRig/Source/ControlRigEditor/Public/EditMode/ControlRigEditMode.h
+++ b/Engine/Plugins/Animation/ControlRig/Source/ControlRigEditor/Public/EditMode/ControlRigEditMode.h
@@ -491,6 +491,7 @@ private:
 
 	bool bIsConstructionEventRunning;
 	TArray<uint32> LastHierarchyHash;
+	TArray<uint32> LastShapeLibraryHash;
 
 	friend class FControlRigEditorModule;
 	friend class FControlRigEditor;
