Control Rig: Fixed description of layered control rig checkbox on sequencer track editor option
#rnx

diff --git a/Engine/Plugins/Animation/ControlRig/Source/ControlRigEditor/Private/Sequencer/ControlRigParameterTrackEditor.cpp b/Engine/Plugins/Animation/ControlRig/Source/ControlRigEditor/Private/Sequencer/ControlRigParameterTrackEditor.cpp
--- a/Engine/Plugins/Animation/ControlRig/Source/ControlRigEditor/Private/Sequencer/ControlRigParameterTrackEditor.cpp
+++ b/Engine/Plugins/Animation/ControlRig/Source/ControlRigEditor/Private/Sequencer/ControlRigParameterTrackEditor.cpp
@@ -1300,7 +1300,7 @@ void FControlRigParameterTrackEditor::HandleAddTrackSubMenu(FMenuBuilder& MenuBu
 	{
 		MenuBuilder.AddMenuEntry(
 		LOCTEXT("IsLayeredControlRig", "Layered"),
-		LOCTEXT("IsLayeredControlRigTooltip", "Add an layered control rig"),
+		LOCTEXT("IsLayeredControlRigTooltip", "When checked, a layered control rig will be added"),
 		FSlateIcon(),
 		FUIAction(
 			FExecuteAction::CreateSP(this, &FControlRigParameterTrackEditor::ToggleIsAdditiveControlRig),
