ControlRig: Support for sockets in the hierarchy

#rnx

diff --git a/Engine/Plugins/Animation/ControlRig/Source/ControlRig/Private/Rigs/RigHierarchy.cpp b/Engine/Plugins/Animation/ControlRig/Source/ControlRig/Private/Rigs/RigHierarchy.cpp
--- a/Engine/Plugins/Animation/ControlRig/Source/ControlRig/Private/Rigs/RigHierarchy.cpp
+++ b/Engine/Plugins/Animation/ControlRig/Source/ControlRig/Private/Rigs/RigHierarchy.cpp
@@ -1058,7 +1058,10 @@ TArray<FRigSocketState> URigHierarchy::GetSocketStates() const
 	for(const FRigElementKey& Key : Keys)
 	{
 		const FRigSocketElement* Socket = FindChecked<FRigSocketElement>(Key);
-		States.Add(Socket->GetSocketState(this));
+		if(!Socket->IsProcedural())
+		{
+			States.Add(Socket->GetSocketState(this));
+		}
 	}
 	return States;
 }
@@ -1096,7 +1099,10 @@ TArray<FRigConnectorState> URigHierarchy::GetConnectorStates() const
 	for(const FRigElementKey& Key : Keys)
 	{
 		const FRigConnectorElement* Connector = FindChecked<FRigConnectorElement>(Key);
-		States.Add(Connector->GetConnectorState(this));
+		if(!Connector->IsProcedural())
+		{
+			States.Add(Connector->GetConnectorState(this));
+		}
 	}
 	return States;
 }
diff --git a/Engine/Plugins/Animation/ControlRig/Source/ControlRig/Private/Rigs/RigHierarchyCache.cpp b/Engine/Plugins/Animation/ControlRig/Source/ControlRig/Private/Rigs/RigHierarchyCache.cpp
--- a/Engine/Plugins/Animation/ControlRig/Source/ControlRig/Private/Rigs/RigHierarchyCache.cpp
+++ b/Engine/Plugins/Animation/ControlRig/Source/ControlRig/Private/Rigs/RigHierarchyCache.cpp
@@ -5,6 +5,16 @@
 
 #include UE_INLINE_GENERATED_CPP_BY_NAME(RigHierarchyCache)
 
+const FRigElementKey& FCachedRigElement::GetResolvedKey() const
+{
+	if(Element)
+	{
+		return Element->GetKey();
+	}
+	static FRigElementKey InvalidKey;
+	return InvalidKey;
+}
+
 bool FCachedRigElement::UpdateCache(const URigHierarchy* InHierarchy)
 {
 	if(InHierarchy)
diff --git a/Engine/Plugins/Animation/ControlRig/Source/ControlRig/Private/Rigs/RigHierarchyController.cpp b/Engine/Plugins/Animation/ControlRig/Source/ControlRig/Private/Rigs/RigHierarchyController.cpp
--- a/Engine/Plugins/Animation/ControlRig/Source/ControlRig/Private/Rigs/RigHierarchyController.cpp
+++ b/Engine/Plugins/Animation/ControlRig/Source/ControlRig/Private/Rigs/RigHierarchyController.cpp
@@ -2577,11 +2577,6 @@ bool URigHierarchyController::AddParent(FRigBaseElement* InChild, FRigBaseElemen
 		return false;
 	}
 
-	if(InParent->GetType() == ERigElementType::Socket)
-	{
-		return false;
-	}
-
 	// single parent children can't be parented multiple times
 	if(FRigSingleParentElement* SingleParentElement = Cast<FRigSingleParentElement>(InChild))
 	{
@@ -3032,12 +3027,6 @@ bool URigHierarchyController::SetParent(FRigElementKey InChild, FRigElementKey I
 		return false;
 	}
 
-	if(InParent.Type == ERigElementType::Socket)
-	{
-		ReportWarningf(TEXT("Cannot parent Child '%s' under a Socket parent."), *InChild.ToString());
-		return false;
-	}
-
 	FRigBaseElement* Child = Hierarchy->Find(InChild);
 	if(Child == nullptr)
 	{
diff --git a/Engine/Plugins/Animation/ControlRig/Source/ControlRig/Private/Units/Execution/RigUnit_Hierarchy.cpp b/Engine/Plugins/Animation/ControlRig/Source/ControlRig/Private/Units/Execution/RigUnit_Hierarchy.cpp
--- a/Engine/Plugins/Animation/ControlRig/Source/ControlRig/Private/Units/Execution/RigUnit_Hierarchy.cpp
+++ b/Engine/Plugins/Animation/ControlRig/Source/ControlRig/Private/Units/Execution/RigUnit_Hierarchy.cpp
@@ -22,7 +22,7 @@ FRigUnit_HierarchyGetParent_Execute()
 
 		if(CachedChild.UpdateCache(Child, ExecuteContext.Hierarchy))
 		{
-			Parent = ExecuteContext.Hierarchy->GetFirstParent(Child);
+			Parent = ExecuteContext.Hierarchy->GetFirstParent(CachedChild.GetResolvedKey());
 			if(Parent.IsValid())
 			{
 				CachedParent.UpdateCache(Parent, ExecuteContext.Hierarchy);
@@ -57,7 +57,7 @@ FRigUnit_HierarchyGetParentsItemArray_Execute()
 		if(CachedChild.UpdateCache(Child, ExecuteContext.Hierarchy))
 		{
 			TArray<FRigElementKey> Keys;
-			FRigElementKey Parent = Child;
+			FRigElementKey Parent = CachedChild.GetResolvedKey();
 			do
 			{
 				if(bIncludeChild || Parent != Child)
@@ -89,15 +89,16 @@ FRigUnit_HierarchyGetChildren_Execute()
 
 		if(CachedParent.UpdateCache(Parent, ExecuteContext.Hierarchy))
 		{
+			const FRigElementKey ResolvedParent = CachedParent.GetResolvedKey();
+
 			TArray<FRigElementKey> Keys;
 
 			if(bIncludeParent)
 			{
-				Keys.Add(Parent);
+				Keys.Add(ResolvedParent);
 			}
 
-			
-			Keys.Append(ExecuteContext.Hierarchy->GetChildren(Parent, bRecursive));
+			Keys.Append(ExecuteContext.Hierarchy->GetChildren(ResolvedParent, bRecursive));
 
 			CachedChildren = FRigElementKeyCollection(Keys);
 		}
@@ -141,7 +142,7 @@ FRigUnit_HierarchyGetSiblingsItemArray_Execute()
 		{
 			TArray<FRigElementKey> Keys;
 
-			FRigElementKey Parent = ExecuteContext.Hierarchy->GetFirstParent(Item);
+			FRigElementKey Parent = ExecuteContext.Hierarchy->GetFirstParent(CachedItem.GetResolvedKey());
 			if(Parent.IsValid())
 			{
 				TArray<FRigElementKey> Children = ExecuteContext.Hierarchy->GetChildren(Parent, false);
diff --git a/Engine/Plugins/Animation/ControlRig/Source/ControlRig/Public/Rigs/RigHierarchyCache.h b/Engine/Plugins/Animation/ControlRig/Source/ControlRig/Public/Rigs/RigHierarchyCache.h
--- a/Engine/Plugins/Animation/ControlRig/Source/ControlRig/Public/Rigs/RigHierarchyCache.h
+++ b/Engine/Plugins/Animation/ControlRig/Source/ControlRig/Public/Rigs/RigHierarchyCache.h
@@ -77,6 +77,8 @@ public:
 		return Key;
 	}
 
+	const FRigElementKey& GetResolvedKey() const;
+
 	const FRigBaseElement* GetElement() const
 	{
 		return Element;
diff --git a/Engine/Plugins/Animation/ControlRig/Source/ControlRigDeveloper/Private/ControlRigBlueprint.cpp b/Engine/Plugins/Animation/ControlRig/Source/ControlRigDeveloper/Private/ControlRigBlueprint.cpp
--- a/Engine/Plugins/Animation/ControlRig/Source/ControlRigDeveloper/Private/ControlRigBlueprint.cpp
+++ b/Engine/Plugins/Animation/ControlRig/Source/ControlRigDeveloper/Private/ControlRigBlueprint.cpp
@@ -316,7 +316,10 @@ bool UControlRigBlueprint::TurnIntoControlRigModule(bool InAutoConvertHierarchy,
 			if(Hierarchy->Num(ERigElementType::Connector) == 0)
 			{
 				static const FName RootName = TEXT("Root");
-				Controller->AddConnector(RootName);
+				static const FString RootDescription = TEXT("This is the default temporary socket used for the root connection.");
+				const FRigElementKey ConnectorKey = Controller->AddConnector(RootName);
+				const FRigElementKey SocketKey = Controller->AddSocket(RootName, FRigElementKey(), FTransform::Identity, false, FRigSocketElement::SocketDefaultColor, RootDescription, false);
+				(void)ResolveConnector(ConnectorKey, SocketKey);
 			}
 		}
 	}
diff --git a/Engine/Plugins/Animation/ControlRig/Source/ControlRigEditor/Private/ControlRigEditorStyle.h b/Engine/Plugins/Animation/ControlRig/Source/ControlRigEditor/Private/ControlRigEditorStyle.h
--- a/Engine/Plugins/Animation/ControlRig/Source/ControlRigEditor/Private/ControlRigEditorStyle.h
+++ b/Engine/Plugins/Animation/ControlRig/Source/ControlRigEditor/Private/ControlRigEditorStyle.h
@@ -117,7 +117,8 @@ public:
 			Set("ControlRig.Tree.ProxyControl", new IMAGE_BRUSH("Slate/ProxyControl1_16x", Icon16x16));
 			Set("ControlRig.Tree.Null", new IMAGE_BRUSH("Slate/Null_16x", Icon16x16));
 			Set("ControlRig.Tree.RigidBody", new IMAGE_BRUSH("Slate/RigidBody_16x", Icon16x16));
-			Set("ControlRig.Tree.Socket", new IMAGE_BRUSH("Slate/Socket_16x", Icon16x16));
+			Set("ControlRig.Tree.Socket_Open", new IMAGE_BRUSH_SVG("Slate/Socket_Open", Icon16x16));
+			Set("ControlRig.Tree.Socket_Closed", new IMAGE_BRUSH_SVG("Slate/Socket_Closed", Icon16x16));
 			{
 				FContentRootBracket Bracket(this, EngineEditorSlateDir);
 				Set("ControlRig.Tree.Connector", new IMAGE_BRUSH_SVG("Starship/Common/SetShowSockets", Icon16x16));
diff --git a/Engine/Plugins/Animation/ControlRig/Source/ControlRigEditor/Private/ControlRigHierarchyCommands.cpp b/Engine/Plugins/Animation/ControlRig/Source/ControlRigEditor/Private/ControlRigHierarchyCommands.cpp
--- a/Engine/Plugins/Animation/ControlRig/Source/ControlRigEditor/Private/ControlRigHierarchyCommands.cpp
+++ b/Engine/Plugins/Animation/ControlRig/Source/ControlRigEditor/Private/ControlRigHierarchyCommands.cpp
@@ -11,7 +11,7 @@ void FControlRigHierarchyCommands::RegisterCommands()
 	UI_COMMAND(AddAnimationChannelItem, "New Animation Channel", "Add new animation channel.", EUserInterfaceActionType::Button, FInputChord());
 	UI_COMMAND(AddNullItem, "New Null", "Add new null at the origin (0, 0, 0) to the hierarchy.", EUserInterfaceActionType::Button, FInputChord());
 	UI_COMMAND(AddConnectorItem, "New Connector", "Add new connector at the origin (0, 0, 0) to the hierarchy.\nNote: This is only available for Rig Modules.", EUserInterfaceActionType::Button, FInputChord());
-	UI_COMMAND(AddSocketItem, "New Socket", "Add new socket at the origin (0, 0, 0) to the hierarchy.\nNote: This is only available for Rig Modules.", EUserInterfaceActionType::Button, FInputChord());
+	UI_COMMAND(AddSocketItem, "New Socket", "Add new socket at the origin (0, 0, 0) to the hierarchy.\nNote: This is only available for Rig Modules and Modular Rigs.", EUserInterfaceActionType::Button, FInputChord());
 	UI_COMMAND(DuplicateItem, "Duplicate", "Duplicate the selected items in the hierarchy.", EUserInterfaceActionType::Button, FInputChord(EKeys::D, EModifierKey::Control));
 	UI_COMMAND(MirrorItem, "Mirror", "Mirror the selected items in the hierarchy.", EUserInterfaceActionType::Button, FInputChord());
 	UI_COMMAND(DeleteItem, "Delete", "Deletes the selected items from the hierarchy.", EUserInterfaceActionType::Button, FInputChord(EKeys::Delete));
diff --git a/Engine/Plugins/Animation/ControlRig/Source/ControlRigEditor/Private/EditMode/ControlRigEditMode.cpp b/Engine/Plugins/Animation/ControlRig/Source/ControlRigEditor/Private/EditMode/ControlRigEditMode.cpp
--- a/Engine/Plugins/Animation/ControlRig/Source/ControlRigEditor/Private/EditMode/ControlRigEditMode.cpp
+++ b/Engine/Plugins/Animation/ControlRig/Source/ControlRigEditor/Private/EditMode/ControlRigEditMode.cpp
@@ -923,6 +923,38 @@ void FControlRigEditMode::Render(const FSceneView* View, FViewport* Viewport, FP
 					}
 				}
 
+				// temporary implementation to draw sockets in 3D
+				if (bIsAssetEditor && (Settings->bDisplaySockets || ControlRig->IsConstructionModeEnabled()))
+				{
+					Hierarchy->ForEach<FRigSocketElement>([this, Hierarchy, PDI](FRigSocketElement* Socket)
+					{
+						const FLinearColor Color = Socket->GetColor(Hierarchy);
+						const uint32 ColorHash = GetTypeHash(Color.ToFColor(true));
+
+						const FMaterialRenderProxy* MaterialProxy;
+						if(const TStrongObjectPtr<UMaterialInstanceDynamic>* ExistingMaterialPtr = SocketMaterials.Find(ColorHash))
+						{
+							MaterialProxy = (*ExistingMaterialPtr)->GetRenderProxy();
+						}
+						else
+						{
+							UMaterial* MaterialBase = GEngine->ArrowMaterial;
+							UMaterialInstanceDynamic* SocketMaterial = UMaterialInstanceDynamic::Create(MaterialBase, NULL);
+							SocketMaterial->SetVectorParameterValue("GizmoColor", Color);
+							SocketMaterials.Add(ColorHash, TStrongObjectPtr<UMaterialInstanceDynamic>(SocketMaterial));
+							MaterialProxy = SocketMaterial->GetRenderProxy();
+						}
+
+						DrawSphere(PDI,
+							Hierarchy->GetGlobalTransform(Socket->GetIndex()).GetLocation(), 
+							FRotator::ZeroRotator,
+							FVector::OneVector * 4.f,
+							8, 8,
+							MaterialProxy,
+							SDPG_Foreground);
+						return true;
+					});
+				}
 				ControlRig->DrawIntoPDI(PDI, ComponentTransform);
 			}
 		}
diff --git a/Engine/Plugins/Animation/ControlRig/Source/ControlRigEditor/Private/EditMode/ControlRigEditModeSettings.h b/Engine/Plugins/Animation/ControlRig/Source/ControlRigEditor/Private/EditMode/ControlRigEditModeSettings.h
--- a/Engine/Plugins/Animation/ControlRig/Source/ControlRigEditor/Private/EditMode/ControlRigEditModeSettings.h
+++ b/Engine/Plugins/Animation/ControlRig/Source/ControlRigEditor/Private/EditMode/ControlRigEditModeSettings.h
@@ -18,6 +18,7 @@ class UControlRigEditModeSettings : public UObject
 	UControlRigEditModeSettings()
 		: bDisplayHierarchy(false)
 		, bDisplayNulls(false)
+		, bDisplaySockets(true)
 		, bHideControlShapes(false)
 		, bShowAllProxyControls(false)
 		, bShowControlsAsOverlay(false)
@@ -48,6 +49,10 @@ public:
 	UPROPERTY(config, EditAnywhere, Category = "Animation Settings")
 	bool bDisplayNulls;
 
+	/** Should we show sockets in the viewport */
+	UPROPERTY(config, EditAnywhere, Category = "Animation Settings")
+	bool bDisplaySockets;
+
 	/** Should we always hide control shapes in viewport */
 	UPROPERTY(config, EditAnywhere, Category = "Animation Settings")
 	bool bHideControlShapes;
diff --git a/Engine/Plugins/Animation/ControlRig/Source/ControlRigEditor/Private/Editor/ControlRigEditor.cpp b/Engine/Plugins/Animation/ControlRig/Source/ControlRigEditor/Private/Editor/ControlRigEditor.cpp
--- a/Engine/Plugins/Animation/ControlRig/Source/ControlRigEditor/Private/Editor/ControlRigEditor.cpp
+++ b/Engine/Plugins/Animation/ControlRig/Source/ControlRigEditor/Private/Editor/ControlRigEditor.cpp
@@ -2121,6 +2121,24 @@ void FControlRigEditor::HandleViewportCreated(const TSharedRef<class IPersonaVie
 					LOCTEXT("ControlRigDisplayNulls", "Display Nulls")
 				);
 
+				InMenuBuilder.AddWidget(
+					SNew(SBox)
+					.HAlign(HAlign_Right)
+					[
+						SNew(SBox)
+						.Padding(FMargin(4.0f, 0.0f, 0.0f, 0.0f))
+						.WidthOverride(100.0f)
+						.IsEnabled(this, &FControlRigEditor::IsToolbarDrawSocketsEnabled)
+						[
+							SNew(SCheckBox)
+							.IsChecked(this, &FControlRigEditor::GetToolbarDrawSockets)
+							.OnCheckStateChanged(this, &FControlRigEditor::OnToolbarDrawSocketsChanged)
+							.ToolTipText(LOCTEXT("ControlRigDrawSocketsToolTip", "If checked all sockets are drawn."))
+						]
+					],
+					LOCTEXT("ControlRigDisplaySockets", "Display Sockets")
+				);
+
 				InMenuBuilder.AddWidget(
 					SNew(SBox)
 					.HAlign(HAlign_Right)
@@ -2297,6 +2315,35 @@ void FControlRigEditor::OnToolbarDrawNullsChanged(ECheckBoxState InNewValue)
 	}
 }
 
+bool FControlRigEditor::IsToolbarDrawSocketsEnabled() const
+{
+	if (const UControlRig* ControlRig = GetControlRig())
+	{
+		if (!ControlRig->IsConstructionModeEnabled())
+		{
+			return true;
+		}
+	}
+	return false;
+}
+
+ECheckBoxState FControlRigEditor::GetToolbarDrawSockets() const
+{
+	if (const UControlRigEditModeSettings* Settings = GetDefault<UControlRigEditModeSettings>())
+	{
+		return Settings->bDisplaySockets ? ECheckBoxState::Checked : ECheckBoxState::Unchecked;
+	}
+	return ECheckBoxState::Unchecked;
+}
+
+void FControlRigEditor::OnToolbarDrawSocketsChanged(ECheckBoxState InNewValue)
+{
+	if (UControlRigEditModeSettings* Settings = GetMutableDefault<UControlRigEditModeSettings>())
+	{
+		Settings->bDisplaySockets = InNewValue == ECheckBoxState::Checked;
+	}
+}
+
 bool FControlRigEditor::IsConstructionModeEnabled() const
 {
 	return GetEventQueue() == ConstructionEventQueue;
@@ -4855,6 +4902,12 @@ void FControlRigEditor::OnPreConstruction_AnyThread(UControlRig* InRig, const FN
 						}
 					}
 				}
+
+				if(RigBlueprint->IsControlRigModule())
+				{
+					RigBlueprint->Hierarchy->RestoreSocketsFromStates(SocketStates);
+					InRig->GetHierarchy()->RestoreSocketsFromStates(SocketStates);
+				}
 			}
 		}
 	}
@@ -4868,8 +4921,6 @@ void FControlRigEditor::OnPostConstruction_AnyThread(UControlRig* InRig, const F
 	{
 		if(RigBlueprint->IsControlRigModule())
 		{
-			RigBlueprint->Hierarchy->RestoreSocketsFromStates(SocketStates);
-			InRig->GetHierarchy()->RestoreSocketsFromStates(SocketStates);
 			RigBlueprint->Hierarchy->RestoreConnectorsFromStates(ConnectorStates);
 		}
 	}
diff --git a/Engine/Plugins/Animation/ControlRig/Source/ControlRigEditor/Private/Editor/ControlRigEditor.h b/Engine/Plugins/Animation/ControlRig/Source/ControlRigEditor/Private/Editor/ControlRigEditor.h
--- a/Engine/Plugins/Animation/ControlRig/Source/ControlRigEditor/Private/Editor/ControlRigEditor.h
+++ b/Engine/Plugins/Animation/ControlRig/Source/ControlRigEditor/Private/Editor/ControlRigEditor.h
@@ -232,6 +232,9 @@ private:
 	bool IsToolbarDrawNullsEnabled() const;
 	ECheckBoxState GetToolbarDrawNulls() const;
 	void OnToolbarDrawNullsChanged(ECheckBoxState InNewValue);
+	bool IsToolbarDrawSocketsEnabled() const;
+	ECheckBoxState GetToolbarDrawSockets() const;
+	void OnToolbarDrawSocketsChanged(ECheckBoxState InNewValue);
 	ECheckBoxState GetToolbarDrawAxesOnSelection() const;
 	void OnToolbarDrawAxesOnSelectionChanged(ECheckBoxState InNewValue);
 	TOptional<float> GetToolbarAxesScale() const;
diff --git a/Engine/Plugins/Animation/ControlRig/Source/ControlRigEditor/Private/Editor/SRigHierarchy.cpp b/Engine/Plugins/Animation/ControlRig/Source/ControlRigEditor/Private/Editor/SRigHierarchy.cpp
--- a/Engine/Plugins/Animation/ControlRig/Source/ControlRigEditor/Private/Editor/SRigHierarchy.cpp
+++ b/Engine/Plugins/Animation/ControlRig/Source/ControlRigEditor/Private/Editor/SRigHierarchy.cpp
@@ -826,9 +826,13 @@ void SRigHierarchy::OnHierarchyModified(ERigHierarchyNotification InNotif, URigH
 		}
 		case ERigHierarchyNotification::ControlSettingChanged:
 		case ERigHierarchyNotification::ConnectorSettingChanged:
+		case ERigHierarchyNotification::SocketColorChanged:
 		{
 			// update color and other settings of the item
-			if(InElement && ((InElement->GetType() == ERigElementType::Control) || (InElement->GetType() == ERigElementType::Connector)))
+			if(InElement && (
+				(InElement->GetType() == ERigElementType::Control) ||
+				(InElement->GetType() == ERigElementType::Connector) ||
+				(InElement->GetType() == ERigElementType::Socket)))
 			{
 				for (int32 RootIndex = 0; RootIndex < TreeView->RootElements.Num(); ++RootIndex)
 				{
@@ -1854,11 +1858,15 @@ bool SRigHierarchy::IsNonProceduralElementSelected() const
 
 bool SRigHierarchy::CanAddElement(const ERigElementType ElementType) const
 {
-	if (ElementType == ERigElementType::Connector ||
-		ElementType == ERigElementType::Socket)
+	if (ElementType == ERigElementType::Connector)
 	{
 		return ControlRigBlueprint->IsControlRigModule();
 	}
+	if (ElementType == ERigElementType::Socket)
+	{
+		return ControlRigBlueprint->IsControlRigModule() ||
+			ControlRigBlueprint->IsModularRig();
+	}
 	return !ControlRigBlueprint->IsControlRigModule();
 }
 
diff --git a/Engine/Plugins/Animation/ControlRig/Source/ControlRigEditor/Private/Editor/SRigHierarchyTreeView.cpp b/Engine/Plugins/Animation/ControlRig/Source/ControlRigEditor/Private/Editor/SRigHierarchyTreeView.cpp
--- a/Engine/Plugins/Animation/ControlRig/Source/ControlRigEditor/Private/Editor/SRigHierarchyTreeView.cpp
+++ b/Engine/Plugins/Animation/ControlRig/Source/ControlRigEditor/Private/Editor/SRigHierarchyTreeView.cpp
@@ -972,6 +972,16 @@ bool SRigHierarchyItem::OnVerifyNameChanged(const FText& InText, FText& OutError
 
 TPair<const FSlateBrush*, FSlateColor> SRigHierarchyItem::GetBrushForElementType(const URigHierarchy* InHierarchy, const FRigElementKey& InKey)
 {
+	static const FSlateBrush* ProxyControlBrush = FControlRigEditorStyle::Get().GetBrush("ControlRig.Tree.ProxyControl"); 
+	static const FSlateBrush* ControlBrush = FControlRigEditorStyle::Get().GetBrush("ControlRig.Tree.Control");
+	static const FSlateBrush* NullBrush = FControlRigEditorStyle::Get().GetBrush("ControlRig.Tree.Null");
+	static const FSlateBrush* BoneImportedBrush = FControlRigEditorStyle::Get().GetBrush("ControlRig.Tree.BoneImported");
+	static const FSlateBrush* BoneUserBrush = FControlRigEditorStyle::Get().GetBrush("ControlRig.Tree.BoneUser");
+	static const FSlateBrush* RigidBodyBrush = FControlRigEditorStyle::Get().GetBrush("ControlRig.Tree.RigidBody");
+	static const FSlateBrush* SocketOpenBrush = FControlRigEditorStyle::Get().GetBrush("ControlRig.Tree.Socket_Open");
+	static const FSlateBrush* SocketClosedBrush = FControlRigEditorStyle::Get().GetBrush("ControlRig.Tree.Socket_Closed");
+	static const FSlateBrush* ConnectorBrush = FControlRigEditorStyle::Get().GetBrush("ControlRig.Tree.Connector");
+	
 	const FSlateBrush* Brush = nullptr;
 	FSlateColor Color = FSlateColor::UseForeground();
 	switch (InKey.Type)
@@ -986,11 +996,11 @@ TPair<const FSlateBrush*, FSlateColor> SRigHierarchyItem::GetBrushForElementType
 				{
 					if(Control->Settings.AnimationType == ERigControlAnimationType::ProxyControl)
 					{
-						Brush = FControlRigEditorStyle::Get().GetBrush("ControlRig.Tree.ProxyControl");
+						Brush = ProxyControlBrush;
 					}
 					else
 					{
-						Brush = FControlRigEditorStyle::Get().GetBrush("ControlRig.Tree.Control");
+						Brush = ControlBrush;
 					}
 					ShapeColor = Control->Settings.ShapeColor;
 				}
@@ -1007,13 +1017,13 @@ TPair<const FSlateBrush*, FSlateColor> SRigHierarchyItem::GetBrushForElementType
 			}
 			else
 			{
-				Brush = FControlRigEditorStyle::Get().GetBrush("ControlRig.Tree.Control");
+				Brush = ControlBrush;
 			}
 			break;
 		}
 		case ERigElementType::Null:
 		{
-			Brush = FControlRigEditorStyle::Get().GetBrush("ControlRig.Tree.Null");
+			Brush = NullBrush;
 			break;
 		}
 		case ERigElementType::Bone:
@@ -1033,13 +1043,13 @@ TPair<const FSlateBrush*, FSlateColor> SRigHierarchyItem::GetBrushForElementType
 			{
 				case ERigBoneType::Imported:
 				{
-					Brush = FControlRigEditorStyle::Get().GetBrush("ControlRig.Tree.BoneImported");
+					Brush = BoneImportedBrush;
 					break;
 				}
 				case ERigBoneType::User:
 				default:
 				{
-					Brush = FControlRigEditorStyle::Get().GetBrush("ControlRig.Tree.BoneUser");
+					Brush = BoneUserBrush;
 					break;
 				}
 			}
@@ -1048,18 +1058,34 @@ TPair<const FSlateBrush*, FSlateColor> SRigHierarchyItem::GetBrushForElementType
 		}
 		case ERigElementType::RigidBody:
 		{
-			Brush = FControlRigEditorStyle::Get().GetBrush("ControlRig.Tree.RigidBody");
+			Brush = RigidBodyBrush;
 			break;
 		}
 		case ERigElementType::Reference:
 		case ERigElementType::Socket:
 		{
-			Brush = FControlRigEditorStyle::Get().GetBrush("ControlRig.Tree.Socket");
+			Brush = SocketOpenBrush;
+
+			if(UControlRig* ControlRig = Cast<UControlRig>(InHierarchy->GetOuter()))
+			{
+				if(const FRigElementKey* ConnectorKey = ControlRig->GetElementKeyRedirector().FindReverse(InKey))
+				{
+					if(ConnectorKey->Type == ERigElementType::Connector)
+					{
+						Brush = SocketClosedBrush;
+					}
+				}
+			}
+
+			if(const FRigSocketElement* Socket = InHierarchy->Find<FRigSocketElement>(InKey))
+			{
+				Color = Socket->GetColor(InHierarchy);
+			}
 			break;
 		}
 		case ERigElementType::Connector:
 		{
-			Brush = FControlRigEditorStyle::Get().GetBrush("ControlRig.Tree.Connector");
+			Brush = ConnectorBrush;
 			break;
 		}
 		default:
diff --git a/Engine/Plugins/Animation/ControlRig/Source/ControlRigEditor/Public/EditMode/ControlRigEditMode.h b/Engine/Plugins/Animation/ControlRig/Source/ControlRigEditor/Public/EditMode/ControlRigEditMode.h
--- a/Engine/Plugins/Animation/ControlRig/Source/ControlRigEditor/Public/EditMode/ControlRigEditMode.h
+++ b/Engine/Plugins/Animation/ControlRig/Source/ControlRigEditor/Public/EditMode/ControlRigEditMode.h
@@ -13,6 +13,8 @@
 #include "UnrealWidgetFwd.h"
 #include "IPersonaEditMode.h"
 #include "Misc/Guid.h"
+#include "Materials/Material.h"
+#include "Materials/MaterialInstanceDynamic.h"
 #include "ControlRigEditMode.generated.h"
 
 
@@ -492,6 +494,7 @@ private:
 	bool bIsConstructionEventRunning;
 	TArray<uint32> LastHierarchyHash;
 	TArray<uint32> LastShapeLibraryHash;
+	TMap<uint32, TStrongObjectPtr<UMaterialInstanceDynamic>> SocketMaterials;
 
 	friend class FControlRigEditorModule;
 	friend class FControlRigEditor;
