Control Rig: Remove connection rule functionality

#rnx

diff --git a/Engine/Plugins/Animation/ControlRig/Source/ControlRig/Private/ControlRig.cpp b/Engine/Plugins/Animation/ControlRig/Source/ControlRig/Private/ControlRig.cpp
--- a/Engine/Plugins/Animation/ControlRig/Source/ControlRig/Private/ControlRig.cpp
+++ b/Engine/Plugins/Animation/ControlRig/Source/ControlRig/Private/ControlRig.cpp
@@ -699,13 +699,6 @@ bool UControlRig::AllConnectorsAreResolved(FString* OutFailureReason, FRigElemen
 				return false;
 			}
 		}
-
-		// now that we have the connection try if we can connect this way
-		const FRigConnectionInfo ConnectionInfo(ConnectionMap, Hierarchy, Hierarchy);
-		if(!Hierarchy->CanConnect(&ConnectionInfo, OutFailureReason, OutConnector))
-		{
-			return false;
-		}
 	}
 	return true;
 }
diff --git a/Engine/Plugins/Animation/ControlRig/Source/ControlRig/Private/ModularRigController.cpp b/Engine/Plugins/Animation/ControlRig/Source/ControlRig/Private/ModularRigController.cpp
--- a/Engine/Plugins/Animation/ControlRig/Source/ControlRig/Private/ModularRigController.cpp
+++ b/Engine/Plugins/Animation/ControlRig/Source/ControlRig/Private/ModularRigController.cpp
@@ -126,12 +126,6 @@ FRigModuleReference* UModularRigController::FindModule(const FString& InPath)
 	return *Cur;
 }
 
-bool UModularRigController::CanConnectConnectorToElement(const FRigModuleConnector& InConnector, const FRigElementKey& InTargetKey, FText& OutErrorMessage)
-{
-	// TODO: Check rules are satisfied
-	return true;
-}
-
 bool UModularRigController::ConnectConnectorToElement(const FRigElementKey& InConnectorKey, const FRigElementKey& InTargetKey, bool bSetupUndo)
 {
 	FString ConnectorParentPath, ConnectorName;
@@ -166,6 +160,7 @@ bool UModularRigController::ConnectConnectorToElement(const FRigElementKey& InCo
 		return false;
 	}
 
+	/*
 	FText ErrorMessage;
 	if (!CanConnectConnectorToElement(*ModuleConnector, InTargetKey, ErrorMessage))
 	{
@@ -173,6 +168,7 @@ bool UModularRigController::ConnectConnectorToElement(const FRigElementKey& InCo
 			*InConnectorKey.Name.ToString(), *InTargetKey.ToString(), *ErrorMessage.ToString());
 		return false;
 	}
+	*/
 
 	UBlueprint* Blueprint = Cast<UBlueprint>(GetOuter());
 	const IRigHierarchyProvider* HierarchyProvider = CastChecked<IRigHierarchyProvider>(Blueprint);
diff --git a/Engine/Plugins/Animation/ControlRig/Source/ControlRig/Private/Rigs/RigConnectionRules.cpp b/Engine/Plugins/Animation/ControlRig/Source/ControlRig/Private/Rigs/RigConnectionRules.cpp
--- a/Engine/Plugins/Animation/ControlRig/Source/ControlRig/Private/Rigs/RigConnectionRules.cpp
+++ b/Engine/Plugins/Animation/ControlRig/Source/ControlRig/Private/Rigs/RigConnectionRules.cpp
@@ -92,33 +92,17 @@ bool FRigConnectionRuleStash::operator==(const FRigConnectionRuleStash& InOther)
 		ExportedText.Equals(InOther.ExportedText, ESearchCase::CaseSensitive);
 }
 
-void FRigConnectionRuleStash::AppendAdditionalConnectors(TArray<FRigElementKey>& OutConnectors) const
-{
-	TSharedPtr<FStructOnScope> Storage;
-	if(const FRigConnectionRule* Rule = Get(Storage))
-	{
-		TArray<FRigElementKey> AdditionalConnectors = Rule->GetAdditionalConnectors();
-		for(const FRigElementKey& AdditionalConnector : AdditionalConnectors)
-		{
-			OutConnectors.AddUnique(AdditionalConnector);
-		}
-	}
-}
-
 uint32 GetTypeHash(const FRigConnectionRuleStash& InRuleStash)
 {
 	return HashCombine(GetTypeHash(InRuleStash.ScriptStructPath), GetTypeHash(InRuleStash.ExportedText));
 }
 
+/*
+
 ////////////////////////////////////////////////////////////////////////////////
 // FRigConnectionRule
 ////////////////////////////////////////////////////////////////////////////////
 
-TArray<FRigElementKey> FRigConnectionRule::GetAdditionalConnectors() const
-{
-	return TArray<FRigElementKey>();
-}
-
 bool FRigConnectionRule::CanConnect(const FRigConnectionInfo* InConnectionInfo, FString* OutFailureReason) const
 {
 	for(const TPair<FRigElementKey, FRigElementKey>& Pair : InConnectionInfo->ConnectionMap)
@@ -178,16 +162,6 @@ bool FRigConnectionRule::CanConnect(const FRigConnectionInfo* InConnectionInfo,
 // FRigAndConnectionRule
 ////////////////////////////////////////////////////////////////////////////////
 
-TArray<FRigElementKey> FRigAndConnectionRule::GetAdditionalConnectors() const
-{
-	TArray<FRigElementKey> AdditionalConnectors;
-	for(const FRigConnectionRuleStash& ChildRule : ChildRules)
-	{
-		ChildRule.AppendAdditionalConnectors(AdditionalConnectors);
-	}
-	return AdditionalConnectors;
-}
-
 bool FRigAndConnectionRule::CanConnect(const FRigConnectionInfo* InConnectionInfo, FString* OutFailureReason) const
 {
 	if(!Super::CanConnect(InConnectionInfo, OutFailureReason))
@@ -214,16 +188,6 @@ bool FRigAndConnectionRule::CanConnect(const FRigConnectionInfo* InConnectionInf
 // FRigOrConnectionRule
 ////////////////////////////////////////////////////////////////////////////////
 
-TArray<FRigElementKey> FRigOrConnectionRule::GetAdditionalConnectors() const
-{
-	TArray<FRigElementKey> AdditionalConnectors;
-	for(const FRigConnectionRuleStash& ChildRule : ChildRules)
-	{
-		ChildRule.AppendAdditionalConnectors(AdditionalConnectors);
-	}
-	return AdditionalConnectors;
-}
-
 bool FRigOrConnectionRule::CanConnect(const FRigConnectionInfo* InConnectionInfo, FString* OutFailureReason) const
 {
 	if(!Super::CanConnect(InConnectionInfo, OutFailureReason))
@@ -301,11 +265,6 @@ bool FRigTagConnectionRule::CanConnect(const FRigConnectionInfo* InConnectionInf
 // FRigChainConnectionRule
 ////////////////////////////////////////////////////////////////////////////////
 
-TArray<FRigElementKey> FRigChainConnectionRule::GetAdditionalConnectors() const
-{
-	return {RootConnector};
-}
-
 bool FRigChainConnectionRule::CanConnect(const FRigConnectionInfo* InConnectionInfo, FString* OutFailureReason) const
 {
 	if(!Super::CanConnect(InConnectionInfo, OutFailureReason))
@@ -390,4 +349,4 @@ bool FRigChainConnectionRule::CanConnect(const FRigConnectionInfo* InConnectionI
 	return true;
 }
 
-
+*/
diff --git a/Engine/Plugins/Animation/ControlRig/Source/ControlRig/Private/Rigs/RigHierarchy.cpp b/Engine/Plugins/Animation/ControlRig/Source/ControlRig/Private/Rigs/RigHierarchy.cpp
--- a/Engine/Plugins/Animation/ControlRig/Source/ControlRig/Private/Rigs/RigHierarchy.cpp
+++ b/Engine/Plugins/Animation/ControlRig/Source/ControlRig/Private/Rigs/RigHierarchy.cpp
@@ -2530,126 +2530,6 @@ void URigHierarchy::Traverse(TFunction<void(FRigBaseElement*, bool& /* continue
 	}
 }
 
-bool URigHierarchy::CanConnect(const FRigConnectionInfo* InConnectionInfo, FString* OutFailureReason, FRigElementKey* OutConnector) const
-{
-	// todo: Check if the rules are satisfied
-	return true;
-#if 0
-	check(InConnectionInfo);
-	check(InConnectionInfo->IsValid());
-
-	// make sure all connectors are provided
-	const TArray<FRigConnectorElement*> ExpectedConnectors = InConnectionInfo->SourceHierarchy->GetConnectors(false);
-	for(const FRigConnectorElement* ExpectedConnector : ExpectedConnectors)
-	{
-		if(!InConnectionInfo->ConnectionMap.Contains(ExpectedConnector->GetKey()))
-		{
-			if(OutFailureReason)
-			{
-				static constexpr TCHAR Format[] = TEXT("Connector key '%s' not provided as part of the connection map.");
-				OutFailureReason->Appendf(Format, *ExpectedConnector->GetKey().ToString());
-			}
-			if(OutConnector)
-			{
-				*OutConnector = ExpectedConnector->GetKey();
-			}
-			return false;
-		}
-	}
-	
-	for(const TPair<FRigElementKey, FRigElementKey>& Pair : InConnectionInfo->ConnectionMap)
-	{
-		const FRigElementKey& ConnectorKey = Pair.Key;
-		const FRigElementKey& TargetKey = Pair.Value;
-
-		if(!ConnectorKey.IsValid())
-		{
-			if(OutFailureReason)
-			{
-				static constexpr TCHAR Format[] = TEXT("Connector key '%s' is not valid.");
-				OutFailureReason->Appendf(Format, *ConnectorKey.ToString());
-			}
-			return false;
-		}
-		if(ConnectorKey.Type != ERigElementType::Connector)
-		{
-			if(OutFailureReason)
-			{
-				static constexpr TCHAR Format[] = TEXT("Connector key '%s' is not a connector.");
-				OutFailureReason->Appendf(Format, *ConnectorKey.ToString());
-			}
-			if(OutConnector)
-			{
-				*OutConnector = ConnectorKey;
-			}
-			return false;
-		}
-		if(!TargetKey.IsValid())
-		{
-			if(OutFailureReason)
-			{
-				static constexpr TCHAR Format[] = TEXT("Target key '%s' is not valid.");
-				OutFailureReason->Appendf(Format, *TargetKey.ToString());
-			}
-			if(OutConnector)
-			{
-				*OutConnector = ConnectorKey;
-			}
-			return false;
-		}
-
-		const FRigConnectorElement* Connector = InConnectionInfo->SourceHierarchy->Find<FRigConnectorElement>(ConnectorKey); 
-		if(Connector == nullptr)
-		{
-			if(OutFailureReason)
-			{
-				static constexpr TCHAR Format[] = TEXT("Connector element '%s' does not exist.");
-				OutFailureReason->Appendf(Format, *ConnectorKey.ToString());
-			}
-			return false;
-		}
-		const FRigBaseElement* Target = Find<FRigBaseElement>(TargetKey);
-		if(Target == nullptr)
-		{
-			if(OutFailureReason)
-			{
-				static constexpr TCHAR Format[] = TEXT("Target element '%s' does not exist.");
-				OutFailureReason->Appendf(Format, *TargetKey.ToString());
-			}
-			if(OutConnector && Connector)
-			{
-				*OutConnector = Connector->GetKey();
-			}
-			return false;
-		}
-		const FRigTransformElement* TargetTransform = Cast<FRigTransformElement>(Target);
-		if(TargetTransform == nullptr)
-		{
-			if(OutFailureReason)
-			{
-				static constexpr TCHAR Format[] = TEXT("Target element '%s' is not a transform.");
-				OutFailureReason->Appendf(Format, *TargetKey.ToString());
-			}
-			if(OutConnector && Connector)
-			{
-				*OutConnector = Connector->GetKey();
-			}
-			return false;
-		}
-
-		if(Connector && !Connector->CanConnect(InConnectionInfo, OutFailureReason))
-		{
-			if(OutConnector)
-			{
-				*OutConnector = Connector->GetKey();
-			}
-			return false;
-		}
-	}
-	return true;
-#endif
-}
-
 const FRigElementKey& URigHierarchy::GetResolvedTarget(const FRigElementKey& InConnectorKey) const
 {
 	if (InConnectorKey.Type == ERigElementType::Connector)
diff --git a/Engine/Plugins/Animation/ControlRig/Source/ControlRig/Private/Rigs/RigHierarchyElements.cpp b/Engine/Plugins/Animation/ControlRig/Source/ControlRig/Private/Rigs/RigHierarchyElements.cpp
--- a/Engine/Plugins/Animation/ControlRig/Source/ControlRig/Private/Rigs/RigHierarchyElements.cpp
+++ b/Engine/Plugins/Animation/ControlRig/Source/ControlRig/Private/Rigs/RigHierarchyElements.cpp
@@ -1553,97 +1553,6 @@ void FRigConnectorElement::Load(FArchive& Ar, ESerializationPhase SerializationP
 	}
 }
 
-bool FRigConnectorElement::CanConnect(const FRigConnectionInfo* InConnectionInfo, FString* OutFailureReason) const
-{
-	check(InConnectionInfo);
-	check(InConnectionInfo->IsValid());
-
-	if(!InConnectionInfo->ConnectionMap.Contains(GetKey()))
-	{
-		if(OutFailureReason)
-		{
-			static constexpr TCHAR Format[] = TEXT("Connector '%s' is not mapped in ConnectionMap.");
-			OutFailureReason->Appendf(Format, *GetKey().ToString());
-		}
-		return false;
-	}
-
-	for(const FRigConnectionRuleStash& Stash : Settings.Rules)
-	{
-		if(!Stash.IsValid())
-		{
-			continue;
-		}
-		
-		TSharedPtr<FStructOnScope> RuleScope;
-		const FRigConnectionRule* Rule = Stash.Get(RuleScope);
-		if(Rule == nullptr)
-		{
-			if(OutFailureReason)
-			{
-				static constexpr TCHAR Format[] = TEXT("Rule '%s' could not be loaded from '%s'");
-				OutFailureReason->Appendf(Format, *Stash.ScriptStructPath, *Stash.ExportedText);
-			}
-			return false;
-		}
-
-		FRigConnectionInfo InfoPerRule = *InConnectionInfo;
-		InfoPerRule.ConnectionMap.Reset();
-		InfoPerRule.ConnectionMap.Add(GetKey(), InConnectionInfo->ConnectionMap.FindChecked(GetKey()));
-
-		const TArray<FRigElementKey> AdditionalConnectorKeys = Rule->GetAdditionalConnectors();
-		for(const FRigElementKey& AdditionalConnectorKey : AdditionalConnectorKeys)
-		{
-			if(AdditionalConnectorKey.Type != ERigElementType::Connector)
-			{
-				if(OutFailureReason)
-				{
-					static constexpr TCHAR Format[] = TEXT("Additionally required connector '%s' is not a connector.");
-					OutFailureReason->Appendf(Format, *AdditionalConnectorKey.ToString());
-				}
-				return false;
-			}
-			if(!InConnectionInfo->SourceHierarchy->Contains(AdditionalConnectorKey))
-			{
-				if(OutFailureReason)
-				{
-					static constexpr TCHAR Format[] = TEXT("Additionally required connector '%s' does not exist.");
-					OutFailureReason->Appendf(Format, *AdditionalConnectorKey.ToString());
-				}
-				return false;
-			}
-			if(!InConnectionInfo->ConnectionMap.Contains(AdditionalConnectorKey))
-			{
-				if(OutFailureReason)
-				{
-					static constexpr TCHAR Format[] = TEXT("Additionally required connector '%s' not provided as part of ConnectionMap.");
-					OutFailureReason->Appendf(Format, *AdditionalConnectorKey.ToString());
-				}
-				return false;
-			}
-			const FRigElementKey& AdditionalTargetKey = InConnectionInfo->ConnectionMap.FindChecked(AdditionalConnectorKey);
-			if(!InConnectionInfo->TargetHierarchy->Contains(AdditionalTargetKey))
-			{
-				if(OutFailureReason)
-				{
-					static constexpr TCHAR Format[] = TEXT("Target element '%s' does not exist.");
-					OutFailureReason->Appendf(Format, *AdditionalConnectorKey.ToString());
-				}
-				return false;
-			}
-
-			InfoPerRule.ConnectionMap.Add(AdditionalConnectorKey, AdditionalTargetKey);
-		}
-
-		if(!Rule->CanConnect(&InfoPerRule, OutFailureReason))
-		{
-			return false;
-		}
-	}
-
-	return true;
-}
-
 FRigConnectorState FRigConnectorElement::GetConnectorState(const URigHierarchy* InHierarchy) const
 {
 	FRigConnectorState State;
diff --git a/Engine/Plugins/Animation/ControlRig/Source/ControlRig/Public/ModularRigController.h b/Engine/Plugins/Animation/ControlRig/Source/ControlRig/Public/ModularRigController.h
--- a/Engine/Plugins/Animation/ControlRig/Source/ControlRig/Public/ModularRigController.h
+++ b/Engine/Plugins/Animation/ControlRig/Source/ControlRig/Public/ModularRigController.h
@@ -27,9 +27,6 @@ class CONTROLRIG_API UModularRigController : public UObject
 	UFUNCTION(BlueprintCallable, Category = "Control Rig | Modules")
 	FString AddModule(const FName& InModuleName, TSubclassOf<UControlRig> InClass, const FString& InParentModulePath, bool bSetupUndo = true);
 
-	UFUNCTION(BlueprintCallable, Category = "Control Rig | Modules")
-	bool CanConnectConnectorToElement(const FRigModuleConnector& InConnector, const FRigElementKey& InTargetKey, FText& OutErrorMessage);
-
 	UFUNCTION(BlueprintCallable, Category = "Control Rig | Modules")
 	bool ConnectConnectorToElement(const FRigElementKey& InConnectorKey, const FRigElementKey& InTargetKey, bool bSetupUndo = true);
 
diff --git a/Engine/Plugins/Animation/ControlRig/Source/ControlRig/Public/Rigs/RigConnectionRules.h b/Engine/Plugins/Animation/ControlRig/Source/ControlRig/Public/Rigs/RigConnectionRules.h
--- a/Engine/Plugins/Animation/ControlRig/Source/ControlRig/Public/Rigs/RigConnectionRules.h
+++ b/Engine/Plugins/Animation/ControlRig/Source/ControlRig/Public/Rigs/RigConnectionRules.h
@@ -18,8 +18,6 @@ struct CONTROLRIG_API FRigConnectionRuleStash
 	FRigConnectionRuleStash();
 	FRigConnectionRuleStash(const FRigConnectionRule* InRule);
 
-	// todo: copy content from FRigVMActionKey
-
 	void Save(FArchive& Ar);
 	void Load(FArchive& Ar);
 	
@@ -37,8 +35,6 @@ struct CONTROLRIG_API FRigConnectionRuleStash
 		return !(*this == InOther);
 	}
 
-	void AppendAdditionalConnectors(TArray<FRigElementKey>& OutConnectors) const;
-
 	UPROPERTY(BlueprintReadWrite, EditAnywhere, Category=Rule)
 	FString ScriptStructPath;
 
@@ -57,16 +53,6 @@ public:
 	virtual ~FRigConnectionRule() {}
 
 	virtual UScriptStruct* GetScriptStruct() const { return FRigConnectionRule::StaticStruct(); }
-
-	/** If the rule requires additional connectors for validation return them here */
-	virtual TArray<FRigElementKey> GetAdditionalConnectors() const;
-
-	/**
-	 * Validate if the rule permits the link.
-	 * Provided are the connectors keys and their hierarchy,
-	 * as well as the desired target keys and their hierarchy.
-	 */
-	virtual bool CanConnect(const FRigConnectionInfo* InConnectionInfo, FString* OutFailureReason) const;
 };
 
 USTRUCT(BlueprintType, DisplayName="And Rule")
@@ -89,9 +75,6 @@ public:
 	virtual ~FRigAndConnectionRule() override {}
 
 	virtual UScriptStruct* GetScriptStruct() const override { return FRigAndConnectionRule::StaticStruct(); }
-	virtual TArray<FRigElementKey> GetAdditionalConnectors() const override;
-
-	virtual bool CanConnect(const FRigConnectionInfo* InConnectionInfo, FString* OutFailureReason) const override;
 
 	UPROPERTY(BlueprintReadWrite, EditAnywhere, Category=Settings)
 	TArray<FRigConnectionRuleStash> ChildRules;
@@ -117,9 +100,6 @@ public:
 	virtual ~FRigOrConnectionRule() override {}
 
 	virtual UScriptStruct* GetScriptStruct() const override { return FRigOrConnectionRule::StaticStruct(); }
-	virtual TArray<FRigElementKey> GetAdditionalConnectors() const override;
-
-	virtual bool CanConnect(const FRigConnectionInfo* InConnectionInfo, FString* OutFailureReason) const override;
 
 	UPROPERTY(BlueprintReadWrite, EditAnywhere, Category=Settings)
 	TArray<FRigConnectionRuleStash> ChildRules;
@@ -144,8 +124,6 @@ public:
 
 	virtual UScriptStruct* GetScriptStruct() const override { return FRigTypeConnectionRule::StaticStruct(); }
 
-	virtual bool CanConnect(const FRigConnectionInfo* InConnectionInfo, FString* OutFailureReason) const override;
-
 	UPROPERTY(BlueprintReadWrite, EditAnywhere, Category=Settings)
 	ERigElementType ElementType;
 };
@@ -169,8 +147,6 @@ public:
 
 	virtual UScriptStruct* GetScriptStruct() const override { return FRigTagConnectionRule::StaticStruct(); }
 
-	virtual bool CanConnect(const FRigConnectionInfo* InConnectionInfo, FString* OutFailureReason) const override;
-
 	UPROPERTY(BlueprintReadWrite, EditAnywhere, Category=Settings)
 	FName Tag;
 };
@@ -198,8 +174,6 @@ public:
 	virtual ~FRigChainConnectionRule() override {}
 
 	virtual UScriptStruct* GetScriptStruct() const override { return FRigChainConnectionRule::StaticStruct(); }
-	virtual TArray<FRigElementKey> GetAdditionalConnectors() const override;
-	virtual bool CanConnect(const FRigConnectionInfo* InConnectionInfo, FString* OutFailureReason) const override;
 
 	UPROPERTY(BlueprintReadWrite, EditAnywhere, Category=Settings)
 	FRigElementKey RootConnector;
diff --git a/Engine/Plugins/Animation/ControlRig/Source/ControlRig/Public/Rigs/RigHierarchy.h b/Engine/Plugins/Animation/ControlRig/Source/ControlRig/Public/Rigs/RigHierarchy.h
--- a/Engine/Plugins/Animation/ControlRig/Source/ControlRig/Public/Rigs/RigHierarchy.h
+++ b/Engine/Plugins/Animation/ControlRig/Source/ControlRig/Public/Rigs/RigHierarchy.h
@@ -3327,17 +3327,6 @@ public:
 	 */
 	void Traverse(TFunction<void(FRigBaseElement*, bool& /* continue */)> PerElementFunction, bool bTowardsChildren = true) const;
 
-	/**
-	 * Returns true if a hierarchy can be linked into another hierarchy. This relies on Connector elements
-	 * as the relevant landmarks providing rules for validation.
-	 */
-	UFUNCTION(BlueprintCallable, Category = URigHierarchy, meta = (DisplayName = "Can Link", ScriptName = "CanConnect"))
-	bool CanConnect_ForBlueprint(FRigConnectionInfo InConnectionInfo) const
-	{
-		return CanConnect(&InConnectionInfo);
-	}
-	bool CanConnect(const FRigConnectionInfo* InConnectionInfo, FString* OutFailureReason = nullptr, FRigElementKey* OutConnector = nullptr) const;
-
 	/**
 	 * Returns the currently resolved target for given connector key
 	 */
diff --git a/Engine/Plugins/Animation/ControlRig/Source/ControlRig/Public/Rigs/RigHierarchyElements.h b/Engine/Plugins/Animation/ControlRig/Source/ControlRig/Public/Rigs/RigHierarchyElements.h
--- a/Engine/Plugins/Animation/ControlRig/Source/ControlRig/Public/Rigs/RigHierarchyElements.h
+++ b/Engine/Plugins/Animation/ControlRig/Source/ControlRig/Public/Rigs/RigHierarchyElements.h
@@ -1751,8 +1751,6 @@ struct CONTROLRIG_API FRigConnectorElement final : public FRigBaseElement
 	virtual void Save(FArchive& A, ESerializationPhase SerializationPhase) override;
 	virtual void Load(FArchive& Ar, ESerializationPhase SerializationPhase) override;
 
-	bool CanConnect(const FRigConnectionInfo* InConnectionInfo, FString* OutFailureReason) const;
-
 	FRigConnectorState GetConnectorState(const URigHierarchy* InHierarchy) const;
 
 private:
diff --git a/Engine/Plugins/Animation/ControlRig/Source/ControlRigEditor/Private/Editor/SRigHierarchy.cpp b/Engine/Plugins/Animation/ControlRig/Source/ControlRigEditor/Private/Editor/SRigHierarchy.cpp
--- a/Engine/Plugins/Animation/ControlRig/Source/ControlRigEditor/Private/Editor/SRigHierarchy.cpp
+++ b/Engine/Plugins/Animation/ControlRig/Source/ControlRigEditor/Private/Editor/SRigHierarchy.cpp
@@ -2697,6 +2697,8 @@ TOptional<EItemDropZone> SRigHierarchy::OnCanAcceptDrop(const FDragDropEvent& Dr
 					{
 						TargetKey = TargetItem->Key;
 					}
+
+					/*
 					FText ErrorMessage;
 					if (Controller->CanConnectConnectorToElement(*PrimaryConnector, TargetKey, ErrorMessage))
 					{
@@ -2706,6 +2708,8 @@ TOptional<EItemDropZone> SRigHierarchy::OnCanAcceptDrop(const FDragDropEvent& Dr
 					{
 						return InvalidDropZone;
 					}
+					*/
+					ReturnDropZone = DropZone;
 				}
 			}
 		}
diff --git a/Engine/Plugins/Animation/ControlRig/Source/ControlRigEditor/Private/Editor/SRigHierarchyTreeView.cpp b/Engine/Plugins/Animation/ControlRig/Source/ControlRigEditor/Private/Editor/SRigHierarchyTreeView.cpp
--- a/Engine/Plugins/Animation/ControlRig/Source/ControlRigEditor/Private/Editor/SRigHierarchyTreeView.cpp
+++ b/Engine/Plugins/Animation/ControlRig/Source/ControlRigEditor/Private/Editor/SRigHierarchyTreeView.cpp
@@ -929,6 +929,7 @@ void SRigHierarchyTreeView::AddConnectorResolveWarningTag(TSharedPtr<FRigTreeEle
 FText SRigHierarchyTreeView::GetConnectorWarningMessage(TSharedPtr<FRigTreeElement> InTreeElement,
 	TWeakObjectPtr<UControlRig> InControlRigPtr, const FRigElementKey InConnectorKey) const
 {
+	/*
 	if(UControlRig* ControlRig = InControlRigPtr.Get())
 	{
 		URigHierarchy* Hierarchy = ControlRig->GetHierarchy();
@@ -960,6 +961,7 @@ FText SRigHierarchyTreeView::GetConnectorWarningMessage(TSharedPtr<FRigTreeEleme
 			return NotResolvedWarning;
 		}
 	}
+	*/
 	return FText();
 }
 
