Control Rig: Allow namespaced names for default control rigs by default

#rnx

diff --git a/Engine/Plugins/Animation/ControlRig/Source/ControlRig/Private/Rigs/RigHierarchy.cpp b/Engine/Plugins/Animation/ControlRig/Source/ControlRig/Private/Rigs/RigHierarchy.cpp
--- a/Engine/Plugins/Animation/ControlRig/Source/ControlRig/Private/Rigs/RigHierarchy.cpp
+++ b/Engine/Plugins/Animation/ControlRig/Source/ControlRig/Private/Rigs/RigHierarchy.cpp
@@ -9,6 +9,7 @@
 #include "Math/ControlRigMathLibrary.h"
 #include "UObject/AnimObjectVersion.h"
 #include "ControlRigObjectVersion.h"
+#include "ModularRig.h"
 #include "Algo/Count.h"
 #include "Serialization/MemoryReader.h"
 #include "Serialization/MemoryWriter.h"
@@ -1258,10 +1259,10 @@ void URigHierarchy::SanitizeName(FRigName& InOutName, bool bAllowNameSpaces)
 	}
 }
 
-FRigName URigHierarchy::GetSanitizedName(const FRigName& InName)
+FRigName URigHierarchy::GetSanitizedName(const FRigName& InName, bool bAllowNameSpaces)
 {
 	FRigName Name = InName;
-	SanitizeName(Name);
+	SanitizeName(Name, bAllowNameSpaces);
 	return Name;
 }
 
@@ -1302,8 +1303,33 @@ bool URigHierarchy::IsNameAvailable(const FRigName& InPotentialNewName, ERigElem
 		return false;
 	}
 
+	bool bAllowNameSpaces = bAllowNameSpaceWhenSanitizingName;
+
+	// try to find a control rig this belongs to
+	const UControlRig* ControlRig = Cast<UControlRig>(GetOuter());
+	if(ControlRig == nullptr)
+	{
+		if(const UBlueprint* Blueprint = GetTypedOuter<UBlueprint>())
+		{
+			if(const UClass* Class = Blueprint->GeneratedClass)
+			{
+				ControlRig = Cast<UControlRig>(Class->GetDefaultObject());
+			}
+		}
+	}
+
+	// allow namespaces on default control rigs (non-module and non-modular)
+	if(ControlRig)
+	{
+		if(!ControlRig->IsRigModule() &&
+			!ControlRig->GetClass()->IsChildOf(UModularRig::StaticClass()))
+		{
+			bAllowNameSpaces = true;
+		}
+	}
+
 	FRigName SanitizedName = UnsanitizedName;
-	SanitizeName(SanitizedName, bAllowNameSpaceWhenSanitizingName);
+	SanitizeName(SanitizedName, bAllowNameSpaces);
 
 	if (SanitizedName != UnsanitizedName)
 	{
@@ -1363,7 +1389,7 @@ bool URigHierarchy::IsDisplayNameAvailable(const FRigElementKey& InParentElement
 	}
 
 	FRigName SanitizedName = UnsanitizedName;
-	SanitizeName(SanitizedName);
+	SanitizeName(SanitizedName, true);
 
 	if (SanitizedName != UnsanitizedName)
 	{
@@ -1380,7 +1406,7 @@ bool URigHierarchy::IsDisplayNameAvailable(const FRigElementKey& InParentElement
 FRigName URigHierarchy::GetSafeNewName(const FRigName& InPotentialNewName, ERigElementType InType) const
 {
 	FRigName SanitizedName = InPotentialNewName;
-	SanitizeName(SanitizedName, false);
+	SanitizeName(SanitizedName, true);
 
 	if(ExecuteContext)
 	{
diff --git a/Engine/Plugins/Animation/ControlRig/Source/ControlRig/Public/Rigs/RigHierarchy.h b/Engine/Plugins/Animation/ControlRig/Source/ControlRig/Public/Rigs/RigHierarchy.h
--- a/Engine/Plugins/Animation/ControlRig/Source/ControlRig/Public/Rigs/RigHierarchy.h
+++ b/Engine/Plugins/Animation/ControlRig/Source/ControlRig/Public/Rigs/RigHierarchy.h
@@ -1668,14 +1668,14 @@ public:
 	 * Sanitizes a name by removing invalid characters.
 	 * @param InOutName The name to sanitize in place.
 	 */
-	static void SanitizeName(FRigName& InOutName, bool bAllowNameSpaces = false);
+	static void SanitizeName(FRigName& InOutName, bool bAllowNameSpaces = true);
 
 	/**
 	 * Sanitizes a name by removing invalid characters.
 	 * @param InName The name to sanitize.
 	 * @return The sanitized name.
  	 */
-	static FRigName GetSanitizedName(const FRigName& InName);
+	static FRigName GetSanitizedName(const FRigName& InName, bool bAllowNameSpaces = true);
 
 	/**
 	 * Returns true if a given name is available.
