Control Rig: Fix hierarchy to allow for more than one primary

#rnx

diff --git a/Engine/Plugins/Animation/ControlRig/Source/ControlRig/Private/Rigs/RigHierarchyController.cpp b/Engine/Plugins/Animation/ControlRig/Source/ControlRig/Private/Rigs/RigHierarchyController.cpp
--- a/Engine/Plugins/Animation/ControlRig/Source/ControlRig/Private/Rigs/RigHierarchyController.cpp
+++ b/Engine/Plugins/Animation/ControlRig/Source/ControlRig/Private/Rigs/RigHierarchyController.cpp
@@ -586,9 +586,21 @@ FRigElementKey URigHierarchyController::AddConnector(FName InName, FRigConnector
 		{
 			if(Connector->Settings.Type == EConnectorType::Primary)
 			{
-				static constexpr TCHAR Format[] = TEXT("Cannot add connector '%s' - there already is a primary connector.");
-				ReportAndNotifyErrorf(Format, *InName.ToString());
-				return FRigElementKey();
+				if(Hierarchy->HasExecuteContext())
+				{
+					const FControlRigExecuteContext& CRContext = Hierarchy->ExecuteContext->GetPublicData<FControlRigExecuteContext>();
+					const FString NameSpace = CRContext.GetRigModuleNameSpace();
+					if(!NameSpace.IsEmpty())
+					{
+						const FName& ConnectorNameSpace = Hierarchy->GetNameMetadata(Connector->GetKey(), URigHierarchy::NameSpaceMetadataName, NAME_None);
+						if(!ConnectorNameSpace.IsNone() && ConnectorNameSpace.ToString().Equals(NameSpace, ESearchCase::CaseSensitive))
+						{
+							static constexpr TCHAR Format[] = TEXT("Cannot add connector '%s' - there already is a primary connector.");
+							ReportAndNotifyErrorf(Format, *InName.ToString());
+							return FRigElementKey();
+						}
+					}
+				}
 			}
 		}
 
