Control Rig: Add asset creation for modular rig + add default root socket

#rnx

diff --git a/Engine/Plugins/Animation/ControlRig/Source/ControlRig/Private/ControlRig.cpp b/Engine/Plugins/Animation/ControlRig/Source/ControlRig/Private/ControlRig.cpp
--- a/Engine/Plugins/Animation/ControlRig/Source/ControlRig/Private/ControlRig.cpp
+++ b/Engine/Plugins/Animation/ControlRig/Source/ControlRig/Private/ControlRig.cpp
@@ -36,6 +36,7 @@
 #include "RigVMCore/RigVMNativized.h"
 #include "UObject/UObjectIterator.h"
 #include "RigVMCore/RigVMAssetUserData.h"
+#include "ModularRig.h"
 
 #include UE_INLINE_GENERATED_CPP_BY_NAME(ControlRig)
 
@@ -1793,6 +1794,29 @@ bool UControlRig::IsRigModuleInstance() const
 	return false;
 }
 
+bool UControlRig::IsModularRig() const
+{
+	if(const UClass* Class = GetClass())
+	{
+		return Class->IsChildOf(UModularRig::StaticClass());
+	}
+	return false;
+}
+
+bool UControlRig::IsStandaloneRig() const
+{
+	return !IsModularRig();
+}
+
+bool UControlRig::IsNativeRig() const
+{
+	if(const UClass* Class = GetClass())
+	{
+		return Class->IsNative();
+	}
+	return false;
+}
+
 UControlRig* UControlRig::GetParentRig() const
 {
 	return GetTypedOuter<UControlRig>();
diff --git a/Engine/Plugins/Animation/ControlRig/Source/ControlRig/Private/Rigs/RigHierarchyController.cpp b/Engine/Plugins/Animation/ControlRig/Source/ControlRig/Private/Rigs/RigHierarchyController.cpp
--- a/Engine/Plugins/Animation/ControlRig/Source/ControlRig/Private/Rigs/RigHierarchyController.cpp
+++ b/Engine/Plugins/Animation/ControlRig/Source/ControlRig/Private/Rigs/RigHierarchyController.cpp
@@ -720,6 +720,30 @@ FRigElementKey URigHierarchyController::AddSocket(FName InName, FRigElementKey I
 	return NewElement->Key;
 }
 
+FRigElementKey URigHierarchyController::AddDefaultRootSocket()
+{
+	FRigElementKey SocketKey;
+	if(const URigHierarchy* CurrentHierarchy = GetHierarchy())
+	{
+		CurrentHierarchy->ForEach<FRigBoneElement>([this, CurrentHierarchy, &SocketKey](const FRigBoneElement* Bone) -> bool
+		{
+			// find first root bone
+			if(CurrentHierarchy->GetNumberOfParents(Bone) == 0)
+			{
+				static const FRigElementKey RootSocketKey(TEXT("Root"), ERigElementType::Socket);
+				SocketKey = AddSocket(RootSocketKey.Name, Bone->GetKey(), FTransform::Identity);
+
+				// stop
+				return false;
+			}
+
+			// continue with the search
+			return true;
+		});
+	}
+	return SocketKey; 
+}
+
 FRigControlSettings URigHierarchyController::GetControlSettings(FRigElementKey InKey) const
 {
 	if(!IsValid())
diff --git a/Engine/Plugins/Animation/ControlRig/Source/ControlRig/Public/ControlRig.h b/Engine/Plugins/Animation/ControlRig/Source/ControlRig/Public/ControlRig.h
--- a/Engine/Plugins/Animation/ControlRig/Source/ControlRig/Public/ControlRig.h
+++ b/Engine/Plugins/Animation/ControlRig/Source/ControlRig/Public/ControlRig.h
@@ -95,6 +95,15 @@ public:
 	// when being interacted with the asset editor
 	bool IsRigModuleInstance() const;
 
+	// Returns true if this rig is a modular rig (of class UModularRig)
+	bool IsModularRig() const;
+
+	// Returns true if this is a standalone rig (of class UControlRig and not modular)
+	bool IsStandaloneRig() const;
+
+	// Returns true if this is a native rig (implemented in C++)
+	bool IsNativeRig() const;
+
 	// Returns the parent rig hosting this module instance
 	UControlRig* GetParentRig() const;
 
diff --git a/Engine/Plugins/Animation/ControlRig/Source/ControlRig/Public/Rigs/RigHierarchyController.h b/Engine/Plugins/Animation/ControlRig/Source/ControlRig/Public/Rigs/RigHierarchyController.h
--- a/Engine/Plugins/Animation/ControlRig/Source/ControlRig/Public/Rigs/RigHierarchyController.h
+++ b/Engine/Plugins/Animation/ControlRig/Source/ControlRig/Public/Rigs/RigHierarchyController.h
@@ -263,6 +263,11 @@ public:
 	UFUNCTION(BlueprintCallable, Category = URigHierarchyController)
 	FRigElementKey AddSocket(FName InName, FRigElementKey InParent, FTransform InTransform, bool bTransformInGlobal = true, const FLinearColor& InColor = FLinearColor::White, const FString& InDescription = TEXT(""), bool bSetupUndo = false, bool bPrintPythonCommand = false);
 
+	/**
+	 * Adds a socket to the first determined root bone the hierarchy
+	 * @return The key for the newly created bone (or an invalid key).
+	 */
+	FRigElementKey AddDefaultRootSocket();
 	/**
 	 * Returns the control settings of a given control
 	 * @param InKey The key of the control to receive the settings for
diff --git a/Engine/Plugins/Animation/ControlRig/Source/ControlRigEditor/Private/ControlRigBlueprintActions.cpp b/Engine/Plugins/Animation/ControlRig/Source/ControlRigEditor/Private/ControlRigBlueprintActions.cpp
--- a/Engine/Plugins/Animation/ControlRig/Source/ControlRigEditor/Private/ControlRigBlueprintActions.cpp
+++ b/Engine/Plugins/Animation/ControlRig/Source/ControlRigEditor/Private/ControlRigBlueprintActions.cpp
@@ -34,6 +34,7 @@
 #include "LevelSequenceEditorBlueprintLibrary.h"
 #include "Sequencer/ControlRigParameterTrackEditor.h"
 #include "Rigs/RigHierarchyController.h"
+#include "ModularRig.h"
 
 #define LOCTEXT_NAMESPACE "ControlRigBlueprintActions"
 
@@ -157,6 +158,7 @@ void FControlRigBlueprintActions::ExtendSketalMeshToolMenu()
 				TArray<UObject*> SelectedObjects = Context->GetSelectedObjects();
 				if (SelectedObjects.Num() > 0)
 				{
+					static constexpr bool bModularRig = true;
 					InSection.AddMenuEntry(
 						"CreateControlRig",
 						LOCTEXT("CreateControlRig", "Control Rig"),
@@ -166,7 +168,20 @@ void FControlRigBlueprintActions::ExtendSketalMeshToolMenu()
 						{
 							for (UObject* SelectedObject : SelectedObjects)
 							{
-								FControlRigBlueprintActions::CreateControlRigFromSkeletalMeshOrSkeleton(SelectedObject);
+								CreateControlRigFromSkeletalMeshOrSkeleton(SelectedObject, !bModularRig);
+							}
+						})
+					);
+					InSection.AddMenuEntry(
+						"CreateModularRig",
+						LOCTEXT("CreateModularRig", "Modular Rig"),
+						LOCTEXT("CreateModularRig_ToolTip", "Creates a modular rig and preconfigures it for this asset"),
+						FSlateIcon(FRigVMEditorStyle::Get().GetStyleSetName(), "RigVM", "RigVM.Unit"),
+						FExecuteAction::CreateLambda([SelectedObjects]()
+						{
+							for (UObject* SelectedObject : SelectedObjects)
+							{
+								CreateControlRigFromSkeletalMeshOrSkeleton(SelectedObject, bModularRig);
 							}
 						})
 					);
@@ -176,12 +191,12 @@ void FControlRigBlueprintActions::ExtendSketalMeshToolMenu()
 	}
 }
 
-UControlRigBlueprint* FControlRigBlueprintActions::CreateNewControlRigAsset(const FString& InDesiredPackagePath)
+UControlRigBlueprint* FControlRigBlueprintActions::CreateNewControlRigAsset(const FString& InDesiredPackagePath, const bool bModularRig)
 {
 	FAssetToolsModule& AssetToolsModule = FModuleManager::LoadModuleChecked<FAssetToolsModule>("AssetTools");
 
 	UControlRigBlueprintFactory* Factory = NewObject<UControlRigBlueprintFactory>();
-	Factory->ParentClass = UControlRig::StaticClass();
+	Factory->ParentClass = bModularRig ? UModularRig::StaticClass() : UControlRig::StaticClass();
 
 	FString UniquePackageName;
 	FString UniqueAssetName;
@@ -196,7 +211,7 @@ UControlRigBlueprint* FControlRigBlueprintActions::CreateNewControlRigAsset(cons
 	return Cast<UControlRigBlueprint>(NewAsset);
 }
 
-UControlRigBlueprint* FControlRigBlueprintActions::CreateControlRigFromSkeletalMeshOrSkeleton(UObject* InSelectedObject)
+UControlRigBlueprint* FControlRigBlueprintActions::CreateControlRigFromSkeletalMeshOrSkeleton(UObject* InSelectedObject, const bool bModularRig)
 {
 	USkeletalMesh* SkeletalMesh = Cast<USkeletalMesh>(InSelectedObject);
 	USkeleton* Skeleton = Cast<USkeleton>(InSelectedObject);
@@ -228,14 +243,18 @@ UControlRigBlueprint* FControlRigBlueprintActions::CreateControlRigFromSkeletalM
 		PackagePath = PackagePath.Left(LastSlashPos);
 	}
 
-	UControlRigBlueprint* NewControlRigBlueprint = CreateNewControlRigAsset(PackagePath / ControlRigName);
+	UControlRigBlueprint* NewControlRigBlueprint = CreateNewControlRigAsset(PackagePath / ControlRigName, bModularRig);
 	if (NewControlRigBlueprint == nullptr)
 	{
 		return nullptr;
 	}
 
-	NewControlRigBlueprint->GetHierarchyController()->ImportBones(*RefSkeleton, NAME_None, false, false, false, false);
-	NewControlRigBlueprint->GetHierarchyController()->ImportCurves(Skeleton, NAME_None, false, false);
+	if(URigHierarchyController* Controller = NewControlRigBlueprint->GetHierarchyController())
+	{
+		Controller->ImportBones(*RefSkeleton, NAME_None, false, false, false, false);
+		Controller->ImportCurves(Skeleton, NAME_None, false, false);
+		Controller->AddDefaultRootSocket();
+	}
 	NewControlRigBlueprint->SourceHierarchyImport = Skeleton;
 	NewControlRigBlueprint->SourceCurveImport = Skeleton;
 	NewControlRigBlueprint->PropagateHierarchyFromBPToInstances();
@@ -245,7 +264,10 @@ UControlRigBlueprint* FControlRigBlueprintActions::CreateControlRigFromSkeletalM
 		NewControlRigBlueprint->SetPreviewMesh(SkeletalMesh);
 	}
 
-	NewControlRigBlueprint->RecompileVM();
+	if(!bModularRig)
+	{
+		NewControlRigBlueprint->RecompileVM();
+	}
 
 	return NewControlRigBlueprint;
 }
diff --git a/Engine/Plugins/Animation/ControlRig/Source/ControlRigEditor/Private/ControlRigBlueprintActions.h b/Engine/Plugins/Animation/ControlRig/Source/ControlRigEditor/Private/ControlRigBlueprintActions.h
--- a/Engine/Plugins/Animation/ControlRig/Source/ControlRigEditor/Private/ControlRigBlueprintActions.h
+++ b/Engine/Plugins/Animation/ControlRig/Source/ControlRigEditor/Private/ControlRigBlueprintActions.h
@@ -31,8 +31,8 @@ public:
 
 	static void ExtendSketalMeshToolMenu();
 
-	static UControlRigBlueprint* CreateNewControlRigAsset(const FString& InDesiredPackagePath);
-	static UControlRigBlueprint* CreateControlRigFromSkeletalMeshOrSkeleton(UObject* InSelectedObject);
+	static UControlRigBlueprint* CreateNewControlRigAsset(const FString& InDesiredPackagePath, const bool bModularRig);
+	static UControlRigBlueprint* CreateControlRigFromSkeletalMeshOrSkeleton(UObject* InSelectedObject, const bool bModularRig);
 
 	static USkeletalMesh* GetSkeletalMeshFromControlRigBlueprint(const FAssetData& InAsset);
 	static void PostSpawningSkeletalMeshActor(AActor* InSpawnedActor, UObject* InAsset);
diff --git a/Engine/Plugins/Animation/ControlRig/Source/ControlRigEditor/Private/ControlRigBlueprintFactory.cpp b/Engine/Plugins/Animation/ControlRig/Source/ControlRigEditor/Private/ControlRigBlueprintFactory.cpp
--- a/Engine/Plugins/Animation/ControlRig/Source/ControlRigEditor/Private/ControlRigBlueprintFactory.cpp
+++ b/Engine/Plugins/Animation/ControlRig/Source/ControlRigEditor/Private/ControlRigBlueprintFactory.cpp
@@ -304,14 +304,14 @@ UObject* UControlRigBlueprintFactory::FactoryCreateNew(UClass* Class, UObject* I
 	return FactoryCreateNew(Class, InParent, Name, Flags, Context, Warn, NAME_None);
 }
 
-UControlRigBlueprint* UControlRigBlueprintFactory::CreateNewControlRigAsset(const FString& InDesiredPackagePath)
+UControlRigBlueprint* UControlRigBlueprintFactory::CreateNewControlRigAsset(const FString& InDesiredPackagePath, const bool bModularRig)
 {
-	return FControlRigBlueprintActions::CreateNewControlRigAsset(InDesiredPackagePath);
+	return FControlRigBlueprintActions::CreateNewControlRigAsset(InDesiredPackagePath, bModularRig);
 }
 
-UControlRigBlueprint* UControlRigBlueprintFactory::CreateControlRigFromSkeletalMeshOrSkeleton(UObject* InSelectedObject)
+UControlRigBlueprint* UControlRigBlueprintFactory::CreateControlRigFromSkeletalMeshOrSkeleton(UObject* InSelectedObject, const bool bModularRig)
 {
-	return FControlRigBlueprintActions::CreateControlRigFromSkeletalMeshOrSkeleton(InSelectedObject);
+	return FControlRigBlueprintActions::CreateControlRigFromSkeletalMeshOrSkeleton(InSelectedObject, bModularRig);
 }
 
 #undef LOCTEXT_NAMESPACE
diff --git a/Engine/Plugins/Animation/ControlRig/Source/ControlRigEditor/Private/Editor/ControlRigEditor.cpp b/Engine/Plugins/Animation/ControlRig/Source/ControlRigEditor/Private/Editor/ControlRigEditor.cpp
--- a/Engine/Plugins/Animation/ControlRig/Source/ControlRigEditor/Private/Editor/ControlRigEditor.cpp
+++ b/Engine/Plugins/Animation/ControlRig/Source/ControlRigEditor/Private/Editor/ControlRigEditor.cpp
@@ -5015,18 +5015,21 @@ void FControlRigEditor::OnPreConstruction_AnyThread(UControlRig* InRig, const FN
 						const TArray<FRigElementKey> Bones = Controller->ImportBones(RigBlueprint->PreviewSkeletalMesh->GetSkeleton());
 						for(const FRigElementKey& Bone : Bones)
 						{
-							if(FRigBaseElement* Element = InRig->GetHierarchy()->Find(Bone))
+							if(FRigBaseElement* Element = Hierarchy->Find(Bone))
 							{
 								Element->CreatedAtInstructionIndex = InstructionIndex;
 							}
 						}
 
+						// create a default socket under the root bone
+						Controller->AddDefaultRootSocket();
+
 						// create a null to store controls under
 						static const FRigElementKey ControlParentKey(TEXT("Controls"), ERigElementType::Null);
 						if(!Hierarchy->Contains(ControlParentKey))
 						{
 							const FRigElementKey Null = Controller->AddNull(ControlParentKey.Name, FRigElementKey(), FTransform::Identity, true, false, false);
-							if(FRigBaseElement* Element = InRig->GetHierarchy()->Find(Null))
+							if(FRigBaseElement* Element = Hierarchy->Find(Null))
 							{
 								Element->CreatedAtInstructionIndex = InstructionIndex;
 							}
diff --git a/Engine/Plugins/Animation/ControlRig/Source/ControlRigEditor/Private/Editor/SRigHierarchy.cpp b/Engine/Plugins/Animation/ControlRig/Source/ControlRigEditor/Private/Editor/SRigHierarchy.cpp
--- a/Engine/Plugins/Animation/ControlRig/Source/ControlRigEditor/Private/Editor/SRigHierarchy.cpp
+++ b/Engine/Plugins/Animation/ControlRig/Source/ControlRigEditor/Private/Editor/SRigHierarchy.cpp
@@ -1706,9 +1706,11 @@ void SRigHierarchy::ImportHierarchy(const FAssetData& InAssetData)
 		// we do this to avoid the editmode / viewport shapes to refresh recursively,
 		// which can add an extreme slowdown depending on the number of bones (n^(n-1))
 		bool bSelectBones = true;
+		bool bIsModularRig = false;
 		if (const UControlRig* CurrentRig = EditorSharedPtr->GetControlRig())
 		{
 			bSelectBones = !CurrentRig->IsConstructionModeEnabled();
+			bIsModularRig = CurrentRig->IsModularRig();
 		}
 
 		URigHierarchyController* Controller = Hierarchy->GetController(true);
@@ -1716,6 +1718,10 @@ void SRigHierarchy::ImportHierarchy(const FAssetData& InAssetData)
 
 		const TArray<FRigElementKey> ImportedBones = Controller->ImportBones(Mesh->GetSkeleton(), NAME_None, false, false, bSelectBones, true, true);
 		Controller->ImportCurves(Mesh->GetSkeleton(), NAME_None, false, true);
+		if(bIsModularRig && Hierarchy->GetSockets().Num() == 0)
+		{
+			(void)Controller->AddDefaultRootSocket();
+		}
 
 		ControlRigBlueprint->SourceHierarchyImport = Mesh->GetSkeleton();
 		ControlRigBlueprint->SourceCurveImport = Mesh->GetSkeleton();
diff --git a/Engine/Plugins/Animation/ControlRig/Source/ControlRigEditor/Public/ControlRigBlueprintFactory.h b/Engine/Plugins/Animation/ControlRig/Source/ControlRigEditor/Public/ControlRigBlueprintFactory.h
--- a/Engine/Plugins/Animation/ControlRig/Source/ControlRigEditor/Public/ControlRigBlueprintFactory.h
+++ b/Engine/Plugins/Animation/ControlRig/Source/ControlRigEditor/Public/ControlRigBlueprintFactory.h
@@ -31,16 +31,18 @@ public:
 	/**
 	 * Create a new control rig asset within the contents space of the project.
 	 * @param InDesiredPackagePath The package path to use for the control rig asset
+	 * @param bModularRig If true the rig will be created as a modular rig
 	 */
 	UFUNCTION(BlueprintCallable, Category = "Control Rig")
-	static UControlRigBlueprint* CreateNewControlRigAsset(const FString& InDesiredPackagePath);
+	static UControlRigBlueprint* CreateNewControlRigAsset(const FString& InDesiredPackagePath, const bool bModularRig = false);
 
 	/**
 	 * Create a new control rig asset within the contents space of the project
 	 * based on a skeletal mesh or skeleton object.
 	 * @param InSelectedObject The SkeletalMesh / Skeleton object to base the control rig asset on
+	 * @param bModularRig If true the rig will be created as a modular rig
 	 */
 	UFUNCTION(BlueprintCallable, Category = "Control Rig")
-	static UControlRigBlueprint* CreateControlRigFromSkeletalMeshOrSkeleton(UObject* InSelectedObject);
+	static UControlRigBlueprint* CreateControlRigFromSkeletalMeshOrSkeleton(UObject* InSelectedObject, const bool bModularRig = false);
 };
 
