Control Rig: Refresh connectors on load and on module change

#rnx

diff --git a/Engine/Plugins/Animation/ControlRig/Source/ControlRigDeveloper/Private/ControlRigBlueprint.cpp b/Engine/Plugins/Animation/ControlRig/Source/ControlRigDeveloper/Private/ControlRigBlueprint.cpp
--- a/Engine/Plugins/Animation/ControlRig/Source/ControlRigDeveloper/Private/ControlRigBlueprint.cpp
+++ b/Engine/Plugins/Animation/ControlRig/Source/ControlRigDeveloper/Private/ControlRigBlueprint.cpp
@@ -925,7 +925,19 @@ void UControlRigBlueprint::PreSave(FObjectPreSaveContext ObjectSaveContext)
 		return true;
 	});
 
-	UpdateExposedModuleConnectors();
+	if(IsControlRigModule())
+	{
+		URigHierarchy* DebuggedHierarchy = Hierarchy;
+		if(UControlRig* DebuggedRig = Cast<UControlRig>(GetObjectBeingDebugged()))
+		{
+			DebuggedHierarchy = DebuggedRig->GetHierarchy();
+		}
+
+		TGuardValue<bool> SuspendNotifGuard(Hierarchy->GetSuspendNotificationsFlag(), true);
+		TGuardValue<bool> SuspendNotifGuardOnDebuggedHierarchy(DebuggedHierarchy->GetSuspendNotificationsFlag(), true);
+
+		UpdateExposedModuleConnectors();
+	}
 
 	if (IsControlRigModule())
 	{
@@ -1166,6 +1178,8 @@ void UControlRigBlueprint::HandlePackageDone()
 			(void)Element->Class.LoadSynchronous();
 			return true;
 		});
+
+		RefreshModuleConnectors();
 		RecompileModularRig();
 	}
 }
@@ -2245,6 +2259,99 @@ void UControlRigBlueprint::OnModularDependencyChanged(URigVMBlueprint* InBluepri
 
 }
 
+void UControlRigBlueprint::RefreshModuleConnectors()
+{
+	if(!IsModularRig())
+	{
+		return;
+	}
+
+	ModularRigModel.ForEachModule([this](const FRigModuleReference* Element) -> bool
+	{
+		RefreshModuleConnectors(Element);
+		return true;
+	});
+}
+
+void UControlRigBlueprint::RefreshModuleConnectors(const FRigModuleReference* InModule)
+{
+	if(!IsModularRig())
+	{
+		return;
+	}
+
+	// avoid dead class pointers
+	if(InModule->Class.Get() == nullptr)
+	{
+		return;
+	}
+	
+	const bool bRemoveAllConnectors = !ModularRigModel.FindModule(InModule->GetPath());
+	
+	if (URigHierarchyController* Controller = GetHierarchyController())
+	{
+		if (UControlRig* CDO = GetControlRigClass()->GetDefaultObject<UControlRig>())
+		{
+			Hierarchy->Modify();
+
+			const FString Namespace = InModule->GetNamespace();
+			const TArray<FRigElementKey> AllConnectors = Hierarchy->GetKeysOfType<FRigConnectorElement>();
+			const TArray<FRigElementKey> ExistingConnectors = AllConnectors.FilterByPredicate([Namespace](const FRigElementKey& ConnectorKey) -> bool
+			{
+				const FString ConnectorName = ConnectorKey.Name.ToString();
+				FString ConnectorNamespace, ShortName = ConnectorName;
+				ConnectorName.Split(UModularRig::NamespaceSeparator, &ConnectorNamespace, &ShortName, ESearchCase::CaseSensitive, ESearchDir::FromEnd);
+				ConnectorNamespace.Append(UModularRig::NamespaceSeparator);
+				return ConnectorNamespace.Equals(Namespace, ESearchCase::CaseSensitive);
+			});
+
+			// setup the module information. this is needed so that newly added
+			// connectors result in the right namespace metadata etc
+			FRigVMExtendedExecuteContext& Context = CDO->GetRigVMExtendedExecuteContext();
+			FRigHierarchyExecuteContextBracket HierarchyContextGuard(Controller->GetHierarchy(), &Context);
+			FControlRigExecuteContext& PublicContext = Context.GetPublicDataSafe<FControlRigExecuteContext>();
+			FControlRigExecuteContextRigModuleGuard RigModuleGuard(PublicContext, InModule->GetNamespace());
+		
+			const UControlRig* ModuleCDO = InModule->Class->GetDefaultObject<UControlRig>();
+			const TArray<FRigModuleConnector>& ExpectedConnectors = ModuleCDO->GetRigModuleSettings().ExposedConnectors;
+
+			// remove the obsolete connectors
+			for(const FRigElementKey& Connector : ExistingConnectors)
+			{
+				const FString ConnectorNameString = Connector.Name.ToString();
+				FString ConnectorNamespace, ShortName = ConnectorNameString;
+				ConnectorNameString.Split(UModularRig::NamespaceSeparator, &ConnectorNamespace, &ShortName, ESearchCase::CaseSensitive, ESearchDir::FromEnd);
+				const bool bConnectorExpected = ExpectedConnectors.ContainsByPredicate(
+					[ShortName](const FRigModuleConnector& ExpectedConnector) -> bool
+					{
+						return ExpectedConnector.Name.Equals(ShortName, ESearchCase::CaseSensitive);
+					}
+				);
+				
+				if(bRemoveAllConnectors || !bConnectorExpected)
+				{
+					(void)Controller->RemoveElement(Connector);
+				}
+			}
+
+			// add the missing expected connectors
+			if(!bRemoveAllConnectors)
+			{
+				for (const FRigModuleConnector& Connector : ExpectedConnectors)
+				{
+					const FName ConnectorName = *Connector.Name;
+					if(!Hierarchy->Contains(FRigElementKey(ConnectorName, ERigElementType::Connector)))
+					{
+						(void)Controller->AddConnector(ConnectorName, Connector.Settings);
+					}
+				}
+			}
+			
+			PropagateHierarchyFromBPToInstances();
+		}
+	}
+}
+
 void UControlRigBlueprint::HandleHierarchyModified(ERigHierarchyNotification InNotification, URigHierarchy* InHierarchy, const FRigBaseElement* InElement)
 {
 #if WITH_EDITOR
@@ -2342,29 +2449,7 @@ void UControlRigBlueprint::HandleRigModulesModified(EModularRigNotification InNo
 		{
 			if (InModule)
 			{
-				if (URigHierarchyController* Controller = GetHierarchyController())
-				{
-					if (UControlRig* CDO = GetControlRigClass()->GetDefaultObject<UControlRig>())
-					{
-						Hierarchy->Modify();
-						
-						FRigVMExtendedExecuteContext& Context = CDO->GetRigVMExtendedExecuteContext();
-						FRigHierarchyExecuteContextBracket HierarchyContextGuard(Controller->GetHierarchy(), &Context);
-			
-						// setup the module information
-						FControlRigExecuteContext& PublicContext = Context.GetPublicDataSafe<FControlRigExecuteContext>();
-						FControlRigExecuteContextRigModuleGuard RigModuleGuard(PublicContext, InModule->GetNamespace());
-			
-						UControlRig* ClassDefaultObject = InModule->Class->GetDefaultObject<UControlRig>();
-						const TArray<FRigModuleConnector>& Connectors = ClassDefaultObject->GetRigModuleSettings().ExposedConnectors;
-						for (const FRigModuleConnector& Connector : Connectors)
-						{
-							Controller->AddConnector(*Connector.Name, Connector.Settings);
-						}
-
-						PropagateHierarchyFromBPToInstances();
-					}
-				}
+				RefreshModuleConnectors(InModule);
 				UpdateModularDependencyDelegates();
 			}
 			break;
@@ -2456,26 +2541,7 @@ void UControlRigBlueprint::HandleRigModulesModified(EModularRigNotification InNo
 		{
 			if (InModule)
 			{
-				// Remove all the connectors belonging this namespace
-				if (URigHierarchyController* Controller = GetHierarchyController())
-				{
-					Hierarchy->Modify();
-					
-					const FString Namespace = InModule->GetNamespace();
-					TArray<FRigElementKey> Connectors = Controller->GetHierarchy()->GetKeysOfType<FRigConnectorElement>();
-					Connectors.Append(Controller->GetHierarchy()->GetKeysOfType<FRigSocketElement>());
-					for (const FRigElementKey& Connector : Connectors)
-					{
-						FString ConnectorName = Connector.Name.ToString();
-						FString ConnectorNamespace, ShortName;
-						ConnectorName.Split(UModularRig::NamespaceSeparator, &ConnectorNamespace, &ShortName, ESearchCase::CaseSensitive, ESearchDir::FromEnd);
-						ConnectorNamespace.Append(UModularRig::NamespaceSeparator);
-						if (ConnectorNamespace == Namespace)
-						{
-							Controller->RemoveElement(Connector);
-						}
-					}
-				}
+				RefreshModuleConnectors(InModule);
 				UpdateModularDependencyDelegates();
 			}
 			break;
diff --git a/Engine/Plugins/Animation/ControlRig/Source/ControlRigDeveloper/Public/ControlRigBlueprint.h b/Engine/Plugins/Animation/ControlRig/Source/ControlRigDeveloper/Public/ControlRigBlueprint.h
--- a/Engine/Plugins/Animation/ControlRig/Source/ControlRigDeveloper/Public/ControlRigBlueprint.h
+++ b/Engine/Plugins/Animation/ControlRig/Source/ControlRigDeveloper/Public/ControlRigBlueprint.h
@@ -300,6 +300,8 @@ public:
 	void UpdateModularDependencyDelegates();
 	void OnModularDependencyVMCompiled(UObject* InBlueprint, URigVM* InVM, FRigVMExtendedExecuteContext& InExecuteContext);
 	void OnModularDependencyChanged(URigVMBlueprint* InBlueprint);
+	void RefreshModuleConnectors();
+	void RefreshModuleConnectors(const FRigModuleReference* InModule);
 
 	/**
 	* Returns the modified event, which can be used to 
