Control Rig: Fix hierarchy hash to include connector settings

#rnx

diff --git a/Engine/Plugins/Animation/ControlRig/Source/ControlRig/Private/Rigs/RigHierarchy.cpp b/Engine/Plugins/Animation/ControlRig/Source/ControlRig/Private/Rigs/RigHierarchy.cpp
--- a/Engine/Plugins/Animation/ControlRig/Source/ControlRig/Private/Rigs/RigHierarchy.cpp
+++ b/Engine/Plugins/Animation/ControlRig/Source/ControlRig/Private/Rigs/RigHierarchy.cpp
@@ -714,6 +714,10 @@ uint32 URigHierarchy::GetTopologyHash(bool bIncludeTopologyVersion, bool bInclud
 		{
 			Hash = HashCombine(Hash, GetTypeHash(ControlElement->Settings));
 		}
+		if(const FRigConnectorElement* ConnectorElement = Cast<FRigConnectorElement>(Element))
+		{
+			Hash = HashCombine(Hash, GetTypeHash(ConnectorElement->Settings));
+		}
 	}
 
 	return Hash;
diff --git a/Engine/Plugins/Animation/ControlRig/Source/ControlRig/Private/Rigs/RigHierarchyController.cpp b/Engine/Plugins/Animation/ControlRig/Source/ControlRig/Private/Rigs/RigHierarchyController.cpp
--- a/Engine/Plugins/Animation/ControlRig/Source/ControlRig/Private/Rigs/RigHierarchyController.cpp
+++ b/Engine/Plugins/Animation/ControlRig/Source/ControlRig/Private/Rigs/RigHierarchyController.cpp
@@ -725,12 +725,17 @@ FRigElementKey URigHierarchyController::AddDefaultRootSocket()
 	FRigElementKey SocketKey;
 	if(const URigHierarchy* CurrentHierarchy = GetHierarchy())
 	{
+		static const FRigElementKey RootSocketKey(TEXT("Root"), ERigElementType::Socket);
+		if(CurrentHierarchy->Contains(RootSocketKey))
+		{
+			return RootSocketKey;
+		}
+
 		CurrentHierarchy->ForEach<FRigBoneElement>([this, CurrentHierarchy, &SocketKey](const FRigBoneElement* Bone) -> bool
 		{
 			// find first root bone
 			if(CurrentHierarchy->GetNumberOfParents(Bone) == 0)
 			{
-				static const FRigElementKey RootSocketKey(TEXT("Root"), ERigElementType::Socket);
 				SocketKey = AddSocket(RootSocketKey.Name, Bone->GetKey(), FTransform::Identity);
 
 				// stop
diff --git a/Engine/Plugins/Animation/ControlRig/Source/ControlRigEditor/Private/Editor/ControlRigEditor.cpp b/Engine/Plugins/Animation/ControlRig/Source/ControlRigEditor/Private/Editor/ControlRigEditor.cpp
--- a/Engine/Plugins/Animation/ControlRig/Source/ControlRigEditor/Private/Editor/ControlRigEditor.cpp
+++ b/Engine/Plugins/Animation/ControlRig/Source/ControlRigEditor/Private/Editor/ControlRigEditor.cpp
@@ -1171,6 +1171,31 @@ bool FControlRigEditor::DetailViewShowsRigElement(FRigElementKey InKey) const
 	return false;
 }
 
+TArray<FRigElementKey> FControlRigEditor::GetSelectedRigElementsFromDetailView() const
+{
+	TArray<FRigElementKey> Elements;
+	
+	TArray< TWeakObjectPtr<UObject> > SelectedObjects = Inspector->GetSelectedObjects();
+	for (TWeakObjectPtr<UObject> SelectedObject : SelectedObjects)
+	{
+		if (SelectedObject.IsValid())
+		{
+			if(URigVMDetailsViewWrapperObject* WrapperObject = Cast<URigVMDetailsViewWrapperObject>(SelectedObject.Get()))
+			{
+				if (const UScriptStruct* WrappedStruct = WrapperObject->GetWrappedStruct())
+				{
+					if (WrappedStruct->IsChildOf(FRigBaseElement::StaticStruct()))
+					{
+						Elements.Add(WrapperObject->GetContent<FRigBaseElement>().GetKey());
+					}
+				}
+			}
+		}
+	}
+
+	return Elements;
+}
+
 void FControlRigEditor::SetDetailViewForRigModules()
 {
 	SetDetailViewForRigModules(ModulesSelected);
@@ -5071,7 +5096,8 @@ void FControlRigEditor::OnPostConstruction_AnyThread(UControlRig* InRig, const F
 			RebindToSkeletalMeshComponent();
 			if(DetailViewShowsAnyRigElement())
 			{
-				SetDetailViewForRigElements();
+				const TArray<FRigElementKey> Keys = GetSelectedRigElementsFromDetailView();
+				SetDetailViewForRigElements(Keys);
 			}
 			
 			if (FControlRigEditorEditMode* EditMode = GetEditMode())
diff --git a/Engine/Plugins/Animation/ControlRig/Source/ControlRigEditor/Private/Editor/ControlRigEditor.h b/Engine/Plugins/Animation/ControlRig/Source/ControlRigEditor/Private/Editor/ControlRigEditor.h
--- a/Engine/Plugins/Animation/ControlRig/Source/ControlRigEditor/Private/Editor/ControlRigEditor.h
+++ b/Engine/Plugins/Animation/ControlRig/Source/ControlRigEditor/Private/Editor/ControlRigEditor.h
@@ -128,6 +128,7 @@ public:
 	void SetDetailViewForRigElements(const TArray<FRigElementKey>& InKeys);
 	bool DetailViewShowsAnyRigElement() const;
 	bool DetailViewShowsRigElement(FRigElementKey InKey) const;
+	TArray<FRigElementKey> GetSelectedRigElementsFromDetailView() const;
 
 	void SetDetailViewForRigModules();
 	void SetDetailViewForRigModules(const TArray<FString> InKeys);
