Control Rig: Fix compile issue
#rnx

diff --git a/Engine/Plugins/Animation/ControlRig/Source/ControlRigEditor/Private/ControlRigEditorStyle.h b/Engine/Plugins/Animation/ControlRig/Source/ControlRigEditor/Private/ControlRigEditorStyle.h
--- a/Engine/Plugins/Animation/ControlRig/Source/ControlRigEditor/Private/ControlRigEditorStyle.h
+++ b/Engine/Plugins/Animation/ControlRig/Source/ControlRigEditor/Private/ControlRigEditorStyle.h
@@ -44,6 +44,7 @@ public:
 		const FVector2D Icon24x24(24.0f, 24.0f);
 		const FVector2D Icon32x32(32.0f, 32.0f);
 		const FVector2D Icon40x40(40.0f, 40.0f);
+		const FVector2D Icon128x128(128.0f, 128.0f);
 		const FString ControlRigPluginContentDir = FPaths::EnginePluginsDir() / TEXT("Animation/ControlRig/Content");
 		const FString EngineEditorSlateDir = FPaths::EngineContentDir() / TEXT("Editor/Slate");
 		SetContentRoot(ControlRigPluginContentDir);
@@ -143,8 +144,8 @@ public:
 
 		// Schematic
 		{
-			Set("ControlRig.Schematic.SocketUnresolved", new IMAGE_BRUSH_SVG("Slate/Socket_Open", CoreStyleConstants::Icon128x128));
-			Set("ControlRig.Schematic.SocketResolved", new IMAGE_BRUSH_SVG("Slate/Socket_Closed", CoreStyleConstants::Icon128x128));
+			Set("ControlRig.Schematic.SocketUnresolved", new IMAGE_BRUSH_SVG("Slate/Socket_Open", Icon128x128));
+			Set("ControlRig.Schematic.SocketResolved", new IMAGE_BRUSH_SVG("Slate/Socket_Closed", Icon128x128));
 		}
 
 		FSlateStyleRegistry::RegisterSlateStyle(*this);
