Sequencer: Control Rig: When adding control rig no longer set up the object binding, we do that later when it evaluates. To fix issues with FK control rig's not getting set up in the outliner needed to add callbacks to make sure they refresh when new objecs are bound.

diff --git a/Engine/Plugins/Animation/ControlRig/Source/ControlRigEditor/Private/ControlRigSequencerEditorLibrary.cpp b/Engine/Plugins/Animation/ControlRig/Source/ControlRigEditor/Private/ControlRigSequencerEditorLibrary.cpp
--- a/Engine/Plugins/Animation/ControlRig/Source/ControlRigEditor/Private/ControlRigSequencerEditorLibrary.cpp
+++ b/Engine/Plugins/Animation/ControlRig/Source/ControlRigEditor/Private/ControlRigSequencerEditorLibrary.cpp
@@ -200,7 +200,7 @@ static TSharedPtr<ISequencer> GetSequencerFromAsset()
 	return Sequencer;
 }
 
-static UMovieSceneControlRigParameterTrack* AddControlRig(ULevelSequence* LevelSequence,const UClass* InClass, UObject* BoundActor, FGuid ObjectBinding, UControlRig* InExistingControlRig, bool bIsAdditiveControlRig)
+static UMovieSceneControlRigParameterTrack* AddControlRig(ULevelSequence* LevelSequence,const UClass* InClass, FGuid ObjectBinding, UControlRig* InExistingControlRig, bool bIsAdditiveControlRig)
 {
 	FSlateApplication::Get().DismissAllMenus();
 	if (!InClass || !InClass->IsChildOf(UControlRig::StaticClass()) ||
@@ -254,8 +254,6 @@ static UMovieSceneControlRigParameterTrack* AddControlRig(ULevelSequence* LevelS
 			ControlRig->SetIsAdditive(bIsAdditiveControlRig);
 		}
 		ControlRig->SetObjectBinding(MakeShared<FControlRigObjectBinding>());
-		ControlRig->GetObjectBinding()->BindToObject(BoundActor);
-		ControlRig->GetDataSourceRegistry()->RegisterDataSource(UControlRig::OwnerComponent, ControlRig->GetObjectBinding()->GetBoundObject());
 		// Do not re-initialize existing control rig
 		if (!InExistingControlRig)
 		{
@@ -330,35 +328,11 @@ UMovieSceneTrack* UControlRigSequencerEditorLibrary::FindOrCreateControlRigTrack
 				}
 			}
 
-			TArray<UObject*, TInlineAllocator<1>> Result;
-			UObject* Context = nullptr;
-			ALevelSequenceActor* OutActor = nullptr;
-			ULevelSequencePlayer* OutPlayer = nullptr;
-			Result = GetBoundObjects(World, LevelSequence, InBinding, &OutPlayer, &OutActor);
-			if (Result.Num() > 0 && Result[0])
-			{
-				UObject* BoundObject = Result[0];
-				USkeleton* Skeleton = nullptr;
-				USkeletalMeshComponent* SkeletalMeshComponent = nullptr;
-				AcquireSkeletonAndSkelMeshCompFromObject(BoundObject, &Skeleton, &SkeletalMeshComponent);
-
-				UControlRig* ControlRig = nullptr;
-				if (SkeletalMeshComponent && SkeletalMeshComponent->GetSkeletalMeshAsset() && SkeletalMeshComponent->GetSkeletalMeshAsset()->GetSkeleton())
-				{
-					UMovieSceneControlRigParameterTrack* Track = AddControlRig(LevelSequence, ControlRigClass, SkeletalMeshComponent, InBinding.BindingID, nullptr, bIsLayeredControlRig);
-
-					if (Track)
-					{
-						BaseTrack = Track;
-					}
-				}
-			}
+			UMovieSceneControlRigParameterTrack* Track = AddControlRig(LevelSequence, ControlRigClass,  InBinding.BindingID, nullptr, bIsLayeredControlRig);
 
-			//no need to stop player was never running, and stop will cause things to get unbound, messing 
-			//with how priorities are set
-			if (OutActor)
+			if (Track)
 			{
-				World->DestroyActor(OutActor);
+				BaseTrack = Track;
 			}
 		}
 	}
@@ -408,7 +382,7 @@ TArray<UMovieSceneTrack*> UControlRigSequencerEditorLibrary::FindOrCreateControl
 
 					if (GoodTrack == nullptr)
 					{
-						GoodTrack = AddControlRig(LevelSequence, CR->GetClass(), BoundActor, InBinding.BindingID, CR, false);
+						GoodTrack = AddControlRig(LevelSequence, CR->GetClass(), InBinding.BindingID, CR, false);
 					}
 					Tracks.Add(GoodTrack);
 				}
diff --git a/Engine/Plugins/Animation/ControlRig/Source/ControlRigEditor/Private/EditMode/SControlRigOutliner.cpp b/Engine/Plugins/Animation/ControlRig/Source/ControlRigEditor/Private/EditMode/SControlRigOutliner.cpp
--- a/Engine/Plugins/Animation/ControlRig/Source/ControlRigEditor/Private/EditMode/SControlRigOutliner.cpp
+++ b/Engine/Plugins/Animation/ControlRig/Source/ControlRigEditor/Private/EditMode/SControlRigOutliner.cpp
@@ -1132,6 +1132,21 @@ void SControlRigOutliner::SetEditMode(FControlRigEditMode& InEditMode)
 	if (FControlRigEditMode* EditMode = static_cast<FControlRigEditMode*>(ModeTools->GetActiveMode(FControlRigEditMode::ModeName)))
 	{
 		TArrayView<TWeakObjectPtr<UControlRig>> ControlRigs = EditMode->GetControlRigs();
+		for (TWeakObjectPtr<UControlRig>& ControlRig : ControlRigs)
+		{
+			if (ControlRig.IsValid())
+			{
+				if (!ControlRig.Get()->ControlRigBound().IsBoundToObject(this))
+				{
+					ControlRig.Get()->ControlRigBound().AddRaw(this, &SControlRigOutliner::HandleOnControlRigBound);
+				}
+				const TSharedPtr<IControlRigObjectBinding> Binding = ControlRig.Get()->GetObjectBinding();
+				if (Binding && !Binding->OnControlRigBind().IsBoundToObject(this))
+				{
+					Binding->OnControlRigBind().AddRaw(this, &SControlRigOutliner::HandleOnObjectBoundToControlRig);
+				}
+			}
+		}
 		HierarchyTreeView->GetTreeView()->SetControlRigs(ControlRigs); //will refresh tree
 	}
 }
@@ -1139,6 +1154,30 @@ void SControlRigOutliner::SetEditMode(FControlRigEditMode& InEditMode)
 void SControlRigOutliner::HandleControlAdded(UControlRig* ControlRig, bool bIsAdded)
 {
 	FControlRigBaseDockableView::HandleControlAdded(ControlRig, bIsAdded);
+	if (ControlRig)
+	{
+		if (bIsAdded == true )
+		{
+			if (!ControlRig->ControlRigBound().IsBoundToObject(this))
+			{
+				ControlRig->ControlRigBound().AddRaw(this, &SControlRigOutliner::HandleOnControlRigBound);
+			}
+			const TSharedPtr<IControlRigObjectBinding> Binding = ControlRig->GetObjectBinding();
+			if (Binding && !Binding->OnControlRigBind().IsBoundToObject(this))
+			{
+				Binding->OnControlRigBind().AddRaw(this, &SControlRigOutliner::HandleOnObjectBoundToControlRig);
+			}
+		}
+		else
+		{
+			ControlRig->ControlRigBound().RemoveAll(this);
+			const TSharedPtr<IControlRigObjectBinding> Binding = ControlRig->GetObjectBinding();
+			if (Binding)
+			{
+				Binding->OnControlRigBind().RemoveAll(this);
+			}
+		}
+	}
 	if (FControlRigEditMode* EditMode = static_cast<FControlRigEditMode*>(ModeTools->GetActiveMode(FControlRigEditMode::ModeName)))
 	{
 		TArrayView<TWeakObjectPtr<UControlRig>> ControlRigs = EditMode->GetControlRigs();
@@ -1146,5 +1185,36 @@ void SControlRigOutliner::HandleControlAdded(UControlRig* ControlRig, bool bIsAd
 	}
 }
 
+void SControlRigOutliner::HandleOnControlRigBound(UControlRig* InControlRig)
+{
+	if (!InControlRig)
+	{
+		return;
+	}
+
+	const TSharedPtr<IControlRigObjectBinding> Binding = InControlRig->GetObjectBinding();
+
+	if (Binding && !Binding->OnControlRigBind().IsBoundToObject(this))
+	{
+		Binding->OnControlRigBind().AddRaw(this, &SControlRigOutliner::HandleOnObjectBoundToControlRig);
+	}
+}
+
+
+void SControlRigOutliner::HandleOnObjectBoundToControlRig(UObject* InObject)
+{
+	//just refresh the views, but do so on nex tick sine with FK control rig's the controls aren't set up
+	//until AFTER we are bound.
+	GEditor->GetTimerManager()->SetTimerForNextTick([this]()
+	{
+		if (FControlRigEditMode* EditMode = static_cast<FControlRigEditMode*>(ModeTools->GetActiveMode(FControlRigEditMode::ModeName)))
+		{
+			TArrayView<TWeakObjectPtr<UControlRig>> ControlRigs = EditMode->GetControlRigs();
+			HierarchyTreeView->GetTreeView()->SetControlRigs(ControlRigs); //will refresh tree
+		}
+	});
+
+
+}
 
 #undef LOCTEXT_NAMESPACE
diff --git a/Engine/Plugins/Animation/ControlRig/Source/ControlRigEditor/Private/EditMode/SControlRigOutliner.h b/Engine/Plugins/Animation/ControlRig/Source/ControlRigEditor/Private/EditMode/SControlRigOutliner.h
--- a/Engine/Plugins/Animation/ControlRig/Source/ControlRigEditor/Private/EditMode/SControlRigOutliner.h
+++ b/Engine/Plugins/Animation/ControlRig/Source/ControlRigEditor/Private/EditMode/SControlRigOutliner.h
@@ -252,9 +252,12 @@ private:
 	virtual void HandleControlAdded(UControlRig* ControlRig, bool bIsAdded) override;
 	virtual void HandleControlSelected(UControlRig* Subject, FRigControlElement* InControl, bool bSelected) override;
 
+	//control rig delegates
+	void HandleOnControlRigBound(UControlRig* InControlRig);
+	void HandleOnObjectBoundToControlRig(UObject* InObject);
+
 	void OnObjectsReplaced(const TMap<UObject*, UObject*>& OldToNewInstanceMap);
 
-	const URigHierarchy* GetHierarchy() const;
 	void HandleSelectionChanged(TSharedPtr<FMultiRigTreeElement> Selection, ESelectInfo::Type SelectInfo);
 
 	/** Hierarchy picker for controls*/
