Control Rig: Introduce rule manager and initial unit tests

#rnx

diff --git a/Engine/Plugins/Animation/ControlRig/Source/ControlRig/Private/ModularRigRuleManager.cpp b/Engine/Plugins/Animation/ControlRig/Source/ControlRig/Private/ModularRigRuleManager.cpp
new file mode 100644
--- /dev/null
+++ b/Engine/Plugins/Animation/ControlRig/Source/ControlRig/Private/ModularRigRuleManager.cpp
@@ -0,0 +1,153 @@
+// Copyright Epic Games, Inc. All Rights Reserved.
+
+#include "ModularRigRuleManager.h"
+
+#include UE_INLINE_GENERATED_CPP_BY_NAME(ModularRigRuleManager)
+
+#define LOCTEXT_NAMESPACE "ModularRigRuleManager"
+
+FModularRigResolveResult UModularRigRuleManager::Resolve(
+	const FRigConnectorElement* InConnector,
+	const FRigModuleInstance* InModule,
+	const FRigElementKeyRedirector& InResolvedConnectors) const
+{
+	FModularRigResolveResult Result;
+	
+	if(!Hierarchy.IsValid())
+	{
+		static const FText MissingHierarchyMessage = LOCTEXT("MissingHierarchyMessage", "The rule manager is missing the hierarchy.");
+		Result.Message = MissingHierarchyMessage;
+		Result.State = EModularRigResolveState::Error;
+		return Result;
+	}
+
+	// start with a full set of possible targets
+	TArray<bool> VisitedElement;
+	VisitedElement.AddZeroed(Hierarchy->Num());
+	Result.Matches.Reserve(Hierarchy->Num());
+	Hierarchy->Traverse([&VisitedElement, &Result](const FRigBaseElement* Element, bool& bContinue)
+	{
+		if(!VisitedElement[Element->GetIndex()])
+		{
+			VisitedElement[Element->GetIndex()] = true;
+			Result.Matches.Emplace(Element->GetKey(), ERigElementResolveState::PossibleTarget, FText());
+		}
+		bContinue = true;
+	}, true);
+
+	FWorkData WorkData;
+	WorkData.Hierarchy = Hierarchy.Get();
+	WorkData.Connector = InConnector;
+	WorkData.Module = InModule;
+	WorkData.ResolvedConnectors = &InResolvedConnectors;
+	WorkData.Result = &Result;
+	ResolveConnector(WorkData);
+	return Result;;
+}
+
+void UModularRigRuleManager::FWorkData::Filter(TFunction<void(FRigElementResolveResult&)> PerMatchFunction)
+{
+	const TArray<FRigElementResolveResult> PreviousMatches = Result->Matches;;
+	Result->Matches.Reset();
+	for(const FRigElementResolveResult& PreviousMatch : PreviousMatches)
+	{
+		FRigElementResolveResult Match = PreviousMatch;
+		PerMatchFunction(Match);
+		if(Match.IsValid())
+		{
+			Result->Matches.Add(Match);
+		}
+		else
+		{
+			Result->Excluded.Add(Match);
+		}
+	}
+}
+
+void UModularRigRuleManager::SetHierarchy(const URigHierarchy* InHierarchy)
+{
+	check(InHierarchy);
+	Hierarchy = InHierarchy;
+}
+
+void UModularRigRuleManager::ResolveConnector(FWorkData& InOutWorkData)
+{
+	FilterIncompatibleTypes(InOutWorkData);
+	FilterInvalidNameSpaces(InOutWorkData);
+	FilterByConnectorRules(InOutWorkData);
+
+	if(InOutWorkData.Result->Matches.IsEmpty())
+	{
+		InOutWorkData.Result->State = EModularRigResolveState::Error;
+	}
+	else
+	{
+		InOutWorkData.Result->State = EModularRigResolveState::Success;
+	}
+}
+
+void UModularRigRuleManager::FilterIncompatibleTypes(FWorkData& InOutWorkData)
+{
+	InOutWorkData.Filter([](FRigElementResolveResult& Result)
+	{
+		if(Result.GetKey().Type == ERigElementType::Curve)
+		{
+			static const FText CurveInvalidTargetMessage = LOCTEXT("CannotConnectToCurves", "Cannot connect to curves.");
+			Result.SetInvalidTarget(CurveInvalidTargetMessage);
+		}
+		if(Result.GetKey().Type == ERigElementType::Connector)
+		{
+			static const FText CurveInvalidTargetMessage = LOCTEXT("CannotConnectToConnectors", "Cannot connect to connectors.");
+			Result.SetInvalidTarget(CurveInvalidTargetMessage);
+		}
+	});
+}
+
+void UModularRigRuleManager::FilterInvalidNameSpaces(FWorkData& InOutWorkData)
+{
+	const FName NameSpace = InOutWorkData.Hierarchy->GetNameMetadata(InOutWorkData.Connector->GetKey(), URigHierarchy::NameSpaceMetadataName, NAME_None);
+	if(NameSpace.IsNone())
+	{
+		return;
+	}
+	
+	const FString NameSpaceString = NameSpace.ToString(); 
+	const FString NameSpacePrefix = NameSpaceString + UModularRig::NamespaceSeparator; 
+	InOutWorkData.Filter([NameSpaceString, NameSpacePrefix, InOutWorkData](FRigElementResolveResult& Result)
+	{
+		const FName MatchNameSpace = InOutWorkData.Hierarchy->GetNameMetadata(Result.GetKey(), URigHierarchy::NameSpaceMetadataName, NAME_None);
+		if(!MatchNameSpace.IsNone())
+		{
+			const FString MatchNameSpaceString = MatchNameSpace.ToString();
+			if(MatchNameSpaceString.Equals(NameSpaceString, ESearchCase::CaseSensitive))
+			{
+				static const FText CannotConnectWithinNameSpaceMessage = LOCTEXT("CannotConnectWithinNameSpace", "Cannot connect within the same namespace.");
+				Result.SetInvalidTarget(CannotConnectWithinNameSpaceMessage);
+			}
+			else if(MatchNameSpaceString.StartsWith(NameSpacePrefix, ESearchCase::CaseSensitive))
+			{
+				static const FText CannotConnectBelowNameSpaceMessage = LOCTEXT("CannotConnectBelowNameSpace", "Cannot connect to element below the connector's namespace.");
+				Result.SetInvalidTarget(CannotConnectBelowNameSpaceMessage);
+			}
+		}
+	});
+}
+
+void UModularRigRuleManager::FilterByConnectorRules(FWorkData& InOutWorkData)
+{
+	const FRigConnectorElement* Connector = InOutWorkData.Connector;
+	for(const FRigConnectionRuleStash& Stash : Connector->Settings.Rules)
+	{
+		TSharedPtr<FStructOnScope> Storage;
+		const FRigConnectionRule* Rule = Stash.Get(Storage);
+
+		InOutWorkData.Filter([Rule, InOutWorkData](FRigElementResolveResult& Result)
+		{
+			const FRigBaseElement* Target = InOutWorkData.Hierarchy->Find(Result.GetKey());
+			check(Target);
+			Result = Rule->Resolve(Target, InOutWorkData.Hierarchy, InOutWorkData.ResolvedConnectors);
+		});
+	}
+}
+
+#undef LOCTEXT_NAMESPACE
diff --git a/Engine/Plugins/Animation/ControlRig/Source/ControlRig/Private/Rigs/RigConnectionRules.cpp b/Engine/Plugins/Animation/ControlRig/Source/ControlRig/Private/Rigs/RigConnectionRules.cpp
--- a/Engine/Plugins/Animation/ControlRig/Source/ControlRig/Private/Rigs/RigConnectionRules.cpp
+++ b/Engine/Plugins/Animation/ControlRig/Source/ControlRig/Private/Rigs/RigConnectionRules.cpp
@@ -97,256 +97,121 @@ uint32 GetTypeHash(const FRigConnectionRuleStash& InRuleStash)
 	return HashCombine(GetTypeHash(InRuleStash.ScriptStructPath), GetTypeHash(InRuleStash.ExportedText));
 }
 
-/*
-
 ////////////////////////////////////////////////////////////////////////////////
 // FRigConnectionRule
 ////////////////////////////////////////////////////////////////////////////////
 
-bool FRigConnectionRule::CanConnect(const FRigConnectionInfo* InConnectionInfo, FString* OutFailureReason) const
+FRigElementResolveResult FRigConnectionRule::Resolve(const FRigBaseElement* InTarget, const URigHierarchy* InHierarchy, const FRigElementKeyRedirector* InRedirector) const
 {
-	for(const TPair<FRigElementKey, FRigElementKey>& Pair : InConnectionInfo->ConnectionMap)
-	{
-		const FRigElementKey& ConnectorKey = Pair.Key;
-		if(!ConnectorKey.IsValid())
-		{
-			if(OutFailureReason)
-			{
-				static constexpr TCHAR Format[] = TEXT("Connector element '%s' is not valid.");
-				OutFailureReason->Appendf(Format, *ConnectorKey.ToString());
-			}
-			return false;
-		}
-		if(ConnectorKey.Type != ERigElementType::Connector)
-		{
-			if(OutFailureReason)
-			{
-				static constexpr TCHAR Format[] = TEXT("Connector element '%s' is not a Connector.");
-				OutFailureReason->Appendf(Format, *ConnectorKey.ToString());
-			}
-			return false;
-		}
-		if(InConnectionInfo->SourceHierarchy->Find<FRigConnectorElement>(ConnectorKey) == nullptr)
-		{
-			if(OutFailureReason)
-			{
-				static constexpr TCHAR Format[] = TEXT("Connector element '%s' does not exist.");
-				OutFailureReason->Appendf(Format, *ConnectorKey.ToString());
-			}
-			return false;
-		}
-		const FRigElementKey& TargetKey = Pair.Value;
-		if(!TargetKey.IsValid())
-		{
-			if(OutFailureReason)
-			{
-				static constexpr TCHAR Format[] = TEXT("Target element '%s' is not valid.");
-				OutFailureReason->Appendf(Format, *TargetKey.ToString());
-			}
-			return false;
-		}
-		if(InConnectionInfo->TargetHierarchy->Find(TargetKey) == nullptr)
-		{
-			if(OutFailureReason)
-			{
-				static constexpr TCHAR Format[] = TEXT("Target element '%s' does not exist.");
-				OutFailureReason->Appendf(Format, *TargetKey.ToString());
-			}
-			return false;
-		}
-	}
-	return true;
+	FRigElementResolveResult Result(InTarget->GetKey());
+	Result.SetInvalidTarget(FText());
+	return Result;
 }
 
 ////////////////////////////////////////////////////////////////////////////////
 // FRigAndConnectionRule
 ////////////////////////////////////////////////////////////////////////////////
 
-bool FRigAndConnectionRule::CanConnect(const FRigConnectionInfo* InConnectionInfo, FString* OutFailureReason) const
+FRigElementResolveResult FRigAndConnectionRule::Resolve(const FRigBaseElement* InTarget, const URigHierarchy* InHierarchy, const FRigElementKeyRedirector* InRedirector) const
 {
-	if(!Super::CanConnect(InConnectionInfo, OutFailureReason))
-	{
-		return false;
-	}
-
+	FRigElementResolveResult Result(InTarget->GetKey());
+	Result.SetPossibleTarget();
+	
 	TSharedPtr<FStructOnScope> Storage;
 	for(const FRigConnectionRuleStash& ChildRule : ChildRules)
 	{
 		if(const FRigConnectionRule* Rule = ChildRule.Get(Storage))
 		{
-			if(!Rule->CanConnect(InConnectionInfo, OutFailureReason))
+			Result = Rule->Resolve(InTarget, InHierarchy, InRedirector);
+			if(!Result.IsValid())
 			{
-				return false;
+				return Result;
 			}
 		}
 	}
 
-	return true;
+	return Result;
 }
 
 ////////////////////////////////////////////////////////////////////////////////
 // FRigOrConnectionRule
 ////////////////////////////////////////////////////////////////////////////////
 
-bool FRigOrConnectionRule::CanConnect(const FRigConnectionInfo* InConnectionInfo, FString* OutFailureReason) const
+FRigElementResolveResult FRigOrConnectionRule::Resolve(const FRigBaseElement* InTarget, const URigHierarchy* InHierarchy, const FRigElementKeyRedirector* InRedirector) const
 {
-	if(!Super::CanConnect(InConnectionInfo, OutFailureReason))
-	{
-		return false;
-	}
-
+	FRigElementResolveResult Result(InTarget->GetKey());
+	Result.SetPossibleTarget();
+	
 	TSharedPtr<FStructOnScope> Storage;
 	for(const FRigConnectionRuleStash& ChildRule : ChildRules)
 	{
 		if(const FRigConnectionRule* Rule = ChildRule.Get(Storage))
 		{
-			if(Rule->CanConnect(InConnectionInfo, OutFailureReason))
+			Result = Rule->Resolve(InTarget, InHierarchy, InRedirector);
+			if(Result.IsValid())
 			{
-				return true;
+				return Result;
 			}
 		}
 	}
 
-	return false;
+	return Result;
 }
 
 ////////////////////////////////////////////////////////////////////////////////
 // FRigTypeConnectionRule
 ////////////////////////////////////////////////////////////////////////////////
 
-bool FRigTypeConnectionRule::CanConnect(const FRigConnectionInfo* InConnectionInfo, FString* OutFailureReason) const
+FRigElementResolveResult FRigTypeConnectionRule::Resolve(const FRigBaseElement* InTarget, const URigHierarchy* InHierarchy, const FRigElementKeyRedirector* InRedirector) const
 {
-	if(!Super::CanConnect(InConnectionInfo, OutFailureReason))
+	FRigElementResolveResult Result(InTarget->GetKey());
+	Result.SetPossibleTarget();
+
+	if(!InTarget->GetKey().IsTypeOf(ElementType))
 	{
-		return false;
+		const FString ExpectedType = StaticEnum<ERigElementType>()->GetDisplayNameTextByValue((int64)ElementType).ToString();
+		static constexpr TCHAR Format[] = TEXT("Element '%s' is not of the expected type (%s).");
+		Result.SetInvalidTarget(FText::FromString(FString::Printf(Format, *InTarget->GetKey().ToString(), *ExpectedType)));
 	}
 	
-	for(const TPair<FRigElementKey, FRigElementKey>& Pair : InConnectionInfo->ConnectionMap)
-	{
-		const FRigElementKey& TargetKey = Pair.Value;
-		if(!TargetKey.IsTypeOf(ElementType))
-		{
-			if(OutFailureReason)
-			{
-				const FString ExpectedType = StaticEnum<ERigElementType>()->GetDisplayNameTextByValue((int64)ElementType).ToString();
-				static constexpr TCHAR Format[] = TEXT("Element '%s' is not of the expected type (%s).");
-				OutFailureReason->Appendf(Format, *TargetKey.ToString(), *ExpectedType);
-			}
-			return false;
-		}
-	}
-	return true;
+	return Result;
 }
 
 ////////////////////////////////////////////////////////////////////////////////
 // FRigTagConnectionRule
 ////////////////////////////////////////////////////////////////////////////////
 
-bool FRigTagConnectionRule::CanConnect(const FRigConnectionInfo* InConnectionInfo, FString* OutFailureReason) const
+FRigElementResolveResult FRigTagConnectionRule::Resolve(const FRigBaseElement* InTarget, const URigHierarchy* InHierarchy, const FRigElementKeyRedirector* InRedirector) const
 {
-	if(!Super::CanConnect(InConnectionInfo, OutFailureReason))
-	{
-		return false;
-	}
-	const FRigElementKey TargetKey = InConnectionInfo->ConnectionMap.begin().Value();
-	if(!InConnectionInfo->TargetHierarchy->HasTag(TargetKey, Tag))
+	FRigElementResolveResult Result(InTarget->GetKey());
+	Result.SetPossibleTarget();
+
+	if(!InHierarchy->HasTag(InTarget->GetKey(), Tag))
 	{
-		if(OutFailureReason)
-		{
-			static constexpr TCHAR Format[] = TEXT("Element '%s' does not contain tag '%s'.");
-			OutFailureReason->Appendf(Format, *TargetKey.ToString(), *Tag.ToString());
-		}
-		return false;
+		static constexpr TCHAR Format[] = TEXT("Element '%s' does not contain tag '%s'.");
+		Result.SetInvalidTarget(FText::FromString(FString::Printf(Format, *InTarget->GetKey().ToString(), *Tag.ToString())));
 	}
-	return true;
+	
+	return Result;
 }
 
 ////////////////////////////////////////////////////////////////////////////////
-// FRigChainConnectionRule
+// FRigOnChainRule
 ////////////////////////////////////////////////////////////////////////////////
 
-bool FRigChainConnectionRule::CanConnect(const FRigConnectionInfo* InConnectionInfo, FString* OutFailureReason) const
+/*
+FRigElementResolveResult FRigOnChainRule::Resolve(const FRigBaseElement* InTarget, const URigHierarchy* InHierarchy, const FRigElementKeyRedirector* InRedirector) const
 {
-	if(!Super::CanConnect(InConnectionInfo, OutFailureReason))
-	{
-		return false;
-	}
-
-	TArray<FRigElementKey> Targets;
-	Targets.Reserve(InConnectionInfo->ConnectionMap.Num());
-	for(const TPair<FRigElementKey, FRigElementKey>& Pair : InConnectionInfo->ConnectionMap)
-	{
-		Targets.Add(Pair.Value);
-	}
-	if(Targets.IsEmpty())
-	{
-		return false;
-	}
-
-	const FRigElementKey EndKey = Targets[0];
-	const FRigElementKey StartKey = Targets.Last();
-
-	FRigElementKey NextKey = EndKey;
-	TArray<FRigElementKey> ReversedChain;
-
-	while(NextKey.IsValid())
-	{
-		ReversedChain.Add(NextKey);
-		if(NextKey == StartKey)
-		{
-			break;
-		}
-		NextKey = InConnectionInfo->TargetHierarchy->GetFirstParent(NextKey);
-	}
+	FRigElementResolveResult Result;
+	Result.State = ERigElementResolveState::PossibleTarget;
 
-	if(ReversedChain.Last() != StartKey)
+	if(!InHierarchy->HasTag(InTarget->GetKey(), Tag))
 	{
-		if(OutFailureReason)
-		{
-			static constexpr TCHAR Format[] = TEXT("Elements '%s' and '%s' are not on a chain.");
-			OutFailureReason->Appendf(Format, *StartKey.ToString(), *EndKey.ToString());
-		}
-		return false;
-	}
-
-	if(ReversedChain.Num() < MinNumBones)
-	{
-		if(OutFailureReason)
-		{
-			static constexpr TCHAR Format[] = TEXT("Chain from '%s' to '%s' only has %d elements (minimum of %d required).");
-			OutFailureReason->Appendf(Format, *StartKey.ToString(), *EndKey.ToString(), ReversedChain.Num(), MinNumBones);
-		}
-		return false;
+		static constexpr TCHAR Format[] = TEXT("Element '%s' does not contain tag '%s'.");
+		Result.State = ERigElementResolveState::InvalidTarget;
+		Result.Message = FText::FromString(FString::Printf(Format, *InTarget->GetKey().ToString(), *Tag.ToString()));
 	}
 	
-	if(ReversedChain.Num() > MaxNumBones && MaxNumBones > 2)
-	{
-		if(OutFailureReason)
-		{
-			static constexpr TCHAR Format[] = TEXT("Chain from '%s' to '%s' has %d elements (maximum of %d allowed).");
-			OutFailureReason->Appendf(Format, *StartKey.ToString(), *EndKey.ToString(), ReversedChain.Num(), MaxNumBones);
-		}
-		return false;
-	}
-
-	if(!bAllowBranches)
-	{
-		for(int32 Index = 1; Index < ReversedChain.Num() - 1; Index++)
-		{
-			NextKey = ReversedChain[Index];
-			if(InConnectionInfo->TargetHierarchy->GetChildren(NextKey).Num() > 1)
-			{
-				if(OutFailureReason)
-				{
-					static constexpr TCHAR Format[] = TEXT("Chain element '%s' contains a branch.");
-					OutFailureReason->Appendf(Format, *NextKey.ToString());
-				}
-				return false;
-			}
-		}
-	}
-
-	return true;
+	return Result;
 }
-
-*/
+*/
\ No newline at end of file
diff --git a/Engine/Plugins/Animation/ControlRig/Source/ControlRig/Private/Rigs/RigHierarchy.cpp b/Engine/Plugins/Animation/ControlRig/Source/ControlRig/Private/Rigs/RigHierarchy.cpp
--- a/Engine/Plugins/Animation/ControlRig/Source/ControlRig/Private/Rigs/RigHierarchy.cpp
+++ b/Engine/Plugins/Animation/ControlRig/Source/ControlRig/Private/Rigs/RigHierarchy.cpp
@@ -5,6 +5,7 @@
 
 #include "Rigs/RigHierarchyElements.h"
 #include "Rigs/RigHierarchyController.h"
+#include "ModularRigRuleManager.h"
 #include "Units/RigUnitContext.h"
 #include "Math/ControlRigMathLibrary.h"
 #include "UObject/AnimObjectVersion.h"
@@ -2741,6 +2742,27 @@ URigHierarchyController* URigHierarchy::GetController(bool bCreateIfNeeded)
 	return nullptr;
 }
 
+UModularRigRuleManager* URigHierarchy::GetRuleManager(bool bCreateIfNeeded)
+{
+	if(RuleManager)
+	{
+		return RuleManager;
+	}
+	else if(bCreateIfNeeded)
+	{
+		if(ensure(!IsGarbageCollecting()))
+		{
+			RuleManager = NewObject<UModularRigRuleManager>(this, TEXT("RuleManager"), RF_Transient);
+			// In case we create this object from async loading thread
+			RuleManager->ClearInternalFlags(EInternalObjectFlags::Async);
+
+			RuleManager->SetHierarchy(this);
+			return RuleManager;
+		}
+	}
+	return nullptr;
+}
+
 void URigHierarchy::IncrementTopologyVersion()
 {
 	TopologyVersion++;
diff --git a/Engine/Plugins/Animation/ControlRig/Source/ControlRig/Private/Rigs/RigHierarchyDefines.cpp b/Engine/Plugins/Animation/ControlRig/Source/ControlRig/Private/Rigs/RigHierarchyDefines.cpp
--- a/Engine/Plugins/Animation/ControlRig/Source/ControlRig/Private/Rigs/RigHierarchyDefines.cpp
+++ b/Engine/Plugins/Animation/ControlRig/Source/ControlRig/Private/Rigs/RigHierarchyDefines.cpp
@@ -362,3 +362,35 @@ FArchive& operator<<(FArchive& Ar, FRigControlValue& Value)
 	return Ar;
 }
 
+////////////////////////////////////////////////////////////////////////////////
+// FRigElementResolveResult
+////////////////////////////////////////////////////////////////////////////////
+
+bool FRigElementResolveResult::IsValid() const
+{
+	return State == ERigElementResolveState::PossibleTarget ||
+		State == ERigElementResolveState::DefaultTarget;
+}
+
+void FRigElementResolveResult::SetInvalidTarget(const FText& InMessage)
+{
+	State = ERigElementResolveState::InvalidTarget;
+	Message = InMessage;
+}
+
+void FRigElementResolveResult::SetPossibleTarget(const FText& InMessage)
+{
+	State = ERigElementResolveState::PossibleTarget;
+	Message = InMessage;
+}
+
+void FRigElementResolveResult::SetDefaultTarget(const FText& InMessage)
+{
+	State = ERigElementResolveState::DefaultTarget;
+	Message = InMessage;
+}
+
+bool FModularRigResolveResult::IsValid() const
+{
+	return State == EModularRigResolveState::Success && !Matches.IsEmpty();
+}
diff --git a/Engine/Plugins/Animation/ControlRig/Source/ControlRig/Private/Rigs/RigHierarchyElements.cpp b/Engine/Plugins/Animation/ControlRig/Source/ControlRig/Private/Rigs/RigHierarchyElements.cpp
--- a/Engine/Plugins/Animation/ControlRig/Source/ControlRig/Private/Rigs/RigHierarchyElements.cpp
+++ b/Engine/Plugins/Animation/ControlRig/Source/ControlRig/Private/Rigs/RigHierarchyElements.cpp
@@ -1406,23 +1406,6 @@ void FRigReferenceElement::CopyPose(FRigBaseElement* InOther, bool bCurrent, boo
 	}
 }
 
-////////////////////////////////////////////////////////////////////////////////
-// FRigConnectionInfo
-////////////////////////////////////////////////////////////////////////////////
-
-FRigConnectionInfo::FRigConnectionInfo(const FRigElementKeyRedirector* InRedirector, const URigHierarchy* InHierarchy)
-	: SourceHierarchy(InHierarchy)
-	, TargetHierarchy(InHierarchy)
-{
-	for(const TPair<FRigElementKey, FCachedRigElement>& Pair : InRedirector->InternalKeyToExternalKey)
-	{
-		if(Pair.Key.Type == ERigElementType::Connector)
-		{
-			ConnectionMap.Add(Pair.Key, Pair.Value.GetKey());
-		}
-	}
-}
-
 ////////////////////////////////////////////////////////////////////////////////
 // FRigConnectorSettings
 ////////////////////////////////////////////////////////////////////////////////
diff --git a/Engine/Plugins/Animation/ControlRig/Source/ControlRig/Public/ModularRigRuleManager.h b/Engine/Plugins/Animation/ControlRig/Source/ControlRig/Public/ModularRigRuleManager.h
new file mode 100644
--- /dev/null
+++ b/Engine/Plugins/Animation/ControlRig/Source/ControlRig/Public/ModularRigRuleManager.h
@@ -0,0 +1,59 @@
+// Copyright Epic Games, Inc. All Rights Reserved.
+
+#pragma once
+
+#include "ModularRig.h"
+#include "ModularRigRuleManager.generated.h"
+
+// A management class to validate rules and pattern match
+UCLASS(BlueprintType)
+class CONTROLRIG_API UModularRigRuleManager : public UObject
+{
+public:
+
+	GENERATED_BODY()
+
+	/***
+	 * Returns the possible targets for the given connector in the current resolval stage
+	 * @param InConnector The connector to resolve
+	 * @param InResolvedConnectors A redirect map of the already resolved connectors
+	 * @return The resolval result including a list of matches
+	 */
+	FModularRigResolveResult Resolve(
+		const FRigConnectorElement* InConnector,
+		const FRigModuleInstance* InModule,
+		const FRigElementKeyRedirector& InResolvedConnectors = FRigElementKeyRedirector()
+	) const;
+
+private:
+
+	struct FWorkData
+	{
+		FWorkData()
+		: Hierarchy(nullptr)
+		, Connector(nullptr)
+		, Module(nullptr)
+		, ResolvedConnectors(nullptr)
+		, Result(nullptr)
+		{
+		}
+
+		void Filter(TFunction<void(FRigElementResolveResult&)> PerMatchFunction);
+
+		const URigHierarchy* Hierarchy;
+		const FRigConnectorElement* Connector;
+		const FRigModuleInstance* Module;
+		const FRigElementKeyRedirector* ResolvedConnectors;
+		FModularRigResolveResult* Result;
+	};
+
+	void SetHierarchy(const URigHierarchy* InHierarchy);
+	static void ResolveConnector(FWorkData& InOutWorkData);
+	static void FilterIncompatibleTypes(FWorkData& InOutWorkData);
+	static void FilterInvalidNameSpaces(FWorkData& InOutWorkData);
+	static void FilterByConnectorRules(FWorkData& InOutWorkData);
+
+	TWeakObjectPtr<const URigHierarchy> Hierarchy;
+
+	friend class URigHierarchy;
+};
\ No newline at end of file
diff --git a/Engine/Plugins/Animation/ControlRig/Source/ControlRig/Public/Rigs/RigConnectionRules.h b/Engine/Plugins/Animation/ControlRig/Source/ControlRig/Public/Rigs/RigConnectionRules.h
--- a/Engine/Plugins/Animation/ControlRig/Source/ControlRig/Public/Rigs/RigConnectionRules.h
+++ b/Engine/Plugins/Animation/ControlRig/Source/ControlRig/Public/Rigs/RigConnectionRules.h
@@ -7,8 +7,10 @@
 #include "UObject/StructOnScope.h"
 #include "RigConnectionRules.generated.h"
 
+struct FRigBaseElement;
 struct FRigConnectionRule;
-struct FRigConnectionInfo;
+class FRigElementKeyRedirector;
+class URigHierarchy;
 
 USTRUCT(BlueprintType)
 struct CONTROLRIG_API FRigConnectionRuleStash
@@ -53,6 +55,7 @@ public:
 	virtual ~FRigConnectionRule() {}
 
 	virtual UScriptStruct* GetScriptStruct() const { return FRigConnectionRule::StaticStruct(); }
+	virtual FRigElementResolveResult Resolve(const FRigBaseElement* InTarget, const URigHierarchy* InHierarchy, const FRigElementKeyRedirector* InRedirector) const;
 };
 
 USTRUCT(BlueprintType, DisplayName="And Rule")
@@ -75,6 +78,7 @@ public:
 	virtual ~FRigAndConnectionRule() override {}
 
 	virtual UScriptStruct* GetScriptStruct() const override { return FRigAndConnectionRule::StaticStruct(); }
+	virtual FRigElementResolveResult Resolve(const FRigBaseElement* InTarget, const URigHierarchy* InHierarchy, const FRigElementKeyRedirector* InRedirector) const override;
 
 	UPROPERTY(BlueprintReadWrite, EditAnywhere, Category=Settings)
 	TArray<FRigConnectionRuleStash> ChildRules;
@@ -100,6 +104,7 @@ public:
 	virtual ~FRigOrConnectionRule() override {}
 
 	virtual UScriptStruct* GetScriptStruct() const override { return FRigOrConnectionRule::StaticStruct(); }
+	virtual FRigElementResolveResult Resolve(const FRigBaseElement* InTarget, const URigHierarchy* InHierarchy, const FRigElementKeyRedirector* InRedirector) const override;
 
 	UPROPERTY(BlueprintReadWrite, EditAnywhere, Category=Settings)
 	TArray<FRigConnectionRuleStash> ChildRules;
@@ -123,6 +128,7 @@ public:
 	virtual ~FRigTypeConnectionRule() override {}
 
 	virtual UScriptStruct* GetScriptStruct() const override { return FRigTypeConnectionRule::StaticStruct(); }
+	virtual FRigElementResolveResult Resolve(const FRigBaseElement* InTarget, const URigHierarchy* InHierarchy, const FRigElementKeyRedirector* InRedirector) const override;
 
 	UPROPERTY(BlueprintReadWrite, EditAnywhere, Category=Settings)
 	ERigElementType ElementType;
@@ -146,44 +152,39 @@ public:
 	virtual ~FRigTagConnectionRule() override {}
 
 	virtual UScriptStruct* GetScriptStruct() const override { return FRigTagConnectionRule::StaticStruct(); }
+	virtual FRigElementResolveResult Resolve(const FRigBaseElement* InTarget, const URigHierarchy* InHierarchy, const FRigElementKeyRedirector* InRedirector) const override;
 
 	UPROPERTY(BlueprintReadWrite, EditAnywhere, Category=Settings)
 	FName Tag;
 };
 
-USTRUCT(BlueprintType, DisplayName="Chain Rule")
-struct CONTROLRIG_API FRigChainConnectionRule : public FRigConnectionRule
+/*
+USTRUCT(BlueprintType, DisplayName="On Chain Rule")
+struct CONTROLRIG_API FRigOnChainRule : public FRigConnectionRule
 {
 	GENERATED_BODY()
 
 public:
 
-	FRigChainConnectionRule()
+	FRigOnChainRule()
 	: MinNumBones(2)
 	, MaxNumBones(0)
-	, bAllowBranches(false)
 	{}
 
-	FRigChainConnectionRule(const FRigElementKey& InRootConnector, int32 InMinNumBones = 2, int32 InMaxNumBones = 0, bool InAllowBranches = false)
-	: RootConnector(InRootConnector)
-	, MinNumBones(InMinNumBones)
+	FRigOnChainRule(int32 InMinNumBones = 2, int32 InMaxNumBones = 0)
+	: MinNumBones(InMinNumBones)
 	, MaxNumBones(InMaxNumBones)
-	, bAllowBranches(InAllowBranches)
 	{}
 
-	virtual ~FRigChainConnectionRule() override {}
+	virtual ~FRigOnChainRule() override {}
 
-	virtual UScriptStruct* GetScriptStruct() const override { return FRigChainConnectionRule::StaticStruct(); }
-
-	UPROPERTY(BlueprintReadWrite, EditAnywhere, Category=Settings)
-	FRigElementKey RootConnector;
+	virtual UScriptStruct* GetScriptStruct() const override { return FRigOnChainRule::StaticStruct(); }
+	virtual FRigElementResolveResult Resolve(const FRigBaseElement* InTarget, const URigHierarchy* InHierarchy, const FRigElementKeyRedirector* InRedirector) const override;
 
 	UPROPERTY(BlueprintReadWrite, EditAnywhere, Category=Settings)
 	int32 MinNumBones;
 
 	UPROPERTY(BlueprintReadWrite, EditAnywhere, Category=Settings)
 	int32 MaxNumBones;
-
-	UPROPERTY(BlueprintReadWrite, EditAnywhere, Category=Settings)
-	bool bAllowBranches;
-};
\ No newline at end of file
+};
+*/
\ No newline at end of file
diff --git a/Engine/Plugins/Animation/ControlRig/Source/ControlRig/Public/Rigs/RigHierarchy.h b/Engine/Plugins/Animation/ControlRig/Source/ControlRig/Public/Rigs/RigHierarchy.h
--- a/Engine/Plugins/Animation/ControlRig/Source/ControlRig/Public/Rigs/RigHierarchy.h
+++ b/Engine/Plugins/Animation/ControlRig/Source/ControlRig/Public/Rigs/RigHierarchy.h
@@ -21,6 +21,7 @@
 class UControlRig;
 class URigHierarchy;
 class URigHierarchyController;
+class UModularRigRuleManager; 
 
 DECLARE_MULTICAST_DELEGATE_ThreeParams(FRigHierarchyModifiedEvent, ERigHierarchyNotification /* type */, URigHierarchy* /* hierarchy */, const FRigBaseElement* /* element */);
 DECLARE_DYNAMIC_MULTICAST_DELEGATE_ThreeParams(FRigHierarchyModifiedDynamicEvent, ERigHierarchyNotification, NotifType, URigHierarchy*, Hierarchy, FRigElementKey, Subject);
@@ -3409,6 +3410,16 @@ public:
 	UFUNCTION(BlueprintCallable, Category = URigHierarchy)
 	URigHierarchyController* GetController(bool bCreateIfNeeded = true);
 
+	/**
+	 * Returns a rule manager for this hierarchy
+	 * Note: If the manager is not available this will return nullptr 
+	 * even if the bCreateIfNeeded flag is set to true.
+	 * @param bCreateIfNeeded Creates a controller if needed
+	 * @return The Controller for this hierarchy
+	 */
+	UFUNCTION(BlueprintCallable, Category = URigHierarchy)
+	UModularRigRuleManager* GetRuleManager(bool bCreateIfNeeded = true);
+
 	/**
 	 * Returns the topology version of this hierarchy
 	 */
@@ -4262,6 +4273,9 @@ private:
 	TObjectPtr<URigHierarchyController> HierarchyController;
 	bool bIsControllerAvailable;
 
+	UPROPERTY(Transient)
+	mutable TObjectPtr<UModularRigRuleManager> RuleManager;
+
 	TMap<FRigElementKey, FRigElementKey> PreviousParentMap;
 
 	/*We save this so Sequencer can remap this after load*/
diff --git a/Engine/Plugins/Animation/ControlRig/Source/ControlRig/Public/Rigs/RigHierarchyCache.h b/Engine/Plugins/Animation/ControlRig/Source/ControlRig/Public/Rigs/RigHierarchyCache.h
--- a/Engine/Plugins/Animation/ControlRig/Source/ControlRig/Public/Rigs/RigHierarchyCache.h
+++ b/Engine/Plugins/Animation/ControlRig/Source/ControlRig/Public/Rigs/RigHierarchyCache.h
@@ -186,7 +186,6 @@ private:
 
 	friend class URigHierarchy;
 	friend class URigHierarchyController;
-	friend struct FRigConnectionInfo;
 	friend class FRigModuleInstanceDetails;
 	friend class UModularRig;
 };
\ No newline at end of file
diff --git a/Engine/Plugins/Animation/ControlRig/Source/ControlRig/Public/Rigs/RigHierarchyDefines.h b/Engine/Plugins/Animation/ControlRig/Source/ControlRig/Public/Rigs/RigHierarchyDefines.h
--- a/Engine/Plugins/Animation/ControlRig/Source/ControlRig/Public/Rigs/RigHierarchyDefines.h
+++ b/Engine/Plugins/Animation/ControlRig/Source/ControlRig/Public/Rigs/RigHierarchyDefines.h
@@ -1785,3 +1785,119 @@ struct CONTROLRIG_API FRigEventContext
 };
 
 DECLARE_MULTICAST_DELEGATE_TwoParams(FRigEventDelegate, URigHierarchy*, const FRigEventContext&);
+
+UENUM()
+enum class ERigElementResolveState : uint8
+{
+	Unknown,
+	InvalidTarget,
+	PossibleTarget,
+	DefaultTarget,
+
+	/** MAX - invalid */
+	Max UMETA(Hidden),
+};
+
+USTRUCT(BlueprintType)
+struct CONTROLRIG_API FRigElementResolveResult
+{
+	GENERATED_BODY()
+
+public:
+	
+	FRigElementResolveResult()
+	: State(ERigElementResolveState::Unknown)
+	{
+	}
+
+	FRigElementResolveResult(const FRigElementKey& InKey, ERigElementResolveState InState = ERigElementResolveState::PossibleTarget, const FText& InMessage = FText())
+		: Key(InKey)
+		, State(InState)
+		, Message(InMessage)
+	{
+	}
+
+	bool IsValid() const;
+	const FRigElementKey& GetKey() const
+	{
+		return Key;
+	}
+	void SetInvalidTarget(const FText& InMessage);
+	void SetPossibleTarget(const FText& InMessage = FText());
+	void SetDefaultTarget(const FText& InMessage = FText());
+
+private:
+	
+	UPROPERTY()
+	FRigElementKey Key;
+
+	UPROPERTY()
+	ERigElementResolveState State;
+
+	UPROPERTY()
+	FText Message;
+
+	friend class UModularRigRuleManager;
+};
+
+UENUM()
+enum class EModularRigResolveState : uint8
+{
+	Success,
+	Error,
+
+	/** MAX - invalid */
+	Max UMETA(Hidden),
+};
+
+
+USTRUCT(BlueprintType)
+struct CONTROLRIG_API FModularRigResolveResult
+{
+	GENERATED_BODY()
+
+public:
+	
+	FModularRigResolveResult()
+		: State(EModularRigResolveState::Success)
+	{
+	}
+
+	bool IsValid() const;
+
+	EModularRigResolveState GetState() const
+	{
+		return State;
+	}
+
+	const FText& GetMessage() const
+	{
+		return Message;
+	}
+
+	const TArray<FRigElementResolveResult>& GetMatches() const
+	{
+		return Matches;
+	}
+
+	const TArray<FRigElementResolveResult>& GetExcluded() const
+	{
+		return Excluded;
+	}
+
+private:
+	
+	UPROPERTY()
+	TArray<FRigElementResolveResult> Matches;
+
+	UPROPERTY()
+	TArray<FRigElementResolveResult> Excluded;
+
+	UPROPERTY()
+	EModularRigResolveState State;
+
+	UPROPERTY()
+	FText Message;
+
+	friend class UModularRigRuleManager;
+};
diff --git a/Engine/Plugins/Animation/ControlRig/Source/ControlRig/Public/Rigs/RigHierarchyElements.h b/Engine/Plugins/Animation/ControlRig/Source/ControlRig/Public/Rigs/RigHierarchyElements.h
--- a/Engine/Plugins/Animation/ControlRig/Source/ControlRig/Public/Rigs/RigHierarchyElements.h
+++ b/Engine/Plugins/Animation/ControlRig/Source/ControlRig/Public/Rigs/RigHierarchyElements.h
@@ -1614,45 +1614,6 @@ private:
 	friend struct FRigBaseElement;
 };
 
-USTRUCT(BlueprintType)
-struct CONTROLRIG_API FRigConnectionInfo
-{
-	GENERATED_BODY()
-
-public:
-	
-	FRigConnectionInfo()
-	: SourceHierarchy(nullptr)
-	, TargetHierarchy(nullptr)
-	{}
-
-	FRigConnectionInfo(const TMap<FRigElementKey, FRigElementKey>& InMap, const URigHierarchy* InSourceHierarchy, const URigHierarchy* InTargetHierarchy)
-	: ConnectionMap(InMap)
-	, SourceHierarchy(InSourceHierarchy)
-	, TargetHierarchy(InTargetHierarchy)
-	{}
-
-	FRigConnectionInfo(const FRigElementKeyRedirector* InRedirector, const URigHierarchy* InHierarchy);
-
-	bool IsValid() const
-	{
-		return (!ConnectionMap.IsEmpty()) && (SourceHierarchy != nullptr) && (TargetHierarchy != nullptr);
-	}
-
-	// The keys of the connectors in the source hierarchy
-	// mapping to the to-be-linked keys of the elements in the target hierarchy
-	UPROPERTY(BlueprintReadWrite, EditAnywhere, Category = Connections)
-	TMap<FRigElementKey, FRigElementKey> ConnectionMap;
-
-	// The hierarchy owning the connectors
-	UPROPERTY(BlueprintReadWrite, EditAnywhere, Category = Connections)
-	const URigHierarchy* SourceHierarchy; 
-
-	// The hierarchy to be linked into
-	UPROPERTY(BlueprintReadWrite, EditAnywhere, Category = Connections)
-	const URigHierarchy* TargetHierarchy; 
-};
-
 UENUM(BlueprintType)
 enum class EConnectorType : uint8
 {
