Control Rig: Treat integer enum controls as non additive
#rnx

diff --git a/Engine/Plugins/Animation/ControlRig/Source/ControlRig/Private/ControlRig.cpp b/Engine/Plugins/Animation/ControlRig/Source/ControlRig/Private/ControlRig.cpp
--- a/Engine/Plugins/Animation/ControlRig/Source/ControlRig/Private/ControlRig.cpp
+++ b/Engine/Plugins/Animation/ControlRig/Source/ControlRig/Private/ControlRig.cpp
@@ -403,7 +403,7 @@ void UControlRig::Evaluate_AnyThread()
 		{
 			if (FRigControlElement* Control = Cast<FRigControlElement>(Element))
 			{
-				if (Control->Settings.ControlType != ERigControlType::Bool)
+				if (Control->CanTreatAsAdditive())
 				{
 					if (Hierarchy->GetActiveParent(Control->GetKey()) != URigHierarchy::GetDefaultParentKey())
 					{
@@ -467,8 +467,8 @@ void UControlRig::Evaluate_AnyThread()
 				{
 					FRigSetControlValueInfo& Info = Value.Value;
 
-					// A bool value is not an additive property. We just overwrite the value.
-					if (Control->Settings.ControlType == ERigControlType::Bool)
+					// A bool/enum value is not an additive property. We just overwrite the value.
+					if (!Control->CanTreatAsAdditive())
 					{
 						const bool bSetupUndo = false; // Rely on the sequencer track to handle undo/redo
 						Hierarchy->SetControlValue(Control, Info.Value, ERigControlValueType::Current, bSetupUndo, false, Info.bPrintPythonCommnds, false);
@@ -548,7 +548,7 @@ TArray<FRigControlElement*> UControlRig::InvertInputPose(const TArray<FRigElemen
 	{
 		if (FRigControlElement* ControlElement = Cast<FRigControlElement>(DynamicHierarchy->Get(PoseElement.Index)))
 		{
-			if (IsAdditive() && ControlElement->Settings.ControlType == ERigControlType::Bool)
+			if (IsAdditive() && !ControlElement->CanTreatAsAdditive())
 			{
 				continue;
 			}
@@ -1983,8 +1983,8 @@ FRigControlValue UControlRig::GetControlValue(FRigControlElement* InControl, con
 		const int32 ControlIndex = ControlsAfterBackwardsSolve.GetIndex(InControl->GetKey());
 		if (ControlIndex != INDEX_NONE)
 		{
-			// Booleans are not additive properties, just return the current value
-			if (InControl->Settings.ControlType == ERigControlType::Bool)
+			// Booleans/Enums are not additive properties, just return the current value
+			if (!InControl->CanTreatAsAdditive())
 			{
 				return GetHierarchy()->GetControlValue(InControl, InValueType);
 			}
diff --git a/Engine/Plugins/Animation/ControlRig/Source/ControlRig/Public/Rigs/RigHierarchyElements.h b/Engine/Plugins/Animation/ControlRig/Source/ControlRig/Public/Rigs/RigHierarchyElements.h
--- a/Engine/Plugins/Animation/ControlRig/Source/ControlRig/Public/Rigs/RigHierarchyElements.h
+++ b/Engine/Plugins/Animation/ControlRig/Source/ControlRig/Public/Rigs/RigHierarchyElements.h
@@ -1443,6 +1443,19 @@ struct CONTROLRIG_API FRigControlElement final : public FRigMultiParentElement
 
 	bool CanDriveControls() const { return Settings.AnimationType == ERigControlAnimationType::ProxyControl || Settings.AnimationType == ERigControlAnimationType::AnimationControl; }
 
+	bool CanTreatAsAdditive() const
+	{
+		if (Settings.ControlType == ERigControlType::Bool)
+		{
+			return false;
+		}
+		if (Settings.ControlType == ERigControlType::Integer && Settings.ControlEnum != nullptr)
+		{
+			return false;
+		}
+		return true;
+	}
+
 	virtual void Save(FArchive& A, ESerializationPhase SerializationPhase) override;
 	virtual void Load(FArchive& Ar, ESerializationPhase SerializationPhase) override;
 
diff --git a/Engine/Plugins/Animation/ControlRig/Source/ControlRigEditor/Private/EditMode/ControlRigEditMode.cpp b/Engine/Plugins/Animation/ControlRig/Source/ControlRigEditor/Private/EditMode/ControlRigEditMode.cpp
--- a/Engine/Plugins/Animation/ControlRig/Source/ControlRigEditor/Private/EditMode/ControlRigEditMode.cpp
+++ b/Engine/Plugins/Animation/ControlRig/Source/ControlRigEditor/Private/EditMode/ControlRigEditMode.cpp
@@ -3031,7 +3031,7 @@ void FControlRigEditMode::ZeroTransforms(bool bSelectionOnly)
 			{
 				if (FRigControlElement* Element = ControlRig->FindControl(Key.Name))
 				{
-					return Element->Settings.ControlType != ERigControlType::Bool;
+					return Element->CanTreatAsAdditive();
 				}
 				return true;
 			});
@@ -3046,10 +3046,10 @@ void FControlRigEditMode::ZeroTransforms(bool bSelectionOnly)
 			TransformElementsToReset.Reserve(Elements.Num());
 			for (const FRigBaseElement* Element : Elements)
 			{
-				// For additive rigs, ignore boolean controls
+				// For additive rigs, ignore non-additive controls
 				if (const FRigControlElement* Control = Cast<FRigControlElement>(Element))
 				{
-					if (ControlRig->IsAdditive() && Control->Settings.ControlType == ERigControlType::Bool)
+					if (ControlRig->IsAdditive() && !Control->CanTreatAsAdditive())
 					{
 						continue;
 					}
