Control Rig: Introduce notion of short names to modular rig

#rnx

diff --git a/Engine/Plugins/Animation/ControlRig/Source/ControlRig/Private/ModularRig.cpp b/Engine/Plugins/Animation/ControlRig/Source/ControlRig/Private/ModularRig.cpp
--- a/Engine/Plugins/Animation/ControlRig/Source/ControlRig/Private/ModularRig.cpp
+++ b/Engine/Plugins/Animation/ControlRig/Source/ControlRig/Private/ModularRig.cpp
@@ -58,6 +58,19 @@ void UModularRig::BeginDestroy()
 #endif
 }
 
+FString FRigModuleInstance::GetShortName() const
+{
+	if(const FRigModuleReference* ModuleReference = GetModuleReference())
+	{
+		const FString ShortName = ModuleReference->GetShortName();
+		if(!ShortName.IsEmpty())
+		{
+			return ShortName;
+		}
+	}
+	return Name.ToString();
+}
+
 FString FRigModuleInstance::GetPath() const
 {
 	if (!ParentPath.IsEmpty())
@@ -108,6 +121,19 @@ bool FRigModuleInstance::ContainsRig(const UControlRig* InRig) const
 	return false;
 }
 
+const FRigModuleReference* FRigModuleInstance::GetModuleReference() const
+{
+	if(const UControlRig* Rig = GetRig())
+	{
+		if(const UModularRig* ModularRig = Cast<UModularRig>(Rig->GetParentRig()))
+		{
+			const FModularRigModel& Model = ModularRig->GetModularRigModel();
+			return Model.FindModule(GetPath());
+		}
+	}
+	return nullptr;
+}
+
 void UModularRig::InitializeVMs(bool bRequestInit)
 {
 	URigVMHost::Initialize(bRequestInit);
diff --git a/Engine/Plugins/Animation/ControlRig/Source/ControlRig/Private/ModularRigController.cpp b/Engine/Plugins/Animation/ControlRig/Source/ControlRig/Private/ModularRigController.cpp
--- a/Engine/Plugins/Animation/ControlRig/Source/ControlRig/Private/ModularRigController.cpp
+++ b/Engine/Plugins/Animation/ControlRig/Source/ControlRig/Private/ModularRigController.cpp
@@ -78,6 +78,7 @@ FString UModularRigController::AddModule(const FName& InModuleName, TSubclassOf<
 	}
 
 	Model->UpdateCachedChildren();
+	UpdateShortNames();
 
 	if (!NewModule)
 	{
@@ -532,6 +533,7 @@ bool UModularRigController::DeleteModule(const FString& InModulePath, bool bSetu
 	Model->DeletedModules.Add(*Module);
 	Model->Modules.RemoveSingle(*Module);
 	Model->UpdateCachedChildren();
+	UpdateShortNames();
 
 	// Fix bindings
 	for (FRigModuleReference& Reference : Model->Modules)
@@ -558,27 +560,27 @@ bool UModularRigController::DeleteModule(const FString& InModulePath, bool bSetu
 	return false;
 }
 
-bool UModularRigController::RenameModule(const FString& InModulePath, const FName& InNewName, bool bSetupUndo)
+FString UModularRigController::RenameModule(const FString& InModulePath, const FName& InNewName, bool bSetupUndo)
 {
 	FRigModuleReference* Module = FindModule(InModulePath);
 	if (!Module)
 	{
 		UE_LOG(LogControlRig, Error, TEXT("Could not find module %s"), *InModulePath);
-		return false;
+		return FString();
 	}
 
 	FText ErrorMessage;
 	if (!CanRenameModule(InModulePath, InNewName, ErrorMessage))
 	{
 		UE_LOG(LogControlRig, Error, TEXT("Could not rename module %s: %s"), *InModulePath, *ErrorMessage.ToString());
-		return false;
+		return FString();
 	}
 
 	const FString OldName = Module->Name.ToString();
 	const FString NewName = InNewName.ToString();
 	if (OldName.Equals(NewName))
 	{
-		return true;
+		return Module->GetPath();
 	}
 	
 #if WITH_EDITOR
@@ -621,16 +623,17 @@ bool UModularRigController::RenameModule(const FString& InModulePath, const FNam
 		};
 	}
 
+	UpdateShortNames();
 	Notify(EModularRigNotification::ModuleRenamed, Module);
 
 #if WITH_EDITOR
 	TransactionPtr.Reset();
 #endif
 	
-	return true;
+	return NewPath;
 }
 
-bool UModularRigController::CanRenameModule(const FString& InModulePath, const FName& InNewName, FText& OutErrorMessage)
+bool UModularRigController::CanRenameModule(const FString& InModulePath, const FName& InNewName, FText& OutErrorMessage) const
 {
 	if (InNewName.IsNone() || InNewName.ToString().IsEmpty())
 	{
@@ -638,54 +641,42 @@ bool UModularRigController::CanRenameModule(const FString& InModulePath, const F
 		return false;
 	}
 
-	FRigModuleReference* Module = FindModule(InModulePath);
-	if (!Module)
+	if(InNewName.ToString().Contains(UModularRig::NamespaceSeparator))
 	{
-		OutErrorMessage = FText::FromString(FString::Printf(TEXT("Module %s not found."), *InModulePath));
+		OutErrorMessage = NSLOCTEXT("ModularRigController", "NameContainsNamespaceSeparator", "Name contains namespace separator ':'.");
 		return false;
 	}
 
-	const FString NewNameStr = InNewName.ToString();
-	const FString NewPath = (Module->ParentPath.IsEmpty()) ? NewNameStr : FString::Printf(TEXT("%s:%s"), *Module->ParentPath, *NewNameStr);
-	if (InModulePath == NewPath)
+	const FRigModuleReference* Module = const_cast<UModularRigController*>(this)->FindModule(InModulePath);
+	if (!Module)
 	{
-		return true;
+		OutErrorMessage = FText::FromString(FString::Printf(TEXT("Module %s not found."), *InModulePath));
+		return false;
 	}
 
-	// Check there are no siblings with the same name
+	FString ErrorMessage;
+	if(!IsNameAvailable(Module->ParentPath, InNewName, &ErrorMessage))
 	{
-		FRigModuleReference* Parent = FindModule(Module->ParentPath);
-		TArray<FRigModuleReference*>* Siblings = &Model->RootModules;
-		if (Parent)
-		{
-			Siblings = &Parent->CachedChildren;
-		}
-		for (FRigModuleReference* Sibling : *Siblings)
-		{
-			if (Sibling->Name.ToString() == NewNameStr)
-			{
-				OutErrorMessage = FText::FromString(FString::Printf(TEXT("Sibling with path %s already exists"), *Sibling->GetPath()));
-				return false;
-			}
-		}
+		OutErrorMessage = FText::FromString(ErrorMessage);
+		return false;
 	}
 	return true;
 }
 
-bool UModularRigController::ReparentModule(const FString& InModulePath, const FString& InNewParentModulePath, bool bSetupUndo)
+FString UModularRigController::ReparentModule(const FString& InModulePath, const FString& InNewParentModulePath, bool bSetupUndo)
 {
 	FRigModuleReference* Module = FindModule(InModulePath);
 	if (!Module)
 	{
 		UE_LOG(LogControlRig, Error, TEXT("Could not find module %s"), *InModulePath);
-		return false;
+		return FString();
 	}
 
 #if WITH_EDITOR
 	TSharedPtr<FScopedTransaction> TransactionPtr;
 	if (bSetupUndo)
 	{
-		TransactionPtr = MakeShared<FScopedTransaction>(NSLOCTEXT("ModularRigController", "RenameModuleTransaction", "Rename Module"));
+		TransactionPtr = MakeShared<FScopedTransaction>(NSLOCTEXT("ModularRigController", "ReparentModuleTransaction", "Reparent Module"));
 		if(UBlueprint* Blueprint = Cast<UBlueprint>(GetOuter()))
 		{
 			Blueprint->Modify();
@@ -698,7 +689,7 @@ bool UModularRigController::ReparentModule(const FString& InModulePath, const FS
 	const FString OldPath = Module->GetPath();
 	Module->PreviousParentPath = Module->ParentPath;
 	Module->ParentPath = (NewParentModule) ? NewParentModule->GetPath() : FString();
-	Module->Name = GetSafeNewName(Module->GetPath());
+	Module->Name = GetSafeNewName(Module->ParentPath, FRigName(Module->Name));
 	const FString NewPath = Module->GetPath();
 
 	// Fix all the subtree namespaces
@@ -711,6 +702,7 @@ bool UModularRigController::ReparentModule(const FString& InModulePath, const FS
 
 
 	Model->UpdateCachedChildren();
+	UpdateShortNames();
 
 	// Fix bindings
 	for (FRigModuleReference& Reference : Model->Modules)
@@ -742,40 +734,218 @@ bool UModularRigController::ReparentModule(const FString& InModulePath, const FS
 
 	Notify(EModularRigNotification::ModuleReparented, Module);
 	
-	return false;
+	return NewPath;
+}
+
+bool UModularRigController::SetModuleShortName(const FString& InModulePath, const FString& InNewShortName, bool bSetupUndo)
+{
+	FRigModuleReference* Module = FindModule(InModulePath);
+	if (!Module)
+	{
+		UE_LOG(LogControlRig, Error, TEXT("Could not find module %s"), *InModulePath);
+		return false;
+	}
+
+	FText ErrorMessage;
+	if (!CanSetModuleShortName(InModulePath, InNewShortName, ErrorMessage))
+	{
+		UE_LOG(LogControlRig, Error, TEXT("Could not rename module %s: %s"), *InModulePath, *ErrorMessage.ToString());
+		return false;
+	}
+
+	const FString OldShortName = Module->GetShortName();
+	const FString NewShortName = InNewShortName;
+	if (OldShortName.Equals(NewShortName))
+	{
+		return true;
+	}
+	
+#if WITH_EDITOR
+	TSharedPtr<FScopedTransaction> TransactionPtr;
+	if (bSetupUndo)
+	{
+		TransactionPtr = MakeShared<FScopedTransaction>(NSLOCTEXT("ModularRigController", "SetModuleShortNameTransaction", "Set Module Display Name"));
+		if(UBlueprint* Blueprint = Cast<UBlueprint>(GetOuter()))
+		{
+			Blueprint->Modify();
+		}
+	}
+#endif
+	
+	Module->ShortName = InNewShortName;
+	Module->bShortNameBasedOnPath = false;
+
+	Notify(EModularRigNotification::ModuleShortNameChanged, Module);
+
+	// update all other display named to avoid collision
+	UpdateShortNames();
+
+#if WITH_EDITOR
+	TransactionPtr.Reset();
+#endif
+	
+	return true;
 }
 
-FName UModularRigController::GetSafeNewName(const FString& InModuleDesiredPath)
+bool UModularRigController::CanSetModuleShortName(const FString& InModulePath, const FString& InNewShortName, FText& OutErrorMessage) const
 {
-	FString ParentPath, DesiredName = InModuleDesiredPath;
-	InModuleDesiredPath.Split(UModularRig::NamespaceSeparator, &ParentPath, &DesiredName, ESearchCase::CaseSensitive, ESearchDir::FromEnd);
+	FString ErrorMessage;
+	if(!IsShortNameAvailable(FRigName(InNewShortName), &ErrorMessage))
+	{
+		OutErrorMessage = FText::FromString(ErrorMessage);
+		return false;
+	}
+	return true;
+}
+
+void UModularRigController::SanitizeName(FRigName& InOutName, bool bAllowNameSpaces)
+{
+	// Sanitize the name
+	FString SanitizedNameString = InOutName.GetName();
+	bool bChangedSomething = false;
+	for (int32 i = 0; i < SanitizedNameString.Len(); ++i)
+	{
+		TCHAR& C = SanitizedNameString[i];
+
+		const bool bGoodChar = FChar::IsAlpha(C) ||					 // Any letter
+			(C == '_') || (C == '-') || (C == '.') || (C == '|') ||	 // _  - .  | anytime
+			(FChar::IsDigit(C)) ||									 // 0-9 anytime
+			((i > 0) && (C== ' '));									 // Space after the first character to support virtual bones
+
+		if (!bGoodChar)
+		{
+			if(bAllowNameSpaces && C == ':')
+			{
+				continue;
+			}
+			
+			C = '_';
+			bChangedSomething = true;
+		}
+	}
+
+	if (SanitizedNameString.Len() > GetMaxNameLength())
+	{
+		SanitizedNameString.LeftChopInline(SanitizedNameString.Len() - GetMaxNameLength());
+		bChangedSomething = true;
+	}
+
+	if(bChangedSomething)
+	{
+		InOutName.SetName(SanitizedNameString);
+	}
+}
+
+FRigName UModularRigController::GetSanitizedName(const FRigName& InName, bool bAllowNameSpaces)
+{
+	FRigName Name = InName;
+	SanitizeName(Name, bAllowNameSpaces);
+	return Name;
+}
+
+bool UModularRigController::IsNameAvailable(const FString& InParentModulePath, const FRigName& InDesiredName, FString* OutErrorMessage) const
+{
+	const FRigName DesiredName = GetSanitizedName(InDesiredName, false);
+	if(DesiredName != InDesiredName)
+	{
+		if(OutErrorMessage)
+		{
+			static const FString ContainsInvalidCharactersMessage = TEXT("Name contains invalid characters.");
+			*OutErrorMessage = ContainsInvalidCharactersMessage;
+		}
+		return false;
+	}
 
 	TArray<FRigModuleReference*>* Children = &Model->RootModules;
-	if (!ParentPath.IsEmpty())
+	if (!InParentModulePath.IsEmpty())
 	{
-		if (FRigModuleReference* Parent = FindModule(ParentPath))
+		if (FRigModuleReference* Parent = const_cast<UModularRigController*>(this)->FindModule(InParentModulePath))
 		{
 			Children = &Parent->CachedChildren;
 		}
 	}
 
+	for (const FRigModuleReference* Child : *Children)
+	{
+		if (FRigName(Child->Name) == DesiredName)
+		{
+			if(OutErrorMessage)
+			{
+				static const FString NameAlreadyInUse = TEXT("This name is already in use.");
+				*OutErrorMessage = NameAlreadyInUse;
+			}
+			return false;
+		}
+	}
+	return true;
+}
+
+bool UModularRigController::IsShortNameAvailable(const FRigName& InDesiredShortName, FString* OutErrorMessage) const
+{
+	const FRigName DesiredShortName = GetSanitizedName(InDesiredShortName, false);
+	if(DesiredShortName != InDesiredShortName)
+	{
+		if(OutErrorMessage)
+		{
+			static const FString ContainsInvalidCharactersMessage = TEXT("Display Name contains invalid characters.");
+			*OutErrorMessage = ContainsInvalidCharactersMessage;
+		}
+		return false;
+	}
+
+	for (const FRigModuleReference& Child : Model->Modules)
+	{
+		if (InDesiredShortName == FRigName(Child.GetShortName()))
+		{
+			if(OutErrorMessage)
+			{
+				static const FString NameAlreadyInUse = TEXT("This name is already in use.");
+				*OutErrorMessage = NameAlreadyInUse;
+			}
+			return false;
+		}
+	}
+	return true;
+}
+
+FRigName UModularRigController::GetSafeNewName(const FString& InParentModulePath, const FRigName& InDesiredName) const
+{
 	bool bSafeToUse = false;
-	FString NewName = DesiredName;
+
+	// create a copy of the desired name so that the string conversion can be cached
+	const FRigName DesiredName = GetSanitizedName(InDesiredName, false);
+	FRigName NewName = DesiredName;
 	int32 Index = 0;
 	while (!bSafeToUse)
 	{
 		bSafeToUse = true;
-		for (FRigModuleReference* Child : *Children)
+		if(!IsNameAvailable(InParentModulePath, NewName))
 		{
-			if (Child->Name == *NewName)
-			{
-				bSafeToUse = false;
-				NewName = FString::Printf(TEXT("%s%d"), *DesiredName, ++Index);
-				break;
-			}
+			bSafeToUse = false;
+			NewName = FString::Printf(TEXT("%s_%d"), *DesiredName.ToString(), ++Index);
 		}
 	}
-	return *NewName;
+	return NewName;
+}
+
+FRigName UModularRigController::GetSafeNewShortName(const FRigName& InDesiredShortName) const
+{
+	bool bSafeToUse = false;
+
+	// create a copy of the desired name so that the string conversion can be cached
+	const FRigName DesiredShortName = GetSanitizedName(InDesiredShortName, true);
+	FRigName NewShortName = DesiredShortName;
+	int32 Index = 0;
+	while (!bSafeToUse)
+	{
+		bSafeToUse = true;
+		if(!IsShortNameAvailable(NewShortName))
+		{
+			bSafeToUse = false;
+			NewShortName = FString::Printf(TEXT("%s_%d"), *DesiredShortName.ToString(), ++Index);
+		}
+	}
+	return NewShortName;
 }
 
 void UModularRigController::Notify(const EModularRigNotification& InNotification, const FRigModuleReference* InElement)
@@ -785,3 +955,69 @@ void UModularRigController::Notify(const EModularRigNotification& InNotification
 		ModifiedEvent.Broadcast(InNotification, InElement);
 	}
 }
+
+void UModularRigController::UpdateShortNames()
+{
+	TMap<FString, int32> TokenToCount;
+
+	// collect all usages of all paths and their segments
+	for(const FRigModuleReference& Module : Model->Modules)
+	{
+		if(Module.bShortNameBasedOnPath)
+		{
+			FString RemainingPath = Module.GetPath();
+			TokenToCount.FindOrAdd(RemainingPath, 0)++;
+			while(RemainingPath.Split(UModularRig::NamespaceSeparator, nullptr, &RemainingPath, ESearchCase::IgnoreCase, ESearchDir::FromStart))
+			{
+				TokenToCount.FindOrAdd(RemainingPath, 0)++;
+			}
+		}
+		else
+		{
+			TokenToCount.FindOrAdd(Module.ShortName, 0)++;
+		}
+	}
+
+	for(FRigModuleReference& Module : Model->Modules)
+	{
+		if(Module.bShortNameBasedOnPath)
+		{
+			FString ShortPath = Module.GetPath();
+			if(!Module.ParentPath.IsEmpty())
+			{
+				FString Left, Right, RemainingPath = Module.GetPath();
+				ShortPath.Reset();
+
+				while (RemainingPath.Split(UModularRig::NamespaceSeparator, &Left, &Right, ESearchCase::IgnoreCase, ESearchDir::FromEnd))
+				{
+					ShortPath = ShortPath.IsEmpty() ? Right : FString::Printf(TEXT("%s:%s"), *Right, *ShortPath);
+
+					// if the short path only exists once - that's what we use for the display name
+					if(TokenToCount.FindChecked(ShortPath) == 1)
+					{
+						RemainingPath.Reset();
+						break;
+					}
+
+					RemainingPath = Left;
+				}
+
+				if(!RemainingPath.IsEmpty())
+				{
+					ShortPath = FString::Printf(TEXT("%s:%s"), *RemainingPath, *ShortPath);
+				}
+			}
+
+			if(!Module.ShortName.Equals(ShortPath, ESearchCase::CaseSensitive))
+			{
+				Module.ShortName = ShortPath;
+				Notify(EModularRigNotification::ModuleShortNameChanged, &Module);
+			}
+		}
+		else
+		{
+			// the display name is user defined so we won't touch it
+		}
+		
+	}
+}
diff --git a/Engine/Plugins/Animation/ControlRig/Source/ControlRig/Private/ModularRigModel.cpp b/Engine/Plugins/Animation/ControlRig/Source/ControlRig/Private/ModularRigModel.cpp
--- a/Engine/Plugins/Animation/ControlRig/Source/ControlRig/Private/ModularRigModel.cpp
+++ b/Engine/Plugins/Animation/ControlRig/Source/ControlRig/Private/ModularRigModel.cpp
@@ -6,6 +6,15 @@
 
 #include UE_INLINE_GENERATED_CPP_BY_NAME(ModularRigModel)
 
+FString FRigModuleReference::GetShortName() const
+{
+	if(!ShortName.IsEmpty())
+	{
+		return ShortName;
+	}
+	return GetPath();
+}
+
 FString FRigModuleReference::GetPath() const
 {
 	if (ParentPath.IsEmpty())
diff --git a/Engine/Plugins/Animation/ControlRig/Source/ControlRig/Public/ModularRig.h b/Engine/Plugins/Animation/ControlRig/Source/ControlRig/Public/ModularRig.h
--- a/Engine/Plugins/Animation/ControlRig/Source/ControlRig/Public/ModularRig.h
+++ b/Engine/Plugins/Animation/ControlRig/Source/ControlRig/Public/ModularRig.h
@@ -68,11 +68,17 @@ public:
 
 	TArray<FRigModuleInstance*> CachedChildren;
 
+	FString GetShortName() const;
+	FString GetLongName() const
+	{
+		return GetPath();
+	}
 	FString GetPath() const;
 	FString GetNamespace() const;
 	UControlRig* GetRig() const;
 	void SetRig(UControlRig* InRig);
 	bool ContainsRig(const UControlRig* InRig) const;
+	const FRigModuleReference* GetModuleReference() const;
 };
 
 USTRUCT(BlueprintType)
diff --git a/Engine/Plugins/Animation/ControlRig/Source/ControlRig/Public/ModularRigController.h b/Engine/Plugins/Animation/ControlRig/Source/ControlRig/Public/ModularRigController.h
--- a/Engine/Plugins/Animation/ControlRig/Source/ControlRig/Public/ModularRigController.h
+++ b/Engine/Plugins/Animation/ControlRig/Source/ControlRig/Public/ModularRigController.h
@@ -45,15 +45,31 @@ class CONTROLRIG_API UModularRigController : public UObject
 	bool DeleteModule(const FString& InModulePath, bool bSetupUndo = true);
 
 	UFUNCTION(BlueprintCallable, Category = "ControlRig | Modules")
-	bool RenameModule(const FString& InModulePath, const FName& InNewName, bool bSetupUndo = true);
-	bool CanRenameModule(const FString& InModulePath, const FName& InNewName, FText& OutErrorMessage);
+	FString RenameModule(const FString& InModulePath, const FName& InNewName, bool bSetupUndo = true);
+	bool CanRenameModule(const FString& InModulePath, const FName& InNewName, FText& OutErrorMessage) const;
 
 	UFUNCTION(BlueprintCallable, Category = "ControlRig | Modules")
-	bool ReparentModule(const FString& InModulePath, const FString& InNewParentModulePath, bool bSetupUndo = true);
-	
-	FName GetSafeNewName(const FString& InModuleDesiredPath);
-	void SetModel(FModularRigModel* InModel) { Model = InModel; }
+	FString ReparentModule(const FString& InModulePath, const FString& InNewParentModulePath, bool bSetupUndo = true);
+
+	UFUNCTION(BlueprintCallable, Category = "ControlRig | Modules")
+	bool SetModuleShortName(const FString& InModulePath, const FString& InNewShortName, bool bSetupUndo = true);
+	bool CanSetModuleShortName(const FString& InModulePath, const FString& InNewShortName, FText& OutErrorMessage) const;
+
+	static int32 GetMaxNameLength() { return 100; }
+	static void SanitizeName(FRigName& InOutName, bool bAllowNameSpaces);
+	static FRigName GetSanitizedName(const FRigName& InName, bool bAllowNameSpaces);
+	bool IsNameAvailable(const FString& InParentModulePath, const FRigName& InDesiredName, FString* OutErrorMessage = nullptr) const;
+	bool IsShortNameAvailable(const FRigName& InDesiredShortName, FString* OutErrorMessage = nullptr) const;
+	FRigName GetSafeNewName(const FString& InParentModulePath, const FRigName& InDesiredName) const;
+	FRigName GetSafeNewShortName(const FRigName& InDesiredShortName) const;
 	FRigModuleReference* FindModule(const FString& InPath);
 	FModularRigModifiedEvent& OnModified() { return ModifiedEvent; }
+
+private:
+
+	void SetModel(FModularRigModel* InModel) { Model = InModel; }
 	void Notify(const EModularRigNotification& InNotification, const FRigModuleReference* InElement);
+	void UpdateShortNames();
+	
+	friend struct FModularRigModel;
 };
\ No newline at end of file
diff --git a/Engine/Plugins/Animation/ControlRig/Source/ControlRig/Public/ModularRigModel.h b/Engine/Plugins/Animation/ControlRig/Source/ControlRig/Public/ModularRigModel.h
--- a/Engine/Plugins/Animation/ControlRig/Source/ControlRig/Public/ModularRigModel.h
+++ b/Engine/Plugins/Animation/ControlRig/Source/ControlRig/Public/ModularRigModel.h
@@ -22,6 +22,8 @@ enum class EModularRigNotification : uint8
 
 	ModuleConfigValueChanged,
 
+	ModuleShortNameChanged,
+
 	/** MAX - invalid */
 	Max UMETA(Hidden),
 };
@@ -33,12 +35,16 @@ struct CONTROLRIG_API FRigModuleReference
 
 	FRigModuleReference()
 		: Name(NAME_None)
+		, ShortName()
+		, bShortNameBasedOnPath(true)
 		, ParentPath(FString())
 		, Class(nullptr)
 	{}
 	
 	FRigModuleReference(const FName& InName, TSubclassOf<UControlRig> InClass, const FString& InParentPath)
 		: Name(InName)
+		, ShortName()
+		, bShortNameBasedOnPath(true)
 		, ParentPath(InParentPath)
 		, Class(InClass)
 	{}
@@ -46,6 +52,12 @@ struct CONTROLRIG_API FRigModuleReference
 	UPROPERTY()
 	FName Name;
 
+	UPROPERTY()
+	FString ShortName;
+
+	UPROPERTY()
+	bool bShortNameBasedOnPath;
+
 	UPROPERTY()
 	FString ParentPath;
 
@@ -69,6 +81,13 @@ struct CONTROLRIG_API FRigModuleReference
 
 	TArray<FRigModuleReference*> CachedChildren;
 
+	FString GetShortName() const;
+
+	FString GetLongName() const
+	{
+		return GetPath();
+	}
+
 	FString GetPath() const;
 
 	FString GetNamespace() const;
diff --git a/Engine/Plugins/Animation/ControlRig/Source/ControlRigEditor/Private/ControlRigModuleDetails.cpp b/Engine/Plugins/Animation/ControlRig/Source/ControlRigEditor/Private/ControlRigModuleDetails.cpp
--- a/Engine/Plugins/Animation/ControlRig/Source/ControlRigEditor/Private/ControlRigModuleDetails.cpp
+++ b/Engine/Plugins/Animation/ControlRig/Source/ControlRigEditor/Private/ControlRigModuleDetails.cpp
@@ -155,20 +155,66 @@ void FRigModuleInstanceDetails::CustomizeDetails(IDetailLayoutBuilder& DetailBui
 
 	IDetailCategoryBuilder& GeneralCategory = DetailBuilder.EditCategory(TEXT("General"), LOCTEXT("General", "General"));
 	{
+		static const FText NameTooltip = LOCTEXT("NameTooltip", "The name is used to determine the long name (the full path) and to provide a unique address within the rig.");
 		GeneralCategory.AddCustomRow(FText::FromString(TEXT("Name")))
 		.NameContent()
 		[
 			SNew(STextBlock)
 			.Text(FText::FromString(TEXT("Name")))
 			.Font(IDetailLayoutBuilder::GetDetailFont())
-			.IsEnabled(true)
+			.ToolTipText(NameTooltip)
+			.IsEnabled(PerModuleInfos.Num() == 1)
 		]
 		.ValueContent()
 		[
 			SNew(SInlineEditableTextBlock)
 			.Font(IDetailLayoutBuilder::GetDetailFont())
 			.Text(this, &FRigModuleInstanceDetails::GetName)
-			.IsEnabled(true)
+			.OnTextCommitted(this, &FRigModuleInstanceDetails::SetName, DetailBuilder.GetPropertyUtilities())
+			.ToolTipText(NameTooltip)
+			.IsEnabled(PerModuleInfos.Num() == 1)
+			.OnVerifyTextChanged(this, &FRigModuleInstanceDetails::OnVerifyNameChanged)
+		];
+
+		static const FText ShortNameTooltip = LOCTEXT("ShortNameTooltip", "The short name is used for the user interface, for example the sequencer channels.\nThis value can be edited and adjusted as needed.");
+		GeneralCategory.AddCustomRow(FText::FromString(TEXT("Short Name")))
+		.NameContent()
+		[
+			SNew(STextBlock)
+			.Text(FText::FromString(TEXT("Short Name")))
+			.Font(IDetailLayoutBuilder::GetDetailFont())
+			.ToolTipText(ShortNameTooltip)
+			.IsEnabled(PerModuleInfos.Num() == 1)
+		]
+		.ValueContent()
+		[
+			SNew(SInlineEditableTextBlock)
+			.Font(IDetailLayoutBuilder::GetDetailFont())
+			.Text(this, &FRigModuleInstanceDetails::GetShortName)
+			.OnTextCommitted(this, &FRigModuleInstanceDetails::SetShortName, DetailBuilder.GetPropertyUtilities())
+			.ToolTipText(ShortNameTooltip)
+			.IsEnabled(PerModuleInfos.Num() == 1)
+			.OnVerifyTextChanged(this, &FRigModuleInstanceDetails::OnVerifyShortNameChanged)
+		];
+
+		static const FText LongNameTooltip = LOCTEXT("LongNameTooltip", "The long name represents a unique address within the rig but isn't used for the user interface.");
+		GeneralCategory.AddCustomRow(FText::FromString(TEXT("Long Name")))
+		.NameContent()
+		[
+			SNew(STextBlock)
+			.Text(FText::FromString(TEXT("Long Name")))
+			.Font(IDetailLayoutBuilder::GetDetailFont())
+			.ToolTipText(LOCTEXT("LongNameTooltip", "The long name represents a unique address within the rig but isn't used for the user interface."))
+			.ToolTipText(LongNameTooltip)
+			.IsEnabled(false)
+		]
+		.ValueContent()
+		[
+			SNew(SInlineEditableTextBlock)
+			.Font(IDetailLayoutBuilder::GetDetailFont())
+			.Text(this, &FRigModuleInstanceDetails::GetLongName)
+			.ToolTipText(LongNameTooltip)
+			.IsEnabled(false)
 		];
 
 		GeneralCategory.AddCustomRow(FText::FromString(TEXT("RigClass")))
@@ -332,15 +378,25 @@ void FRigModuleInstanceDetails::CustomizeDetails(IDetailLayoutBuilder& DetailBui
 
 FText FRigModuleInstanceDetails::GetName() const
 {
+	const FRigModuleInstance* FirstModule = PerModuleInfos[0].GetModule();
+	if(FirstModule == nullptr)
+	{
+		return FText();
+	}
+	
+	const FName FirstValue = FirstModule->Name;
 	if(PerModuleInfos.Num() > 1)
 	{
 		bool bSame = true;
 		for (int32 i=1; i<PerModuleInfos.Num(); ++i)
 		{
-			if (PerModuleInfos[i].GetModule()->Name !=  PerModuleInfos[0].GetModule()->Name)
+			if(const FRigModuleInstance* Module = PerModuleInfos[i].GetModule())
 			{
-				bSame = false;
-				break;
+				if (Module->Name.IsEqual(FirstValue, ENameCase::CaseSensitive))
+				{
+					bSame = false;
+					break;
+				}
 			}
 		}
 		if (!bSame)
@@ -348,20 +404,146 @@ FText FRigModuleInstanceDetails::GetName() const
 			return ControlRigModuleDetailsMultipleValues;
 		}
 	}
-	return FText::FromName(PerModuleInfos[0].GetModule()->Name);
+	return FText::FromName(FirstValue);
 }
 
-FText FRigModuleInstanceDetails::GetRigClassPath() const
+void FRigModuleInstanceDetails::SetName(const FText& InValue, ETextCommit::Type InCommitType, const TSharedRef<IPropertyUtilities> PropertyUtilities)
+{
+	if(InValue.IsEmpty())
+	{
+		return;
+	}
+	
+	check(PerModuleInfos.Num() == 1);
+	FPerModuleInfo& Info = PerModuleInfos[0];
+	if (const FRigModuleInstance* ModuleInstance = Info.GetModule())
+	{
+		if (UControlRigBlueprint* Blueprint = Info.GetBlueprint())
+		{
+			UModularRigController* Controller = Blueprint->GetModularRigController();
+			const FString OldPath = ModuleInstance->GetPath();
+			(void)Controller->RenameModule(OldPath, *InValue.ToString(), true);
+		}
+	}
+}
+
+bool FRigModuleInstanceDetails::OnVerifyNameChanged(const FText& InText, FText& OutErrorMessage)
+{
+	check(PerModuleInfos.Num() == 1);
+
+	if(InText.IsEmpty())
+	{
+		static const FText EmptyNameIsNotAllowed = LOCTEXT("EmptyNameIsNotAllowed", "Empty name is not allowed.");
+		OutErrorMessage = EmptyNameIsNotAllowed;
+		return false;
+	}
+
+	const FPerModuleInfo& Info = PerModuleInfos[0];
+	if (const FRigModuleInstance* ModuleInstance = Info.GetModule())
+	{
+		if (UControlRigBlueprint* Blueprint = Info.GetBlueprint())
+		{
+			UModularRigController* Controller = Blueprint->GetModularRigController();
+			return Controller->CanRenameModule(ModuleInstance->GetPath(), *InText.ToString(), OutErrorMessage);
+		}
+	}
+
+	return true;
+}
+
+FText FRigModuleInstanceDetails::GetShortName() const
 {
+	const FRigModuleInstance* FirstModule = PerModuleInfos[0].GetModule();
+	if(FirstModule == nullptr)
+	{
+		return FText();
+	}
+
+	const FString FirstValue = FirstModule->GetShortName();
 	if(PerModuleInfos.Num() > 1)
 	{
 		bool bSame = true;
 		for (int32 i=1; i<PerModuleInfos.Num(); ++i)
 		{
-			if (PerModuleInfos[i].GetModule()->GetRig()->GetClass() !=  PerModuleInfos[0].GetModule()->GetRig()->GetClass())
+			if(const FRigModuleInstance* Module = PerModuleInfos[i].GetModule())
 			{
-				bSame = false;
-				break;
+				if (Module->GetShortName().Equals(FirstValue, ESearchCase::CaseSensitive))
+				{
+					bSame = false;
+					break;
+				}
+			}
+		}
+		if (!bSame)
+		{
+			return ControlRigModuleDetailsMultipleValues;
+		}
+	}
+	return FText::FromString(FirstValue);
+}
+
+void FRigModuleInstanceDetails::SetShortName(const FText& InValue, ETextCommit::Type InCommitType, const TSharedRef<IPropertyUtilities> PropertyUtilities)
+{
+	if(InValue.IsEmpty())
+	{
+		return;
+	}
+	
+	check(PerModuleInfos.Num() == 1);
+	const FPerModuleInfo& Info = PerModuleInfos[0];
+	if (const FRigModuleInstance* ModuleInstance = Info.GetModule())
+	{
+		if (UControlRigBlueprint* Blueprint = Info.GetBlueprint())
+		{
+			UModularRigController* Controller = Blueprint->GetModularRigController();
+			Controller->SetModuleShortName(ModuleInstance->GetPath(), *InValue.ToString(), true);
+		}
+	}
+}
+
+bool FRigModuleInstanceDetails::OnVerifyShortNameChanged(const FText& InText, FText& OutErrorMessage)
+{
+	check(PerModuleInfos.Num() == 1);
+
+	if(InText.IsEmpty())
+	{
+		return true;
+	}
+
+	const FPerModuleInfo& Info = PerModuleInfos[0];
+	if (const FRigModuleInstance* ModuleInstance = Info.GetModule())
+	{
+		if (UControlRigBlueprint* Blueprint = Info.GetBlueprint())
+		{
+			UModularRigController* Controller = Blueprint->GetModularRigController();
+			return Controller->CanSetModuleShortName(ModuleInstance->GetPath(), *InText.ToString(), OutErrorMessage);
+		}
+	}
+
+	return false;
+}
+
+FText FRigModuleInstanceDetails::GetLongName() const
+{
+	const FRigModuleInstance* FirstModule = PerModuleInfos[0].GetModule();
+	if(FirstModule == nullptr)
+	{
+		return FText();
+	}
+
+	const FString FirstValue = FirstModule->GetLongName();
+	if(PerModuleInfos.Num() > 1)
+	{
+		bool bSame = true;
+		for (int32 i=1; i<PerModuleInfos.Num(); ++i)
+		{
+			if(const FRigModuleInstance* Module = PerModuleInfos[i].GetModule())
+			{
+				if (Module->GetLongName().Equals(FirstValue, ESearchCase::CaseSensitive))
+				{
+					bSame = false;
+					break;
+				}
 			}
 		}
 		if (!bSame)
@@ -369,6 +551,32 @@ FText FRigModuleInstanceDetails::GetRigClassPath() const
 			return ControlRigModuleDetailsMultipleValues;
 		}
 	}
+	return FText::FromString(FirstValue);
+}
+FText FRigModuleInstanceDetails::GetRigClassPath() const
+{
+	if(PerModuleInfos.Num() > 1)
+	{
+		if(const FRigModuleInstance* FirstModule = PerModuleInfos[0].GetModule())
+		{
+			bool bSame = true;
+			for (int32 i=1; i<PerModuleInfos.Num(); ++i)
+			{
+				if(const FRigModuleInstance* Module = PerModuleInfos[i].GetModule())
+				{
+					if (Module->GetRig()->GetClass() !=  FirstModule->GetRig()->GetClass())
+					{
+						bSame = false;
+						break;
+					}
+				}
+			}
+			if (!bSame)
+			{
+				return ControlRigModuleDetailsMultipleValues;
+			}
+		}
+	}
 
 	if (const FRigModuleInstance* Module = PerModuleInfos[0].GetModule())
 	{
@@ -440,7 +648,11 @@ const FRigModuleInstanceDetails::FPerModuleInfo& FRigModuleInstanceDetails::Find
 {
 	const FPerModuleInfo* Info = FindModuleByPredicate([InPath](const FPerModuleInfo& Info)
 	{
-		return Info.GetModule()->GetPath() == InPath;
+		if(const FRigModuleInstance* Module = Info.GetModule())
+		{
+			return Module->GetPath() == InPath;
+		}
+		return false;
 	});
 
 	if(Info)
@@ -484,8 +696,11 @@ void FRigModuleInstanceDetails::OnElementNameChanged(TSharedPtr<FString> InItem,
 					TargetKey->Name = **InItem;
 				}
 
-				FRigElementKey NamespacedConnector(*FString::Printf(TEXT("%s:%s"), *Info.GetModule()->GetPath(), *Connector.Name.ToString()), ERigElementType::Connector);
-				Controller->ConnectConnectorToElement(NamespacedConnector, *TargetKey);
+				if(const FRigModuleInstance* Module = Info.GetModule())
+				{
+					FRigElementKey NamespacedConnector(*FString::Printf(TEXT("%s:%s"), *Module->GetPath(), *Connector.Name.ToString()), ERigElementType::Connector);
+					Controller->ConnectConnectorToElement(NamespacedConnector, *TargetKey);
+				}
 			}
 		}
 	}
@@ -504,8 +719,11 @@ void FRigModuleInstanceDetails::OnElementTypeChanged(ERigElementType InElementTy
 			{
 				TargetKey->Type = InElementType;
 
-				FRigElementKey NamespacedConnector(*FString::Printf(TEXT("%s:%s"), *Info.GetModule()->GetPath(), *Connector.Name.ToString()), ERigElementType::Connector);
-				Controller->ConnectConnectorToElement(NamespacedConnector, *TargetKey);
+				if(const FRigModuleInstance* Module = Info.GetModule())
+				{
+					FRigElementKey NamespacedConnector(*FString::Printf(TEXT("%s:%s"), *Module->GetPath(), *Connector.Name.ToString()), ERigElementType::Connector);
+					Controller->ConnectConnectorToElement(NamespacedConnector, *TargetKey);
+				}
 			}
 		}
 	}
@@ -527,8 +745,11 @@ FReply FRigModuleInstanceDetails::OnGetSelectedClicked(FRigElementKey Connector)
 		const TArray<FRigElementKey>& Selected = Blueprint->Hierarchy->GetSelectedKeys();
 		if (Selected.Num() > 0)
 		{
-			FRigElementKey NamespacedConnector(*FString::Printf(TEXT("%s:%s"), *PerModuleInfos[0].GetModule()->GetPath(), *Connector.Name.ToString()), ERigElementType::Connector);
-			Blueprint->GetModularRigController()->ConnectConnectorToElement(NamespacedConnector, Selected[0]);
+			if(const FRigModuleInstance* Module = PerModuleInfos[0].GetModule())
+			{
+				FRigElementKey NamespacedConnector(*FString::Printf(TEXT("%s:%s"), *Module->GetPath(), *Connector.Name.ToString()), ERigElementType::Connector);
+				Blueprint->GetModularRigController()->ConnectConnectorToElement(NamespacedConnector, Selected[0]);
+			}
 		}
 	}
 	return FReply::Handled();
diff --git a/Engine/Plugins/Animation/ControlRig/Source/ControlRigEditor/Private/ControlRigModuleDetails.h b/Engine/Plugins/Animation/ControlRig/Source/ControlRigEditor/Private/ControlRigModuleDetails.h
--- a/Engine/Plugins/Animation/ControlRig/Source/ControlRigEditor/Private/ControlRigModuleDetails.h
+++ b/Engine/Plugins/Animation/ControlRig/Source/ControlRigEditor/Private/ControlRigModuleDetails.h
@@ -22,6 +22,7 @@
 #include "ScopedTransaction.h"
 #include "Styling/AppStyle.h"
 #include "Algo/Transform.h"
+#include "IPropertyUtilities.h"
 
 class IPropertyHandle;
 
@@ -39,6 +40,12 @@ public:
 	}
 
 	FText GetName() const;
+	void SetName(const FText& InValue, ETextCommit::Type InCommitType, const TSharedRef<IPropertyUtilities> PropertyUtilities);
+	bool OnVerifyNameChanged(const FText& InText, FText& OutErrorMessage);
+	FText GetShortName() const;
+	void SetShortName(const FText& InValue, ETextCommit::Type InCommitType, const TSharedRef<IPropertyUtilities> PropertyUtilities);
+	bool OnVerifyShortNameChanged(const FText& InText, FText& OutErrorMessage);
+	FText GetLongName() const;
 	FText GetRigClassPath() const;
 	TArray<FRigModuleConnector> GetConnectors() const;
 	FRigElementKeyRedirector GetConnections() const;
diff --git a/Engine/Plugins/Animation/ControlRig/Source/ControlRigEditor/Private/Editor/ControlRigEditor.cpp b/Engine/Plugins/Animation/ControlRig/Source/ControlRigEditor/Private/Editor/ControlRigEditor.cpp
--- a/Engine/Plugins/Animation/ControlRig/Source/ControlRigEditor/Private/Editor/ControlRigEditor.cpp
+++ b/Engine/Plugins/Animation/ControlRig/Source/ControlRigEditor/Private/Editor/ControlRigEditor.cpp
@@ -2710,7 +2710,7 @@ void FControlRigEditor::HandleSchematicDrop(SSchematicGraphPanel* InPanel, SSche
 			{
 				if (UModularRigController* Controller = Blueprint->GetModularRigController())
 				{
-					const FName ModuleName = Controller->GetSafeNewName(AssetBlueprint->RigModuleSettings.Identifier.Name);
+					const FName ModuleName = Controller->GetSafeNewName(FString(), FRigName(AssetBlueprint->RigModuleSettings.Identifier.Name));
 					const FString ModulePath = Controller->AddModule(ModuleName, AssetBlueprint->GetControlRigClass(), FString());
 					if(!ModulePath.IsEmpty())
 					{
@@ -3887,6 +3887,7 @@ void FControlRigEditor::HandleModularRigModified(EModularRigNotification InNotif
 			}
 			ModulesSelected.Remove(OldPath);
 			ModulesSelected.Add(InModule->GetPath());
+			RefreshDetailView();
 
 			// todo: update SchematicGraph
 			break;
diff --git a/Engine/Plugins/Animation/ControlRig/Source/ControlRigEditor/Private/Editor/SModularRigModel.cpp b/Engine/Plugins/Animation/ControlRig/Source/ControlRigEditor/Private/Editor/SModularRigModel.cpp
--- a/Engine/Plugins/Animation/ControlRig/Source/ControlRigEditor/Private/Editor/SModularRigModel.cpp
+++ b/Engine/Plugins/Animation/ControlRig/Source/ControlRigEditor/Private/Editor/SModularRigModel.cpp
@@ -510,14 +510,13 @@ void SModularRigModel::HandleNewItem(UClass* InClass, const FString &InParentPat
 	
 	if (ControlRigBlueprint.IsValid())
 	{
+		UModularRigController* Controller = ControlRigBlueprint->GetModularRigController();
+		
 		FString ClassName = InClass->GetName();
 		ClassName.RemoveFromEnd(TEXT("_C"));
-		FString PathName = InParentPath.IsEmpty() ? *ClassName : FString::Printf(TEXT("%s:%s"), *InParentPath, *ClassName);
-		const FName Name = CreateUniqueName(*PathName);
-		ControlRigBlueprint->GetModularRigController()->AddModule(Name, InClass, InParentPath);
-
-		FString NewPathName = InParentPath.IsEmpty() ? *Name.ToString() : FString::Printf(TEXT("%s:%s"), *InParentPath, *Name.ToString());
-		TSharedPtr<FModularRigTreeElement> Element = TreeView->FindElement(NewPathName);
+		const FRigName Name = Controller->GetSafeNewName(InParentPath, FRigName(ClassName));
+		const FString NewPath = Controller->AddModule(Name, InClass, InParentPath);
+		TSharedPtr<FModularRigTreeElement> Element = TreeView->FindElement(NewPath);
 		if (Element.IsValid())
 		{
 			TreeView->SetSelection({Element});
@@ -567,8 +566,7 @@ FName SModularRigModel::HandleRenameModule(const FString& InOldPath, const FName
 		UModularRigController* Controller = ControlRigBlueprint->GetModularRigController();
 		check(Controller);
 
-		FName ResultingName = NAME_None;
-		if (Controller->RenameModule(InOldPath, InNewName))
+		if (!Controller->RenameModule(InOldPath, InNewName).IsEmpty())
 		{
 			return InNewName;
 		}
@@ -709,11 +707,6 @@ UModularRig* SModularRigModel::GetDefaultModularRig() const
 }
 
 
-FName SModularRigModel::CreateUniqueName(const FName& InBasePath) const
-{
-	return ControlRigBlueprint->GetModularRigController()->GetSafeNewName(InBasePath.ToString());
-}
-
 void SModularRigModel::OnRequestDetailsInspection(const FString& InKey)
 {
 	if(!ControlRigEditor.IsValid())
diff --git a/Engine/Plugins/Animation/ControlRig/Source/ControlRigEditor/Private/Editor/SModularRigModel.h b/Engine/Plugins/Animation/ControlRig/Source/ControlRigEditor/Private/Editor/SModularRigModel.h
--- a/Engine/Plugins/Animation/ControlRig/Source/ControlRigEditor/Private/Editor/SModularRigModel.h
+++ b/Engine/Plugins/Animation/ControlRig/Source/ControlRigEditor/Private/Editor/SModularRigModel.h
@@ -145,7 +145,6 @@ private:
 	UModularRig* GetModularRig() const;
 	UModularRig* GetDefaultModularRig() const;
 	const UModularRig* GetModularRigForTreeView() const { return GetModularRig(); }
-	FName CreateUniqueName(const FName& InBasePath) const;
 	void OnRequestDetailsInspection(const FString& InKey);
 	void HandlePreCompileModularRigs(URigVMBlueprint* InBlueprint);
 	void HandlePostCompileModularRigs(URigVMBlueprint* InBlueprint);
diff --git a/Engine/Plugins/Animation/ControlRig/Source/ControlRigEditor/Private/Editor/SRigHierarchy.cpp b/Engine/Plugins/Animation/ControlRig/Source/ControlRigEditor/Private/Editor/SRigHierarchy.cpp
--- a/Engine/Plugins/Animation/ControlRig/Source/ControlRigEditor/Private/Editor/SRigHierarchy.cpp
+++ b/Engine/Plugins/Animation/ControlRig/Source/ControlRigEditor/Private/Editor/SRigHierarchy.cpp
@@ -2821,8 +2821,9 @@ FReply SRigHierarchy::OnAcceptDrop(const FDragDropEvent& DragDropEvent, EItemDro
 			{
 				if (UModularRigController* Controller = ControlRigBlueprint->GetModularRigController())
 				{
-					const FName ModuleName = Controller->GetSafeNewName(AssetBlueprint->RigModuleSettings.Identifier.Name);
-					const FString ModulePath = Controller->AddModule(ModuleName, AssetBlueprint->GetControlRigClass(), FString());
+					static const FString ParentPath = FString();
+					const FRigName ModuleName = Controller->GetSafeNewName(ParentPath, FRigName(AssetBlueprint->RigModuleSettings.Identifier.Name));
+					const FString ModulePath = Controller->AddModule(ModuleName, AssetBlueprint->GetControlRigClass(), ParentPath);
 					if(TargetItem.IsValid() && !ModulePath.IsEmpty())
 					{
 						FRigElementKey PrimaryConnectorKey;
