Control Rig: Update module identifier after save-as, duplicate or rename

#rnx

diff --git a/Engine/Plugins/Animation/ControlRig/Source/ControlRigDeveloper/Private/ControlRigBlueprint.cpp b/Engine/Plugins/Animation/ControlRig/Source/ControlRigDeveloper/Private/ControlRigBlueprint.cpp
--- a/Engine/Plugins/Animation/ControlRig/Source/ControlRigDeveloper/Private/ControlRigBlueprint.cpp
+++ b/Engine/Plugins/Animation/ControlRig/Source/ControlRigDeveloper/Private/ControlRigBlueprint.cpp
@@ -1362,9 +1362,26 @@ void UControlRigBlueprint::PostDuplicate(bool bDuplicateForPIE)
 		ModularController->OnModified().AddUObject(this, &UControlRigBlueprint::HandleRigModulesModified);
 	}
 
+	// update the rig module identifier after save-as or duplicate asset
+	if(IsControlRigModule())
+	{
+		RigModuleSettings.Identifier.Name = URigHierarchy::GetSanitizedName(FRigName(GetName())).ToString();
+	}
+
 	ModularRigModel.UpdateCachedChildren();
 }
 
+void UControlRigBlueprint::PostRename(UObject* OldOuter, const FName OldName)
+{
+	Super::PostRename(OldOuter, OldName);
+
+	// update the rig module identifier after renaming the asset
+	if(IsControlRigModule())
+	{
+		RigModuleSettings.Identifier.Name = URigHierarchy::GetSanitizedName(FRigName(GetName())).ToString();
+	}
+}
+
 TArray<UControlRigBlueprint*> UControlRigBlueprint::GetCurrentlyOpenRigBlueprints()
 {
 	return sCurrentlyOpenedRigBlueprints;
diff --git a/Engine/Plugins/Animation/ControlRig/Source/ControlRigDeveloper/Public/ControlRigBlueprint.h b/Engine/Plugins/Animation/ControlRig/Source/ControlRigDeveloper/Public/ControlRigBlueprint.h
--- a/Engine/Plugins/Animation/ControlRig/Source/ControlRigDeveloper/Public/ControlRigBlueprint.h
+++ b/Engine/Plugins/Animation/ControlRig/Source/ControlRigDeveloper/Public/ControlRigBlueprint.h
@@ -78,6 +78,7 @@ public:
 	virtual void PostLoad() override;
 	virtual void PostTransacted(const FTransactionObjectEvent& TransactionEvent) override;
 	virtual void PostDuplicate(bool bDuplicateForPIE) override;
+	virtual void PostRename(UObject* OldOuter, const FName OldName) override;
 	virtual bool RequiresForceLoadMembers(UObject* InObject) const override;
 
 	virtual bool SupportsGlobalVariables() const override { return true; }
