Control Rig: Perform details customization for any module config value

#rnx

diff --git a/Engine/Plugins/Animation/ControlRig/Source/ControlRigEditor/Private/ControlRigModuleDetails.cpp b/Engine/Plugins/Animation/ControlRig/Source/ControlRigEditor/Private/ControlRigModuleDetails.cpp
--- a/Engine/Plugins/Animation/ControlRig/Source/ControlRigEditor/Private/ControlRigModuleDetails.cpp
+++ b/Engine/Plugins/Animation/ControlRig/Source/ControlRigEditor/Private/ControlRigModuleDetails.cpp
@@ -153,6 +153,9 @@ void FRigModuleInstanceDetails::CustomizeDetails(IDetailLayoutBuilder& DetailBui
 		return;
 	}
 
+	TArray<FName> OriginalCategoryNames;
+	DetailBuilder.GetCategoryNames(OriginalCategoryNames);
+
 	IDetailCategoryBuilder& GeneralCategory = DetailBuilder.EditCategory(TEXT("General"), LOCTEXT("General", "General"));
 	{
 		static const FText NameTooltip = LOCTEXT("NameTooltip", "The name is used to determine the long name (the full path) and to provide a unique address within the rig.");
@@ -278,100 +281,103 @@ void FRigModuleInstanceDetails::CustomizeDetails(IDetailLayoutBuilder& DetailBui
 		}
 	}
 
-	IDetailCategoryBuilder& ConfigValuesCategory = DetailBuilder.EditCategory(TEXT("Default"), LOCTEXT("ConfigValues", "Config Values"));
+	for(const FName& OriginalCategoryName : OriginalCategoryNames)
 	{
-		IPropertyAccessEditor& PropertyAccessEditor = IModularFeatures::Get().GetModularFeature<IPropertyAccessEditor>("PropertyAccessEditor");
-
-		TArray<TSharedRef<IPropertyHandle>> DefaultProperties;
-		ConfigValuesCategory.GetDefaultProperties(DefaultProperties, true, true);
-
-		for(const TSharedRef<IPropertyHandle>& DefaultProperty : DefaultProperties)
+		IDetailCategoryBuilder& Category = DetailBuilder.EditCategory(OriginalCategoryName);
 		{
-			const FProperty* Property = DefaultProperty->GetProperty();
-			if(Property == nullptr)
-			{
-				DetailBuilder.HideProperty(DefaultProperty);
-				continue;
-			}
+			IPropertyAccessEditor& PropertyAccessEditor = IModularFeatures::Get().GetModularFeature<IPropertyAccessEditor>("PropertyAccessEditor");
 
-			// skip advanced properties for now
-			const bool bAdvancedDisplay = Property->HasAnyPropertyFlags(CPF_AdvancedDisplay);
-			if(bAdvancedDisplay)
-			{
-				DetailBuilder.HideProperty(DefaultProperty);
-				continue;
-			}
+			TArray<TSharedRef<IPropertyHandle>> DefaultProperties;
+			Category.GetDefaultProperties(DefaultProperties, true, true);
 
-			// skip non-public properties for now
-			const bool bIsPublic = Property->HasAnyPropertyFlags(CPF_Edit | CPF_EditConst);
-			const bool bIsInstanceEditable = !Property->HasAnyPropertyFlags(CPF_DisableEditOnInstance);
-			if(!bIsPublic || !bIsInstanceEditable)
+			for(const TSharedRef<IPropertyHandle>& DefaultProperty : DefaultProperties)
 			{
-				DetailBuilder.HideProperty(DefaultProperty);
-				continue;
-			}
+				const FProperty* Property = DefaultProperty->GetProperty();
+				if(Property == nullptr)
+				{
+					DetailBuilder.HideProperty(DefaultProperty);
+					continue;
+				}
 
-			const FSimpleDelegate OnValueChangedDelegate = FSimpleDelegate::CreateSP(this, &FRigModuleInstanceDetails::OnConfigValueChanged, Property->GetFName());
-			DefaultProperty->SetOnPropertyValueChanged(OnValueChangedDelegate);
-			DefaultProperty->SetOnChildPropertyValueChanged(OnValueChangedDelegate);
+				// skip advanced properties for now
+				const bool bAdvancedDisplay = Property->HasAnyPropertyFlags(CPF_AdvancedDisplay);
+				if(bAdvancedDisplay)
+				{
+					DetailBuilder.HideProperty(DefaultProperty);
+					continue;
+				}
 
-			FPropertyBindingWidgetArgs BindingArgs;
-			BindingArgs.Property = (FProperty*)Property;
-			BindingArgs.CurrentBindingText = TAttribute<FText>::CreateLambda([this, Property]()
-			{
-				return GetBindingText(Property);
-			});
-			BindingArgs.CurrentBindingImage = TAttribute<const FSlateBrush*>::CreateLambda([this, Property]()
-			{
-				return GetBindingImage(Property);
-			});
-			BindingArgs.CurrentBindingColor = TAttribute<FLinearColor>::CreateLambda([this, Property]()
-			{
-				return GetBindingColor(Property);
-			});
+				// skip non-public properties for now
+				const bool bIsPublic = Property->HasAnyPropertyFlags(CPF_Edit | CPF_EditConst);
+				const bool bIsInstanceEditable = !Property->HasAnyPropertyFlags(CPF_DisableEditOnInstance);
+				if(!bIsPublic || !bIsInstanceEditable)
+				{
+					DetailBuilder.HideProperty(DefaultProperty);
+					continue;
+				}
+
+				const FSimpleDelegate OnValueChangedDelegate = FSimpleDelegate::CreateSP(this, &FRigModuleInstanceDetails::OnConfigValueChanged, Property->GetFName());
+				DefaultProperty->SetOnPropertyValueChanged(OnValueChangedDelegate);
+				DefaultProperty->SetOnChildPropertyValueChanged(OnValueChangedDelegate);
 
-			BindingArgs.OnCanBindProperty.BindLambda([](const FProperty* InProperty) -> bool { return true; });
-			BindingArgs.OnCanBindToClass.BindLambda([](UClass* InClass) -> bool { return false; });
-			BindingArgs.OnCanRemoveBinding.BindRaw(this, &FRigModuleInstanceDetails::CanRemoveBinding);
-			BindingArgs.OnRemoveBinding.BindSP(this, &FRigModuleInstanceDetails::HandleRemoveBinding);
-
-			BindingArgs.bGeneratePureBindings = true;
-			BindingArgs.bAllowNewBindings = true;
-			BindingArgs.bAllowArrayElementBindings = false;
-			BindingArgs.bAllowStructMemberBindings = false;
-			BindingArgs.bAllowUObjectFunctions = false;
-
-			BindingArgs.MenuExtender = MakeShareable(new FExtender);
-			BindingArgs.MenuExtender->AddMenuExtension(
-				"Properties",
-				EExtensionHook::After,
-				nullptr,
-				FMenuExtensionDelegate::CreateSPLambda(this, [this, Property](FMenuBuilder& MenuBuilder)
+				FPropertyBindingWidgetArgs BindingArgs;
+				BindingArgs.Property = (FProperty*)Property;
+				BindingArgs.CurrentBindingText = TAttribute<FText>::CreateLambda([this, Property]()
 				{
-					FillBindingMenu(MenuBuilder, Property);
-				})
-			);
+					return GetBindingText(Property);
+				});
+				BindingArgs.CurrentBindingImage = TAttribute<const FSlateBrush*>::CreateLambda([this, Property]()
+				{
+					return GetBindingImage(Property);
+				});
+				BindingArgs.CurrentBindingColor = TAttribute<FLinearColor>::CreateLambda([this, Property]()
+				{
+					return GetBindingColor(Property);
+				});
+
+				BindingArgs.OnCanBindProperty.BindLambda([](const FProperty* InProperty) -> bool { return true; });
+				BindingArgs.OnCanBindToClass.BindLambda([](UClass* InClass) -> bool { return false; });
+				BindingArgs.OnCanRemoveBinding.BindRaw(this, &FRigModuleInstanceDetails::CanRemoveBinding);
+				BindingArgs.OnRemoveBinding.BindSP(this, &FRigModuleInstanceDetails::HandleRemoveBinding);
+
+				BindingArgs.bGeneratePureBindings = true;
+				BindingArgs.bAllowNewBindings = true;
+				BindingArgs.bAllowArrayElementBindings = false;
+				BindingArgs.bAllowStructMemberBindings = false;
+				BindingArgs.bAllowUObjectFunctions = false;
+
+				BindingArgs.MenuExtender = MakeShareable(new FExtender);
+				BindingArgs.MenuExtender->AddMenuExtension(
+					"Properties",
+					EExtensionHook::After,
+					nullptr,
+					FMenuExtensionDelegate::CreateSPLambda(this, [this, Property](FMenuBuilder& MenuBuilder)
+					{
+						FillBindingMenu(MenuBuilder, Property);
+					})
+				);
 
-			// todo: remove the original row.
-
-			ConfigValuesCategory.AddCustomRow(DefaultProperty->GetPropertyDisplayName())
-			.NameContent()
-			[
-				DefaultProperty->CreatePropertyNameWidget()
-			]
-
-			// note: this doesn't work for some reason. seeking help from the editor team
-			.ValueContent()
-			[
-				DefaultProperty->CreatePropertyValueWidget()
-				// todo: if the property is bound / or partially bound
-				// mark the property value widget as disabled / read only.
-			]
-			
-			.ExtensionContent()
-			[
-				PropertyAccessEditor.MakePropertyBindingWidget(nullptr, BindingArgs)
-			];
+				// todo: remove the original row.
+
+				Category.AddCustomRow(DefaultProperty->GetPropertyDisplayName())
+				.NameContent()
+				[
+					DefaultProperty->CreatePropertyNameWidget()
+				]
+
+				// note: this doesn't work for some reason. seeking help from the editor team
+				.ValueContent()
+				[
+					DefaultProperty->CreatePropertyValueWidget()
+					// todo: if the property is bound / or partially bound
+					// mark the property value widget as disabled / read only.
+				]
+				
+				.ExtensionContent()
+				[
+					PropertyAccessEditor.MakePropertyBindingWidget(nullptr, BindingArgs)
+				];
+			}
 		}
 	}
 }
