Control Rig: Avoid crashes when opening same asset twice

#rnx

diff --git a/Engine/Plugins/Animation/ControlRig/Source/ControlRigEditor/Private/Editor/ControlRigEditor.cpp b/Engine/Plugins/Animation/ControlRig/Source/ControlRigEditor/Private/Editor/ControlRigEditor.cpp
--- a/Engine/Plugins/Animation/ControlRig/Source/ControlRigEditor/Private/Editor/ControlRigEditor.cpp
+++ b/Engine/Plugins/Animation/ControlRig/Source/ControlRigEditor/Private/Editor/ControlRigEditor.cpp
@@ -151,6 +151,12 @@ FControlRigEditor::~FControlRigEditor()
 			EditMode->OnEditorClosed();
 		}
 
+		if (RigBlueprint->IsModularRig())
+		{
+			RigBlueprint->OnSetObjectBeingDebugged().RemoveAll(&SchematicModel);
+			RigBlueprint->OnHierarchyModified().RemoveAll(&SchematicModel);
+		}
+
 		RigBlueprint->OnRigTypeChanged().RemoveAll(this);
 		if (RigBlueprint->IsModularRig())
 		{
