Control Rig: Allow interactive changes to control shapes

#rnx

diff --git a/Engine/Plugins/Animation/ControlRig/Source/ControlRigEditor/Private/EditMode/ControlRigEditMode.cpp b/Engine/Plugins/Animation/ControlRig/Source/ControlRigEditor/Private/EditMode/ControlRigEditMode.cpp
--- a/Engine/Plugins/Animation/ControlRig/Source/ControlRigEditor/Private/EditMode/ControlRigEditMode.cpp
+++ b/Engine/Plugins/Animation/ControlRig/Source/ControlRigEditor/Private/EditMode/ControlRigEditMode.cpp
@@ -3564,9 +3564,17 @@ void FControlRigEditMode::RecreateControlShapeActors(const TArray<FRigElementKey
 						{
 							const FTransform ShapeTransform = ControlRig->GetHierarchy()->GetControlShapeTransform(ControlElement, ERigTransformType::CurrentLocal);
 							FTransform MeshTransform = FTransform::Identity;
-							if (const FControlRigShapeDefinition* Gizmo = UControlRigShapeLibrary::GetShapeByName(ControlElement->Settings.ShapeName, ControlRig->GetShapeLibraries(), ControlRig->ShapeLibraryNameMap))
+							if (const FControlRigShapeDefinition* ShapeDef = UControlRigShapeLibrary::GetShapeByName(ControlElement->Settings.ShapeName, ControlRig->GetShapeLibraries(), ControlRig->ShapeLibraryNameMap))
 							{
-								MeshTransform = Gizmo->Transform;
+								MeshTransform = ShapeDef->Transform;
+
+								if(UStaticMesh* ShapeMesh = ShapeDef->StaticMesh.LoadSynchronous())
+								{
+									if(ShapeActor->StaticMeshComponent->GetStaticMesh() != ShapeMesh)
+									{
+										ShapeActor->StaticMeshComponent->SetStaticMesh(ShapeMesh);
+									}
+								}
 							}
 							ShapeActor->StaticMeshComponent->SetRelativeTransform(MeshTransform * ShapeTransform);
 						}
