Control Rig: Draw unresolved connectors in the schematic panel
#rnx

diff --git a/Engine/Plugins/Animation/ControlRig/Source/ControlRig/Public/Rigs/RigHierarchyCache.h b/Engine/Plugins/Animation/ControlRig/Source/ControlRig/Public/Rigs/RigHierarchyCache.h
--- a/Engine/Plugins/Animation/ControlRig/Source/ControlRig/Public/Rigs/RigHierarchyCache.h
+++ b/Engine/Plugins/Animation/ControlRig/Source/ControlRig/Public/Rigs/RigHierarchyCache.h
@@ -188,4 +188,5 @@ private:
 	friend class URigHierarchyController;
 	friend class FRigModuleInstanceDetails;
 	friend class UModularRig;
+	friend struct FControlRigSchematicModel;
 };
\ No newline at end of file
diff --git a/Engine/Plugins/Animation/ControlRig/Source/ControlRig/Public/Rigs/RigHierarchyDefines.h b/Engine/Plugins/Animation/ControlRig/Source/ControlRig/Public/Rigs/RigHierarchyDefines.h
--- a/Engine/Plugins/Animation/ControlRig/Source/ControlRig/Public/Rigs/RigHierarchyDefines.h
+++ b/Engine/Plugins/Animation/ControlRig/Source/ControlRig/Public/Rigs/RigHierarchyDefines.h
@@ -1465,6 +1465,19 @@ public:
 		, Name(InName)
 	{}
 
+	FRigElementKey(const FString& InKeyString)
+	{
+		FString TypeStr, NameStr;
+		check(InKeyString.Split(TEXT("("), &TypeStr, &NameStr));
+		NameStr.RemoveFromEnd(TEXT(")"));
+		Name = *NameStr;
+
+		const UEnum* ElementTypeEnum = StaticEnum<ERigElementType>();
+		Type = (ERigElementType)ElementTypeEnum->GetValueByName(*TypeStr);
+		check(Type > ERigElementType::None && Type < ERigElementType::All);
+		
+	}
+
 	void Serialize(FArchive& Ar);
 	void Save(FArchive& Ar);
 	void Load(FArchive& Ar);
diff --git a/Engine/Plugins/Animation/ControlRig/Source/ControlRigEditor/Private/ControlRigEditorStyle.h b/Engine/Plugins/Animation/ControlRig/Source/ControlRigEditor/Private/ControlRigEditorStyle.h
--- a/Engine/Plugins/Animation/ControlRig/Source/ControlRigEditor/Private/ControlRigEditorStyle.h
+++ b/Engine/Plugins/Animation/ControlRig/Source/ControlRigEditor/Private/ControlRigEditorStyle.h
@@ -146,6 +146,9 @@ public:
 		{
 			Set("ControlRig.Schematic.SocketUnresolved", new IMAGE_BRUSH_SVG("Slate/Socket_Open", Icon128x128));
 			Set("ControlRig.Schematic.SocketResolved", new IMAGE_BRUSH_SVG("Slate/Socket_Closed", Icon128x128));
+			Set("ControlRig.Schematic.ConnectorPrimary", new IMAGE_BRUSH_SVG("Slate/Connector_Primary", Icon128x128));
+			Set("ControlRig.Schematic.ConnectorSecondary", new IMAGE_BRUSH_SVG("Slate/Connector_Secondary", Icon128x128));
+			Set("ControlRig.Schematic.ConnectorOptional", new IMAGE_BRUSH_SVG("Slate/Connector_Optional", Icon128x128));
 		}
 
 		FSlateStyleRegistry::RegisterSlateStyle(*this);
diff --git a/Engine/Plugins/Animation/ControlRig/Source/ControlRigEditor/Private/Editor/ControlRigEditor.cpp b/Engine/Plugins/Animation/ControlRig/Source/ControlRigEditor/Private/Editor/ControlRigEditor.cpp
--- a/Engine/Plugins/Animation/ControlRig/Source/ControlRigEditor/Private/Editor/ControlRigEditor.cpp
+++ b/Engine/Plugins/Animation/ControlRig/Source/ControlRigEditor/Private/Editor/ControlRigEditor.cpp
@@ -155,6 +155,7 @@ FControlRigEditor::~FControlRigEditor()
 		{
 			RigBlueprint->OnSetObjectBeingDebugged().RemoveAll(&SchematicModel);
 			RigBlueprint->OnHierarchyModified().RemoveAll(&SchematicModel);
+			RigBlueprint->GetModularRigController()->OnModified().RemoveAll(&SchematicModel);
 		}
 
 		RigBlueprint->OnRigTypeChanged().RemoveAll(this);
@@ -328,8 +329,10 @@ void FControlRigEditor::InitRigVMEditor(const EToolkitMode::Type Mode, const TSh
 
 		if (ControlRigBlueprint->IsModularRig())
 		{
+			SchematicModel.SetEditor(SharedThis(this));
 			ControlRigBlueprint->OnSetObjectBeingDebugged().AddRaw(&SchematicModel, &FControlRigSchematicModel::OnSetObjectBeingDebugged);
 			ControlRigBlueprint->OnHierarchyModified().AddRaw(&SchematicModel, &FControlRigSchematicModel::OnHierarchyModified);
+			ControlRigBlueprint->GetModularRigController()->OnModified().AddRaw(&SchematicModel, &FControlRigSchematicModel::HandleModularRigModified);
 		}
 
 		ControlRigBlueprint->OnRigTypeChanged().AddSP(this, &FControlRigEditor::HandleRigTypeChanged);
@@ -2268,8 +2271,14 @@ void FControlRigEditor::HandleViewportCreated(const TSharedRef<class IPersonaVie
 			if (Blueprint->IsModularRig())
 			{
 				TSharedRef<SSchematicGraphPanel> SchematicPanel = SNew(SSchematicGraphPanel)
-																.GraphData(&SchematicModel.SchematicGraph)
-																.IsOverlay(true);
+																.GraphData(&SchematicModel)
+																.IsOverlay(true)
+																.PaddingLeft(30)
+																.PaddingRight(30)
+																.PaddingTop(60)
+																.PaddingBottom(60)
+																.PaddingInterNode(5)
+				;
 				InViewport->AddOverlayWidget(SchematicPanel);
 				HandleSchematicViewportCreated(SchematicPanel);
 			}
@@ -2565,9 +2574,9 @@ void FControlRigEditor::CacheNameLists()
 
 void FControlRigEditor::HandleSchematicViewportCreated(const TSharedRef<SSchematicGraphPanel>& InViewport)
 {
-	InViewport->UpdateNodeWidgetDelegate.BindSP(this, &FControlRigEditor::HandleUpdateSchematicNodes);
-	InViewport->OnNodeClickedDelegate.BindSP(this, &FControlRigEditor::HandleSchematicNodeClicked);
-	InViewport->OnDropDelegate.BindSP(this, &FControlRigEditor::HandleSchematicDrop);
+	InViewport->UpdateNodeWidgetDelegate.BindRaw(&SchematicModel, &FControlRigSchematicModel::HandleUpdateSchematicNodes);
+	InViewport->OnNodeClickedDelegate.BindRaw(&SchematicModel, &FControlRigSchematicModel::HandleSchematicNodeClicked);
+	InViewport->OnDropDelegate.BindRaw(&SchematicModel, &FControlRigSchematicModel::HandleSchematicDrop);
 }
 
 FVector2D FControlRigEditor::ComputePersonaProjectedScreenPos(const FVector& InWorldPos)
@@ -2613,136 +2622,6 @@ FVector2D FControlRigEditor::ComputePersonaProjectedScreenPos(const FVector& InW
 	return FVector2D::ZeroVector;
 }
 
-void FControlRigEditor::HandleUpdateSchematicNodes(SSchematicGraphPanel* InPanel, TSharedPtr<SSchematicGraphNode> InNode)
-{
-	UControlRigBlueprint* Blueprint = GetControlRigBlueprint();
-	if (!Blueprint)
-	{
-		return;
-	}
-
-	if (InPanel->bIsOverlay)
-	{
-		FTransform Transform;
-		if (UControlRig* DebuggedControlRig = Cast<UControlRig>(GetBlueprintObj()->GetObjectBeingDebugged()))
-		{
-			FRigElementKey SocketKey(*InNode->NodeData->Name, ERigElementType::Socket);
-			Transform = DebuggedControlRig->GetHierarchy()->GetGlobalTransform(SocketKey);
-			FVector2D PixelPos = ComputePersonaProjectedScreenPos(Transform.GetLocation());
-			InNode->SetPosition(PixelPos, true);
-
-			
-			if (const FRigElementKey* ConnectorKey = DebuggedControlRig->ElementKeyRedirector.FindReverse(SocketKey))
-			{
-				InNode->Brush = *FControlRigEditorStyle::Get().GetBrush( "ControlRig.Schematic.SocketResolved");
-				InNode->Size->Set(InNode->OriginalSize*0.5);
-			}
-			else
-			{
-				InNode->Brush = *FControlRigEditorStyle::Get().GetBrush( "ControlRig.Schematic.SocketUnresolved");
-				InNode->Size->Set(InNode->OriginalSize);
-			}
-
-			if (FRigSocketElement* SocketElement = Cast<FRigSocketElement>(DebuggedControlRig->GetHierarchy()->Find(SocketKey)))
-			{
-				InNode->Brush.TintColor = SocketElement->GetColor(DebuggedControlRig->GetHierarchy());
-			}
-			else
-			{
-				InNode->Brush.TintColor = FStyleColors::AccentBlue;
-			}
-		}
-
-	}
-	else
-	{
-		InNode->SetPosition(FVector2D(0, 0));
-		InNode->Brush = *FAppStyle::GetBrush("WhiteTexture");
-	}
-}
-
-void FControlRigEditor::HandleSchematicNodeClicked(SSchematicGraphPanel* InPanel, SSchematicGraphNode* InNode)
-{
-	for (FSchematicGraphNode* Node : SchematicModel.SchematicGraph.Nodes)
-	{
-		Node->bIsSelected = Node == InNode->NodeData;
-	}
-}
-
-void FControlRigEditor::HandleSchematicDrop(SSchematicGraphPanel* InPanel, SSchematicGraphNode* InNode, const FDragDropEvent& InDragDropEvent)
-{
-	UControlRigBlueprint* Blueprint = GetControlRigBlueprint();
-	if (!Blueprint)
-	{
-		return;
-	}
-
-	UControlRig* ControlRig = Blueprint->GetDebuggedControlRig();
-	if (!ControlRig)
-	{
-		return;
-	}
-
-	URigHierarchy* Hierarchy = ControlRig->GetHierarchy();
-	if (!Hierarchy)
-	{
-		return;
-	}
-	
-	FRigBaseElement* Element = Hierarchy->Find(FRigElementKey(*InNode->NodeData->Name, ERigElementType::Socket));
-	if (!Element)
-	{
-		return;
-	}
-
-	FRigSocketElement* Socket = Cast<FRigSocketElement>(Element);
-	if (!Socket)
-	{
-		return;
-	}
-
-	TSharedPtr<FAssetDragDropOp> AssetDragDropOp = InDragDropEvent.GetOperationAs<FAssetDragDropOp>();
-	if (AssetDragDropOp.IsValid())
-	{
-		for (const FAssetData& AssetData : AssetDragDropOp->GetAssets())
-		{
-			UClass* AssetClass = AssetData.GetClass();
-			if (!AssetClass->IsChildOf(UControlRigBlueprint::StaticClass()))
-			{
-				continue;
-			}
-
-			if(UControlRigBlueprint* AssetBlueprint = Cast<UControlRigBlueprint>(AssetData.GetAsset()))
-			{
-				if (UModularRigController* Controller = Blueprint->GetModularRigController())
-				{
-					const FName ModuleName = Controller->GetSafeNewName(FString(), FRigName(AssetBlueprint->RigModuleSettings.Identifier.Name));
-					const FString ModulePath = Controller->AddModule(ModuleName, AssetBlueprint->GetControlRigClass(), FString());
-					if(!ModulePath.IsEmpty())
-					{
-						FRigElementKey PrimaryConnectorKey;
-						TArray<FRigConnectorElement*> Connectors = Hierarchy->GetElementsOfType<FRigConnectorElement>();
-						for (FRigConnectorElement* Connector : Connectors)
-						{
-							if (Connector->Settings.Type == EConnectorType::Primary)
-							{
-								FString Path, Name;
-								Connector->GetName().Split(UModularRig::NamespaceSeparator, &Path, &Name, ESearchCase::CaseSensitive, ESearchDir::FromEnd);
-								if (Path == ModulePath)
-								{
-									PrimaryConnectorKey = Connector->GetKey();
-									break;
-								}
-							}
-						}
-						Controller->ConnectConnectorToElement(PrimaryConnectorKey, Socket->GetKey());
-					}
-				}
-			}
-		}
-	}
-}
-
 void FControlRigEditor::HandlePreviewMeshChanged(USkeletalMesh* InOldSkeletalMesh, USkeletalMesh* InNewSkeletalMesh)
 {
 	RebindToSkeletalMeshComponent();
diff --git a/Engine/Plugins/Animation/ControlRig/Source/ControlRigEditor/Private/Editor/ControlRigEditor.h b/Engine/Plugins/Animation/ControlRig/Source/ControlRigEditor/Private/Editor/ControlRigEditor.h
--- a/Engine/Plugins/Animation/ControlRig/Source/ControlRigEditor/Private/Editor/ControlRigEditor.h
+++ b/Engine/Plugins/Animation/ControlRig/Source/ControlRigEditor/Private/Editor/ControlRigEditor.h
@@ -192,6 +192,8 @@ public:
 	FText GetConnectorWarningText() const;
 	FReply OnNavigateToConnectorWarning() const;
 	FSimpleMulticastDelegate& OnRequestNavigateToConnectorWarning() { return RequestNavigateToConnectorWarningDelegate; }
+
+	FVector2D ComputePersonaProjectedScreenPos(const FVector& InWorldPos);
 	
 protected:
 
@@ -245,10 +247,6 @@ private:
 
 	/** Handle schematic panel setup*/
 	void HandleSchematicViewportCreated(const TSharedRef<class SSchematicGraphPanel>& InViewport);
-	void HandleUpdateSchematicNodes(SSchematicGraphPanel* InPanel, TSharedPtr<SSchematicGraphNode> InNode);
-	void HandleSchematicNodeClicked(SSchematicGraphPanel* InPanel, SSchematicGraphNode* InNode);
-	void HandleSchematicDrop(SSchematicGraphPanel* InPanel, SSchematicGraphNode* InNode, const FDragDropEvent& InDragDropEvent);
-	FVector2D ComputePersonaProjectedScreenPos(const FVector& InWorldPos);
 
 		/** Handle switching skeletal meshes */
 	void HandlePreviewMeshChanged(USkeletalMesh* InOldSkeletalMesh, USkeletalMesh* InNewSkeletalMesh);
diff --git a/Engine/Plugins/Animation/ControlRig/Source/ControlRigEditor/Private/Editor/ControlRigSchematicModel.cpp b/Engine/Plugins/Animation/ControlRig/Source/ControlRigEditor/Private/Editor/ControlRigSchematicModel.cpp
--- a/Engine/Plugins/Animation/ControlRig/Source/ControlRigEditor/Private/Editor/ControlRigSchematicModel.cpp
+++ b/Engine/Plugins/Animation/ControlRig/Source/ControlRigEditor/Private/Editor/ControlRigSchematicModel.cpp
@@ -2,6 +2,36 @@
 
 #include "ControlRigSchematicModel.h"
 #include "ControlRig.h"
+#include "ControlRigBlueprint.h"
+#include "ControlRigEditor.h"
+#include "ControlRigEditorStyle.h"
+#include "ModularRig.h"
+#include "DragAndDrop/AssetDragDropOp.h"
+
+FControlRigSchematicModel::~FControlRigSchematicModel()
+{
+	if(ControlRigBeingDebuggedPtr.IsValid())
+	{
+		if(UControlRig* ControlRigBeingDebugged = ControlRigBeingDebuggedPtr.Get())
+		{
+			if(!ControlRigBeingDebugged->HasAnyFlags(RF_BeginDestroyed))
+			{
+				ControlRigBeingDebugged->GetHierarchy()->OnModified().RemoveAll(this);
+			}
+		}
+	}
+
+	ControlRigBeingDebuggedPtr.Reset();
+	ControlRigEditor.Reset();
+	ControlRigBlueprint.Reset();
+}
+
+void FControlRigSchematicModel::SetEditor(const TSharedRef<FControlRigEditor>& InEditor)
+{
+	ControlRigEditor = InEditor;
+	ControlRigBlueprint = ControlRigEditor.Pin()->GetControlRigBlueprint();
+	ControlRigBeingDebuggedPtr = ControlRigBlueprint->GetDebuggedControlRig();
+}
 
 void FControlRigSchematicModel::OnSetObjectBeingDebugged(UObject* InObject)
 {
@@ -42,9 +72,10 @@ void FControlRigSchematicModel::OnHierarchyModified(ERigHierarchyNotification In
 		{
 			if (InElement)
 			{
-				if (InElement->GetType() == ERigElementType::Socket)
+				if (InElement->GetType() == ERigElementType::Socket ||
+					InElement->GetType() == ERigElementType::Connector)
 				{
-					SchematicGraph.AddNode(InElement->GetName());
+					AddNode(InElement->GetKey().ToString());
 				}
 			}
 			break;
@@ -53,10 +84,12 @@ void FControlRigSchematicModel::OnHierarchyModified(ERigHierarchyNotification In
 		{
 			if (InElement)
 			{
-				if (InElement->GetType() == ERigElementType::Socket)
+				if (InElement->GetType() == ERigElementType::Socket ||
+					InElement->GetType() == ERigElementType::Connector)
 				{
 					const FString OldNameStr = InHierarchy->GetPreviousName(InElement->GetKey()).ToString();
-					SchematicGraph.RenameNode(OldNameStr, InElement->GetName());
+					FRigElementKey OldKey(*OldNameStr, InElement->GetType());
+					RenameNode(OldKey.ToString(), InElement->GetKey().ToString());
 				}
 			}
 			break;
@@ -65,22 +98,209 @@ void FControlRigSchematicModel::OnHierarchyModified(ERigHierarchyNotification In
 		{
 			if (InElement)
 			{
-				if (InElement->GetType() == ERigElementType::Socket)
+				if (InElement->GetType() == ERigElementType::Socket ||
+					InElement->GetType() == ERigElementType::Connector)
 				{
-					SchematicGraph.RemoveNode(InElement->GetName());
+					RemoveNode(InElement->GetKey().ToString());
 				}
 			}
 			break;
 		}
 		case ERigHierarchyNotification::HierarchyReset:
 		{
-			SchematicGraph.Reset();
+			Reset();
 			TArray<FRigSocketElement*> Sockets = InHierarchy->GetElementsOfType<FRigSocketElement>();
 			for (FRigSocketElement* Socket : Sockets)
 			{
-				SchematicGraph.AddNode(Socket->GetFName().ToString());
+				AddNode(Socket->GetKey().ToString());
+			}
+			if (ControlRigBeingDebuggedPtr.IsValid())
+			{
+				FRigElementKeyRedirector& Redirector = ControlRigBeingDebuggedPtr->GetElementKeyRedirector();
+				TArray<FRigConnectorElement*> Connectors = InHierarchy->GetElementsOfType<FRigConnectorElement>();
+				for (FRigConnectorElement* Connector : Connectors)
+				{
+					if (!Redirector.Contains(Connector->GetKey()))
+					{
+						AddNode(Connector->GetKey().ToString());
+					}
+				}
+			}
+			break;
+		}
+	}
+}
+
+void FControlRigSchematicModel::HandleModularRigModified(EModularRigNotification InNotification, const FRigModuleReference* InModule)
+{
+	switch (InNotification)
+	{
+		case EModularRigNotification::ConnectionChanged:
+		{
+			if (ControlRigBeingDebuggedPtr.IsValid())
+			{
+				FRigElementKeyRedirector& Redirector = ControlRigBeingDebuggedPtr->GetElementKeyRedirector();
+				for (TPair<FRigElementKey, FRigElementKey>& Pair : Redirector.ExternalKeys)
+				{
+					RemoveNode(Pair.Key.ToString());
+				}
 			}
 			break;
 		}
 	}
 }
+
+void FControlRigSchematicModel::HandleSchematicNodeClicked(SSchematicGraphPanel* InPanel, SSchematicGraphNode* InNode)
+{
+	for (FSchematicGraphNode* Node : Nodes)
+	{
+		Node->bIsSelected = Node == InNode->NodeData;
+	}
+}
+
+void FControlRigSchematicModel::HandleUpdateSchematicNodes(SSchematicGraphPanel* InPanel, TSharedPtr<SSchematicGraphNode> InNode)
+{
+	if (!ControlRigBlueprint.IsValid())
+	{
+		return;
+	}
+
+	if (InPanel->bIsOverlay)
+	{
+		FTransform Transform;
+		if (ControlRigBeingDebuggedPtr.IsValid())
+		{
+			URigHierarchy* Hierarchy = ControlRigBeingDebuggedPtr->GetHierarchy();
+			FRigElementKey ElementKey(InNode->NodeData->Name);
+
+			if (ElementKey.IsTypeOf(ERigElementType::Socket))
+			{
+				Transform = Hierarchy->GetGlobalTransform(ElementKey);
+				FVector2D PixelPos = ControlRigEditor.Pin()->ComputePersonaProjectedScreenPos(Transform.GetLocation());
+				InNode->SetPosition(PixelPos, true);
+
+				if (const FRigElementKey* ConnectorKey = ControlRigBeingDebuggedPtr->GetElementKeyRedirector().FindReverse(ElementKey))
+				{
+					InNode->Brush = *FControlRigEditorStyle::Get().GetBrush( "ControlRig.Schematic.SocketResolved");
+					InNode->Size->Set(InNode->OriginalSize*0.5);
+				}
+				else
+				{
+					InNode->Brush = *FControlRigEditorStyle::Get().GetBrush( "ControlRig.Schematic.SocketUnresolved");
+					InNode->Size->Set(InNode->OriginalSize);
+				}
+
+				if (FRigSocketElement* SocketElement = Cast<FRigSocketElement>(Hierarchy->Find(ElementKey)))
+				{
+					InNode->Brush.TintColor = SocketElement->GetColor(Hierarchy);
+				}
+				else
+				{
+					InNode->Brush.TintColor = FStyleColors::AccentBlue;
+				}
+			}
+			else if(ElementKey.IsTypeOf(ERigElementType::Connector))
+			{
+				InNode->NodeData->Placement = ESchematicGraphNodePlacement::BottomRight;
+				InNode->Size->Set(InNode->OriginalSize);
+				if (FRigConnectorElement* ConnectorElement = Cast<FRigConnectorElement>(Hierarchy->Find(ElementKey)))
+				{
+					if (ConnectorElement->Settings.Type == EConnectorType::Primary)
+					{
+						InNode->Brush = *FControlRigEditorStyle::Get().GetBrush( "ControlRig.Schematic.ConnectorPrimary");
+					}
+					else if(ConnectorElement->Settings.Type == EConnectorType::Secondary)
+					{
+						if (ConnectorElement->Settings.bOptional)
+						{
+							InNode->Brush = *FControlRigEditorStyle::Get().GetBrush( "ControlRig.Schematic.ConnectorOptional");
+						}
+						else
+						{
+							InNode->Brush = *FControlRigEditorStyle::Get().GetBrush( "ControlRig.Schematic.ConnectorSecondary");
+						}
+					}
+				}
+			}
+		}
+
+	}
+	else
+	{
+		InNode->SetPosition(FVector2D(0, 0));
+		InNode->Brush = *FAppStyle::GetBrush("WhiteTexture");
+	}
+}
+
+void FControlRigSchematicModel::HandleSchematicDrop(SSchematicGraphPanel* InPanel, SSchematicGraphNode* InNode, const FDragDropEvent& InDragDropEvent)
+{
+	if (!ControlRigBlueprint.IsValid())
+	{
+		return;
+	}
+
+	UControlRig* ControlRig = ControlRigBlueprint->GetDebuggedControlRig();
+	if (!ControlRig)
+	{
+		return;
+	}
+
+	URigHierarchy* Hierarchy = ControlRig->GetHierarchy();
+	if (!Hierarchy)
+	{
+		return;
+	}
+	
+	FRigBaseElement* Element = Hierarchy->Find(FRigElementKey(InNode->NodeData->Name));
+	if (!Element)
+	{
+		return;
+	}
+
+	FRigSocketElement* Socket = Cast<FRigSocketElement>(Element);
+	if (!Socket)
+	{
+		return;
+	}
+
+	TSharedPtr<FAssetDragDropOp> AssetDragDropOp = InDragDropEvent.GetOperationAs<FAssetDragDropOp>();
+	if (AssetDragDropOp.IsValid())
+	{
+		for (const FAssetData& AssetData : AssetDragDropOp->GetAssets())
+		{
+			UClass* AssetClass = AssetData.GetClass();
+			if (!AssetClass->IsChildOf(UControlRigBlueprint::StaticClass()))
+			{
+				continue;
+			}
+
+			if(UControlRigBlueprint* AssetBlueprint = Cast<UControlRigBlueprint>(AssetData.GetAsset()))
+			{
+				if (UModularRigController* Controller = ControlRigBlueprint->GetModularRigController())
+				{
+					const FName ModuleName = Controller->GetSafeNewName(FString(), FRigName(AssetBlueprint->RigModuleSettings.Identifier.Name));
+					const FString ModulePath = Controller->AddModule(ModuleName, AssetBlueprint->GetControlRigClass(), FString());
+					if(!ModulePath.IsEmpty())
+					{
+						FRigElementKey PrimaryConnectorKey;
+						TArray<FRigConnectorElement*> Connectors = Hierarchy->GetElementsOfType<FRigConnectorElement>();
+						for (FRigConnectorElement* Connector : Connectors)
+						{
+							if (Connector->Settings.Type == EConnectorType::Primary)
+							{
+								FString Path, Name;
+								Connector->GetName().Split(UModularRig::NamespaceSeparator, &Path, &Name, ESearchCase::CaseSensitive, ESearchDir::FromEnd);
+								if (Path == ModulePath)
+								{
+									PrimaryConnectorKey = Connector->GetKey();
+									break;
+								}
+							}
+						}
+						Controller->ConnectConnectorToElement(PrimaryConnectorKey, Socket->GetKey());
+					}
+				}
+			}
+		}
+	}
+}
\ No newline at end of file
diff --git a/Engine/Plugins/Animation/ControlRig/Source/ControlRigEditor/Private/Editor/ControlRigSchematicModel.h b/Engine/Plugins/Animation/ControlRig/Source/ControlRigEditor/Private/Editor/ControlRigSchematicModel.h
--- a/Engine/Plugins/Animation/ControlRig/Source/ControlRigEditor/Private/Editor/ControlRigSchematicModel.h
+++ b/Engine/Plugins/Animation/ControlRig/Source/ControlRigEditor/Private/Editor/ControlRigSchematicModel.h
@@ -2,19 +2,33 @@
 
 #pragma once
 
+#include "ModularRigModel.h"
 #include "SchematicGraphPanel/SSchematicGraphPanel.h"
 #include "Rigs/RigHierarchyDefines.h"
 
+class FControlRigEditor;
 class UControlRig;
 class URigHierarchy;
 struct FRigBaseElement;
 
 /** Model for the schematic views */
-struct FControlRigSchematicModel 
+struct FControlRigSchematicModel : FSchematicGraph
 {
-	FSchematicGraph SchematicGraph;
-	TWeakObjectPtr<UControlRig> ControlRigBeingDebuggedPtr;
+	~FControlRigSchematicModel();
+	
+	void SetEditor(const TSharedRef<FControlRigEditor>& InEditor);
 
 	void OnSetObjectBeingDebugged(UObject* InObject);
 	void OnHierarchyModified(ERigHierarchyNotification InNotif, URigHierarchy* InHierarchy, const FRigBaseElement* InElement);
+	void HandleModularRigModified(EModularRigNotification InNotification, const FRigModuleReference* InModule);
+
+	void HandleUpdateSchematicNodes(SSchematicGraphPanel* InPanel, TSharedPtr<SSchematicGraphNode> InNode);
+	void HandleSchematicNodeClicked(SSchematicGraphPanel* InPanel, SSchematicGraphNode* InNode);
+	void HandleSchematicDrop(SSchematicGraphPanel* InPanel, SSchematicGraphNode* InNode, const FDragDropEvent& InDragDropEvent);
+
+private:
+	
+	TWeakPtr<FControlRigEditor> ControlRigEditor;
+	TWeakObjectPtr<UControlRigBlueprint> ControlRigBlueprint;
+	TWeakObjectPtr<UControlRig> ControlRigBeingDebuggedPtr;
 };
\ No newline at end of file
diff --git a/Engine/Source/Developer/AnimationWidgets/Private/SchematicGraphPanel/SSchematicGraphPanel.cpp b/Engine/Source/Developer/AnimationWidgets/Private/SchematicGraphPanel/SSchematicGraphPanel.cpp
--- a/Engine/Source/Developer/AnimationWidgets/Private/SchematicGraphPanel/SSchematicGraphPanel.cpp
+++ b/Engine/Source/Developer/AnimationWidgets/Private/SchematicGraphPanel/SSchematicGraphPanel.cpp
@@ -13,7 +13,10 @@ void FSchematicGraph::Reset()
 	Nodes.Reset();
 	Links.Reset();
 
-	OnGraphReset().Broadcast();
+	if (OnGraphResetDelegate.IsBound())
+	{
+		OnGraphReset().Broadcast();
+	}
 }
 
 bool FSchematicGraph::AddNode(const FString& InName)
@@ -30,7 +33,11 @@ bool FSchematicGraph::AddNode(const FString& InName)
 	new (NewElement) FSchematicGraphNode();
 	NewElement->Name = InName;
 	Nodes.Add(NewElement);
-	OnNodeAddedDelegate.Broadcast(NewElement);
+
+	if (OnNodeAddedDelegate.IsBound())
+	{
+		OnNodeAddedDelegate.Broadcast(NewElement);
+	}
 	return true;
 }
 
@@ -43,6 +50,10 @@ bool FSchematicGraph::RenameNode(const FString& InOldName, const FString& InNewN
 
 	if (FoundNode)
 	{
+		if (OnNodeRenamedDelegate.IsBound())
+		{
+			OnNodeRenamedDelegate.Broadcast(*FoundNode);
+		}
 		(*FoundNode)->Name = InNewName;
 		return true;
 	}
@@ -58,7 +69,10 @@ bool FSchematicGraph::RemoveNode(const FString& InName)
 	});
 	if (FoundNode)
 	{
-		OnNodeRemovedDelegate.Broadcast(*FoundNode);
+		if (OnNodeRemovedDelegate.IsBound())
+		{
+			OnNodeRemovedDelegate.Broadcast(*FoundNode);
+		}
 		Nodes.Remove(*FoundNode);
 		FMemory::Free(*FoundNode);
 		return true;
@@ -78,6 +92,7 @@ void SSchematicGraphNode::Construct(const FArguments& InArgs)
 	Scale = TAnimatedAttribute<float>::Create(FloatInterpSettings, 1.f);
 	Brush = *FAppStyle::GetBrush("WhiteTexture");
 	Brush.TintColor = EStyleColor::AccentWhite;
+	SetToolTipText(FText::FromString(NodeData->Name));
 }
 
 FVector2D SSchematicGraphNode::ComputeDesiredSize(float LayoutScaleMultiplier) const
@@ -161,6 +176,7 @@ void SSchematicGraphPanel::SetSchematicGraph(FSchematicGraph* InGraphData)
 	{
 		GraphData->OnNodeAdded().RemoveAll(this);
 		GraphData->OnNodeRemoved().RemoveAll(this);
+		GraphData->OnNodeRenamed().RemoveAll(this);
 		GraphData->OnGraphReset().RemoveAll(this);
 	}
 	
@@ -169,6 +185,7 @@ void SSchematicGraphPanel::SetSchematicGraph(FSchematicGraph* InGraphData)
 	{
 		GraphData->OnNodeAdded().AddSP(this, &SSchematicGraphPanel::AddNode);
 		GraphData->OnNodeRemoved().AddSP(this, &SSchematicGraphPanel::RemoveNode);
+		GraphData->OnNodeRenamed().AddSP(this, &SSchematicGraphPanel::RenameNode);
 		GraphData->OnGraphReset().AddSP(this, &SSchematicGraphPanel::RebuildPanel);
 	}
 }
@@ -177,6 +194,11 @@ void SSchematicGraphPanel::Construct(const FArguments& InArgs)
 {
 	GraphData = InArgs._GraphData;
 	bIsOverlay = InArgs._IsOverlay;
+	PaddingLeft = InArgs._PaddingLeft;
+	PaddingRight = InArgs._PaddingRight;
+	PaddingTop = InArgs._PaddingTop;
+	PaddingBottom = InArgs._PaddingBottom;
+	PaddingInterNode = InArgs._PaddingInterNode;
 	UpdateNodeWidgetDelegate = InArgs._OnUpdateNodeWidget;
 	OnNodeClickedDelegate = InArgs._OnNodeClicked;
 	OnDropDelegate = InArgs._OnDrop;
@@ -191,6 +213,7 @@ void SSchematicGraphPanel::Construct(const FArguments& InArgs)
 	{
 		GraphData->OnNodeAdded().AddSP(this, &SSchematicGraphPanel::AddNode);
 		GraphData->OnNodeRemoved().AddSP(this, &SSchematicGraphPanel::RemoveNode);
+		GraphData->OnNodeRenamed().AddSP(this, &SSchematicGraphPanel::RenameNode);
 		GraphData->OnGraphReset().AddSP(this, &SSchematicGraphPanel::RebuildPanel);
 	}
 }
@@ -239,6 +262,28 @@ void SSchematicGraphPanel::RemoveNode(FSchematicGraphNode* InNodeToRemove)
 	}
 }
 
+void SSchematicGraphPanel::RenameNode(FSchematicGraphNode* InNodeToRename)
+{
+	for (int32 Iter = 0; Iter != Children.Num(); ++Iter)
+	{
+		TSharedRef<SSchematicGraphNode> Child = GetChild(Iter);
+		if (Child->NodeData == InNodeToRename)
+		{
+			Child->SetToolTipText(FText::FromString(InNodeToRename->Name));
+			break;
+		}
+	}
+	for (int32 Iter = 0; Iter != VisibleChildren.Num(); ++Iter)
+	{
+		TSharedRef<SSchematicGraphNode> Child = StaticCastSharedRef<SSchematicGraphNode>(VisibleChildren[Iter]);
+		if (Child->NodeData == InNodeToRename)
+		{
+			Child->SetToolTipText(FText::FromString(InNodeToRename->Name));
+			break;
+		}
+	}
+}
+
 int32 SSchematicGraphPanel::OnPaint(const FPaintArgs& Args, const FGeometry& AllottedGeometry, const FSlateRect& MyCullingRect, FSlateWindowElementList& OutDrawElements, int32 LayerId, const FWidgetStyle& InWidgetStyle, bool bParentEnabled) const
 {
 	const int32 NodeLayerId = LayerId + 1;
@@ -331,11 +376,48 @@ void SSchematicGraphPanel::Tick(float DeltaTime)
 		}
 	}
 	SetFadeBackground(bFadeBackgroundBlanket);
-	
+
+	uint16 TopRightNodes = 0;
+	uint16 TopLeftNodes = 0;
+	uint16 BottomRightNodes = 0;
+	uint16 BottomLeftNodes = 0;
+	FVector2D Size = GetCachedGeometry().Size;
 	for (int32 i=0; i<Children.Num(); ++i)
 	{
 		TSharedRef<SSchematicGraphNode> SNode = GetChild(i);
-		UpdateNodeWidgetDelegate.Execute(this, SNode.ToSharedPtr());
+		UpdateNodeWidgetDelegate.ExecuteIfBound(this, SNode.ToSharedPtr());
+		
+		switch(SNode->NodeData->Placement)
+		{
+			case ESchematicGraphNodePlacement::Panel:
+			{
+				break;
+			}
+			case ESchematicGraphNodePlacement::TopLeft:
+			{
+				SNode->SetPosition(FVector2d(PaddingLeft + SNode->OriginalSize.X*0.5, PaddingTop + SNode->OriginalSize.Y*0.5 + ((SNode->OriginalSize.Y + PaddingInterNode) * TopLeftNodes)), true);
+				TopLeftNodes++;
+				break;
+			}
+			case ESchematicGraphNodePlacement::TopRight:
+			{
+				SNode->SetPosition(FVector2d(Size.X - PaddingRight - (SNode->OriginalSize.X * 0.5), PaddingTop + SNode->OriginalSize.Y*0.5 + ((SNode->OriginalSize.Y + PaddingInterNode) * TopRightNodes)), true);
+				TopRightNodes++;
+				break;
+			}
+			case ESchematicGraphNodePlacement::BottomLeft:
+			{
+				SNode->SetPosition(FVector2d(PaddingLeft + SNode->OriginalSize.X*0.5, Size.Y - PaddingBottom - (SNode->OriginalSize.Y * 0.5) - ((SNode->OriginalSize.Y + PaddingInterNode)  * BottomLeftNodes)), true);
+				BottomLeftNodes++;
+				break;
+			}
+			case ESchematicGraphNodePlacement::BottomRight:
+			{
+				SNode->SetPosition(FVector2d(Size.X - PaddingRight - (SNode->OriginalSize.X * 0.5), Size.Y - PaddingBottom - (SNode->OriginalSize.Y * 0.5) - ((SNode->OriginalSize.Y + PaddingInterNode)  * BottomRightNodes)), true);
+				BottomRightNodes++;
+				break;
+			}
+		}
 	}
 }
 
diff --git a/Engine/Source/Developer/AnimationWidgets/Private/SchematicGraphPanel/SchematicViewportTabSummoner.cpp b/Engine/Source/Developer/AnimationWidgets/Private/SchematicGraphPanel/SchematicViewportTabSummoner.cpp
--- a/Engine/Source/Developer/AnimationWidgets/Private/SchematicGraphPanel/SchematicViewportTabSummoner.cpp
+++ b/Engine/Source/Developer/AnimationWidgets/Private/SchematicGraphPanel/SchematicViewportTabSummoner.cpp
@@ -38,7 +38,12 @@ TSharedRef<SWidget> FSchematicViewportTabSummoner::CreateTabBody(const FWorkflow
 {
 	TSharedRef<SSchematicGraphPanel> Viewport = SNew(SSchematicGraphPanel)
 											.GraphData(SchematicGraph)
-											.IsOverlay(false);
+											.IsOverlay(false)
+											.PaddingLeft(30)
+											.PaddingRight(30)
+											.PaddingTop(30)
+											.PaddingBottom(30)
+											.PaddingInterNode(5);
 	OnViewportCreated.ExecuteIfBound(Viewport);
 	return Viewport;
 }
diff --git a/Engine/Source/Developer/AnimationWidgets/Public/SchematicGraphPanel/SSchematicGraphPanel.h b/Engine/Source/Developer/AnimationWidgets/Public/SchematicGraphPanel/SSchematicGraphPanel.h
--- a/Engine/Source/Developer/AnimationWidgets/Public/SchematicGraphPanel/SSchematicGraphPanel.h
+++ b/Engine/Source/Developer/AnimationWidgets/Public/SchematicGraphPanel/SSchematicGraphPanel.h
@@ -7,14 +7,25 @@
 #include "TickableEditorObject.h"
 #include "Animation/AnimatedAttribute.h"
 
+enum ESchematicGraphNodePlacement
+{
+	Panel,
+	TopLeft,
+	TopRight,
+	BottomLeft,
+	BottomRight
+};
+
 struct ANIMATIONWIDGETS_API FSchematicGraphNode
 {
 	FString Name;
 	bool bIsSelected = false;
+	ESchematicGraphNodePlacement Placement;
 };
 
 DECLARE_EVENT_OneParam(FSchematicGraph, FOnNodeAdded, FSchematicGraphNode*);
 DECLARE_EVENT_OneParam(FSchematicGraph, FOnNodeRemoved, FSchematicGraphNode*);
+DECLARE_EVENT_OneParam(FSchematicGraph, FOnNodeRenamed, FSchematicGraphNode*);
 DECLARE_EVENT(FSchematicGraph, FOnGraphReset);
 
 class ANIMATIONWIDGETS_API FSchematicGraph
@@ -32,6 +43,8 @@ public:
 	FOnNodeAdded& OnNodeAdded() { return OnNodeAddedDelegate; }
 	FOnNodeRemoved OnNodeRemovedDelegate;
 	FOnNodeRemoved& OnNodeRemoved() { return OnNodeRemovedDelegate; }
+	FOnNodeRenamed OnNodeRenamedDelegate;
+	FOnNodeRenamed& OnNodeRenamed() { return OnNodeRenamedDelegate; }
 	FOnGraphReset OnGraphResetDelegate;
 	FOnGraphReset& OnGraphReset() { return OnGraphResetDelegate; }
 };
@@ -95,12 +108,22 @@ public:
 	SLATE_BEGIN_ARGS(SSchematicGraphPanel) {}
 	SLATE_ARGUMENT(bool, IsOverlay)
 	SLATE_ARGUMENT(FSchematicGraph*, GraphData)
+	SLATE_ARGUMENT(int32, PaddingLeft)
+	SLATE_ARGUMENT(int32, PaddingRight)
+	SLATE_ARGUMENT(int32, PaddingTop)
+	SLATE_ARGUMENT(int32, PaddingBottom)
+	SLATE_ARGUMENT(int32, PaddingInterNode)
 	SLATE_EVENT(FUpdateNodeWidget, OnUpdateNodeWidget)
 	SLATE_EVENT(FOnNodeClicked, OnNodeClicked)
 	SLATE_EVENT(FOnDrop, OnDrop)
 	SLATE_END_ARGS()
 
-	~SSchematicGraphPanel(){}
+	~SSchematicGraphPanel()
+	{
+		UpdateNodeWidgetDelegate.Unbind();
+		OnNodeClickedDelegate.Unbind();
+		OnDropDelegate.Unbind();
+	}
 
 	void SetSchematicGraph(FSchematicGraph* InGraphData);
 	void Construct(const FArguments& InArgs);
@@ -108,6 +131,7 @@ public:
 	void RebuildPanel();
 	void AddNode(FSchematicGraphNode* NodeToAdd);
 	void RemoveNode(FSchematicGraphNode* NodeToRemove);
+	void RenameNode(FSchematicGraphNode* NodeToRemove);
 
 	// SNodePanel interface
 	virtual int32 OnPaint(const FPaintArgs& Args, const FGeometry& AllottedGeometry, const FSlateRect& MyCullingRect, FSlateWindowElementList& OutDrawElements, int32 LayerId, const FWidgetStyle& InWidgetStyle, bool bParentEnabled) const override;
@@ -128,6 +152,11 @@ public:
 	void SetFadeBackground(bool bInFade) { FadeBackgroundAlpha->Set(bInFade ? 0.5f : 0.f); }
 
 	bool bIsOverlay;
+	int32 PaddingLeft = 0;
+	int32 PaddingRight = 0;
+	int32 PaddingTop = 0;
+	int32 PaddingBottom = 0;
+	int32 PaddingInterNode = 0;
 	TSharedPtr<TAnimatedAttribute<float>> FadeBackgroundAlpha;
 	FSchematicGraph* GraphData;
 	FUpdateNodeWidget UpdateNodeWidgetDelegate;
