Control Rig: Improve short name usage for UI and sequencer

* Use local hierarchy inside of RigHierarchyController
* introduce URigHierarchy::JoinNameSpace to avoid string formats all over the place

#rnx

diff --git a/Engine/Plugins/Animation/ControlRig/Source/ControlRig/Private/ModularRig.cpp b/Engine/Plugins/Animation/ControlRig/Source/ControlRig/Private/ModularRig.cpp
--- a/Engine/Plugins/Animation/ControlRig/Source/ControlRig/Private/ModularRig.cpp
+++ b/Engine/Plugins/Animation/ControlRig/Source/ControlRig/Private/ModularRig.cpp
@@ -75,7 +75,7 @@ FString FRigModuleInstance::GetPath() const
 {
 	if (!ParentPath.IsEmpty())
 	{
-		return FString::Printf(TEXT("%s:%s"), *ParentPath, *Name.ToString()); 
+		return URigHierarchy::JoinNameSpace(ParentPath, Name.ToString()); 
 	}
 	return Name.ToString();
 }
@@ -325,6 +325,7 @@ void UModularRig::ExecuteQueue()
 			// forward the draw interface to each module
 			RigPublicContext.SetDrawInterface(PublicContext.GetDrawInterface());
 			RigPublicContext.SetDrawContainer(PublicContext.GetDrawContainer());
+			RigPublicContext.RigModuleInstance = ExecutionElement.ModuleInstance;
 
 			// re-initialize the module in case only the VM side got recompiled.
 			// this happens when the user relies on auto recompilation when editing the
diff --git a/Engine/Plugins/Animation/ControlRig/Source/ControlRig/Private/ModularRigController.cpp b/Engine/Plugins/Animation/ControlRig/Source/ControlRig/Private/ModularRigController.cpp
--- a/Engine/Plugins/Animation/ControlRig/Source/ControlRig/Private/ModularRigController.cpp
+++ b/Engine/Plugins/Animation/ControlRig/Source/ControlRig/Private/ModularRigController.cpp
@@ -315,7 +315,7 @@ TArray<FString> UModularRigController::GetPossibleBindings(const FString& InModu
 			for (const FRigVMExternalVariable& Variable : Variables)
 			{
 				FText ErrorMessage;
-				const FString SourceVariablePath = FString::Printf(TEXT("%s:%s"), *CurModulePath, *Variable.Name.ToString());
+				const FString SourceVariablePath = URigHierarchy::JoinNameSpace(CurModulePath, Variable.Name.ToString());
 				if (CanBindModuleVariable(InModulePath, InVariableName, SourceVariablePath, ErrorMessage))
 				{
 					PossibleBindings.Add(SourceVariablePath);
@@ -394,7 +394,7 @@ bool UModularRigController::CanBindModuleVariable(const FString& InModulePath, c
 		return false;
 	}
 
-	FString SourcePath = (SourceModulePath.IsEmpty()) ? SourceVariableName : FString::Printf(TEXT("%s:%s"), *SourceModulePath, *SourceVariableName);
+	FString SourcePath = (SourceModulePath.IsEmpty()) ? SourceVariableName : URigHierarchy::JoinNameSpace(SourceModulePath, SourceVariableName);
 	if (!RigVMTypeUtils::AreCompatible(SourceProperty, TargetProperty))
 	{
 		FString TargetPath = FString::Printf(TEXT("%s.%s"), *InModulePath, *InVariableName.ToString());
@@ -410,7 +410,7 @@ bool UModularRigController::BindModuleVariable(const FString& InModulePath, cons
 	FText ErrorMessage;
 	if (!CanBindModuleVariable(InModulePath, InVariableName, InSourcePath, ErrorMessage))
 	{
-		UE_LOG(LogControlRig, Error, TEXT("Could not bind module variable %s:%s : %s"), *InModulePath, *InVariableName.ToString(), *ErrorMessage.ToString());
+		UE_LOG(LogControlRig, Error, TEXT("Could not bind module variable %s : %s"), *URigHierarchy::JoinNameSpace(InModulePath, InVariableName.ToString()), *ErrorMessage.ToString());
 		return false;
 	}
 	
@@ -439,7 +439,7 @@ bool UModularRigController::BindModuleVariable(const FString& InModulePath, cons
 		}
 	}
 
-	FString SourcePath = (SourceModulePath.IsEmpty()) ? SourceVariableName : FString::Printf(TEXT("%s:%s"), *SourceModulePath, *SourceVariableName);
+	FString SourcePath = (SourceModulePath.IsEmpty()) ? SourceVariableName : URigHierarchy::JoinNameSpace(SourceModulePath, SourceVariableName);
 
 #if WITH_EDITOR
 	TSharedPtr<FScopedTransaction> TransactionPtr;
@@ -595,8 +595,8 @@ FString UModularRigController::RenameModule(const FString& InModulePath, const F
 	}
 #endif
 	
-	const FString OldPath = (Module->ParentPath.IsEmpty()) ? OldName : FString::Printf(TEXT("%s:%s"), *Module->ParentPath, *OldName);
-	const FString NewPath = (Module->ParentPath.IsEmpty()) ? *NewName :  FString::Printf(TEXT("%s:%s"), *Module->ParentPath, *NewName);
+	const FString OldPath = (Module->ParentPath.IsEmpty()) ? OldName : URigHierarchy::JoinNameSpace(Module->ParentPath, OldName);
+	const FString NewPath = (Module->ParentPath.IsEmpty()) ? *NewName :  URigHierarchy::JoinNameSpace(Module->ParentPath, NewName);
 	Module->PreviousName = Module->Name;
 	Module->Name = InNewName;
 	TArray<FRigModuleReference*> Children;
@@ -618,7 +618,7 @@ FString UModularRigController::RenameModule(const FString& InModulePath, const F
 			Binding.Value.Split(UModularRig::NamespaceSeparator, &ModulePath, &VariableName, ESearchCase::CaseSensitive, ESearchDir::FromEnd);
 			if (ModulePath == OldPath)
 			{
-				Binding.Value = FString::Printf(TEXT("%s:%s"), *NewPath, *VariableName);
+				Binding.Value = URigHierarchy::JoinNameSpace(NewPath, VariableName);
 			}
 		};
 	}
@@ -714,7 +714,7 @@ FString UModularRigController::ReparentModule(const FString& InModulePath, const
 			Binding.Value.Split(UModularRig::NamespaceSeparator, &ModulePath, &VariableName, ESearchCase::CaseSensitive, ESearchDir::FromEnd);
 			if (ModulePath == OldPath)
 			{
-				Binding.Value = FString::Printf(TEXT("%s:%s"), *NewPath, *VariableName);
+				Binding.Value = URigHierarchy::JoinNameSpace(NewPath, VariableName);
 				ModulePath = NewPath;
 			}
 
@@ -990,7 +990,7 @@ void UModularRigController::UpdateShortNames()
 
 				while (RemainingPath.Split(UModularRig::NamespaceSeparator, &Left, &Right, ESearchCase::IgnoreCase, ESearchDir::FromEnd))
 				{
-					ShortPath = ShortPath.IsEmpty() ? Right : FString::Printf(TEXT("%s:%s"), *Right, *ShortPath);
+					ShortPath = ShortPath.IsEmpty() ? Right : URigHierarchy::JoinNameSpace(Right, ShortPath);
 
 					// if the short path only exists once - that's what we use for the display name
 					if(TokenToCount.FindChecked(ShortPath) == 1)
@@ -1004,7 +1004,7 @@ void UModularRigController::UpdateShortNames()
 
 				if(!RemainingPath.IsEmpty())
 				{
-					ShortPath = FString::Printf(TEXT("%s:%s"), *RemainingPath, *ShortPath);
+					ShortPath = URigHierarchy::JoinNameSpace(RemainingPath, ShortPath);
 				}
 			}
 
diff --git a/Engine/Plugins/Animation/ControlRig/Source/ControlRig/Private/Rigs/RigHierarchy.cpp b/Engine/Plugins/Animation/ControlRig/Source/ControlRig/Private/Rigs/RigHierarchy.cpp
--- a/Engine/Plugins/Animation/ControlRig/Source/ControlRig/Private/Rigs/RigHierarchy.cpp
+++ b/Engine/Plugins/Animation/ControlRig/Source/ControlRig/Private/Rigs/RigHierarchy.cpp
@@ -1226,6 +1226,16 @@ TArray<FRigElementKey> URigHierarchy::GetSelectedKeys(ERigElementType InTypeFilt
 	return Selection;
 }
 
+FRigName URigHierarchy::JoinNameSpace(const FRigName& InLeft, const FRigName& InRight)
+{
+	return FRigName(JoinNameSpace(InLeft.ToString(), InRight.ToString()));
+}
+
+FString URigHierarchy::JoinNameSpace(const FString& InLeft, const FString& InRight)
+{
+	return InLeft + UModularRig::NamespaceSeparator + InRight;
+}
+
 void URigHierarchy::SanitizeName(FRigName& InOutName, bool bAllowNameSpaces)
 {
 	// Sanitize the name
@@ -1497,6 +1507,38 @@ FRigName URigHierarchy::GetSafeNewDisplayName(const FRigElementKey& InParentElem
 	return Name;
 }
 
+FText URigHierarchy::GetDisplayNameForUI(const FRigBaseElement* InElement, bool bIncludeNameSpace) const
+{
+	check(InElement);
+
+	const FName& DisplayName = InElement->GetDisplayName();
+	
+	if(bIncludeNameSpace)
+	{
+		const FName ModuleShortName = GetNameMetadata(InElement->Key, ShortModuleNameMetadataName, NAME_None);
+		if(!ModuleShortName.IsNone())
+		{
+			FString DisplayNameString = DisplayName.ToString();
+			(void)DisplayNameString.Split(UModularRig::NamespaceSeparator, nullptr, &DisplayNameString, ESearchCase::CaseSensitive, ESearchDir::FromEnd);
+			
+			const FString ModuleShortNameString = ModuleShortName.ToString();
+			const FString ModuleDisplayName = JoinNameSpace(ModuleShortNameString, DisplayNameString);
+			return FText::FromString(ModuleDisplayName);
+		}
+	}
+
+	return FText::FromName(DisplayName);
+}
+
+FText URigHierarchy::GetDisplayNameForUI(const FRigElementKey& InKey, bool bIncludeNameSpace) const
+{
+	if(const FRigBaseElement* Element = Find(InKey))
+	{
+		return GetDisplayNameForUI(Element, bIncludeNameSpace);
+	}
+	return FText();
+}
+
 int32 URigHierarchy::GetPoseVersion(const FRigElementKey& InKey) const
 {
 	if(const FRigTransformElement* TransformElement = Find<FRigTransformElement>(InKey))
diff --git a/Engine/Plugins/Animation/ControlRig/Source/ControlRig/Private/Rigs/RigHierarchyController.cpp b/Engine/Plugins/Animation/ControlRig/Source/ControlRig/Private/Rigs/RigHierarchyController.cpp
--- a/Engine/Plugins/Animation/ControlRig/Source/ControlRig/Private/Rigs/RigHierarchyController.cpp
+++ b/Engine/Plugins/Animation/ControlRig/Source/ControlRig/Private/Rigs/RigHierarchyController.cpp
@@ -44,7 +44,7 @@ void URigHierarchyController::SetHierarchy(URigHierarchy* InHierarchy)
 	{
 		// make sure making multiple valid SetHierarchy() calls won't lead to accumulated delegates
 		// though it should not happen in the first place
-		if (Hierarchy.IsValid())
+		if (URigHierarchy* Hierarchy = WeakHierarchy.Get())
 		{
 			if(!Hierarchy->HasAnyFlags(RF_BeginDestroyed) && Hierarchy->IsValidLowLevel())
 			{
@@ -56,8 +56,8 @@ void URigHierarchyController::SetHierarchy(URigHierarchy* InHierarchy)
 		URigHierarchy* OuterHierarchy = Cast<URigHierarchy>(GetOuter());
 		if (ensure(OuterHierarchy) && ::IsValid(OuterHierarchy))
 		{
-			Hierarchy = OuterHierarchy;
-			Hierarchy->OnModified().AddUObject(this, &URigHierarchyController::HandleHierarchyModified);
+			WeakHierarchy = OuterHierarchy;
+			WeakHierarchy->OnModified().AddUObject(this, &URigHierarchyController::HandleHierarchyModified);
 		}
 	}
 	else
@@ -80,6 +80,8 @@ bool URigHierarchyController::SelectElement(FRigElementKey InKey, bool bSelect,
 		return SetSelection(KeysToSelect);
 	}
 
+	URigHierarchy* Hierarchy = GetHierarchy();
+
 	if(URigHierarchy* HierarchyForSelection = Hierarchy->HierarchyForSelectionPtr.Get())
 	{
 		if(URigHierarchyController* ControllerForSelection = HierarchyForSelection->GetController())
@@ -93,7 +95,7 @@ bool URigHierarchyController::SelectElement(FRigElementKey InKey, bool bSelect,
 	{
 		if(const FCachedRigElement* Cache = Hierarchy->ElementKeyRedirector->Find(Key))
 		{
-			if(const_cast<FCachedRigElement*>(Cache)->UpdateCache(Hierarchy.Get()))
+			if(const_cast<FCachedRigElement*>(Cache)->UpdateCache(Hierarchy))
 			{
 				Key = Cache->GetKey();
 			}
@@ -144,6 +146,8 @@ bool URigHierarchyController::SetSelection(const TArray<FRigElementKey>& InKeys,
 		return false;
 	}
 
+	URigHierarchy* Hierarchy = GetHierarchy();
+
 	if(URigHierarchy* HierarchyForSelection = Hierarchy->HierarchyForSelectionPtr.Get())
 	{
 		if(URigHierarchyController* ControllerForSelection = HierarchyForSelection->GetController())
@@ -211,6 +215,8 @@ FRigElementKey URigHierarchyController::AddBone(FName InName, FRigElementKey InP
 		return FRigElementKey();
 	}
 
+	URigHierarchy* Hierarchy = GetHierarchy();
+
 #if WITH_EDITOR
 	TSharedPtr<FScopedTransaction> TransactionPtr;
 	if(bSetupUndo)
@@ -273,6 +279,8 @@ FRigElementKey URigHierarchyController::AddNull(FName InName, FRigElementKey InP
 		return FRigElementKey();
 	}
 
+	URigHierarchy* Hierarchy = GetHierarchy();
+
 #if WITH_EDITOR
 	TSharedPtr<FScopedTransaction> TransactionPtr;
 	if(bSetupUndo)
@@ -338,6 +346,8 @@ FRigElementKey URigHierarchyController::AddControl(
 		return FRigElementKey();
 	}
 
+	URigHierarchy* Hierarchy = GetHierarchy();
+
 #if WITH_EDITOR
 	TSharedPtr<FScopedTransaction> TransactionPtr;
 	if(bSetupUndo)
@@ -421,6 +431,8 @@ FRigElementKey URigHierarchyController::AddAnimationChannel(FName InName, FRigEl
 		return FRigElementKey();
 	}
 
+	URigHierarchy* Hierarchy = GetHierarchy();
+
 	if(const FRigControlElement* ParentControl = Hierarchy->Find<FRigControlElement>(InParentControl))
 	{
 		InSettings.AnimationType = ERigControlAnimationType::AnimationChannel;
@@ -440,6 +452,8 @@ FRigElementKey URigHierarchyController::AddCurve(FName InName, float InValue, bo
 		return FRigElementKey();
 	}
 
+	URigHierarchy* Hierarchy = GetHierarchy();
+
 #if WITH_EDITOR
 	TSharedPtr<FScopedTransaction> TransactionPtr;
 	if(bSetupUndo)
@@ -489,6 +503,8 @@ FRigElementKey URigHierarchyController::AddRigidBody(FName InName, FRigElementKe
 		return FRigElementKey();
 	}
 
+	URigHierarchy* Hierarchy = GetHierarchy();
+
 #if WITH_EDITOR
 	TSharedPtr<FScopedTransaction> TransactionPtr;
 	if(bSetupUndo)
@@ -528,7 +544,7 @@ FRigElementKey URigHierarchyController::AddRigidBody(FName InName, FRigElementKe
 	}
 #endif
 
-	Hierarchy->EnsureCacheValidity();
+	GetHierarchy()->EnsureCacheValidity();
 
 	return NewElement->Key;
 }
@@ -541,6 +557,8 @@ FRigElementKey URigHierarchyController::AddReference(FName InName, FRigElementKe
 		return FRigElementKey();
 	}
 
+	URigHierarchy* Hierarchy = GetHierarchy();
+
 #if WITH_EDITOR
 	TSharedPtr<FScopedTransaction> TransactionPtr;
 	if(bSetupUndo)
@@ -578,6 +596,8 @@ FRigElementKey URigHierarchyController::AddConnector(FName InName, FRigConnector
 		return FRigElementKey();
 	}
 
+	URigHierarchy* Hierarchy = GetHierarchy();
+
 	// only allow to add one primary connector
 	if(InSettings.Type == EConnectorType::Primary)
 	{
@@ -756,6 +776,8 @@ FRigControlSettings URigHierarchyController::GetControlSettings(FRigElementKey I
 		return FRigControlSettings();
 	}
 
+	URigHierarchy* Hierarchy = GetHierarchy();
+
 	FRigControlElement* ControlElement = Hierarchy->Find<FRigControlElement>(InKey);
 	if(ControlElement == nullptr)
 	{
@@ -772,6 +794,8 @@ bool URigHierarchyController::SetControlSettings(FRigElementKey InKey, FRigContr
 		return false;
 	}
 
+	URigHierarchy* Hierarchy = GetHierarchy();
+
 	FRigControlElement* ControlElement = Hierarchy->Find<FRigControlElement>(InKey);
 	if(ControlElement == nullptr)
 	{
@@ -827,6 +851,8 @@ TArray<FRigElementKey> URigHierarchyController::ImportBones(const FReferenceSkel
 	TArray<FRigElementKey> BonesToSelect;
 	TMap<FName, FName> BoneNameMap;
 
+	URigHierarchy* Hierarchy = GetHierarchy();
+
 	Hierarchy->ResetPoseToInitial();
 
 	const TArray<FMeshBoneInfo>& BoneInfos = InSkeleton.GetRefBoneInfo();
@@ -1073,7 +1099,9 @@ TArray<FRigElementKey> URigHierarchyController::ImportCurves(USkeleton* InSkelet
 		return Keys;
 	}
 
-	InSkeleton->ForEachCurveMetaData([this, InNameSpace, &Keys, bSetupUndo](const FName& InCurveName, const FCurveMetaData& InMetaData)
+	URigHierarchy* Hierarchy = GetHierarchy();
+
+	InSkeleton->ForEachCurveMetaData([this, Hierarchy, InNameSpace, &Keys, bSetupUndo](const FName& InCurveName, const FCurveMetaData& InMetaData)
 	{
 		FName Name = InCurveName;
 		if (!InNameSpace.IsNone())
@@ -1123,6 +1151,8 @@ FString URigHierarchyController::ExportSelectionToText() const
 	{
 		return FString();
 	}
+
+	URigHierarchy* Hierarchy = GetHierarchy();
 	return ExportToText(Hierarchy->GetSelectedKeys());
 }
 
@@ -1133,6 +1163,7 @@ FString URigHierarchyController::ExportToText(TArray<FRigElementKey> InKeys) con
 		return FString();
 	}
 
+	URigHierarchy* Hierarchy = GetHierarchy();
 	Hierarchy->ComputeAllTransforms();
 
 	// sort the keys by traversal order
@@ -1325,6 +1356,8 @@ TArray<FRigElementKey> URigHierarchyController::ImportFromText(FString InContent
 		ReportAndNotifyErrorf(TEXT("Some elements were not allowed to be pasted."));
 	}
 
+	URigHierarchy* Hierarchy = GetHierarchy();
+
 #if WITH_EDITOR
 	TSharedPtr<FScopedTransaction> TransactionPtr;
 	if(bSetupUndo)
@@ -1340,7 +1373,7 @@ TArray<FRigElementKey> URigHierarchyController::ImportFromText(FString InContent
 		KeyMap.Add(Element->GetKey(), Element->GetKey());
 	}
 
-	FRigHierarchyInteractionBracket InteractionBracket(Hierarchy.Get());
+	FRigHierarchyInteractionBracket InteractionBracket(Hierarchy);
 
 	for(const FRigHierarchyCopyPasteContentPerElement& PerElementData : Data.Elements)
 	{
@@ -1532,8 +1565,10 @@ TArray<FRigElementKey> URigHierarchyController::ImportFromText(FString InContent
 
 TArray<FRigElementKey> URigHierarchyController::ImportFromHierarchyContainer(const FRigHierarchyContainer& InContainer, bool bIsCopyAndPaste)
 {
+	URigHierarchy* Hierarchy = GetHierarchy();
+
 	TMap<FRigElementKey, FRigElementKey> KeyMap;;
-	FRigHierarchyInteractionBracket InteractionBracket(Hierarchy.Get());
+	FRigHierarchyInteractionBracket InteractionBracket(Hierarchy);
 
 	for(const FRigBone& Bone : InContainer.BoneHierarchy)
 	{
@@ -1669,6 +1704,8 @@ TArray<FRigElementKey> URigHierarchyController::ImportFromHierarchyContainer(con
 #if WITH_EDITOR
 TArray<FString> URigHierarchyController::GeneratePythonCommands()
 {
+	URigHierarchy* Hierarchy = GetHierarchy();
+
 	TArray<FString> Commands;
 	Hierarchy->Traverse([&](FRigBaseElement* Element, bool& bContinue)
 	{
@@ -1781,7 +1818,9 @@ TArray<FString> URigHierarchyController::GetAddControlPythonCommands(FRigControl
 			
 		Commands.Append(URigHierarchy::ControlSettingsToPythonCommands(Settings, SettingsStr));	
 	}
-		
+
+	URigHierarchy* Hierarchy = GetHierarchy();
+
 	FRigControlValue Value = Hierarchy->GetControlValue(Control->GetKey(), ERigControlValueType::Initial);
 	FString ValueStr = Value.ToPythonString(Settings.ControlType);
 	
@@ -1803,6 +1842,8 @@ TArray<FString> URigHierarchyController::GetAddControlPythonCommands(FRigControl
 
 TArray<FString> URigHierarchyController::GetAddCurvePythonCommands(FRigCurveElement* Curve) const
 {
+	URigHierarchy* Hierarchy = GetHierarchy();
+
 	// FRigElementKey AddCurve(FName InName, float InValue = 0.f, bool bSetupUndo = true);
 	return {FString::Printf(TEXT("hierarchy_controller.add_curve('%s', %f)"),
 		*Curve->GetName(),
@@ -1942,7 +1983,7 @@ void URigHierarchyController::Notify(ERigHierarchyNotification InNotifType, cons
 			return;
 		}
 	}	
-	Hierarchy->Notify(InNotifType, InElement);
+	GetHierarchy()->Notify(InNotifType, InElement);
 }
 
 void URigHierarchyController::HandleHierarchyModified(ERigHierarchyNotification InNotifType, URigHierarchy* InHierarchy, const FRigBaseElement* InElement) const
@@ -1952,7 +1993,7 @@ void URigHierarchyController::HandleHierarchyModified(ERigHierarchyNotification
 		return;
 	}
 	ensure(IsValid());
-	ensure(InHierarchy == Hierarchy);
+	ensure(InHierarchy == GetHierarchy());
 	ModifiedEvent.Broadcast(InNotifType, InHierarchy, InElement);
 }
 
@@ -1960,7 +2001,7 @@ bool URigHierarchyController::IsValid() const
 {
 	// If we're pending kill, it's fine for our Hierarchy to also be pending kill
 	const bool bPendingKillAcceptable = ::IsValid(this);
-	return Hierarchy.IsValid(bPendingKillAcceptable);
+	return WeakHierarchy.IsValid(bPendingKillAcceptable);
 }
 
 FName URigHierarchyController::GetSafeNewName(const FName& InDesiredName, ERigElementType InElementType) const
@@ -1974,13 +2015,15 @@ FName URigHierarchyController::GetSafeNewName(const FName& InDesiredName, ERigEl
 		Name.SetName(Name.GetName().Mid(Index + 1));
 	}
 	
-	return Hierarchy->GetSafeNewName(Name, InElementType).GetFName();
+	return GetHierarchy()->GetSafeNewName(Name, InElementType).GetFName();
 }
 
 int32 URigHierarchyController::AddElement(FRigBaseElement* InElementToAdd, FRigBaseElement* InFirstParent, bool bMaintainGlobalTransform, const FName& InDesiredName)
 {
 	ensure(IsValid());
 
+	URigHierarchy* Hierarchy = GetHierarchy();
+
 	InElementToAdd->CachedNameString.Reset();
 	InElementToAdd->SubIndex = Hierarchy->Num(InElementToAdd->Key.Type);
 	InElementToAdd->Index = Hierarchy->Elements.Add(InElementToAdd);
@@ -1992,19 +2035,6 @@ int32 URigHierarchyController::AddElement(FRigBaseElement* InElementToAdd, FRigB
 	FRigName DesiredName = InDesiredName;
 	URigHierarchy::SanitizeName(DesiredName);
 
-	// Short name of an element is ModuleName:ElementName (or ElementName if not in a module)
-	FString ModulePath, ElementName = InElementToAdd->GetName();
-	if (ElementName.Split(UModularRig::NamespaceSeparator, &ModulePath, &ElementName, ESearchCase::CaseSensitive, ESearchDir::FromEnd))
-	{
-		FString ModuleName = ModulePath, ModuleParentPath;
-		ModulePath.Split(UModularRig::NamespaceSeparator, &ModuleParentPath, &ModuleName, ESearchCase::CaseSensitive, ESearchDir::FromEnd);
-		const FString ShortName = FString::Printf(TEXT("%s:%s"), *ModuleName, *ElementName);
-		if(!ShortName.IsEmpty())
-		{
-			Hierarchy->SetNameMetadata(InElementToAdd->Key, URigHierarchy::ShortNameMetadataName, *ShortName);
-		}
-	}
-
 	if(!InDesiredName.IsNone() &&
 		!InElementToAdd->GetFName().IsEqual(DesiredName.GetFName(), ENameCase::CaseSensitive))
 	{
@@ -2024,7 +2054,17 @@ int32 URigHierarchyController::AddElement(FRigBaseElement* InElementToAdd, FRigB
 
 				if(Hierarchy->ElementKeyRedirector)
 				{
-					Hierarchy->ElementKeyRedirector->Add(FRigElementKey(DesiredName.GetFName(), InElementToAdd->Key.Type), InElementToAdd->Key, Hierarchy.Get());
+					Hierarchy->ElementKeyRedirector->Add(FRigElementKey(DesiredName.GetFName(), InElementToAdd->Key.Type), InElementToAdd->Key, Hierarchy);
+				}
+			}
+
+			// also try to find the live module and retrieve the shortname from that to update the element's short name
+			if(const FRigModuleInstance* Module = CRContext.GetRigModuleInstance())
+			{
+				const FString ModuleShortName = Module->GetShortName();
+				if(!ModuleShortName.IsEmpty())
+				{
+					Hierarchy->SetNameMetadata(InElementToAdd->Key, URigHierarchy::ShortModuleNameMetadataName, *ModuleShortName);
 				}
 			}
 		}
@@ -2056,6 +2096,8 @@ bool URigHierarchyController::RemoveElement(FRigElementKey InElement, bool bSetu
 		return false;
 	}
 
+	URigHierarchy* Hierarchy = GetHierarchy();
+
 	FRigBaseElement* Element = Hierarchy->Find(InElement);
 	if(Element == nullptr)
 	{
@@ -2103,6 +2145,8 @@ bool URigHierarchyController::RemoveElement(FRigBaseElement* InElement)
 		return false;
 	}
 
+	URigHierarchy* Hierarchy = GetHierarchy();
+
 	// make sure this element is part of this hierarchy
 	ensure(Hierarchy->FindChecked(InElement->Key) == InElement);
 	ensure(InElement->OwnedInstances == 1);
@@ -2223,6 +2267,8 @@ FRigElementKey URigHierarchyController::RenameElement(FRigElementKey InElement,
 		return FRigElementKey();
 	}
 
+	URigHierarchy* Hierarchy = GetHierarchy();
+
 	FRigBaseElement* Element = Hierarchy->Find(InElement);
 	if(Element == nullptr)
 	{
@@ -2277,6 +2323,8 @@ bool URigHierarchyController::ReorderElement(FRigElementKey InElement, int32 InI
 		return false;
 	}
 
+	URigHierarchy* Hierarchy = GetHierarchy();
+
 	FRigBaseElement* Element = Hierarchy->Find(InElement);
 	if(Element == nullptr)
 	{
@@ -2326,6 +2374,8 @@ FName URigHierarchyController::SetDisplayName(FRigElementKey InControl, FName In
 		return NAME_None;
 	}
 
+	URigHierarchy* Hierarchy = GetHierarchy();
+
 	FRigControlElement* ControlElement = Hierarchy->Find<FRigControlElement>(InControl);
 	if(ControlElement == nullptr)
 	{
@@ -2383,6 +2433,8 @@ bool URigHierarchyController::RenameElement(FRigBaseElement* InElement, const FN
 
 	const FRigElementKey OldKey = InElement->GetKey();
 
+	URigHierarchy* Hierarchy = GetHierarchy();
+
 	// deselect the key that no longer exists
 	// no need to trigger a reselect since we always clear selection after rename
 	const bool bWasSelected = Hierarchy->IsSelected(InElement); 
@@ -2461,6 +2513,8 @@ bool URigHierarchyController::ReorderElement(FRigBaseElement* InElement, int32 I
 
 	InIndex = FMath::Max<int32>(InIndex, 0);
 
+	URigHierarchy* Hierarchy = GetHierarchy();
+
 	TArray<FRigBaseElement*> LocalElements;
 	if(const FRigBaseElement* ParentElement = Hierarchy->GetFirstParent(InElement))
 	{
@@ -2526,6 +2580,8 @@ FName URigHierarchyController::SetDisplayName(FRigControlElement* InControlEleme
 		return NAME_None;
 	}
 
+	URigHierarchy* Hierarchy = GetHierarchy();
+
 	FRigElementKey ParentElementKey;
 	if(const FRigBaseElement* ParentElement = Hierarchy->GetFirstParent(InControlElement))
 	{
@@ -2575,6 +2631,8 @@ bool URigHierarchyController::AddParent(FRigElementKey InChild, FRigElementKey I
 		return false;
 	}
 
+	URigHierarchy* Hierarchy = GetHierarchy();
+
 	FRigBaseElement* Child = Hierarchy->Find(InChild);
 	if(Child == nullptr)
 	{
@@ -2674,6 +2732,8 @@ bool URigHierarchyController::AddParent(FRigBaseElement* InChild, FRigBaseElemen
 		}
 	}
 
+	URigHierarchy* Hierarchy = GetHierarchy();
+
 	if(Hierarchy->IsParentedTo(InParent, InChild))
 	{
 		ReportErrorf(TEXT("Cannot parent '%s' to '%s' - would cause a cycle."), *InChild->Key.ToString(), *InParent->Key.ToString());
@@ -2805,6 +2865,8 @@ bool URigHierarchyController::RemoveParent(FRigElementKey InChild, FRigElementKe
 		return false;
 	}
 
+	URigHierarchy* Hierarchy = GetHierarchy();
+
 	FRigBaseElement* Child = Hierarchy->Find(InChild);
 	if(Child == nullptr)
 	{
@@ -2867,6 +2929,8 @@ bool URigHierarchyController::RemoveParent(FRigBaseElement* InChild, FRigBaseEle
 		return false;
 	}
 
+	URigHierarchy* Hierarchy = GetHierarchy();
+
 	// single parent children can't be parented multiple times
 	if(FRigSingleParentElement* SingleParentElement = Cast<FRigSingleParentElement>(InChild))
 	{
@@ -2996,6 +3060,8 @@ bool URigHierarchyController::RemoveAllParents(FRigElementKey InChild, bool bMai
 		return false;
 	}
 
+	URigHierarchy* Hierarchy = GetHierarchy();
+
 	FRigBaseElement* Child = Hierarchy->Find(InChild);
 	if(Child == nullptr)
 	{
@@ -3068,6 +3134,8 @@ bool URigHierarchyController::SetParent(FRigElementKey InChild, FRigElementKey I
 		return false;
 	}
 
+	URigHierarchy* Hierarchy = GetHierarchy();
+
 	FRigBaseElement* Child = Hierarchy->Find(InChild);
 	if(Child == nullptr)
 	{
@@ -3153,14 +3221,15 @@ TArray<FRigElementKey> URigHierarchyController::DuplicateElements(TArray<FRigEle
 	}
 #endif
 
-	Hierarchy->EnsureCacheValidity();
+	GetHierarchy()->EnsureCacheValidity();
 	
 	return Result;
 }
 
 TArray<FRigElementKey> URigHierarchyController::MirrorElements(TArray<FRigElementKey> InKeys, FRigVMMirrorSettings InSettings, bool bSelectNewElements, bool bSetupUndo, bool bPrintPythonCommands)
 {
-	FRigHierarchyInteractionBracket InteractionBracket(Hierarchy.Get());
+	URigHierarchy* Hierarchy = GetHierarchy();
+	FRigHierarchyInteractionBracket InteractionBracket(Hierarchy);
 
 	TArray<FRigElementKey> OriginalKeys = Hierarchy->SortKeys(InKeys);
 	TArray<FRigElementKey> DuplicatedKeys = DuplicateElements(OriginalKeys, bSelectNewElements, bSetupUndo);
@@ -3346,9 +3415,9 @@ void URigHierarchyController::ReportWarning(const FString& InMessage) const
 	}
 
 	FString Message = InMessage;
-	if (Hierarchy.IsValid())
+	if (const URigHierarchy* Hierarchy = GetHierarchy())
 	{
-		if (UPackage* Package = Cast<UPackage>(Hierarchy->GetOutermost()))
+		if (const UPackage* Package = Cast<UPackage>(Hierarchy->GetOutermost()))
 		{
 			Message = FString::Printf(TEXT("%s : %s"), *Package->GetPathName(), *InMessage);
 		}
@@ -3371,9 +3440,9 @@ void URigHierarchyController::ReportError(const FString& InMessage) const
 	}
 
 	FString Message = InMessage;
-	if (Hierarchy.IsValid())
+	if (const URigHierarchy* Hierarchy = GetHierarchy())
 	{
-		if (UPackage* Package = Cast<UPackage>(Hierarchy->GetOutermost()))
+		if (const UPackage* Package = Cast<UPackage>(Hierarchy->GetOutermost()))
 		{
 			Message = FString::Printf(TEXT("%s : %s"), *Package->GetPathName(), *InMessage);
 		}
diff --git a/Engine/Plugins/Animation/ControlRig/Source/ControlRig/Private/Sequencer/MovieSceneControlRigParameterSection.cpp b/Engine/Plugins/Animation/ControlRig/Source/ControlRig/Private/Sequencer/MovieSceneControlRigParameterSection.cpp
--- a/Engine/Plugins/Animation/ControlRig/Source/ControlRig/Private/Sequencer/MovieSceneControlRigParameterSection.cpp
+++ b/Engine/Plugins/Animation/ControlRig/Source/ControlRig/Private/Sequencer/MovieSceneControlRigParameterSection.cpp
@@ -1633,7 +1633,7 @@ EMovieSceneChannelProxyType UMovieSceneControlRigParameterSection::CacheChannelP
 				if(const FRigControlElement* ParentControlElement = Cast<FRigControlElement>(Hierarchy->GetFirstParent(ControlElement)))
 				{
 					ParentControlName = ParentControlElement->GetFName();
-					Group = FText::FromName(ParentControlElement->GetDisplayName());
+					Group = Hierarchy->GetDisplayNameForUI(ParentControlElement);
 				}
 			}
 
@@ -1652,7 +1652,7 @@ EMovieSceneChannelProxyType UMovieSceneControlRigParameterSection::CacheChannelP
 							if (Group.IsEmpty())
 							{
 								ControlChannelMap.Add(Scalar.ParameterName, FChannelMapInfo(ControlIndex, TotalIndex, FloatChannelIndex, INDEX_NONE, NAME_None, MaskIndex, CategoryIndex));
-								Group = FText::FromName(ControlElement->GetDisplayName());
+								Group = Hierarchy->GetDisplayNameForUI(ControlElement);
 								if (bEnabled)
 								{
 									++CategoryIndex;
@@ -1666,7 +1666,7 @@ EMovieSceneChannelProxyType UMovieSceneControlRigParameterSection::CacheChannelP
 							}
 
 							FParameterFloatChannelEditorData EditorData(ControlRig, Scalar.ParameterName, bEnabled, Group, TotalIndex);
-							EditorData.MetaData.DisplayText = FText::FromName(ControlElement->GetDisplayName());
+							EditorData.MetaData.DisplayText = Hierarchy->GetDisplayNameForUI(ControlElement);
 							Channels.Add(Scalar.ParameterCurve, EditorData.MetaData, EditorData.ExternalValues);
 							FloatChannelIndex += 1;
 							TotalIndex += 1;
@@ -1685,7 +1685,7 @@ EMovieSceneChannelProxyType UMovieSceneControlRigParameterSection::CacheChannelP
 							if (Group.IsEmpty())
 							{
 								ControlChannelMap.Add(Bool.ParameterName, FChannelMapInfo(ControlIndex, TotalIndex, BoolChannelIndex, INDEX_NONE,BoolChannelTypeName,MaskIndex, CategoryIndex));
-								Group = FText::FromName(ControlElement->GetDisplayName());
+								Group = Hierarchy->GetDisplayNameForUI(ControlElement);
 								if (bEnabled)
 								{
 									++CategoryIndex;
@@ -1699,7 +1699,7 @@ EMovieSceneChannelProxyType UMovieSceneControlRigParameterSection::CacheChannelP
 							}
 
 							FMovieSceneChannelMetaData MetaData(Bool.ParameterName, Group, Group, bEnabled);
-							MetaData.DisplayText = FText::FromName(ControlElement->GetDisplayName());
+							MetaData.DisplayText = Hierarchy->GetDisplayNameForUI(ControlElement);
 							MetaData.SortOrder = TotalIndex++;
 							BoolChannelIndex += 1;
 							ControlIndex += 1;
@@ -1722,7 +1722,7 @@ EMovieSceneChannelProxyType UMovieSceneControlRigParameterSection::CacheChannelP
 								if (Group.IsEmpty())
 								{
 									ControlChannelMap.Add(Enum.ParameterName, FChannelMapInfo(ControlIndex, TotalIndex, EnumChannelIndex,INDEX_NONE, EnumChannelTypeName,MaskIndex, CategoryIndex));
-									Group = FText::FromName(ControlElement->GetDisplayName());
+									Group = Hierarchy->GetDisplayNameForUI(ControlElement);
 									if (bEnabled)
 									{
 										++CategoryIndex;
@@ -1736,7 +1736,7 @@ EMovieSceneChannelProxyType UMovieSceneControlRigParameterSection::CacheChannelP
 								}
 
 								FMovieSceneChannelMetaData MetaData(Enum.ParameterName, Group, Group, bEnabled);
-								MetaData.DisplayText = FText::FromName(ControlElement->GetDisplayName());
+								MetaData.DisplayText = Hierarchy->GetDisplayNameForUI(ControlElement);
 								EnumChannelIndex += 1;
 								ControlIndex += 1;
 								MetaData.SortOrder = TotalIndex++;
@@ -1756,7 +1756,7 @@ EMovieSceneChannelProxyType UMovieSceneControlRigParameterSection::CacheChannelP
 								if (Group.IsEmpty())
 								{
 									ControlChannelMap.Add(Integer.ParameterName, FChannelMapInfo(ControlIndex, TotalIndex, IntegerChannelIndex,INDEX_NONE,IntegerChannelTypeName,MaskIndex, CategoryIndex));
-									Group = FText::FromName(ControlElement->GetDisplayName());
+									Group = Hierarchy->GetDisplayNameForUI(ControlElement);
 									if (bEnabled)
 									{
 										++CategoryIndex;
@@ -1770,7 +1770,7 @@ EMovieSceneChannelProxyType UMovieSceneControlRigParameterSection::CacheChannelP
 								}
 
 								FMovieSceneChannelMetaData MetaData(Integer.ParameterName, Group, Group, bEnabled);
-								MetaData.DisplayText = FText::FromName(ControlElement->GetDisplayName());
+								MetaData.DisplayText = Hierarchy->GetDisplayNameForUI(ControlElement);
 								IntegerChannelIndex += 1;
 								ControlIndex += 1;
 								MetaData.SortOrder = TotalIndex++;
@@ -1797,7 +1797,7 @@ EMovieSceneChannelProxyType UMovieSceneControlRigParameterSection::CacheChannelP
 								{
 									++CategoryIndex;
 								}
-								Group = FText::FromName(ControlElement->GetDisplayName());
+								Group = Hierarchy->GetDisplayNameForUI(ControlElement);
 							}
 							else
 							{
@@ -1833,7 +1833,7 @@ EMovieSceneChannelProxyType UMovieSceneControlRigParameterSection::CacheChannelP
 								{
 									++CategoryIndex;
 								}
-								Group = FText::FromName(ControlElement->GetDisplayName());
+								Group = Hierarchy->GetDisplayNameForUI(ControlElement);
 							}
 							else
 							{
@@ -1913,7 +1913,7 @@ EMovieSceneChannelProxyType UMovieSceneControlRigParameterSection::CacheChannelP
 								{
 									++CategoryIndex;
 								}
-								Group = FText::FromName(ControlElement->GetDisplayName());
+								Group = Hierarchy->GetDisplayNameForUI(ControlElement);
 							}
 							else
 							{
diff --git a/Engine/Plugins/Animation/ControlRig/Source/ControlRig/Private/Units/Execution/RigUnit_RigModules.cpp b/Engine/Plugins/Animation/ControlRig/Source/ControlRig/Private/Units/Execution/RigUnit_RigModules.cpp
--- a/Engine/Plugins/Animation/ControlRig/Source/ControlRig/Private/Units/Execution/RigUnit_RigModules.cpp
+++ b/Engine/Plugins/Animation/ControlRig/Source/ControlRig/Private/Units/Execution/RigUnit_RigModules.cpp
@@ -37,7 +37,7 @@ FRigUnit_GetItemShortName_Execute()
 
 	if(const URigHierarchy* Hierarchy = ExecuteContext.Hierarchy)
 	{
-		ShortName = Hierarchy->GetNameMetadata(Item, URigHierarchy::ShortNameMetadataName, NAME_None);
+		ShortName = *Hierarchy->GetDisplayNameForUI(Item).ToString();
 	}
 
 	if(ShortName.IsNone())
diff --git a/Engine/Plugins/Animation/ControlRig/Source/ControlRig/Public/Rigs/RigHierarchy.h b/Engine/Plugins/Animation/ControlRig/Source/ControlRig/Public/Rigs/RigHierarchy.h
--- a/Engine/Plugins/Animation/ControlRig/Source/ControlRig/Public/Rigs/RigHierarchy.h
+++ b/Engine/Plugins/Animation/ControlRig/Source/ControlRig/Public/Rigs/RigHierarchy.h
@@ -165,7 +165,7 @@ public:
 	typedef TPair<int32, TArray<int32>> TElementDependencyMapPair;
 	typedef TTuple<int32, int32, int32, ERigTransformType::Type> TInstructionSliceElement;
 	inline static const FName TagMetadataName = TEXT("Tags");
-	inline static const FName ShortNameMetadataName = TEXT("ShortName");
+	inline static const FName ShortModuleNameMetadataName = TEXT("ShortModuleName");
 	inline static const FName DesiredNameMetadataName = TEXT("DesiredName");
 	inline static const FName DesiredKeyMetadataName = TEXT("DesiredKey");
 	inline static const FName ModuleMetadataName = TEXT("Module");
@@ -1659,6 +1659,12 @@ public:
 		return Result;
 	}
 
+	/**
+	 * Returns the two name sections with the right namespace separator
+	 */
+	static FString JoinNameSpace(const FString& InLeft, const FString& InRight);
+	static FRigName JoinNameSpace(const FRigName& InLeft, const FRigName& InRight);
+
 	/**
 	 * Returns the max allowed length for a name within the hierarchy.
 	 * @return Returns the max allowed length for a name within the hierarchy.
@@ -1712,6 +1718,12 @@ public:
 	 */
 	FRigName GetSafeNewDisplayName(const FRigElementKey& InParentElement, const FRigName& InPotentialNewDisplayName) const;
 
+	/**
+	 * Returns the display label for an element to be used for the UI
+	 */
+	FText GetDisplayNameForUI(const FRigBaseElement* InElement, bool bIncludeNameSpace = true) const;
+	FText GetDisplayNameForUI(const FRigElementKey& InKey, bool bIncludeNameSpace = true) const;
+
 	/**
 	 * Returns the modified event, which can be used to 
 	 * subscribe to topological changes happening within the hierarchy.
diff --git a/Engine/Plugins/Animation/ControlRig/Source/ControlRig/Public/Rigs/RigHierarchyController.h b/Engine/Plugins/Animation/ControlRig/Source/ControlRig/Public/Rigs/RigHierarchyController.h
--- a/Engine/Plugins/Animation/ControlRig/Source/ControlRig/Public/Rigs/RigHierarchyController.h
+++ b/Engine/Plugins/Animation/ControlRig/Source/ControlRig/Public/Rigs/RigHierarchyController.h
@@ -18,7 +18,7 @@ public:
 
 	URigHierarchyController()
 	: bReportWarningsAndErrors(true)
-	, Hierarchy(nullptr)
+	, WeakHierarchy(nullptr)
 	, bSuspendAllNotifications(false)
 	, bSuspendSelectionNotifications(false)
 	, bSuspendPythonPrinting(false)
@@ -34,7 +34,7 @@ public:
 	UFUNCTION(BlueprintCallable, Category = URigHierarchyController)
 	URigHierarchy* GetHierarchy() const
 	{
-		return Hierarchy.Get();
+		return WeakHierarchy.Get();
 	}
 
 	// Sets the hierarchy currently linked to this controller
@@ -631,7 +631,7 @@ public:
 private:
 
 	UPROPERTY(transient)
-	TWeakObjectPtr<URigHierarchy> Hierarchy;
+	TWeakObjectPtr<URigHierarchy> WeakHierarchy;
 
 	FRigHierarchyModifiedEvent ModifiedEvent;
 	void Notify(ERigHierarchyNotification InNotifType, const FRigBaseElement* InElement);
diff --git a/Engine/Plugins/Animation/ControlRig/Source/ControlRig/Public/Units/RigUnitContext.h b/Engine/Plugins/Animation/ControlRig/Source/ControlRig/Public/Units/RigUnitContext.h
--- a/Engine/Plugins/Animation/ControlRig/Source/ControlRig/Public/Units/RigUnitContext.h
+++ b/Engine/Plugins/Animation/ControlRig/Source/ControlRig/Public/Units/RigUnitContext.h
@@ -13,6 +13,7 @@
 
 class UControlRig;
 class UControlRigShapeLibrary;
+struct FRigModuleInstance;
 
 /**
  * The type of interaction happening on a rig
@@ -169,6 +170,14 @@ public:
 		return RigModuleNameSpace;
 	}
 
+	/**
+	 * Returns the module this unit is running inside of (or nullptr)
+	 */
+	const FRigModuleInstance* GetRigModuleInstance() const
+	{
+		return RigModuleInstance;
+	}
+
 	/**
 	 * Adapts a metadata name according to rig module namespace.
 	 */
@@ -189,6 +198,7 @@ public:
 private:
 	FString RigModuleNameSpace;
 	uint32 RigModuleNameSpaceHash;
+	const FRigModuleInstance* RigModuleInstance;
 
 	friend class FControlRigExecuteContextRigModuleGuard;
 	friend class UModularRig;
diff --git a/Engine/Plugins/Animation/ControlRig/Source/ControlRigDeveloper/Private/ControlRigBlueprint.cpp b/Engine/Plugins/Animation/ControlRig/Source/ControlRigDeveloper/Private/ControlRigBlueprint.cpp
--- a/Engine/Plugins/Animation/ControlRig/Source/ControlRigDeveloper/Private/ControlRigBlueprint.cpp
+++ b/Engine/Plugins/Animation/ControlRig/Source/ControlRigDeveloper/Private/ControlRigBlueprint.cpp
@@ -2491,13 +2491,13 @@ void UControlRigBlueprint::HandleRigModulesModified(EModularRigNotification InNo
 						FString OldNamespace;
 						if (InNotification == EModularRigNotification::ModuleRenamed)
 						{
-							OldNamespace = (InModule->ParentPath.IsEmpty()) ? *InModule->PreviousName.ToString() : FString::Printf(TEXT("%s:%s"), *InModule->ParentPath, *InModule->PreviousName.ToString());
+							OldNamespace = (InModule->ParentPath.IsEmpty()) ? *InModule->PreviousName.ToString() : URigHierarchy::JoinNameSpace(InModule->ParentPath, InModule->PreviousName.ToString());
 						}
 						else if(InNotification == EModularRigNotification::ModuleReparented)
 						{
-							OldNamespace = (InModule->PreviousParentPath.IsEmpty()) ? *InModule->Name.ToString() : FString::Printf(TEXT("%s:%s"), *InModule->PreviousParentPath, *InModule->Name.ToString());
+							OldNamespace = (InModule->PreviousParentPath.IsEmpty()) ? *InModule->Name.ToString() : URigHierarchy::JoinNameSpace(InModule->PreviousParentPath, InModule->Name.ToString());
 						}
-						FString NewNamespace = (InModule->ParentPath.IsEmpty()) ? *InModule->Name.ToString() : FString::Printf(TEXT("%s:%s"), *InModule->ParentPath, *InModule->Name.ToString());
+						FString NewNamespace = (InModule->ParentPath.IsEmpty()) ? *InModule->Name.ToString() : URigHierarchy::JoinNameSpace(InModule->ParentPath, InModule->Name.ToString());
 						OldNamespace.Append(UModularRig::NamespaceSeparator);
 						NewNamespace.Append(UModularRig::NamespaceSeparator);
 						
diff --git a/Engine/Plugins/Animation/ControlRig/Source/ControlRigEditor/Private/ControlRigModuleDetails.cpp b/Engine/Plugins/Animation/ControlRig/Source/ControlRigEditor/Private/ControlRigModuleDetails.cpp
--- a/Engine/Plugins/Animation/ControlRig/Source/ControlRigEditor/Private/ControlRigModuleDetails.cpp
+++ b/Engine/Plugins/Animation/ControlRig/Source/ControlRigEditor/Private/ControlRigModuleDetails.cpp
@@ -704,7 +704,7 @@ void FRigModuleInstanceDetails::OnElementNameChanged(TSharedPtr<FString> InItem,
 
 				if(const FRigModuleInstance* Module = Info.GetModule())
 				{
-					FRigElementKey NamespacedConnector(*FString::Printf(TEXT("%s:%s"), *Module->GetPath(), *Connector.Name.ToString()), ERigElementType::Connector);
+					FRigElementKey NamespacedConnector(*URigHierarchy::JoinNameSpace(Module->GetPath(), Connector.Name.ToString()), ERigElementType::Connector);
 					Controller->ConnectConnectorToElement(NamespacedConnector, *TargetKey);
 				}
 			}
@@ -727,7 +727,7 @@ void FRigModuleInstanceDetails::OnElementTypeChanged(ERigElementType InElementTy
 
 				if(const FRigModuleInstance* Module = Info.GetModule())
 				{
-					FRigElementKey NamespacedConnector(*FString::Printf(TEXT("%s:%s"), *Module->GetPath(), *Connector.Name.ToString()), ERigElementType::Connector);
+					FRigElementKey NamespacedConnector(*URigHierarchy::JoinNameSpace(Module->GetPath(), Connector.Name.ToString()), ERigElementType::Connector);
 					Controller->ConnectConnectorToElement(NamespacedConnector, *TargetKey);
 				}
 			}
@@ -753,7 +753,7 @@ FReply FRigModuleInstanceDetails::OnGetSelectedClicked(FRigElementKey Connector)
 		{
 			if(const FRigModuleInstance* Module = PerModuleInfos[0].GetModule())
 			{
-				FRigElementKey NamespacedConnector(*FString::Printf(TEXT("%s:%s"), *Module->GetPath(), *Connector.Name.ToString()), ERigElementType::Connector);
+				FRigElementKey NamespacedConnector(*URigHierarchy::JoinNameSpace(Module->GetPath(), Connector.Name.ToString()), ERigElementType::Connector);
 				Blueprint->GetModularRigController()->ConnectConnectorToElement(NamespacedConnector, Selected[0]);
 			}
 		}
diff --git a/Engine/Plugins/Animation/ControlRig/Source/ControlRigEditor/Private/Editor/ControlRigEditor.cpp b/Engine/Plugins/Animation/ControlRig/Source/ControlRigEditor/Private/Editor/ControlRigEditor.cpp
--- a/Engine/Plugins/Animation/ControlRig/Source/ControlRigEditor/Private/Editor/ControlRigEditor.cpp
+++ b/Engine/Plugins/Animation/ControlRig/Source/ControlRigEditor/Private/Editor/ControlRigEditor.cpp
@@ -151,12 +151,9 @@ FControlRigEditor::~FControlRigEditor()
 			EditMode->OnEditorClosed();
 		}
 
-		if (RigBlueprint->IsModularRig())
-		{
-			RigBlueprint->OnSetObjectBeingDebugged().RemoveAll(&SchematicModel);
-			RigBlueprint->OnHierarchyModified().RemoveAll(&SchematicModel);
-			RigBlueprint->GetModularRigController()->OnModified().RemoveAll(&SchematicModel);
-		}
+		RigBlueprint->OnSetObjectBeingDebugged().RemoveAll(&SchematicModel);
+		RigBlueprint->OnHierarchyModified().RemoveAll(&SchematicModel);
+		RigBlueprint->GetModularRigController()->OnModified().RemoveAll(&SchematicModel);
 
 		RigBlueprint->OnRigTypeChanged().RemoveAll(this);
 		if (RigBlueprint->IsModularRig())
@@ -3764,11 +3761,11 @@ void FControlRigEditor::HandleModularRigModified(EModularRigNotification InNotif
 			FString OldPath;
 			if (InNotification == EModularRigNotification::ModuleRenamed)
 			{
-				OldPath = FString::Printf(TEXT("%s:%s"), *InModule->ParentPath, *InModule->PreviousName.ToString());
+				OldPath = URigHierarchy::JoinNameSpace(InModule->ParentPath, InModule->PreviousName.ToString());
 			}
 			else
 			{
-				OldPath = FString::Printf(TEXT("%s:%s"), *InModule->PreviousParentPath, *InModule->Name.ToString());
+				OldPath = URigHierarchy::JoinNameSpace(InModule->PreviousParentPath, InModule->Name.ToString());
 			}
 			ModulesSelected.Remove(OldPath);
 			ModulesSelected.Add(InModule->GetPath());
diff --git a/Engine/Plugins/Animation/ControlRig/Source/ControlRigEditor/Private/Editor/SRigHierarchyTreeView.cpp b/Engine/Plugins/Animation/ControlRig/Source/ControlRigEditor/Private/Editor/SRigHierarchyTreeView.cpp
--- a/Engine/Plugins/Animation/ControlRig/Source/ControlRigEditor/Private/Editor/SRigHierarchyTreeView.cpp
+++ b/Engine/Plugins/Animation/ControlRig/Source/ControlRigEditor/Private/Editor/SRigHierarchyTreeView.cpp
@@ -46,7 +46,7 @@ FRigTreeElement::FRigTreeElement(const FRigElementKey& InKey, TWeakPtr<SRigHiera
 	{
 		if(const URigHierarchy* Hierarchy = InTreeView.Pin()->GetRigTreeDelegates().GetHierarchy())
 		{
-			ShortName = Hierarchy->GetNameMetadata(InKey, URigHierarchy::ShortNameMetadataName, ShortName);
+			ShortName = *Hierarchy->GetDisplayNameForUI(InKey).ToString();
 			
 			const FRigTreeDisplaySettings& Settings = InTreeView.Pin()->GetRigTreeDelegates().GetDisplaySettings();
 			RefreshDisplaySettings(Hierarchy, Settings);
@@ -519,7 +519,7 @@ bool SRigHierarchyTreeView::AddElement(const FRigBaseElement* InElement)
 								FName Name = ConnectorKey.Name;
 								if (GetRigTreeDelegates().GetDisplaySettings().bUseShortName)
 								{
-									Name = Hierarchy->GetNameMetadata(ConnectorKey, URigHierarchy::ShortNameMetadataName, ConnectorKey.Name);
+									Name = *Hierarchy->GetDisplayNameForUI(ConnectorKey).ToString();
 								}
 								TagArguments.Text(FText::FromName(Name));
 								TagArguments.TooltipText(FText::FromName(ConnectorKey.Name));
