Sequencer: Constraint Baking: Fix crash baking onto rotators/position/scale controls

diff --git a/Engine/Plugins/Animation/ControlRig/Source/ControlRigEditor/Private/Tools/ConstraintBaker.cpp b/Engine/Plugins/Animation/ControlRig/Source/ControlRigEditor/Private/Tools/ConstraintBaker.cpp
--- a/Engine/Plugins/Animation/ControlRig/Source/ControlRigEditor/Private/Tools/ConstraintBaker.cpp
+++ b/Engine/Plugins/Animation/ControlRig/Source/ControlRigEditor/Private/Tools/ConstraintBaker.cpp
@@ -335,8 +335,15 @@ void FConstraintBaker::Bake(UWorld* InWorld,
 	// remove transform keys in active sections if baking all frames
 	if (InSettings.IsSet() && InSettings->BakingKeySettings == EBakingKeySettings::AllFrames)
 	{
+		EMovieSceneTransformChannel ChannelsToKey = InConstraint->GetChannelsToKey();
+		//if we aren't doing a full transform change ChannelsToKey to make sure it's the ones present
+		//assuming if it's the wrong mix(like rotator with a position constraint) we wouldn't be here anyway
+		if (FloatTransformChannels.Num() == 3 || DoubleTransformChannels.Num() == 3)
+		{
+			ChannelsToKey = EMovieSceneTransformChannel::Translation;
+		}
 		DeleteTransformKeysInActiveRanges(
-			FloatTransformChannels, DoubleTransformChannels, InConstraint->GetChannelsToKey(),
+			FloatTransformChannels, DoubleTransformChannels, ChannelsToKey,
 			FramesToBake, ConstraintFrames, PreviousConstraintValues);
 	}
 
