Control Rig: Fix nonunity error
#rnx

diff --git a/Engine/Plugins/Animation/ControlRig/Source/ControlRigEditor/Private/Editor/ControlRigSchematicModel.h b/Engine/Plugins/Animation/ControlRig/Source/ControlRigEditor/Private/Editor/ControlRigSchematicModel.h
--- a/Engine/Plugins/Animation/ControlRig/Source/ControlRigEditor/Private/Editor/ControlRigSchematicModel.h
+++ b/Engine/Plugins/Animation/ControlRig/Source/ControlRigEditor/Private/Editor/ControlRigSchematicModel.h
@@ -5,6 +5,7 @@
 #include "ModularRigModel.h"
 #include "SchematicGraphPanel/SSchematicGraphPanel.h"
 #include "Rigs/RigHierarchyDefines.h"
+#include "ControlRigBlueprint.h"
 
 class FControlRigEditor;
 class UControlRig;
