Control Rig: Fix axes not showing when display axes is ok
#rnx

diff --git a/Engine/Plugins/Animation/ControlRig/Source/ControlRigEditor/Private/EditMode/ControlRigEditMode.cpp b/Engine/Plugins/Animation/ControlRig/Source/ControlRigEditor/Private/EditMode/ControlRigEditMode.cpp
--- a/Engine/Plugins/Animation/ControlRig/Source/ControlRigEditor/Private/EditMode/ControlRigEditMode.cpp
+++ b/Engine/Plugins/Animation/ControlRig/Source/ControlRigEditor/Private/EditMode/ControlRigEditMode.cpp
@@ -881,7 +881,8 @@ void FControlRigEditMode::Render(const FSceneView* View, FViewport* Viewport, FP
 						});
 				}
 
-				const bool bIsAssetEditor = Viewport->GetClient()->GetWorld()->IsPreviewWorld();
+				EWorldType::Type WorldType = Viewport->GetClient()->GetWorld()->WorldType;
+				const bool bIsAssetEditor = WorldType == EWorldType::Editor || WorldType == EWorldType::EditorPreview;
 				if (bIsAssetEditor && (Settings->bDisplayNulls || ControlRig->IsConstructionModeEnabled()))
 				{
 					TArray<FTransform> SpaceTransforms;
@@ -905,21 +906,18 @@ void FControlRigEditMode::Render(const FSceneView* View, FViewport* Viewport, FP
 
 				if (bIsAssetEditor && (Settings->bDisplayAxesOnSelection && Settings->AxisScale > SMALL_NUMBER))
 				{
-					if (ControlRig->GetWorld() && ControlRig->GetWorld()->IsPreviewWorld())
-					{
-						TArray<FRigElementKey> SelectedRigElements = GetSelectedRigElements(ControlRig);
-						const float Scale = Settings->AxisScale;
-						PDI->AddReserveLines(SDPG_Foreground, SelectedRigElements.Num() * 3);
+					TArray<FRigElementKey> SelectedRigElements = GetSelectedRigElements(ControlRig);
+					const float Scale = Settings->AxisScale;
+					PDI->AddReserveLines(SDPG_Foreground, SelectedRigElements.Num() * 3);
 
-						for (const FRigElementKey& SelectedElement : SelectedRigElements)
-						{
-							FTransform ElementTransform = Hierarchy->GetGlobalTransform(SelectedElement);
-							ElementTransform = ElementTransform * ComponentTransform;
+					for (const FRigElementKey& SelectedElement : SelectedRigElements)
+					{
+						FTransform ElementTransform = Hierarchy->GetGlobalTransform(SelectedElement);
+						ElementTransform = ElementTransform * ComponentTransform;
 
-							PDI->DrawLine(ElementTransform.GetTranslation(), ElementTransform.TransformPosition(FVector(Scale, 0.f, 0.f)), FLinearColor::Red, SDPG_Foreground);
-							PDI->DrawLine(ElementTransform.GetTranslation(), ElementTransform.TransformPosition(FVector(0.f, Scale, 0.f)), FLinearColor::Green, SDPG_Foreground);
-							PDI->DrawLine(ElementTransform.GetTranslation(), ElementTransform.TransformPosition(FVector(0.f, 0.f, Scale)), FLinearColor::Blue, SDPG_Foreground);
-						}
+						PDI->DrawLine(ElementTransform.GetTranslation(), ElementTransform.TransformPosition(FVector(Scale, 0.f, 0.f)), FLinearColor::Red, SDPG_Foreground);
+						PDI->DrawLine(ElementTransform.GetTranslation(), ElementTransform.TransformPosition(FVector(0.f, Scale, 0.f)), FLinearColor::Green, SDPG_Foreground);
+						PDI->DrawLine(ElementTransform.GetTranslation(), ElementTransform.TransformPosition(FVector(0.f, 0.f, Scale)), FLinearColor::Blue, SDPG_Foreground);
 					}
 				}
 
