Control Rig: Import ref skeleton when authoring a rig module

#rnx

diff --git a/Engine/Plugins/Animation/ControlRig/Source/ControlRigEditor/Private/Editor/ControlRigEditor.cpp b/Engine/Plugins/Animation/ControlRig/Source/ControlRigEditor/Private/Editor/ControlRigEditor.cpp
--- a/Engine/Plugins/Animation/ControlRig/Source/ControlRigEditor/Private/Editor/ControlRigEditor.cpp
+++ b/Engine/Plugins/Animation/ControlRig/Source/ControlRigEditor/Private/Editor/ControlRigEditor.cpp
@@ -5004,7 +5004,8 @@ void FControlRigEditor::OnPreConstruction_AnyThread(UControlRig* InRig, const FN
 						}
 
 						// import the bones for the preview hierarchy
-						const TArray<FRigElementKey> Bones = Controller->ImportBones(RigBlueprint->PreviewSkeletalMesh->GetSkeleton());
+						// use the ref skeleton so we'll only see the bones that are actually part of the mesh
+						const TArray<FRigElementKey> Bones = Controller->ImportBones(RigBlueprint->PreviewSkeletalMesh->GetRefSkeleton(), NAME_None, false, false, false, false);
 						for(const FRigElementKey& Bone : Bones)
 						{
 							if(FRigBaseElement* Element = Hierarchy->Find(Bone))
