Control Rig: Further work on Module Rule Manager

#rnx

diff --git a/Engine/Plugins/Animation/ControlRig/Source/ControlRig/Private/ModularRigRuleManager.cpp b/Engine/Plugins/Animation/ControlRig/Source/ControlRig/Private/ModularRigRuleManager.cpp
--- a/Engine/Plugins/Animation/ControlRig/Source/ControlRig/Private/ModularRigRuleManager.cpp
+++ b/Engine/Plugins/Animation/ControlRig/Source/ControlRig/Private/ModularRigRuleManager.cpp
@@ -6,12 +6,13 @@
 
 #define LOCTEXT_NAMESPACE "ModularRigRuleManager"
 
-FModularRigResolveResult UModularRigRuleManager::Resolve(
+FModularRigResolveResult UModularRigRuleManager::FindMatches(
 	const FRigConnectorElement* InConnector,
 	const FRigModuleInstance* InModule,
 	const FRigElementKeyRedirector& InResolvedConnectors) const
 {
 	FModularRigResolveResult Result;
+	Result.Connector = InConnector->GetKey();
 	
 	if(!Hierarchy.IsValid())
 	{
@@ -45,6 +46,64 @@ FModularRigResolveResult UModularRigRuleManager::Resolve(
 	return Result;;
 }
 
+FModularRigResolveResult UModularRigRuleManager::FindMatchesForPrimaryConnector(const FRigModuleInstance* InModule) const
+{
+	static const FRigElementKeyRedirector EmptyRedirector;
+	FRigConnectionRuleInput RuleInput;
+	RuleInput.Hierarchy = Hierarchy.Get();
+	RuleInput.Module = InModule;
+	RuleInput.Redirector = &EmptyRedirector;
+
+	FModularRigResolveResult Result;
+	
+	if(const FRigConnectorElement* PrimaryConnector = RuleInput.FindPrimaryConnector(&Result.Message))
+	{
+		Result = FindMatches(PrimaryConnector, InModule, EmptyRedirector);
+	}
+	else
+	{
+		Result.State = EModularRigResolveState::Error;
+	}
+
+	return Result;
+}
+
+TArray<FModularRigResolveResult> UModularRigRuleManager::FindMatchesForSecondaryConnectors(const FRigModuleInstance* InModule, const FRigElementKeyRedirector& InResolvedConnectors) const
+{
+	FRigConnectionRuleInput RuleInput;
+	RuleInput.Hierarchy = Hierarchy.Get();
+	RuleInput.Module = InModule;
+	RuleInput.Redirector = &InResolvedConnectors;
+
+	TArray<FModularRigResolveResult> Results;
+
+	const TArray<const FRigConnectorElement*> SecondaryConnectors = RuleInput.FindSecondaryConnectors(false /* optional */);
+	for(const FRigConnectorElement* SecondaryConnector : SecondaryConnectors)
+	{
+		Results.Add(FindMatches(SecondaryConnector, InModule, InResolvedConnectors));
+	}
+
+	return Results;
+}
+
+TArray<FModularRigResolveResult> UModularRigRuleManager::FindMatchesForOptionalConnectors(const FRigModuleInstance* InModule, const FRigElementKeyRedirector& InResolvedConnectors) const
+{
+	FRigConnectionRuleInput RuleInput;
+	RuleInput.Hierarchy = Hierarchy.Get();
+	RuleInput.Module = InModule;
+	RuleInput.Redirector = &InResolvedConnectors;
+
+	TArray<FModularRigResolveResult> Results;
+
+	const TArray<const FRigConnectorElement*> OptionalConnectors = RuleInput.FindSecondaryConnectors(true /* optional */);
+	for(const FRigConnectorElement* OptionalConnector : OptionalConnectors)
+	{
+		Results.Add(FindMatches(OptionalConnector, InModule, InResolvedConnectors));
+	}
+
+	return Results;
+}
+
 void UModularRigRuleManager::FWorkData::Filter(TFunction<void(FRigElementResolveResult&)> PerMatchFunction)
 {
 	const TArray<FRigElementResolveResult> PreviousMatches = Result->Matches;;
@@ -75,6 +134,7 @@ void UModularRigRuleManager::ResolveConnector(FWorkData& InOutWorkData)
 	FilterIncompatibleTypes(InOutWorkData);
 	FilterInvalidNameSpaces(InOutWorkData);
 	FilterByConnectorRules(InOutWorkData);
+	FilterByConnectorEvent(InOutWorkData);
 
 	if(InOutWorkData.Result->Matches.IsEmpty())
 	{
@@ -141,13 +201,55 @@ void UModularRigRuleManager::FilterByConnectorRules(FWorkData& InOutWorkData)
 		TSharedPtr<FStructOnScope> Storage;
 		const FRigConnectionRule* Rule = Stash.Get(Storage);
 
-		InOutWorkData.Filter([Rule, InOutWorkData](FRigElementResolveResult& Result)
+		FRigConnectionRuleInput RuleInput;
+		RuleInput.Hierarchy = InOutWorkData.Hierarchy;
+		RuleInput.Redirector = InOutWorkData.ResolvedConnectors;
+		RuleInput.Module = InOutWorkData.Module;
+		
+		InOutWorkData.Filter([Rule, RuleInput](FRigElementResolveResult& Result)
 		{
-			const FRigBaseElement* Target = InOutWorkData.Hierarchy->Find(Result.GetKey());
+			const FRigBaseElement* Target = RuleInput.Hierarchy->Find(Result.GetKey());
 			check(Target);
-			Result = Rule->Resolve(Target, InOutWorkData.Hierarchy, InOutWorkData.ResolvedConnectors);
+			Result = Rule->Resolve(Target, RuleInput);
 		});
 	}
 }
 
+void UModularRigRuleManager::FilterByConnectorEvent(FWorkData& InOutWorkData)
+{
+	// this may be null during unit tests
+	if(InOutWorkData.Module == nullptr)
+	{
+		return;
+	}
+	
+	// see if we are nested below a modular rig
+	UModularRig* ModularRig = InOutWorkData.Hierarchy->GetTypedOuter<UModularRig>();
+	if(ModularRig == nullptr)
+	{
+		return;
+	}
+	
+	FModularRigResolveResult* Result = InOutWorkData.Result;
+
+	// todo: run the VM event to perform filtering
+	// ModularRig->ExecuteConnectorEvent(InOutWorkData->Connector->GetKey(), Module, Redirector)
+	
+	// move the default match to the front of the list
+	if(!Result->Matches.IsEmpty())
+	{
+		const int32 DefaultMatchIndex = Result->Matches.IndexOfByPredicate([](const FRigElementResolveResult& ElementResult) -> bool
+		{
+			return ElementResult.State == ERigElementResolveState::DefaultTarget;
+		}) ;
+
+		if(DefaultMatchIndex != INDEX_NONE)
+		{
+			const FRigElementResolveResult DefaultResult = Result->Matches[DefaultMatchIndex];
+			Result->Matches.RemoveAt(DefaultMatchIndex);
+			Result->Matches.Insert(DefaultResult, 0);
+		};
+	}
+}
+
 #undef LOCTEXT_NAMESPACE
diff --git a/Engine/Plugins/Animation/ControlRig/Source/ControlRig/Private/Rigs/RigConnectionRules.cpp b/Engine/Plugins/Animation/ControlRig/Source/ControlRig/Private/Rigs/RigConnectionRules.cpp
--- a/Engine/Plugins/Animation/ControlRig/Source/ControlRig/Private/Rigs/RigConnectionRules.cpp
+++ b/Engine/Plugins/Animation/ControlRig/Source/ControlRig/Private/Rigs/RigConnectionRules.cpp
@@ -3,6 +3,7 @@
 #include "Rigs/RigConnectionRules.h"
 #include "Rigs/RigHierarchyElements.h"
 #include "Rigs/RigHierarchy.h"
+#include "ModularRig.h"
 
 #include UE_INLINE_GENERATED_CPP_BY_NAME(RigConnectionRules)
 
@@ -97,11 +98,103 @@ uint32 GetTypeHash(const FRigConnectionRuleStash& InRuleStash)
 	return HashCombine(GetTypeHash(InRuleStash.ScriptStructPath), GetTypeHash(InRuleStash.ExportedText));
 }
 
+////////////////////////////////////////////////////////////////////////////////
+// FRigConnectionRuleInput
+////////////////////////////////////////////////////////////////////////////////
+
+const FRigConnectorElement* FRigConnectionRuleInput::FindPrimaryConnector(FText* OutErrorMessage) const
+{
+	check(Hierarchy);
+	check(Module);
+
+	const FName ModuleNameSpace = *Module->GetPath();
+
+	const FRigConnectorElement* PrimaryConnector = nullptr;
+	Hierarchy->ForEach<FRigConnectorElement>(
+		[this, ModuleNameSpace, &PrimaryConnector](FRigConnectorElement* Connector) -> bool
+		{
+			if(Connector->Settings.Type == EConnectorType::Primary)
+			{
+				const FName ConnectorNameSpace = Hierarchy->GetNameMetadata(Connector->GetKey(), URigHierarchy::NameSpaceMetadataName, NAME_None);
+				if(!ConnectorNameSpace.IsNone() && ConnectorNameSpace.IsEqual(ModuleNameSpace, ENameCase::CaseSensitive))
+				{
+					PrimaryConnector = Connector;
+					return false; // stop the search
+				}
+			}
+			return true; // continue the search
+		}
+	);
+
+	if(PrimaryConnector == nullptr)
+	{
+		if(OutErrorMessage)
+		{
+			static constexpr TCHAR Format[] = TEXT("No primary connector found for module '%s'.");
+			*OutErrorMessage = FText::FromString(FString::Printf(Format, *Module->GetPath()));
+		}
+	}
+
+	return PrimaryConnector;
+}
+
+TArray<const FRigConnectorElement*> FRigConnectionRuleInput::FindSecondaryConnectors(bool bOptional, FText* OutErrorMessage) const
+{
+	check(Hierarchy);
+	check(Module);
+
+	const FName ModuleNameSpace = *Module->GetPath();
+
+	TArray<const FRigConnectorElement*> SecondaryConnectors;
+	Hierarchy->ForEach<FRigConnectorElement>(
+		[this, bOptional, ModuleNameSpace, &SecondaryConnectors](FRigConnectorElement* Connector) -> bool
+		{
+			if(Connector->Settings.Type == EConnectorType::Secondary && Connector->Settings.bOptional == bOptional)
+			{
+				const FName ConnectorNameSpace = Hierarchy->GetNameMetadata(Connector->GetKey(), URigHierarchy::NameSpaceMetadataName, NAME_None);
+				if(!ConnectorNameSpace.IsNone() && ConnectorNameSpace.IsEqual(ModuleNameSpace, ENameCase::CaseSensitive))
+				{
+					SecondaryConnectors.Add(Connector);
+				}
+			}
+			return true; // continue the search
+		}
+	);
+
+	return SecondaryConnectors;
+}
+
+const FRigTransformElement* FRigConnectionRuleInput::ResolveConnector(const FRigConnectorElement* InConnector, FText* OutErrorMessage) const
+{
+	check(Redirector);
+
+	if(const FCachedRigElement* Target = Redirector->Find(InConnector->GetKey()))
+	{
+		return Cast<FRigTransformElement>(Target->GetElement());
+	}
+
+	if(OutErrorMessage)
+	{
+		static constexpr TCHAR Format[] = TEXT("Resolved target not found for connector '%s'.");
+		*OutErrorMessage = FText::FromString(FString::Printf(Format, *InConnector->GetName()));
+	}
+	return nullptr;
+}
+
+const FRigTransformElement* FRigConnectionRuleInput::ResolvePrimaryConnector(FText* OutErrorMessage) const
+{
+	if(const FRigConnectorElement* Connector = FindPrimaryConnector(OutErrorMessage))
+	{
+		return ResolveConnector(Connector, OutErrorMessage);
+	}
+	return nullptr;
+}
+
 ////////////////////////////////////////////////////////////////////////////////
 // FRigConnectionRule
 ////////////////////////////////////////////////////////////////////////////////
 
-FRigElementResolveResult FRigConnectionRule::Resolve(const FRigBaseElement* InTarget, const URigHierarchy* InHierarchy, const FRigElementKeyRedirector* InRedirector) const
+FRigElementResolveResult FRigConnectionRule::Resolve(const FRigBaseElement* InTarget, const FRigConnectionRuleInput& InRuleInput) const
 {
 	FRigElementResolveResult Result(InTarget->GetKey());
 	Result.SetInvalidTarget(FText());
@@ -112,7 +205,7 @@ FRigElementResolveResult FRigConnectionRule::Resolve(const FRigBaseElement* InTa
 // FRigAndConnectionRule
 ////////////////////////////////////////////////////////////////////////////////
 
-FRigElementResolveResult FRigAndConnectionRule::Resolve(const FRigBaseElement* InTarget, const URigHierarchy* InHierarchy, const FRigElementKeyRedirector* InRedirector) const
+FRigElementResolveResult FRigAndConnectionRule::Resolve(const FRigBaseElement* InTarget, const FRigConnectionRuleInput& InRuleInput) const
 {
 	FRigElementResolveResult Result(InTarget->GetKey());
 	Result.SetPossibleTarget();
@@ -122,7 +215,7 @@ FRigElementResolveResult FRigAndConnectionRule::Resolve(const FRigBaseElement* I
 	{
 		if(const FRigConnectionRule* Rule = ChildRule.Get(Storage))
 		{
-			Result = Rule->Resolve(InTarget, InHierarchy, InRedirector);
+			Result = Rule->Resolve(InTarget, InRuleInput);
 			if(!Result.IsValid())
 			{
 				return Result;
@@ -137,7 +230,7 @@ FRigElementResolveResult FRigAndConnectionRule::Resolve(const FRigBaseElement* I
 // FRigOrConnectionRule
 ////////////////////////////////////////////////////////////////////////////////
 
-FRigElementResolveResult FRigOrConnectionRule::Resolve(const FRigBaseElement* InTarget, const URigHierarchy* InHierarchy, const FRigElementKeyRedirector* InRedirector) const
+FRigElementResolveResult FRigOrConnectionRule::Resolve(const FRigBaseElement* InTarget, const FRigConnectionRuleInput& InRuleInput) const
 {
 	FRigElementResolveResult Result(InTarget->GetKey());
 	Result.SetPossibleTarget();
@@ -147,7 +240,7 @@ FRigElementResolveResult FRigOrConnectionRule::Resolve(const FRigBaseElement* In
 	{
 		if(const FRigConnectionRule* Rule = ChildRule.Get(Storage))
 		{
-			Result = Rule->Resolve(InTarget, InHierarchy, InRedirector);
+			Result = Rule->Resolve(InTarget, InRuleInput);
 			if(Result.IsValid())
 			{
 				return Result;
@@ -162,7 +255,7 @@ FRigElementResolveResult FRigOrConnectionRule::Resolve(const FRigBaseElement* In
 // FRigTypeConnectionRule
 ////////////////////////////////////////////////////////////////////////////////
 
-FRigElementResolveResult FRigTypeConnectionRule::Resolve(const FRigBaseElement* InTarget, const URigHierarchy* InHierarchy, const FRigElementKeyRedirector* InRedirector) const
+FRigElementResolveResult FRigTypeConnectionRule::Resolve(const FRigBaseElement* InTarget, const FRigConnectionRuleInput& InRuleInput) const
 {
 	FRigElementResolveResult Result(InTarget->GetKey());
 	Result.SetPossibleTarget();
@@ -181,12 +274,12 @@ FRigElementResolveResult FRigTypeConnectionRule::Resolve(const FRigBaseElement*
 // FRigTagConnectionRule
 ////////////////////////////////////////////////////////////////////////////////
 
-FRigElementResolveResult FRigTagConnectionRule::Resolve(const FRigBaseElement* InTarget, const URigHierarchy* InHierarchy, const FRigElementKeyRedirector* InRedirector) const
+FRigElementResolveResult FRigTagConnectionRule::Resolve(const FRigBaseElement* InTarget, const FRigConnectionRuleInput& InRuleInput) const
 {
 	FRigElementResolveResult Result(InTarget->GetKey());
 	Result.SetPossibleTarget();
 
-	if(!InHierarchy->HasTag(InTarget->GetKey(), Tag))
+	if(!InRuleInput.GetHierarchy()->HasTag(InTarget->GetKey(), Tag))
 	{
 		static constexpr TCHAR Format[] = TEXT("Element '%s' does not contain tag '%s'.");
 		Result.SetInvalidTarget(FText::FromString(FString::Printf(Format, *InTarget->GetKey().ToString(), *Tag.ToString())));
@@ -195,12 +288,64 @@ FRigElementResolveResult FRigTagConnectionRule::Resolve(const FRigBaseElement* I
 	return Result;
 }
 
+////////////////////////////////////////////////////////////////////////////////
+// FRigChildOfPrimaryConnectionRule
+////////////////////////////////////////////////////////////////////////////////
+
+FRigElementResolveResult FRigChildOfPrimaryConnectionRule::Resolve(const FRigBaseElement* InTarget, const FRigConnectionRuleInput& InRuleInput) const
+{
+	FRigElementResolveResult Result(InTarget->GetKey());
+	Result.SetPossibleTarget();
+
+	// find the primary resolved target
+	FText ErrorMessage;
+	const FRigTransformElement* PrimaryTarget = InRuleInput.ResolvePrimaryConnector(&ErrorMessage);
+	if(PrimaryTarget == nullptr)
+	{
+		Result.SetInvalidTarget(ErrorMessage);
+		return Result;
+	}
+
+	static constexpr TCHAR IsNotAChildOfFormat[] = TEXT("Target '%s' is not a child of the primary.");
+	static constexpr TCHAR IsAlreadyUsedForPrimaryFormat[] = TEXT("Target '%s' is already used for the primary.");
+	if(InTarget == PrimaryTarget)
+	{
+		Result.SetInvalidTarget(FText::FromString(FString::Printf(IsAlreadyUsedForPrimaryFormat, *InTarget->GetKey().ToString())));
+		return Result;
+	}
+
+	// for sockets we use the parent for resolve
+	if(PrimaryTarget->GetType() == ERigElementType::Socket)
+	{
+		if(const FRigTransformElement* FirstParent =
+			Cast<FRigTransformElement>(InRuleInput.GetHierarchy()->GetFirstParent(PrimaryTarget)))
+		{
+			PrimaryTarget = FirstParent;
+
+			if(InTarget == PrimaryTarget)
+			{
+				Result.SetInvalidTarget(FText::FromString(FString::Printf(IsNotAChildOfFormat, *InTarget->GetKey().ToString())));
+				return Result;
+			}
+		}
+	}
+
+	static const URigHierarchy::TElementDependencyMap EmptyDependencyMap;
+	if(!InRuleInput.GetHierarchy()->IsParentedTo(
+		const_cast<FRigBaseElement*>(InTarget), const_cast<FRigTransformElement*>(PrimaryTarget), EmptyDependencyMap))
+	{
+		Result.SetInvalidTarget(FText::FromString(FString::Printf(IsNotAChildOfFormat, *InTarget->GetKey().ToString())));
+	}
+
+	return Result;
+}
+
 ////////////////////////////////////////////////////////////////////////////////
 // FRigOnChainRule
 ////////////////////////////////////////////////////////////////////////////////
 
 /*
-FRigElementResolveResult FRigOnChainRule::Resolve(const FRigBaseElement* InTarget, const URigHierarchy* InHierarchy, const FRigElementKeyRedirector* InRedirector) const
+FRigElementResolveResult FRigOnChainRule::Resolve(const FRigBaseElement* InTarget, const FRigConnectionRuleInput& InRuleInput) const
 {
 	FRigElementResolveResult Result;
 	Result.State = ERigElementResolveState::PossibleTarget;
diff --git a/Engine/Plugins/Animation/ControlRig/Source/ControlRig/Public/ModularRigRuleManager.h b/Engine/Plugins/Animation/ControlRig/Source/ControlRig/Public/ModularRigRuleManager.h
--- a/Engine/Plugins/Animation/ControlRig/Source/ControlRig/Public/ModularRigRuleManager.h
+++ b/Engine/Plugins/Animation/ControlRig/Source/ControlRig/Public/ModularRigRuleManager.h
@@ -14,17 +14,53 @@ public:
 	GENERATED_BODY()
 
 	/***
-	 * Returns the possible targets for the given connector in the current resolval stage
+	 * Returns the possible targets for the given connector in the current resolve stage
+	 * Note: This method is thread-safe. 
 	 * @param InConnector The connector to resolve
+	 * @param InModule The module the connector belongs to
 	 * @param InResolvedConnectors A redirect map of the already resolved connectors
-	 * @return The resolval result including a list of matches
+	 * @return The resolve result including a list of matches
 	 */
-	FModularRigResolveResult Resolve(
+	FModularRigResolveResult FindMatches(
 		const FRigConnectorElement* InConnector,
 		const FRigModuleInstance* InModule,
 		const FRigElementKeyRedirector& InResolvedConnectors = FRigElementKeyRedirector()
 	) const;
 
+	/***
+	 * Returns the possible targets for the primary connector in the current resolve stage
+	 * Note: This method is thread-safe. 
+	 * @param InModule The module the connector belongs to
+	 * @return The resolve result including a list of matches
+	 */
+	FModularRigResolveResult FindMatchesForPrimaryConnector(
+		const FRigModuleInstance* InModule
+	) const;
+
+	/***
+	 * Returns the possible targets for each secondary connector
+	 * Note: This method is thread-safe. 
+	 * @param InModule The module the secondary connectors belongs to
+	 * @param InResolvedConnectors A redirect map of the already resolved connectors
+	 * @return The resolve result including a list of matches for each connector
+	 */
+	TArray<FModularRigResolveResult> FindMatchesForSecondaryConnectors(
+		const FRigModuleInstance* InModule,
+		const FRigElementKeyRedirector& InResolvedConnectors = FRigElementKeyRedirector()
+	) const;
+
+	/***
+	 * Returns the possible targets for each optional connector
+	 * Note: This method is thread-safe. 
+	 * @param InModule The module the optional connectors belongs to
+	 * @param InResolvedConnectors A redirect map of the already resolved connectors
+	 * @return The resolve result including a list of matches for each connector
+	 */
+	TArray<FModularRigResolveResult> FindMatchesForOptionalConnectors(
+		const FRigModuleInstance* InModule,
+		const FRigElementKeyRedirector& InResolvedConnectors = FRigElementKeyRedirector()
+	) const;
+
 private:
 
 	struct FWorkData
@@ -52,6 +88,7 @@ private:
 	static void FilterIncompatibleTypes(FWorkData& InOutWorkData);
 	static void FilterInvalidNameSpaces(FWorkData& InOutWorkData);
 	static void FilterByConnectorRules(FWorkData& InOutWorkData);
+	static void FilterByConnectorEvent(FWorkData& InOutWorkData);
 
 	TWeakObjectPtr<const URigHierarchy> Hierarchy;
 
diff --git a/Engine/Plugins/Animation/ControlRig/Source/ControlRig/Public/Rigs/RigConnectionRules.h b/Engine/Plugins/Animation/ControlRig/Source/ControlRig/Public/Rigs/RigConnectionRules.h
--- a/Engine/Plugins/Animation/ControlRig/Source/ControlRig/Public/Rigs/RigConnectionRules.h
+++ b/Engine/Plugins/Animation/ControlRig/Source/ControlRig/Public/Rigs/RigConnectionRules.h
@@ -11,6 +11,10 @@ struct FRigBaseElement;
 struct FRigConnectionRule;
 class FRigElementKeyRedirector;
 class URigHierarchy;
+struct FRigModuleInstance;
+struct FRigBaseElement;
+struct FRigTransformElement;
+struct FRigConnectorElement;
 
 USTRUCT(BlueprintType)
 struct CONTROLRIG_API FRigConnectionRuleStash
@@ -44,6 +48,47 @@ struct CONTROLRIG_API FRigConnectionRuleStash
 	FString ExportedText;
 };
 
+struct CONTROLRIG_API FRigConnectionRuleInput
+{
+public:
+	
+	FRigConnectionRuleInput()
+	: Hierarchy(nullptr)
+	, Module(nullptr)
+	, Redirector(nullptr)
+	{
+	}
+
+	const URigHierarchy* GetHierarchy() const
+	{
+		return Hierarchy;
+	}
+	
+	const FRigModuleInstance* GetModule() const
+	{
+		return Module;
+	}
+	
+	const FRigElementKeyRedirector* GetRedirector() const
+	{
+		return Redirector;
+	}
+
+	const FRigConnectorElement* FindPrimaryConnector(FText* OutErrorMessage = nullptr) const;
+   	TArray<const FRigConnectorElement*> FindSecondaryConnectors(bool bOptional, FText* OutErrorMessage = nullptr) const;
+
+	const FRigTransformElement* ResolveConnector(const FRigConnectorElement* InConnector, FText* OutErrorMessage) const;
+	const FRigTransformElement* ResolvePrimaryConnector(FText* OutErrorMessage = nullptr) const;
+
+private:
+
+	const URigHierarchy* Hierarchy;
+	const FRigModuleInstance* Module;
+	const FRigElementKeyRedirector* Redirector;
+
+	friend class UModularRigRuleManager;
+};
+
 USTRUCT(meta=(Hidden))
 struct CONTROLRIG_API FRigConnectionRule
 {
@@ -55,7 +100,7 @@ public:
 	virtual ~FRigConnectionRule() {}
 
 	virtual UScriptStruct* GetScriptStruct() const { return FRigConnectionRule::StaticStruct(); }
-	virtual FRigElementResolveResult Resolve(const FRigBaseElement* InTarget, const URigHierarchy* InHierarchy, const FRigElementKeyRedirector* InRedirector) const;
+	virtual FRigElementResolveResult Resolve(const FRigBaseElement* InTarget, const FRigConnectionRuleInput& InRuleInput) const;
 };
 
 USTRUCT(BlueprintType, DisplayName="And Rule")
@@ -78,7 +123,7 @@ public:
 	virtual ~FRigAndConnectionRule() override {}
 
 	virtual UScriptStruct* GetScriptStruct() const override { return FRigAndConnectionRule::StaticStruct(); }
-	virtual FRigElementResolveResult Resolve(const FRigBaseElement* InTarget, const URigHierarchy* InHierarchy, const FRigElementKeyRedirector* InRedirector) const override;
+	virtual FRigElementResolveResult Resolve(const FRigBaseElement* InTarget, const FRigConnectionRuleInput& InRuleInput) const override;
 
 	UPROPERTY(BlueprintReadWrite, EditAnywhere, Category=Settings)
 	TArray<FRigConnectionRuleStash> ChildRules;
@@ -104,7 +149,7 @@ public:
 	virtual ~FRigOrConnectionRule() override {}
 
 	virtual UScriptStruct* GetScriptStruct() const override { return FRigOrConnectionRule::StaticStruct(); }
-	virtual FRigElementResolveResult Resolve(const FRigBaseElement* InTarget, const URigHierarchy* InHierarchy, const FRigElementKeyRedirector* InRedirector) const override;
+	virtual FRigElementResolveResult Resolve(const FRigBaseElement* InTarget, const FRigConnectionRuleInput& InRuleInput) const override;
 
 	UPROPERTY(BlueprintReadWrite, EditAnywhere, Category=Settings)
 	TArray<FRigConnectionRuleStash> ChildRules;
@@ -128,7 +173,7 @@ public:
 	virtual ~FRigTypeConnectionRule() override {}
 
 	virtual UScriptStruct* GetScriptStruct() const override { return FRigTypeConnectionRule::StaticStruct(); }
-	virtual FRigElementResolveResult Resolve(const FRigBaseElement* InTarget, const URigHierarchy* InHierarchy, const FRigElementKeyRedirector* InRedirector) const override;
+	virtual FRigElementResolveResult Resolve(const FRigBaseElement* InTarget, const FRigConnectionRuleInput& InRuleInput) const override;
 
 	UPROPERTY(BlueprintReadWrite, EditAnywhere, Category=Settings)
 	ERigElementType ElementType;
@@ -152,12 +197,27 @@ public:
 	virtual ~FRigTagConnectionRule() override {}
 
 	virtual UScriptStruct* GetScriptStruct() const override { return FRigTagConnectionRule::StaticStruct(); }
-	virtual FRigElementResolveResult Resolve(const FRigBaseElement* InTarget, const URigHierarchy* InHierarchy, const FRigElementKeyRedirector* InRedirector) const override;
+	virtual FRigElementResolveResult Resolve(const FRigBaseElement* InTarget, const FRigConnectionRuleInput& InRuleInput) const override;
 
 	UPROPERTY(BlueprintReadWrite, EditAnywhere, Category=Settings)
 	FName Tag;
 };
 
+USTRUCT(BlueprintType, DisplayName="Child of Primary")
+struct CONTROLRIG_API FRigChildOfPrimaryConnectionRule : public FRigConnectionRule
+{
+	GENERATED_BODY()
+
+public:
+
+	FRigChildOfPrimaryConnectionRule()
+	{}
+
+	virtual ~FRigChildOfPrimaryConnectionRule() override {}
+
+	virtual UScriptStruct* GetScriptStruct() const override { return FRigChildOfPrimaryConnectionRule::StaticStruct(); }
+	virtual FRigElementResolveResult Resolve(const FRigBaseElement* InTarget, const FRigConnectionRuleInput& InRuleInput) const override;
+};
 /*
 USTRUCT(BlueprintType, DisplayName="On Chain Rule")
 struct CONTROLRIG_API FRigOnChainRule : public FRigConnectionRule
@@ -179,7 +239,7 @@ public:
 	virtual ~FRigOnChainRule() override {}
 
 	virtual UScriptStruct* GetScriptStruct() const override { return FRigOnChainRule::StaticStruct(); }
-	virtual FRigElementResolveResult Resolve(const FRigBaseElement* InTarget, const URigHierarchy* InHierarchy, const FRigElementKeyRedirector* InRedirector) const override;
+	virtual FRigElementResolveResult Resolve(const FRigBaseElement* InTarget, const FRigConnectionRuleInput& InRuleInput) const override;
 
 	UPROPERTY(BlueprintReadWrite, EditAnywhere, Category=Settings)
 	int32 MinNumBones;
diff --git a/Engine/Plugins/Animation/ControlRig/Source/ControlRig/Public/Rigs/RigHierarchyDefines.h b/Engine/Plugins/Animation/ControlRig/Source/ControlRig/Public/Rigs/RigHierarchyDefines.h
--- a/Engine/Plugins/Animation/ControlRig/Source/ControlRig/Public/Rigs/RigHierarchyDefines.h
+++ b/Engine/Plugins/Animation/ControlRig/Source/ControlRig/Public/Rigs/RigHierarchyDefines.h
@@ -1878,6 +1878,11 @@ public:
 
 	bool IsValid() const;
 
+	const FRigElementKey& GetConnectorKey() const
+	{
+		return Connector;
+	}
+
 	EModularRigResolveState GetState() const
 	{
 		return State;
@@ -1899,7 +1904,10 @@ public:
 	}
 
 private:
-	
+
+	UPROPERTY()
+	FRigElementKey Connector;
+
 	UPROPERTY()
 	TArray<FRigElementResolveResult> Matches;
 
diff --git a/Engine/Plugins/Animation/ControlRig/Source/ControlRigEditor/Private/ControlRigEditorStyle.h b/Engine/Plugins/Animation/ControlRig/Source/ControlRigEditor/Private/ControlRigEditorStyle.h
--- a/Engine/Plugins/Animation/ControlRig/Source/ControlRigEditor/Private/ControlRigEditorStyle.h
+++ b/Engine/Plugins/Animation/ControlRig/Source/ControlRigEditor/Private/ControlRigEditorStyle.h
@@ -36,6 +36,8 @@ class FControlRigEditorStyle final
 public:
 	FControlRigEditorStyle()
 		: FSlateStyleSet("ControlRigEditorStyle")
+		, ConnectorUserInterfaceColor(0.0, 112.f/255.f, 224.f/255.f)
+		, SocketUserInterfaceColor(0.0, 112.f/255.f, 224.f/255.f)
 	{
 		const FVector2D Icon10x10(10.0f, 10.0f);
 		const FVector2D Icon14x14(14.0f, 14.0f);
@@ -166,6 +168,8 @@ public:
 	}
 
 	FSlateColor SpacePickerSelectColor;
+	FLinearColor ConnectorUserInterfaceColor;
+	FLinearColor SocketUserInterfaceColor;
 };
 
 #undef TTF_FONT
diff --git a/Engine/Plugins/Animation/ControlRig/Source/ControlRigEditor/Private/Editor/SRigHierarchy.cpp b/Engine/Plugins/Animation/ControlRig/Source/ControlRigEditor/Private/Editor/SRigHierarchy.cpp
--- a/Engine/Plugins/Animation/ControlRig/Source/ControlRigEditor/Private/Editor/SRigHierarchy.cpp
+++ b/Engine/Plugins/Animation/ControlRig/Source/ControlRigEditor/Private/Editor/SRigHierarchy.cpp
@@ -2087,8 +2087,14 @@ void SRigHierarchy::HandleNewItem(ERigElementType InElementType, bool bIsAnimati
 						}
 					}
 
+					const bool bIsPrimary = Hierarchy->GetConnectorKeys(false).Num() == 0;
 					FRigConnectorSettings Settings;
-					Settings.Type = Hierarchy->GetConnectorKeys(false).Num() == 0 ? EConnectorType::Primary : EConnectorType::Secondary;
+					Settings.Type = bIsPrimary ? EConnectorType::Primary : EConnectorType::Secondary;
+						if(!bIsPrimary)
+						{
+							Settings.Rules.Reset();
+							Settings.AddRule(FRigChildOfPrimaryConnectionRule());
+						}
 					NewItemKey = Controller->AddConnector(NewElementName, Settings, true);
 					(void)ResolveConnector(NewItemKey, ParentKey);
 					break;
diff --git a/Engine/Plugins/Animation/ControlRig/Source/ControlRigEditor/Private/Editor/SRigHierarchyTreeView.cpp b/Engine/Plugins/Animation/ControlRig/Source/ControlRigEditor/Private/Editor/SRigHierarchyTreeView.cpp
--- a/Engine/Plugins/Animation/ControlRig/Source/ControlRigEditor/Private/Editor/SRigHierarchyTreeView.cpp
+++ b/Engine/Plugins/Animation/ControlRig/Source/ControlRigEditor/Private/Editor/SRigHierarchyTreeView.cpp
@@ -524,7 +524,9 @@ bool SRigHierarchyTreeView::AddElement(const FRigBaseElement* InElement)
 								TagArguments.Text(FText::FromName(Name));
 								TagArguments.TooltipText(FText::FromName(ConnectorKey.Name));
 								TagArguments.TextColor(FSlateColor(FLinearColor::White));
-								TagArguments.Color(FLinearColor(0.0, 112.f/255.f, 224.f/255.f));
+
+								static const FLinearColor ConnectorColor = FControlRigEditorStyle::Get().ConnectorUserInterfaceColor;
+								TagArguments.Color(ConnectorColor);
 								TagArguments.AllowDragDrop(true);
 								FString Identifier;
 								FRigElementKey::StaticStruct()->ExportText(Identifier, &ConnectorKey, nullptr, nullptr, PPF_None, nullptr);
@@ -742,7 +744,7 @@ void SRigHierarchyTreeView::RefreshTreeView(bool bRebuildContent)
 				bContinue = true;
 			});
 
-			// add all of the connectors. their parent relationship in the tree represents resolval
+			// add all of the connectors. their parent relationship in the tree represents resolve
 			for(FRigConnectorElement* Connector : Connectors)
 			{
 				AddElement(Connector);
