Control Rig: Draw sockets as blue nulls

#rnx

diff --git a/Engine/Plugins/Animation/ControlRig/Source/ControlRigEditor/Private/EditMode/ControlRigEditMode.cpp b/Engine/Plugins/Animation/ControlRig/Source/ControlRigEditor/Private/EditMode/ControlRigEditMode.cpp
--- a/Engine/Plugins/Animation/ControlRig/Source/ControlRigEditor/Private/EditMode/ControlRigEditMode.cpp
+++ b/Engine/Plugins/Animation/ControlRig/Source/ControlRigEditor/Private/EditMode/ControlRigEditMode.cpp
@@ -66,6 +66,8 @@
 #include "Editor/EditorPerProjectUserSettings.h"
 #include "TransformConstraint.h"
 #include "Animation/DebugSkelMeshComponent.h"
+#include "Materials/Material.h"
+#include "ControlRigEditorStyle.h"
 
 #include UE_INLINE_GENERATED_CPP_BY_NAME(ControlRigEditMode)
 
@@ -922,34 +924,21 @@ void FControlRigEditMode::Render(const FSceneView* View, FViewport* Viewport, FP
 				}
 
 				// temporary implementation to draw sockets in 3D
-				if (bIsAssetEditor && (Settings->bDisplaySockets || ControlRig->IsConstructionModeEnabled()))
+				if (bIsAssetEditor && (Settings->bDisplaySockets || ControlRig->IsConstructionModeEnabled()) && Settings->AxisScale > SMALL_NUMBER)
 				{
-					Hierarchy->ForEach<FRigSocketElement>([this, Hierarchy, PDI](FRigSocketElement* Socket)
+					const float Scale = Settings->AxisScale;
+					PDI->AddReserveLines(SDPG_Foreground, Hierarchy->Num(ERigElementType::Socket) * 3);
+					static const FLinearColor SocketColor = FControlRigEditorStyle::Get().SocketUserInterfaceColor;
+
+					Hierarchy->ForEach<FRigSocketElement>([this, Hierarchy, PDI, ComponentTransform, Scale](FRigSocketElement* Socket)
 					{
-						const FLinearColor Color = Socket->GetColor(Hierarchy);
-						const uint32 ColorHash = GetTypeHash(Color.ToFColor(true));
+						FTransform ElementTransform = Hierarchy->GetGlobalTransform(Socket->GetIndex());
+						ElementTransform = ElementTransform * ComponentTransform;
 
-						const FMaterialRenderProxy* MaterialProxy;
-						if(const TStrongObjectPtr<UMaterialInstanceDynamic>* ExistingMaterialPtr = SocketMaterials.Find(ColorHash))
-						{
-							MaterialProxy = (*ExistingMaterialPtr)->GetRenderProxy();
-						}
-						else
-						{
-							UMaterial* MaterialBase = GEngine->ArrowMaterial;
-							UMaterialInstanceDynamic* SocketMaterial = UMaterialInstanceDynamic::Create(MaterialBase, NULL);
-							SocketMaterial->SetVectorParameterValue("GizmoColor", Color);
-							SocketMaterials.Add(ColorHash, TStrongObjectPtr<UMaterialInstanceDynamic>(SocketMaterial));
-							MaterialProxy = SocketMaterial->GetRenderProxy();
-						}
+						PDI->DrawLine(ElementTransform.GetTranslation(), ElementTransform.TransformPosition(FVector(Scale, 0.f, 0.f)), SocketColor, SDPG_Foreground);
+						PDI->DrawLine(ElementTransform.GetTranslation(), ElementTransform.TransformPosition(FVector(0.f, Scale, 0.f)), SocketColor, SDPG_Foreground);
+						PDI->DrawLine(ElementTransform.GetTranslation(), ElementTransform.TransformPosition(FVector(0.f, 0.f, Scale)), SocketColor, SDPG_Foreground);
 
-						DrawSphere(PDI,
-							Hierarchy->GetGlobalTransform(Socket->GetIndex()).GetLocation(), 
-							FRotator::ZeroRotator,
-							FVector::OneVector * 4.f,
-							8, 8,
-							MaterialProxy,
-							SDPG_Foreground);
 						return true;
 					});
 				}
diff --git a/Engine/Plugins/Animation/ControlRig/Source/ControlRigEditor/Public/EditMode/ControlRigEditMode.h b/Engine/Plugins/Animation/ControlRig/Source/ControlRigEditor/Public/EditMode/ControlRigEditMode.h
--- a/Engine/Plugins/Animation/ControlRig/Source/ControlRigEditor/Public/EditMode/ControlRigEditMode.h
+++ b/Engine/Plugins/Animation/ControlRig/Source/ControlRigEditor/Public/EditMode/ControlRigEditMode.h
@@ -13,12 +13,8 @@
 #include "UnrealWidgetFwd.h"
 #include "IPersonaEditMode.h"
 #include "Misc/Guid.h"
-#include "Materials/Material.h"
-#include "Materials/MaterialInstanceDynamic.h"
 #include "ControlRigEditMode.generated.h"
 
-
-
 class FEditorViewportClient;
 class FRigVMEditor;
 class FViewport;
@@ -494,7 +490,6 @@ private:
 	bool bIsConstructionEventRunning;
 	TArray<uint32> LastHierarchyHash;
 	TArray<uint32> LastShapeLibraryHash;
-	TMap<uint32, TStrongObjectPtr<UMaterialInstanceDynamic>> SocketMaterials;
 
 	friend class FControlRigEditorModule;
 	friend class FControlRigEditor;
