Control Rig: Support for marquee selection in the viewport

diff --git a/Engine/Plugins/Animation/ControlRig/Source/ControlRig/Private/Settings/ControlRigSettings.cpp b/Engine/Plugins/Animation/ControlRig/Source/ControlRig/Private/Settings/ControlRigSettings.cpp
--- a/Engine/Plugins/Animation/ControlRig/Source/ControlRig/Private/Settings/ControlRigSettings.cpp
+++ b/Engine/Plugins/Animation/ControlRig/Source/ControlRig/Private/Settings/ControlRigSettings.cpp
@@ -29,6 +29,7 @@ UControlRigEditorSettings::UControlRigEditorSettings(const FObjectInitializer& O
 	BackwardsAndForwardsBorderColor = FLinearColor::Blue;
 	bShowStackedHierarchy = false;
 	MaxStackSize = 16;
+	bLeftMouseDragDoesMarquee = true;
 #endif
 }
 
diff --git a/Engine/Plugins/Animation/ControlRig/Source/ControlRig/Public/Settings/ControlRigSettings.h b/Engine/Plugins/Animation/ControlRig/Source/ControlRig/Public/Settings/ControlRigSettings.h
--- a/Engine/Plugins/Animation/ControlRig/Source/ControlRig/Public/Settings/ControlRigSettings.h
+++ b/Engine/Plugins/Animation/ControlRig/Source/ControlRig/Public/Settings/ControlRigSettings.h
@@ -122,6 +122,12 @@ public:
 	UPROPERTY(EditAnywhere, config, Category = Hierarchy, meta = (EditCondition = "bShowStackedHierarchy"))
 	int32 MaxStackSize;
 
+	/**
+	 * If turned on we'll offer box / marquee selection in the control rig editor viewport.
+	 */
+	UPROPERTY(EditAnywhere, config, Category = Hierarchy)
+	bool bLeftMouseDragDoesMarquee;
+
 #endif
 
 	static UControlRigEditorSettings * Get() { return GetMutableDefault<UControlRigEditorSettings>(); }
diff --git a/Engine/Plugins/Animation/ControlRig/Source/ControlRigEditor/Private/EditMode/ControlRigEditMode.cpp b/Engine/Plugins/Animation/ControlRig/Source/ControlRigEditor/Private/EditMode/ControlRigEditMode.cpp
--- a/Engine/Plugins/Animation/ControlRig/Source/ControlRigEditor/Private/EditMode/ControlRigEditMode.cpp
+++ b/Engine/Plugins/Animation/ControlRig/Source/ControlRigEditor/Private/EditMode/ControlRigEditMode.cpp
@@ -68,6 +68,9 @@
 #include "Animation/DebugSkelMeshComponent.h"
 #include "Materials/Material.h"
 #include "ControlRigEditorStyle.h"
+#include "DragTool_BoxSelect.h"
+#include "DragTool_FrustumSelect.h"
+#include "AnimationEditorViewportClient.h"
 
 #include UE_INLINE_GENERATED_CPP_BY_NAME(ControlRigEditMode)
 
@@ -762,7 +765,9 @@ TSet<FName> FControlRigEditMode::GetActiveControlsFromSequencer(UControlRig* Con
 
 
 void FControlRigEditMode::Render(const FSceneView* View, FViewport* Viewport, FPrimitiveDrawInterface* PDI)
-{	
+{
+	DragToolHandler.Render3DDragTool(View, PDI);
+
 	const UControlRigEditModeSettings* Settings = GetDefault<UControlRigEditModeSettings>();
 	const bool bIsInGameView = !AreEditingControlRigDirectly() ? (ViewportToGameView.Find(Viewport) && ViewportToGameView[Viewport]) : false;
 	bool bRender = !Settings->bHideControlShapes;
@@ -948,6 +953,12 @@ void FControlRigEditMode::Render(const FSceneView* View, FViewport* Viewport, FP
 	}
 }
 
+void FControlRigEditMode::DrawHUD(FEditorViewportClient* ViewportClient, FViewport* Viewport, const FSceneView* View, FCanvas* Canvas)
+{
+	IPersonaEditMode::DrawHUD(ViewportClient, Viewport, View, Canvas);
+	DragToolHandler.RenderDragTool(View, Canvas);
+}
+
 bool FControlRigEditMode::InputKey(FEditorViewportClient* InViewportClient, FViewport* InViewport, FKey InKey, EInputEvent InEvent)
 {
 	if (InEvent != IE_Released)
@@ -1025,6 +1036,15 @@ bool FControlRigEditMode::EndTracking(FEditorViewportClient* InViewportClient, F
 		return true;
 	}
 
+	if (IsMovingCamera(InViewport))
+	{
+		return true;
+	}
+	if (DragToolHandler.EndTracking(InViewportClient, InViewport))
+	{
+		return true;
+	}
+
 	const bool bWasInteracting = bManipulatorMadeChange && InteractionType != (uint8)EControlRigInteractionType::None;
 	
 	InteractionType = (uint8)EControlRigInteractionType::None;
@@ -1081,6 +1101,17 @@ bool FControlRigEditMode::StartTracking(FEditorViewportClient* InViewportClient,
 		return true;
 	}
 
+	if (IsMovingCamera(InViewport))
+	{
+		InViewportClient->SetCurrentWidgetAxis(EAxisList::None);
+		return true;
+	}
+	if (IsDoingDrag(InViewport))
+	{
+		DragToolHandler.MakeDragTool(InViewportClient);
+		return DragToolHandler.StartTracking(InViewportClient, InViewport);
+	}
+
 	InteractionType = GetInteractionType(InViewportClient);
 	bIsTracking = true;
 	
@@ -1761,10 +1792,117 @@ bool FControlRigEditMode::FrustumSelect(const FConvexVolume& InFrustum, FEditorV
 			}
 		}
 	}
+
+	EWorldType::Type WorldType = InViewportClient->GetWorld()->WorldType;
+	const bool bIsAssetEditor = WorldType == EWorldType::Editor || WorldType == EWorldType::EditorPreview;
+
+	if (bIsAssetEditor)
+	{
+		float BoneRadius = 1;
+		EBoneDrawMode::Type BoneDrawMode = EBoneDrawMode::None;
+		if (const FAnimationViewportClient* AnimViewportClient = static_cast<FAnimationViewportClient*>(InViewportClient))
+		{
+			BoneDrawMode = AnimViewportClient->GetBoneDrawMode();
+			BoneRadius = AnimViewportClient->GetBoneDrawSize();
+		}
+
+		if(BoneDrawMode != EBoneDrawMode::None)
+		{
+			for(TWeakObjectPtr<UControlRig> WeakControlRig : RuntimeControlRigs)
+			{
+				if(UControlRig* ControlRig = WeakControlRig.Get())
+				{
+					if(URigHierarchy* Hierarchy = ControlRig->GetHierarchy())
+					{
+						TArray<FRigBoneElement*> Bones = Hierarchy->GetBones();
+						for(int32 Index = 0; Index < Bones.Num(); Index++)
+						{
+							const int32 BoneIndex = Bones[Index]->GetIndex();
+							const TArray<int32> Children = Hierarchy->GetChildren(BoneIndex);
+
+							const FVector Start = Hierarchy->GetGlobalTransform(BoneIndex).GetLocation();
+
+							if(InFrustum.IntersectSphere(Start, 0.1f * BoneRadius))
+							{
+								bSomethingSelected = true;
+								SetRigElementSelection(ControlRig, ERigElementType::Bone, Bones[Index]->GetFName(), true);
+								continue;
+							}
+
+							bool bSelectedBone = false;
+							for(int32 ChildIndex : Children)
+							{
+								if(Hierarchy->Get(ChildIndex)->GetType() != ERigElementType::Bone)
+								{
+									continue;
+								}
+								
+								const FVector End = Hierarchy->GetGlobalTransform(ChildIndex).GetLocation();
+
+								const float BoneLength = (End - Start).Size();
+								const float Radius = FMath::Max(BoneLength * 0.05f, 0.1f) * BoneRadius;
+								const int32 Steps = FMath::CeilToInt(BoneLength / (Radius * 1.5f) + 0.5);
+								const FVector Step = (End - Start) / FVector::FReal(Steps - 1);
+
+								// intersect segment-wise along the bone
+								FVector Position = Start;
+								for(int32 StepIndex = 0; StepIndex < Steps; StepIndex++)
+								{
+									if(InFrustum.IntersectSphere(Position, Radius))
+									{
+										bSomethingSelected = true;
+										bSelectedBone = true;
+										SetRigElementSelection(ControlRig, ERigElementType::Bone, Bones[Index]->GetFName(), true);
+										break;
+									}
+									Position += Step;
+								}
+
+								if(bSelectedBone)
+								{
+									break;
+								}
+							}
+						}
+					}
+				}
+			}
+		}
+
+		for(TWeakObjectPtr<UControlRig> WeakControlRig : RuntimeControlRigs)
+		{
+			if(UControlRig* ControlRig = WeakControlRig.Get())
+			{
+				if (Settings->bDisplayNulls || ControlRig->IsConstructionModeEnabled())
+				{
+					if(URigHierarchy* Hierarchy = ControlRig->GetHierarchy())
+					{
+						TArray<FRigNullElement*> Nulls = Hierarchy->GetNulls();
+						for(int32 Index = 0; Index < Nulls.Num(); Index++)
+						{
+							const int32 NullIndex = Nulls[Index]->GetIndex();
+
+							const FTransform Transform = Hierarchy->GetGlobalTransform(NullIndex);
+							const FVector Origin = Transform.GetLocation();
+							const float MaxScale = Transform.GetMaximumAxisScale();
+
+							if(InFrustum.IntersectSphere(Origin, MaxScale * Settings->AxisScale))
+							{
+								bSomethingSelected = true;
+								SetRigElementSelection(ControlRig, ERigElementType::Null, Nulls[Index]->GetFName(), true);
+							}
+						}
+					}
+				}
+			}
+		}
+	}
+
 	if (bSomethingSelected == true)
 	{
 		return true;
 	}
+	
 	ScopedTransaction.Cancel();
 	//if only selecting controls return true to stop any more selections
 	if (Settings && Settings->bOnlySelectRigControls)
@@ -1781,22 +1919,29 @@ void FControlRigEditMode::SelectNone()
 	FEdMode::SelectNone();
 }
 
-static bool IsDoingDrag(const TWeakPtr<ISequencer>& WeakSequencer, FViewport* InViewport, EAxisList::Type CurrentAxis)
+bool FControlRigEditMode::IsMovingCamera(FViewport* InViewport) const
 {
-	if (WeakSequencer.IsValid())
-	{
-		ISequencer* Sequencer = WeakSequencer.Pin().Get();
-		const USequencerSettings* SequencerSettings = Sequencer->GetSequencerSettings();
-		const bool LeftMouseButtonDown = InViewport->KeyState(EKeys::LeftMouseButton);
-		const bool bIsCtrlKeyDown = InViewport->KeyState(EKeys::LeftControl) || InViewport->KeyState(EKeys::RightControl);
-		const bool bIsAltKeyDown = InViewport->KeyState(EKeys::LeftAlt) || InViewport->KeyState(EKeys::RightAlt);
-
-		//if shfit is down we still want to drag
+	const bool LeftMouseButtonDown = InViewport->KeyState(EKeys::LeftMouseButton);
+	const bool bIsAltKeyDown = InViewport->KeyState(EKeys::LeftAlt) || InViewport->KeyState(EKeys::RightAlt);
+	return LeftMouseButtonDown && bIsAltKeyDown;
+}
 
-		return LeftMouseButtonDown && (CurrentAxis == EAxisList::None) && !bIsCtrlKeyDown && !bIsAltKeyDown && (SequencerSettings ? SequencerSettings->GetLeftMouseDragDoesMarquee() : false);
+bool FControlRigEditMode::IsDoingDrag(FViewport* InViewport) const
+{
+	if(!UControlRigEditorSettings::Get()->bLeftMouseDragDoesMarquee)
+	{
+		return false;
 	}
-	return false;
+	
+	const bool LeftMouseButtonDown = InViewport->KeyState(EKeys::LeftMouseButton);
+	const bool bIsCtrlKeyDown = InViewport->KeyState(EKeys::LeftControl) || InViewport->KeyState(EKeys::RightControl);
+	const bool bIsAltKeyDown = InViewport->KeyState(EKeys::LeftAlt) || InViewport->KeyState(EKeys::RightAlt);
+	EAxisList::Type CurrentAxis = GetCurrentWidgetAxis();
+	
+	//if shift is down we still want to drag
+	return LeftMouseButtonDown && (CurrentAxis == EAxisList::None) && !bIsCtrlKeyDown && !bIsAltKeyDown;
 }
+
 bool FControlRigEditMode::InputDelta(FEditorViewportClient* InViewportClient, FViewport* InViewport, FVector& InDrag, FRotator& InRot, FVector& InScale)
 {
 	if (IsDragAnimSliderToolPressed(InViewport)) //this is needed to make sure we get all of the processed mouse events, for some reason the above may not return true
@@ -1805,6 +1950,10 @@ bool FControlRigEditMode::InputDelta(FEditorViewportClient* InViewportClient, FV
 		return true;
 	}
 
+	if (IsDoingDrag(InViewport))
+	{
+		return DragToolHandler.InputDelta(InViewportClient, InViewport, InDrag, InRot, InScale);
+	}
 
 	FVector Drag = InDrag;
 	FRotator Rot = InRot;
@@ -4845,6 +4994,96 @@ void FControlRigEditMode::OnEditorClosed()
 	ControlRigsToRecreate.Reset();
 }
 
+FControlRigEditMode::FMarqueeDragTool::FMarqueeDragTool()
+	: bIsDeletingDragTool(false)
+{
+}
+
+bool FControlRigEditMode::FMarqueeDragTool::StartTracking(FEditorViewportClient* InViewportClient, FViewport* InViewport)
+{
+	return (DragTool.IsValid() && InViewportClient->GetCurrentWidgetAxis() == EAxisList::None);
+}
+
+bool FControlRigEditMode::FMarqueeDragTool::EndTracking(FEditorViewportClient* InViewportClient, FViewport* InViewport)
+{
+	if (!bIsDeletingDragTool)
+	{
+		// Ending the drag tool may pop up a modal dialog which can cause unwanted reentrancy - protect against this.
+		TGuardValue<bool> RecursionGuard(bIsDeletingDragTool, true);
+
+		// Delete the drag tool if one exists.
+		if (DragTool.IsValid())
+		{
+			if (DragTool->IsDragging())
+			{
+				DragTool->EndDrag();
+			}
+			DragTool.Reset();
+			return true;
+		}
+	}
+	
+	return false;
+}
+
+void FControlRigEditMode::FMarqueeDragTool::MakeDragTool(FEditorViewportClient* InViewportClient)
+{
+	DragTool.Reset();
+	if (InViewportClient->IsOrtho())
+	{
+		DragTool = MakeShareable( new FDragTool_ActorBoxSelect(InViewportClient) );
+	}
+	else
+	{
+		DragTool = MakeShareable( new FDragTool_ActorFrustumSelect(InViewportClient) );
+	}
+}
+
+bool FControlRigEditMode::FMarqueeDragTool::InputDelta(FEditorViewportClient* InViewportClient, FViewport* InViewport, FVector& InDrag, FRotator& InRot, FVector& InScale)
+{
+	if (DragTool.IsValid() == false || InViewportClient->GetCurrentWidgetAxis() != EAxisList::None)
+	{
+		return false;
+	}
+	if (DragTool->IsDragging() == false)
+	{
+		int32 InX = InViewport->GetMouseX();
+		int32 InY = InViewport->GetMouseY();
+		FVector2D Start(InX, InY);
+
+		DragTool->StartDrag(InViewportClient, GEditor->ClickLocation,Start);
+	}
+	const bool bUsingDragTool = UsingDragTool();
+	if (bUsingDragTool == false)
+	{
+		return false;
+	}
+
+	DragTool->AddDelta(InDrag);
+	return true;
+}
+
+bool FControlRigEditMode::FMarqueeDragTool::UsingDragTool() const
+{
+	return DragTool.IsValid() && DragTool->IsDragging();
+}
+
+void FControlRigEditMode::FMarqueeDragTool::Render3DDragTool(const FSceneView* View, FPrimitiveDrawInterface* PDI)
+{
+	if (DragTool.IsValid())
+	{
+		DragTool->Render3D(View, PDI);
+	}
+}
+
+void FControlRigEditMode::FMarqueeDragTool::RenderDragTool(const FSceneView* View, FCanvas* Canvas)
+{
+	if (DragTool.IsValid())
+	{
+		DragTool->Render(View, Canvas);
+	}
+}
+
 void FControlRigEditMode::DestroyShapesActors(UControlRig* ControlRig)
 {
 	if (ControlRig == nullptr)
diff --git a/Engine/Plugins/Animation/ControlRig/Source/ControlRigEditor/Public/EditMode/ControlRigEditMode.h b/Engine/Plugins/Animation/ControlRig/Source/ControlRigEditor/Public/EditMode/ControlRigEditMode.h
--- a/Engine/Plugins/Animation/ControlRig/Source/ControlRigEditor/Public/EditMode/ControlRigEditMode.h
+++ b/Engine/Plugins/Animation/ControlRig/Source/ControlRigEditor/Public/EditMode/ControlRigEditMode.h
@@ -13,6 +13,7 @@
 #include "UnrealWidgetFwd.h"
 #include "IPersonaEditMode.h"
 #include "Misc/Guid.h"
+#include "EditorDragTools.h"
 #include "ControlRigEditMode.generated.h"
 
 class FEditorViewportClient;
@@ -111,6 +112,7 @@ public:
 	virtual void Exit() override;
 	virtual void Tick(FEditorViewportClient* ViewportClient, float DeltaTime) override;
 	virtual void Render(const FSceneView* View, FViewport* Viewport, FPrimitiveDrawInterface* PDI) override;
+	virtual void DrawHUD(FEditorViewportClient* ViewportClient, FViewport* Viewport, const FSceneView* View, FCanvas* Canvas) override;
 	virtual bool InputKey(FEditorViewportClient* InViewportClient, FViewport* InViewport, FKey InKey, EInputEvent InEvent) override;
 	virtual bool EndTracking(FEditorViewportClient* InViewportClient, FViewport* InViewport) override;
 	virtual bool StartTracking(FEditorViewportClient* InViewportClient, FViewport* InViewport) override;
@@ -467,12 +469,45 @@ private:
 	void NotifyDrivenControls(UControlRig* InControlRig, const FRigElementKey& InKey);
 	void UpdateSelectabilityOnSkeletalMeshes(UControlRig* InControlRig, bool bEnabled);
 
+	bool IsMovingCamera(FViewport* InViewport) const;
+	bool IsDoingDrag(FViewport* InViewport) const;
+
 	// world clean up handlers
 	FDelegateHandle OnWorldCleanupHandle;
 	void OnWorldCleanup(UWorld* World, bool bSessionEnded, bool bCleanupResources);
 	UWorld* WorldPtr = nullptr;
 
 	void OnEditorClosed();
+
+	struct FMarqueeDragTool
+    {
+    	FMarqueeDragTool();
+    	~FMarqueeDragTool() {};
+    
+    	bool StartTracking(FEditorViewportClient* InViewportClient, FViewport* InViewport);
+    	bool EndTracking(FEditorViewportClient* InViewportClient, FViewport* InViewport);
+    	void MakeDragTool(FEditorViewportClient* InViewportClient);
+    	bool InputDelta(FEditorViewportClient* InViewportClient, FViewport* InViewport, FVector& InDrag, FRotator& InRot, FVector& InScale);
+    
+    	bool UsingDragTool() const;
+    	void Render3DDragTool(const FSceneView* View, FPrimitiveDrawInterface* PDI);
+    	void RenderDragTool(const FSceneView* View, FCanvas* Canvas);
+    
+    private:
+    	
+    	/**
+    	 * If there is a dragging tool being used, this will point to it.
+    	 * Gets newed/deleted in StartTracking/EndTracking.
+    	 */
+    	TSharedPtr<FDragTool> DragTool;
+    
+    	/** Tracks whether the drag tool is in the process of being deleted (to protect against reentrancy) */
+    	bool bIsDeletingDragTool = false;
+
+		FControlRigEditMode* EditMode;
+    };
+
+	FMarqueeDragTool DragToolHandler;
 	
 	TArray<TWeakObjectPtr<UControlRig>> RuntimeControlRigs;
 	TMap<UControlRig*,TStrongObjectPtr<UControlRigEditModeDelegateHelper>> DelegateHelpers;
diff --git a/Engine/Source/Editor/UnrealEd/Private/DragTool_BoxSelect.cpp b/Engine/Source/Editor/UnrealEd/Private/DragTool_BoxSelect.cpp
--- a/Engine/Source/Editor/UnrealEd/Private/DragTool_BoxSelect.cpp
+++ b/Engine/Source/Editor/UnrealEd/Private/DragTool_BoxSelect.cpp
@@ -34,10 +34,11 @@ namespace UE::LevelEditor::Private
 {
 	TArray<FTypedElementHandle> GetElementsIntersectingBox(const AActor* Actor,
 		const FBox& InBox,
-		FLevelEditorViewportClient* Viewport,
+		const FEditorViewportClient* InEditorViewport,
+		const FLevelEditorViewportClient* InLevelViewport,
 		const FWorldSelectionElementArgs& SelectionArgs)
 	{
-		if (Actor && (!Viewport || !Actor->IsA(AVolume::StaticClass()) || !Viewport->IsVolumeVisibleInViewport(*Actor)))
+		if (Actor && (!InEditorViewport || !Actor->IsA(AVolume::StaticClass()) || (InLevelViewport ? !InLevelViewport->IsVolumeVisibleInViewport(*Actor) : false)))
 		{
 			if (FTypedElementHandle ActorHandle = UEngineElementsLibrary::AcquireEditorActorElementHandle(Actor))
 			{
@@ -99,7 +100,7 @@ void FDragTool_ActorBoxSelect::AddDelta( const FVector& InDelta )
 	FDragTool::AddDelta( InDelta );
 
 	FIntPoint MousePos;
-	LevelViewportClient->Viewport->GetMousePos(MousePos);
+	EditorViewportClient->Viewport->GetMousePos(MousePos);
 
 	End = FVector(MousePos); 
 	EndWk = End;
@@ -120,7 +121,7 @@ void FDragTool_ActorBoxSelect::AddDelta( const FVector& InDelta )
 			SelectionSet,
 			ETypedElementSelectionMethod::Primary,
 			FTypedElementSelectionOptions(),
-			&(LevelViewportClient->EngineShowFlags),
+			&(EditorViewportClient->EngineShowFlags),
 			bStrictDragSelection,
 			bGeometryMode
 		};
@@ -137,7 +138,7 @@ void FDragTool_ActorBoxSelect::AddDelta( const FVector& InDelta )
 		for( FActorIterator It(IteratorWorld); It; ++It )
 		{
 			AActor& Actor = **It;
-			const bool bActorHitByBox = !UE::LevelEditor::Private::GetElementsIntersectingBox( &Actor, SelBBox, LevelViewportClient, SeletionArgs ).IsEmpty();
+			const bool bActorHitByBox = !UE::LevelEditor::Private::GetElementsIntersectingBox( &Actor, SelBBox, EditorViewportClient, LevelViewportClient, SeletionArgs ).IsEmpty();
 
 
 			if( bActorHitByBox )
@@ -201,7 +202,7 @@ void FDragTool_ActorBoxSelect::EndDrag()
 	const bool bEditorModeHandledBoxSelection = ModeTools->BoxSelect(SelBBox, bLeftMouseButtonDown);
 
 	// Let the component visualizers try to handle the selection.
-	const bool bComponentVisHandledSelection = !bEditorModeHandledBoxSelection && GUnrealEd->ComponentVisManager.HandleBoxSelect(SelBBox, LevelViewportClient, LevelViewportClient->Viewport);
+	const bool bComponentVisHandledSelection = !bEditorModeHandledBoxSelection && GUnrealEd->ComponentVisManager.HandleBoxSelect(SelBBox, EditorViewportClient, EditorViewportClient->Viewport);
 
 	// If the edit mode didn't handle the selection, try normal actor box selection.
 	if ( !bEditorModeHandledBoxSelection && !bComponentVisHandledSelection )
@@ -221,7 +222,7 @@ void FDragTool_ActorBoxSelect::EndDrag()
 			SelectionSet,
 			ETypedElementSelectionMethod::Primary,
 			ElementSelectionOption,
-			&(LevelViewportClient->EngineShowFlags),
+			&(EditorViewportClient->EngineShowFlags),
 			bStrictDragSelection,
 			bGeometryMode
 		};
@@ -229,7 +230,7 @@ void FDragTool_ActorBoxSelect::EndDrag()
 		// Select all element that are within the selection box area.  Be aware that certain modes do special processing below.	
 		bool bSelectionChanged = false;
 		UWorld* IteratorWorld = GWorld;
-		const TArray<FName>& HiddenLayers = LevelViewportClient->ViewHiddenLayers;
+		const TArray<FName>& HiddenLayers = LevelViewportClient ? LevelViewportClient->ViewHiddenLayers : TArray<FName>();
 		TArray<FTypedElementHandle> Handles;
 
 		for( FActorIterator It(IteratorWorld); It; ++It )
@@ -250,7 +251,7 @@ void FDragTool_ActorBoxSelect::EndDrag()
 			// Select the actor or its child elements
 			if( bActorIsVisible )
 			{
-				Handles.Append( UE::LevelEditor::Private::GetElementsIntersectingBox( Actor, SelBBox, LevelViewportClient, SeletionArgs ) );
+				Handles.Append( UE::LevelEditor::Private::GetElementsIntersectingBox( Actor, SelBBox, EditorViewportClient, LevelViewportClient, SeletionArgs ) );
 			}
 		}
 
@@ -303,12 +304,12 @@ void FDragTool_ActorBoxSelect::Render(const FSceneView* View, FCanvas* Canvas)
 void FDragTool_ActorBoxSelect::CalculateBox( FBox& OutBox )
 {
 	FSceneViewFamilyContext ViewFamily(FSceneViewFamily::ConstructionValues(
-		LevelViewportClient->Viewport,
-		LevelViewportClient->GetScene(),
-		LevelViewportClient->EngineShowFlags)
-		.SetRealtimeUpdate(LevelViewportClient->IsRealtime()));
+		EditorViewportClient->Viewport,
+		EditorViewportClient->GetScene(),
+		EditorViewportClient->EngineShowFlags)
+		.SetRealtimeUpdate(EditorViewportClient->IsRealtime()));
 
-	FSceneView* View = LevelViewportClient->CalcSceneView(&ViewFamily);
+	FSceneView* View = EditorViewportClient->CalcSceneView(&ViewFamily);
 
 	FVector3f StartFloat{ Start };
 	FVector3f EndFloat{ End };
@@ -324,7 +325,7 @@ void FDragTool_ActorBoxSelect::CalculateBox( FBox& OutBox )
 	OutBox += TransformedStart;
 	OutBox += TransformedEnd;
 
-	switch(LevelViewportClient->ViewportType)
+	switch(EditorViewportClient->ViewportType)
 	{
 	case LVT_OrthoXY:
 	case LVT_OrthoNegativeXY:
diff --git a/Engine/Source/Editor/UnrealEd/Private/DragTool_FrustumSelect.cpp b/Engine/Source/Editor/UnrealEd/Private/DragTool_FrustumSelect.cpp
--- a/Engine/Source/Editor/UnrealEd/Private/DragTool_FrustumSelect.cpp
+++ b/Engine/Source/Editor/UnrealEd/Private/DragTool_FrustumSelect.cpp
@@ -36,10 +36,11 @@ namespace UE::LevelEditor::Private
 {
 	TArray<FTypedElementHandle> GetElementsIntersectingFrustum(const AActor* Actor,
 		const FConvexVolume& InFrustum,
-		FLevelEditorViewportClient* Viewport,
+		const FEditorViewportClient* EditorViewport,
+		const FLevelEditorViewportClient* LevelViewport,
 		const FWorldSelectionElementArgs& SelectionArgs)
 	{
-		if (Actor && (!Viewport || !Actor->IsA(AVolume::StaticClass()) || !Viewport->IsVolumeVisibleInViewport(*Actor)))
+		if (Actor && (!EditorViewport || !Actor->IsA(AVolume::StaticClass()) || (LevelViewport ? !LevelViewport->IsVolumeVisibleInViewport(*Actor) : false)))
 		{
 			if (FTypedElementHandle ActorHandle = UEngineElementsLibrary::AcquireEditorActorElementHandle(Actor))
 			{
@@ -57,7 +58,7 @@ namespace UE::LevelEditor::Private
 void FDragTool_ActorFrustumSelect::AddDelta( const FVector& InDelta )
 {
 	FIntPoint MousePos;
-	LevelViewportClient->Viewport->GetMousePos(MousePos);
+	EditorViewportClient->Viewport->GetMousePos(MousePos);
 
 	EndWk = FVector(MousePos);
 	End = EndWk;
@@ -86,8 +87,8 @@ void FDragTool_ActorFrustumSelect::EndDrag()
 	UBrushEditingSubsystem* BrushSubsystem = GEditor->GetEditorSubsystem<UBrushEditingSubsystem>();
 	const bool bGeometryMode = BrushSubsystem ? BrushSubsystem->IsGeometryEditorModeActive() : false;
 
-	FSceneViewFamilyContext ViewFamily(FSceneViewFamily::ConstructionValues(LevelViewportClient->Viewport, LevelViewportClient->GetScene(), LevelViewportClient->EngineShowFlags ));
-	FSceneView* SceneView = LevelViewportClient->CalcSceneView(&ViewFamily);
+	FSceneViewFamilyContext ViewFamily(FSceneViewFamily::ConstructionValues(EditorViewportClient->Viewport, EditorViewportClient->GetScene(), EditorViewportClient->EngineShowFlags ));
+	FSceneView* SceneView = EditorViewportClient->CalcSceneView(&ViewFamily);
 
 	// Generate a frustum out of the dragged box
 	FConvexVolume Frustum;
@@ -107,10 +108,10 @@ void FDragTool_ActorFrustumSelect::EndDrag()
 	const bool bStrictDragSelection = GetDefault<ULevelEditorViewportSettings>()->bStrictBoxSelection;
 
 	// Let the editor mode try to handle the selection.
-	const bool bEditorModeHandledSelection = ModeTools->FrustumSelect(Frustum, LevelViewportClient, bLeftMouseButtonDown);
+	const bool bEditorModeHandledSelection = ModeTools->FrustumSelect(Frustum, EditorViewportClient, bLeftMouseButtonDown);
 
 	// Let the component visualizers try to handle the selection.
-	const bool bComponentVisHandledSelection = !bEditorModeHandledSelection && GUnrealEd->ComponentVisManager.HandleFrustumSelect(Frustum, LevelViewportClient, LevelViewportClient->Viewport);
+	const bool bComponentVisHandledSelection = !bEditorModeHandledSelection && GUnrealEd->ComponentVisManager.HandleFrustumSelect(Frustum, EditorViewportClient, EditorViewportClient->Viewport);
 
 	if( !bEditorModeHandledSelection && !bComponentVisHandledSelection)
 	{
@@ -122,13 +123,13 @@ void FDragTool_ActorFrustumSelect::EndDrag()
 			SelectionSet,
 			ETypedElementSelectionMethod::Primary,
 			FTypedElementSelectionOptions(),
-			&(LevelViewportClient->EngineShowFlags),
+			&(EditorViewportClient->EngineShowFlags),
 			bStrictDragSelection,
 			bGeometryMode
 		};
 
-		const int32 ViewportSizeX = LevelViewportClient->Viewport->GetSizeXY().X;
-		const int32 ViewportSizeY = LevelViewportClient->Viewport->GetSizeXY().Y;
+		const int32 ViewportSizeX = EditorViewportClient->Viewport->GetSizeXY().X;
+		const int32 ViewportSizeY = EditorViewportClient->Viewport->GetSizeXY().Y;
 
 		if( Start.X > End.X )
 		{
@@ -145,10 +146,10 @@ void FDragTool_ActorFrustumSelect::EndDrag()
 		if (bTransparentBoxSelection)
 		{
 			// Get a list of frustum-culled actors
-			for(FActorIterator It(LevelViewportClient->GetWorld()); It; ++It)
+			for(FActorIterator It(EditorViewportClient->GetWorld()); It; ++It)
 			{
 				AActor* Actor = *It;
-				ElementsToSelect.Append(UE::LevelEditor::Private::GetElementsIntersectingFrustum(Actor, Frustum, LevelViewportClient, SeletionArgs));
+				ElementsToSelect.Append(UE::LevelEditor::Private::GetElementsIntersectingFrustum(Actor, Frustum, EditorViewportClient, LevelViewportClient, SeletionArgs));
 			}
 		}
 		else
@@ -167,7 +168,7 @@ void FDragTool_ActorFrustumSelect::EndDrag()
 			// Typed Element selection
 			{
 				FTypedElementListRef ElementList = UTypedElementRegistry::GetInstance()->CreateElementList();
-				LevelViewportClient->Viewport->GetElementHandlesInRect(BoxRect, ElementList);
+				EditorViewportClient->Viewport->GetElementHandlesInRect(BoxRect, ElementList);
 
 				if (bStrictDragSelection)
 				{
@@ -196,7 +197,7 @@ void FDragTool_ActorFrustumSelect::EndDrag()
 			// We need this old code to support the BSP
 			TSet<AActor*> BSPActors;
 			TSet<UModel*> HitModels;
-			LevelViewportClient->Viewport->GetActorsAndModelsInHitProxy( BoxRect, BSPActors, HitModels );
+			EditorViewportClient->Viewport->GetActorsAndModelsInHitProxy( BoxRect, BSPActors, HitModels );
 			BSPActors.Empty(HitModels.Num());
 
 			
@@ -227,7 +228,7 @@ void FDragTool_ActorFrustumSelect::EndDrag()
 					AActor* Actor = *It;
 					if (bStrictDragSelection)
 					{
-						ElementsToSelect.Append(UE::LevelEditor::Private::GetElementsIntersectingFrustum(Actor, Frustum, LevelViewportClient, SeletionArgs));
+						ElementsToSelect.Append(UE::LevelEditor::Private::GetElementsIntersectingFrustum(Actor, Frustum, EditorViewportClient, LevelViewportClient, SeletionArgs));
 					}
 					else
 					{
@@ -281,7 +282,7 @@ void FDragTool_ActorFrustumSelect::CalculateFrustum( FSceneView* View, FConvexVo
 {
 	if( bUseBoxFrustum )
 	{
-		FVector CamPoint = LevelViewportClient->GetViewLocation();
+		FVector CamPoint = EditorViewportClient->GetViewLocation();
 		FVector BoxPoint1, BoxPoint2, BoxPoint3, BoxPoint4;
 		FVector WorldDir1, WorldDir2, WorldDir3, WorldDir4;
 		// Deproject the four corners of the selection box
diff --git a/Engine/Source/Editor/UnrealEd/Public/DragTool_BoxSelect.h b/Engine/Source/Editor/UnrealEd/Public/DragTool_BoxSelect.h
--- a/Engine/Source/Editor/UnrealEd/Public/DragTool_BoxSelect.h
+++ b/Engine/Source/Editor/UnrealEd/Public/DragTool_BoxSelect.h
@@ -17,9 +17,16 @@ class UModel;
 class FDragTool_ActorBoxSelect : public FDragTool
 {
 public:
-	explicit FDragTool_ActorBoxSelect( FLevelEditorViewportClient* InLevelViewportClient )
+	UNREALED_API explicit FDragTool_ActorBoxSelect( FLevelEditorViewportClient* InLevelViewportClient )
 		: FDragTool(InLevelViewportClient->GetModeTools())
 		, LevelViewportClient( InLevelViewportClient )
+		, EditorViewportClient( InLevelViewportClient)
+	{}
+
+	UNREALED_API explicit FDragTool_ActorBoxSelect( FEditorViewportClient* InEditorViewportClient )
+		: FDragTool(InEditorViewportClient->GetModeTools())
+		, LevelViewportClient( nullptr)
+		, EditorViewportClient( InEditorViewportClient )
 	{}
 
 	/**
@@ -28,20 +35,20 @@ public:
 	 * @param	InViewportClient	The viewport client in which the drag event occurred.
 	 * @param	InStart				Where the mouse was when the drag started.
 	 */
-	virtual void StartDrag(FEditorViewportClient* InViewportClient, const FVector& InStart, const FVector2D& InStartScreen) override;
+	UNREALED_API virtual void StartDrag(FEditorViewportClient* InViewportClient, const FVector& InStart, const FVector2D& InStartScreen) override;
 
 	/* Updates the drag tool's end location with the specified delta.  The end location is
 	 * snapped to the editor constraints if bUseSnapping is true.
 	 *
 	 * @param	InDelta		A delta of mouse movement.
 	 */
-	virtual void AddDelta(const FVector& InDelta) override;
+	UNREALED_API virtual void AddDelta(const FVector& InDelta) override;
 
 	/**
 	* Ends a mouse drag behavior (the user has let go of the mouse button).
 	*/
-	virtual void EndDrag() override;
-	virtual void Render(const FSceneView* View, FCanvas* Canvas) override;
+	UNREALED_API virtual void EndDrag() override;
+	UNREALED_API virtual void Render(const FSceneView* View, FCanvas* Canvas) override;
 
 private:
 	/** 
@@ -77,5 +84,6 @@ private:
 	TArray<UModel*> ModelsToCheck;
 
 	FLevelEditorViewportClient* LevelViewportClient;
+	FEditorViewportClient* EditorViewportClient;
 };
 
diff --git a/Engine/Source/Editor/UnrealEd/Public/DragTool_FrustumSelect.h b/Engine/Source/Editor/UnrealEd/Public/DragTool_FrustumSelect.h
--- a/Engine/Source/Editor/UnrealEd/Public/DragTool_FrustumSelect.h
+++ b/Engine/Source/Editor/UnrealEd/Public/DragTool_FrustumSelect.h
@@ -17,9 +17,16 @@ class UModel;
 class FDragTool_ActorFrustumSelect : public FDragTool
 {
 public:
-	explicit FDragTool_ActorFrustumSelect(FLevelEditorViewportClient* InLevelViewportClient)
+	UNREALED_API explicit FDragTool_ActorFrustumSelect(FLevelEditorViewportClient* InLevelViewportClient)
 		: FDragTool(InLevelViewportClient->GetModeTools())
 		, LevelViewportClient( InLevelViewportClient )
+		, EditorViewportClient( InLevelViewportClient)
+	{}
+
+	UNREALED_API explicit FDragTool_ActorFrustumSelect(FEditorViewportClient* InEditorViewportClient)
+	: FDragTool(InEditorViewportClient->GetModeTools())
+	, LevelViewportClient( nullptr )
+	, EditorViewportClient( InEditorViewportClient)
 	{}
 
 	/* Updates the drag tool's end location with the specified delta.  The end location is
@@ -81,4 +88,5 @@ private:
 	TArray<UModel*> ModelsToCheck;
 
 	FLevelEditorViewportClient* LevelViewportClient;
+	FEditorViewportClient* EditorViewportClient;
 };
