Control Rig: Follow up to 30232821

#rnx

diff --git a/Engine/Plugins/Animation/ControlRig/Source/ControlRigEditor/Private/EditMode/ControlRigEditMode.cpp b/Engine/Plugins/Animation/ControlRig/Source/ControlRigEditor/Private/EditMode/ControlRigEditMode.cpp
--- a/Engine/Plugins/Animation/ControlRig/Source/ControlRigEditor/Private/EditMode/ControlRigEditMode.cpp
+++ b/Engine/Plugins/Animation/ControlRig/Source/ControlRigEditor/Private/EditMode/ControlRigEditMode.cpp
@@ -1794,7 +1794,9 @@ bool FControlRigEditMode::FrustumSelect(const FConvexVolume& InFrustum, FEditorV
 	}
 
 	EWorldType::Type WorldType = InViewportClient->GetWorld()->WorldType;
-	const bool bIsAssetEditor = WorldType == EWorldType::Editor || WorldType == EWorldType::EditorPreview;
+	const bool bIsAssetEditor =
+		(WorldType == EWorldType::Editor || WorldType == EWorldType::EditorPreview) &&
+			!InViewportClient->IsLevelEditorClient();
 
 	if (bIsAssetEditor)
 	{
