Control Rig: Fix static analysis error
#rnx

diff --git a/Engine/Plugins/Animation/ControlRig/Source/ControlRigEditor/Private/EditMode/ControlRigEditMode.cpp b/Engine/Plugins/Animation/ControlRig/Source/ControlRigEditor/Private/EditMode/ControlRigEditMode.cpp
--- a/Engine/Plugins/Animation/ControlRig/Source/ControlRigEditor/Private/EditMode/ControlRigEditMode.cpp
+++ b/Engine/Plugins/Animation/ControlRig/Source/ControlRigEditor/Private/EditMode/ControlRigEditMode.cpp
@@ -1746,6 +1746,11 @@ bool FControlRigEditMode::BoxSelect(FBox& InBox, bool InSelect)
 bool FControlRigEditMode::FrustumSelect(const FConvexVolume& InFrustum, FEditorViewportClient* InViewportClient, bool InSelect)
 {
 	const UControlRigEditModeSettings* Settings = GetDefault<UControlRigEditModeSettings>();
+	if (!Settings)
+	{
+		return false;
+	}
+	
 	//need to check for a zero frustum since ComponentIsTouchingSelectionFrustum will return true, selecting everything, when this is the case
 	const bool bMalformedFrustum = (InFrustum.Planes[0].IsNearlyZero() && InFrustum.Planes[2].IsNearlyZero()) || (InFrustum.Planes[3].IsNearlyZero() &&
 		InFrustum.Planes[4].IsNearlyZero());
