ControlRigEditMode: Fix static analysis issue

#rnx

diff --git a/Engine/Plugins/Animation/ControlRig/Source/ControlRigEditor/Private/EditMode/ControlRigEditMode.cpp b/Engine/Plugins/Animation/ControlRig/Source/ControlRigEditor/Private/EditMode/ControlRigEditMode.cpp
--- a/Engine/Plugins/Animation/ControlRig/Source/ControlRigEditor/Private/EditMode/ControlRigEditMode.cpp
+++ b/Engine/Plugins/Animation/ControlRig/Source/ControlRigEditor/Private/EditMode/ControlRigEditMode.cpp
@@ -1912,7 +1912,7 @@ bool FControlRigEditMode::FrustumSelect(const FConvexVolume& InFrustum, FEditorV
 	
 	ScopedTransaction.Cancel();
 	//if only selecting controls return true to stop any more selections
-	if (Settings && Settings->bOnlySelectRigControls)
+	if (Settings->bOnlySelectRigControls)
 	{
 		return true;
 	}
