Control Rig: fixed marquee selection conflicting with new TRS gizmos
    - don't handle tracking if the ITF already handled it

diff --git a/Engine/Plugins/Animation/ControlRig/Source/ControlRigEditor/Private/EditMode/ControlRigEditMode.cpp b/Engine/Plugins/Animation/ControlRig/Source/ControlRigEditor/Private/EditMode/ControlRigEditMode.cpp
--- a/Engine/Plugins/Animation/ControlRig/Source/ControlRigEditor/Private/EditMode/ControlRigEditMode.cpp
+++ b/Engine/Plugins/Animation/ControlRig/Source/ControlRigEditor/Private/EditMode/ControlRigEditMode.cpp
@@ -1926,24 +1926,30 @@ void FControlRigEditMode::SelectNone()
 	FEdMode::SelectNone();
 }
 
-bool FControlRigEditMode::IsMovingCamera(FViewport* InViewport) const
+bool FControlRigEditMode::IsMovingCamera(const FViewport* InViewport) const
 {
 	const bool LeftMouseButtonDown = InViewport->KeyState(EKeys::LeftMouseButton);
 	const bool bIsAltKeyDown = InViewport->KeyState(EKeys::LeftAlt) || InViewport->KeyState(EKeys::RightAlt);
 	return LeftMouseButtonDown && bIsAltKeyDown;
 }
 
-bool FControlRigEditMode::IsDoingDrag(FViewport* InViewport) const
+bool FControlRigEditMode::IsDoingDrag(const FViewport* InViewport) const
 {
 	if(!UControlRigEditorSettings::Get()->bLeftMouseDragDoesMarquee)
 	{
 		return false;
 	}
+
+	if (Owner && Owner->GetInteractiveToolsContext()->InputRouter->HasActiveMouseCapture())
+	{
+		// don't start dragging if the ITF handled tracking event first   
+		return false;
+	}
 	
 	const bool LeftMouseButtonDown = InViewport->KeyState(EKeys::LeftMouseButton);
 	const bool bIsCtrlKeyDown = InViewport->KeyState(EKeys::LeftControl) || InViewport->KeyState(EKeys::RightControl);
 	const bool bIsAltKeyDown = InViewport->KeyState(EKeys::LeftAlt) || InViewport->KeyState(EKeys::RightAlt);
-	EAxisList::Type CurrentAxis = GetCurrentWidgetAxis();
+	const EAxisList::Type CurrentAxis = GetCurrentWidgetAxis();
 	
 	//if shift is down we still want to drag
 	return LeftMouseButtonDown && (CurrentAxis == EAxisList::None) && !bIsCtrlKeyDown && !bIsAltKeyDown;
diff --git a/Engine/Plugins/Animation/ControlRig/Source/ControlRigEditor/Public/EditMode/ControlRigEditMode.h b/Engine/Plugins/Animation/ControlRig/Source/ControlRigEditor/Public/EditMode/ControlRigEditMode.h
--- a/Engine/Plugins/Animation/ControlRig/Source/ControlRigEditor/Public/EditMode/ControlRigEditMode.h
+++ b/Engine/Plugins/Animation/ControlRig/Source/ControlRigEditor/Public/EditMode/ControlRigEditMode.h
@@ -469,8 +469,8 @@ private:
 	void NotifyDrivenControls(UControlRig* InControlRig, const FRigElementKey& InKey);
 	void UpdateSelectabilityOnSkeletalMeshes(UControlRig* InControlRig, bool bEnabled);
 
-	bool IsMovingCamera(FViewport* InViewport) const;
-	bool IsDoingDrag(FViewport* InViewport) const;
+	bool IsMovingCamera(const FViewport* InViewport) const;
+	bool IsDoingDrag(const FViewport* InViewport) const;
 
 	// world clean up handlers
 	FDelegateHandle OnWorldCleanupHandle;
