Control Rig: Fixed issue when switching LODs the rig would return to T-pose for one frame
#rnx

diff --git a/Engine/Plugins/Animation/ControlRig/Source/ControlRig/Private/AnimNode_ControlRigBase.cpp b/Engine/Plugins/Animation/ControlRig/Source/ControlRig/Private/AnimNode_ControlRigBase.cpp
--- a/Engine/Plugins/Animation/ControlRig/Source/ControlRig/Private/AnimNode_ControlRigBase.cpp
+++ b/Engine/Plugins/Animation/ControlRig/Source/ControlRig/Private/AnimNode_ControlRigBase.cpp
@@ -456,6 +456,8 @@ void FAnimNode_ControlRigBase::ExecuteControlRig(FPoseContext& InOutput)
 				return;
 			}
 
+			TGuardValue<bool> ResetCurrentTransfromsAfterConstructionGuard(ControlRig->bResetCurrentTransformsAfterConstruction, true);
+
 			if(Hierarchy->IsTracingChanges())
 			{
 				Hierarchy->StorePoseForTrace(TEXT("FAnimNode_ControlRigBase::BeforeEvaluate"));
diff --git a/Engine/Plugins/Animation/ControlRig/Source/ControlRig/Private/ControlRig.cpp b/Engine/Plugins/Animation/ControlRig/Source/ControlRig/Private/ControlRig.cpp
--- a/Engine/Plugins/Animation/ControlRig/Source/ControlRig/Private/ControlRig.cpp
+++ b/Engine/Plugins/Animation/ControlRig/Source/ControlRig/Private/ControlRig.cpp
@@ -78,6 +78,7 @@ UControlRig::UControlRig(const FObjectInitializer& ObjectInitializer)
 #endif
 	, bCopyHierarchyBeforeConstruction(true)
 	, bResetInitialTransformsBeforeConstruction(true)
+	, bResetCurrentTransformsAfterConstruction(false)
 	, bManipulationEnabled(false)
 	, PreConstructionBracket(0)
 	, PostConstructionBracket(0)
@@ -993,6 +994,13 @@ bool UControlRig::Execute(const FName& InEventName)
 			// disable selection notifications from the hierarchy
 			TGuardValue<bool> DisableSelectionNotifications(GetHierarchy()->GetController(true)->bSuspendSelectionNotifications, true);
 			{
+				FRigPose CurrentPose;
+				// We might want to reset the input pose after construction
+				if (bResetCurrentTransformsAfterConstruction)
+				{
+					CurrentPose = GetHierarchy()->GetPose(false, ERigElementType::ToResetAfterConstructionEvent, FRigElementKeyCollection());
+				}
+				
 				{
 					// Copy the hierarchy from the default object onto this one
 #if WITH_EDITOR
@@ -1047,6 +1055,12 @@ bool UControlRig::Execute(const FName& InEventName)
 					FControlRigBracketScope BracketScope(PostConstructionBracket);
 					PostConstructionEvent.Broadcast(this, FRigUnit_PrepareForExecution::EventName);
 				}
+
+				// Reset the input pose after construction
+				if (CurrentPose.Num() > 0)
+				{
+					GetHierarchy()->SetPose(CurrentPose, ERigTransformType::CurrentLocal);
+				}
 			}
 			
 			// set it here to reestablish the selection. the notifications
diff --git a/Engine/Plugins/Animation/ControlRig/Source/ControlRig/Public/ControlRig.h b/Engine/Plugins/Animation/ControlRig/Source/ControlRig/Public/ControlRig.h
--- a/Engine/Plugins/Animation/ControlRig/Source/ControlRig/Public/ControlRig.h
+++ b/Engine/Plugins/Animation/ControlRig/Source/ControlRig/Public/ControlRig.h
@@ -575,6 +575,7 @@ private:
 protected:
 	bool bCopyHierarchyBeforeConstruction;
 	bool bResetInitialTransformsBeforeConstruction;
+	bool bResetCurrentTransformsAfterConstruction;
 	bool bManipulationEnabled;
 
 	int32 PreConstructionBracket;
