Sequencer: Control Rig: Clear out delegates in outliner

diff --git a/Engine/Plugins/Animation/ControlRig/Source/ControlRigEditor/Private/EditMode/SControlRigOutliner.cpp b/Engine/Plugins/Animation/ControlRig/Source/ControlRigEditor/Private/EditMode/SControlRigOutliner.cpp
--- a/Engine/Plugins/Animation/ControlRig/Source/ControlRigEditor/Private/EditMode/SControlRigOutliner.cpp
+++ b/Engine/Plugins/Animation/ControlRig/Source/ControlRigEditor/Private/EditMode/SControlRigOutliner.cpp
@@ -1027,6 +1027,22 @@ SControlRigOutliner::SControlRigOutliner()
 SControlRigOutliner::~SControlRigOutliner()
 {
 	FCoreUObjectDelegates::OnObjectsReplaced.RemoveAll(this);
+	if (FControlRigEditMode* EditMode = static_cast<FControlRigEditMode*>(ModeTools->GetActiveMode(FControlRigEditMode::ModeName)))
+	{
+		TArrayView<TWeakObjectPtr<UControlRig>> ControlRigs = EditMode->GetControlRigs();
+		for (TWeakObjectPtr<UControlRig>& ControlRig : ControlRigs)
+		{
+			if (ControlRig.IsValid())
+			{
+				ControlRig.Get()->ControlRigBound().RemoveAll(this);
+				const TSharedPtr<IControlRigObjectBinding> Binding = ControlRig.Get()->GetObjectBinding();
+				if (Binding)
+				{
+					Binding->OnControlRigBind().RemoveAll(this);
+				}
+			}
+		}
+	}
 }
 
 void SControlRigOutliner::HandleControlSelected(UControlRig* Subject, FRigControlElement* ControlElement, bool bSelected)
