Control Rig: Enable rule manager for Rig Hierarchy tree

#rnx

diff --git a/Engine/Plugins/Animation/ControlRig/Source/ControlRig/Private/ModularRig.cpp b/Engine/Plugins/Animation/ControlRig/Source/ControlRig/Private/ModularRig.cpp
--- a/Engine/Plugins/Animation/ControlRig/Source/ControlRig/Private/ModularRig.cpp
+++ b/Engine/Plugins/Animation/ControlRig/Source/ControlRig/Private/ModularRig.cpp
@@ -668,6 +668,10 @@ void UModularRig::DiscardModuleRig(UControlRig* InControlRig)
 
 const FRigModuleInstance* UModularRig::FindModule(const FString& InPath) const
 {
+	if(InPath.EndsWith(NamespaceSeparator))
+	{
+		return FindModule(InPath.Left(InPath.Len() - 1));
+	}
 	return Modules.FindByPredicate([InPath](const FRigModuleInstance& Module)
 	{
 		return Module.GetPath() == InPath;
diff --git a/Engine/Plugins/Animation/ControlRig/Source/ControlRig/Private/Rigs/RigHierarchyDefines.cpp b/Engine/Plugins/Animation/ControlRig/Source/ControlRig/Private/Rigs/RigHierarchyDefines.cpp
--- a/Engine/Plugins/Animation/ControlRig/Source/ControlRig/Private/Rigs/RigHierarchyDefines.cpp
+++ b/Engine/Plugins/Animation/ControlRig/Source/ControlRig/Private/Rigs/RigHierarchyDefines.cpp
@@ -390,7 +390,36 @@ void FRigElementResolveResult::SetDefaultTarget(const FText& InMessage)
 	Message = InMessage;
 }
 
+////////////////////////////////////////////////////////////////////////////////
+// FModularRigResolveResult
+////////////////////////////////////////////////////////////////////////////////
+
 bool FModularRigResolveResult::IsValid() const
 {
 	return State == EModularRigResolveState::Success && !Matches.IsEmpty();
 }
+
+bool FModularRigResolveResult::ContainsMatch(const FRigElementKey& InKey, FString* OutErrorMessage) const
+{
+	if(Matches.ContainsByPredicate([InKey](const FRigElementResolveResult& InMatch) -> bool
+	{
+		return InMatch.GetKey() == InKey;
+	}))
+	{
+		return true;
+	}
+
+	if(OutErrorMessage)
+	{
+		if(const FRigElementResolveResult* Mismatch = Excluded.FindByPredicate([InKey](const FRigElementResolveResult& InMatch) -> bool
+		{
+			return InMatch.GetKey() == InKey;
+		}))
+		{
+			*OutErrorMessage = Mismatch->GetMessage().ToString();
+			return true;
+		}
+	}
+	
+	return false;
+}
diff --git a/Engine/Plugins/Animation/ControlRig/Source/ControlRig/Public/Rigs/RigHierarchyDefines.h b/Engine/Plugins/Animation/ControlRig/Source/ControlRig/Public/Rigs/RigHierarchyDefines.h
--- a/Engine/Plugins/Animation/ControlRig/Source/ControlRig/Public/Rigs/RigHierarchyDefines.h
+++ b/Engine/Plugins/Animation/ControlRig/Source/ControlRig/Public/Rigs/RigHierarchyDefines.h
@@ -1835,6 +1835,14 @@ public:
 	{
 		return Key;
 	}
+	const ERigElementResolveState& GetState() const
+	{
+		return State;
+	}
+	const FText& GetMessage() const
+	{
+		return Message;
+	}
 	void SetInvalidTarget(const FText& InMessage);
 	void SetPossibleTarget(const FText& InMessage = FText());
 	void SetDefaultTarget(const FText& InMessage = FText());
@@ -1903,6 +1911,8 @@ public:
 		return Excluded;
 	}
 
+	bool ContainsMatch(const FRigElementKey& InKey, FString* OutErrorMessage = nullptr) const;
+
 private:
 
 	UPROPERTY()
diff --git a/Engine/Plugins/Animation/ControlRig/Source/ControlRigEditor/Private/Editor/SRigHierarchy.cpp b/Engine/Plugins/Animation/ControlRig/Source/ControlRigEditor/Private/Editor/SRigHierarchy.cpp
--- a/Engine/Plugins/Animation/ControlRig/Source/ControlRigEditor/Private/Editor/SRigHierarchy.cpp
+++ b/Engine/Plugins/Animation/ControlRig/Source/ControlRigEditor/Private/Editor/SRigHierarchy.cpp
@@ -46,6 +46,7 @@
 #include "Widgets/Notifications/SNotificationList.h"
 #include "Styling/AppStyle.h"
 #include "ControlRigSkeletalMeshComponent.h"
+#include "ModularRigRuleManager.h"
 #include "Sequencer/ControlRigLayerInstance.h"
 #include "Algo/MinElement.h"
 #include "Algo/MaxElement.h"
@@ -156,6 +157,7 @@ void SRigHierarchy::Construct(const FArguments& InArgs, TSharedRef<FControlRigEd
 	Delegates.OnDragDetected = FOnDragDetected::CreateSP(this, &SRigHierarchy::OnDragDetected);
 	Delegates.OnGetResolvedKey = FOnRigTreeGetResolvedKey::CreateSP(this, &SRigHierarchy::OnGetResolvedKey);
 	Delegates.OnRequestDetailsInspection = FOnRigTreeRequestDetailsInspection::CreateSP(this, &SRigHierarchy::OnRequestDetailsInspection);
+	Delegates.OnRigTreeElementKeyTagDragDetected = FOnRigTreeRequestDetailsInspection::CreateSP(this, &SRigHierarchy::OnElementKeyTagDragDetected);
 
 	ChildSlot
 	[
@@ -2503,6 +2505,8 @@ FReply SRigHierarchy::OnDragDetected(const FGeometry& MyGeometry, const FPointer
 	{
 		if (ControlRigEditor.IsValid())
 		{
+			UpdateConnectorMatchesOnDrag(DraggedElements);
+			
 			TSharedRef<FRigElementHierarchyDragDropOp> DragDropOp = FRigElementHierarchyDragDropOp::New(MoveTemp(DraggedElements));
 			DragDropOp->OnPerformDropToGraph.BindSP(ControlRigEditor.Pin().Get(), &FControlRigEditor::OnGraphNodeDropToPerform);
 			return FReply::Handled().BeginDragDrop(DragDropOp);
@@ -2553,6 +2557,13 @@ TOptional<EItemDropZone> SRigHierarchy::OnCanAcceptDrop(const FDragDropEvent& Dr
 
 				if(RigDragDropOp->IsDraggingSingleConnector() || RigDragDropOp->IsDraggingSingleSocket())
 				{
+					if(const FModularRigResolveResult* ResolveResult = DragRigResolveResults.Find(DraggedKey))
+					{
+						if(!ResolveResult->ContainsMatch(TargetKey))
+						{
+							return InvalidDropZone;
+						}
+					}
 					if(DropZone != EItemDropZone::OntoItem)
 					{
 						return InvalidDropZone;
@@ -2575,12 +2586,6 @@ TOptional<EItemDropZone> SRigHierarchy::OnCanAcceptDrop(const FDragDropEvent& Dr
 			return InvalidDropZone;
 		}
 
-		// resolve connector rules
-		if(RigDragDropOp->IsDraggingSingleConnector())
-		{
-			// todo resolve rules
-		}
-
 		switch (TargetKey.Type)
 		{
 			case ERigElementType::Bone:
@@ -2656,9 +2661,18 @@ TOptional<EItemDropZone> SRigHierarchy::OnCanAcceptDrop(const FDragDropEvent& Dr
 			FRigElementKey DraggedKey;
 			FRigElementKey::StaticStruct()->ImportText(*TagDragDropOp->GetIdentifier(), &DraggedKey, nullptr, EPropertyPortFlags::PPF_None, nullptr, FRigElementKey::StaticStruct()->GetName(), true);
 
-			if(Hierarchy->Contains(DraggedKey))
+			if(Hierarchy->Contains(DraggedKey) && TargetItem.IsValid())
 			{
-				// todo: apply rules
+				if(DraggedKey.Type == ERigElementType::Connector)
+				{
+					if(const FModularRigResolveResult* ResolveResult = DragRigResolveResults.Find(DraggedKey))
+					{
+						if(!ResolveResult->ContainsMatch(TargetItem->Key))
+						{
+							return InvalidDropZone;
+						}
+					}
+				}
 				ReturnDropZone = DropZone;
 			}
 			else if(!TargetItem.IsValid())
@@ -2858,6 +2872,44 @@ FReply SRigHierarchy::OnAcceptDrop(const FDragDropEvent& DragDropEvent, EItemDro
 	return FReply::Unhandled();
 }
 
+void SRigHierarchy::OnElementKeyTagDragDetected(const FRigElementKey& InDraggedTag)
+{
+	UpdateConnectorMatchesOnDrag({InDraggedTag});
+}
+
+void SRigHierarchy::UpdateConnectorMatchesOnDrag(const TArray<FRigElementKey>& InDraggedKeys)
+{
+	DragRigResolveResults.Reset();
+	if(ControlRigBeingDebuggedPtr.IsValid())
+	{
+		const UModularRig* ControlRig = Cast<UModularRig>(ControlRigBeingDebuggedPtr.Get());
+		if(URigHierarchy* Hierarchy = ControlRig->GetHierarchy())
+		{
+			if(const UModularRigRuleManager* RuleManager = Hierarchy->GetRuleManager())
+			{
+				for(const FRigElementKey& DraggedElement : InDraggedKeys)
+				{
+					if(DraggedElement.Type == ERigElementType::Connector)
+					{
+						if(const FRigConnectorElement* Connector = Hierarchy->Find<FRigConnectorElement>(DraggedElement))
+						{
+							const FName NameSpace = Hierarchy->GetNameMetadata(Connector->GetKey(), URigHierarchy::NameSpaceMetadataName, NAME_None);
+							if(!NameSpace.IsNone())
+							{
+								if(const FRigModuleInstance* Module = ControlRig->FindModule(NameSpace.ToString()))
+								{
+									const FModularRigResolveResult ResolveResult = RuleManager->FindMatches(Connector, Module, ControlRig->ElementKeyRedirector); 
+									DragRigResolveResults.Add(DraggedElement, ResolveResult);
+								}
+							}
+						}
+					}
+				}
+			}
+		}
+	}
+}
+
 FName SRigHierarchy::HandleRenameElement(const FRigElementKey& OldKey, const FString& NewName)
 {
 	ClearDetailPanel();
diff --git a/Engine/Plugins/Animation/ControlRig/Source/ControlRigEditor/Private/Editor/SRigHierarchy.h b/Engine/Plugins/Animation/ControlRig/Source/ControlRigEditor/Private/Editor/SRigHierarchy.h
--- a/Engine/Plugins/Animation/ControlRig/Source/ControlRigEditor/Private/Editor/SRigHierarchy.h
+++ b/Engine/Plugins/Animation/ControlRig/Source/ControlRigEditor/Private/Editor/SRigHierarchy.h
@@ -153,6 +153,8 @@ private:
 	FReply OnDragDetected(const FGeometry& MyGeometry, const FPointerEvent& MouseEvent);
 	TOptional<EItemDropZone> OnCanAcceptDrop(const FDragDropEvent& DragDropEvent, EItemDropZone DropZone, TSharedPtr<FRigTreeElement> TargetItem);
 	FReply OnAcceptDrop(const FDragDropEvent& DragDropEvent, EItemDropZone DropZone, TSharedPtr<FRigTreeElement> TargetItem);
+	void OnElementKeyTagDragDetected(const FRigElementKey& InDraggedTag);
+	void UpdateConnectorMatchesOnDrag(const TArray<FRigElementKey>& InDraggedKeys);
 
 	static const FName ContextMenuName;
 	static void CreateContextMenu();
@@ -245,6 +247,7 @@ private:
 	bool bIsConstructionEventRunning;
 	uint32 LastHierarchyHash;
 	TArray<FRigElementKey> SelectionBeforeConstruction;
+	TMap<FRigElementKey, FModularRigResolveResult> DragRigResolveResults;
 
 public:
 	FName HandleRenameElement(const FRigElementKey& OldKey, const FString& NewName);
diff --git a/Engine/Plugins/Animation/ControlRig/Source/ControlRigEditor/Private/Editor/SRigHierarchyTagWidget.cpp b/Engine/Plugins/Animation/ControlRig/Source/ControlRigEditor/Private/Editor/SRigHierarchyTagWidget.cpp
--- a/Engine/Plugins/Animation/ControlRig/Source/ControlRigEditor/Private/Editor/SRigHierarchyTagWidget.cpp
+++ b/Engine/Plugins/Animation/ControlRig/Source/ControlRigEditor/Private/Editor/SRigHierarchyTagWidget.cpp
@@ -119,6 +119,15 @@ FReply SRigHierarchyTagWidget::OnDragDetected(const FGeometry& MyGeometry, const
 		if (MouseEvent.IsMouseButtonDown(EKeys::LeftMouseButton) && !Identifier.Get().IsEmpty())
 		{
 			TSharedRef<FRigHierarchyTagDragDropOp> DragDropOp = FRigHierarchyTagDragDropOp::New(SharedThis(this));
+
+			FRigElementKey DraggedKey;
+			FRigElementKey::StaticStruct()->ImportText(*Identifier.Get(), &DraggedKey, nullptr, EPropertyPortFlags::PPF_None, nullptr, FRigElementKey::StaticStruct()->GetName(), true);
+
+			if(DraggedKey.IsValid())
+			{
+				(void)OnElementKeyDragDetectedDelegate.ExecuteIfBound(DraggedKey);
+			}
+
 			return FReply::Handled().BeginDragDrop(DragDropOp);
 		}
 	}
diff --git a/Engine/Plugins/Animation/ControlRig/Source/ControlRigEditor/Private/Editor/SRigHierarchyTagWidget.h b/Engine/Plugins/Animation/ControlRig/Source/ControlRigEditor/Private/Editor/SRigHierarchyTagWidget.h
--- a/Engine/Plugins/Animation/ControlRig/Source/ControlRigEditor/Private/Editor/SRigHierarchyTagWidget.h
+++ b/Engine/Plugins/Animation/ControlRig/Source/ControlRigEditor/Private/Editor/SRigHierarchyTagWidget.h
@@ -4,10 +4,13 @@
 
 #include "CoreMinimal.h"
 #include "Widgets/SCompoundWidget.h"
+#include "Rigs/RigHierarchyDefines.h"
 #include "DragAndDrop/DecoratedDragDropOp.h"
 
 class SRigHierarchyTagWidget;
 
+DECLARE_DELEGATE_OneParam(FOnRigTreeElementKeyTagDragDetected, const FRigElementKey&);
+
 class FRigHierarchyTagDragDropOp : public FDecoratedDragDropOp
 {
 public:
@@ -62,6 +65,8 @@ public:
 	virtual int32 OnPaint(const FPaintArgs& Args, const FGeometry& AllottedGeometry, const FSlateRect& MyCullingRect, FSlateWindowElementList& OutDrawElements, int32 LayerId, const FWidgetStyle& InWidgetStyle, bool bParentEnabled) const override;
 	virtual FReply OnMouseButtonDown(const FGeometry& MyGeometry, const FPointerEvent& MouseEvent) override;
 	virtual FReply OnDragDetected(const FGeometry& MyGeometry, const FPointerEvent& MouseEvent) override;
+
+	FOnRigTreeElementKeyTagDragDetected& OnElementKeyDragDetected() { return OnElementKeyDragDetectedDelegate; }
 private:
 
 	FMargin GetTextPadding() const;
@@ -75,6 +80,7 @@ private:
 	TAttribute<FString> Identifier;
 	bool bAllowDragDrop = false;
 	FSimpleDelegate OnClicked;
+	FOnRigTreeElementKeyTagDragDetected OnElementKeyDragDetectedDelegate;
 
 	friend class FRigHierarchyTagDragDropOp;
 };
diff --git a/Engine/Plugins/Animation/ControlRig/Source/ControlRigEditor/Private/Editor/SRigHierarchyTreeView.cpp b/Engine/Plugins/Animation/ControlRig/Source/ControlRigEditor/Private/Editor/SRigHierarchyTreeView.cpp
--- a/Engine/Plugins/Animation/ControlRig/Source/ControlRigEditor/Private/Editor/SRigHierarchyTreeView.cpp
+++ b/Engine/Plugins/Animation/ControlRig/Source/ControlRigEditor/Private/Editor/SRigHierarchyTreeView.cpp
@@ -191,10 +191,13 @@ void SRigHierarchyItem::Construct(const FArguments& InArgs, const TSharedRef<STa
 
 	for(const SRigHierarchyTagWidget::FArguments& TagArguments : InRigTreeElement->Tags)
 	{
+		TSharedRef<SRigHierarchyTagWidget> TagWidget = SArgumentNew(TagArguments, SRigHierarchyTagWidget);
+		TagWidget->OnElementKeyDragDetected().BindSP(InTreeView.Get(), &SRigHierarchyTreeView::OnElementKeyTagDragDetected);
+		
 		HorizontalBox->AddSlot()
 		.AutoWidth()
 		[
-			SArgumentNew(TagArguments, SRigHierarchyTagWidget)
+			TagWidget
 		];
 	}
 
@@ -861,6 +864,11 @@ void SRigHierarchyTreeView::HandleGetChildrenForTree(TSharedPtr<FRigTreeElement>
 	OutChildren = InItem.Get()->Children;
 }
 
+void SRigHierarchyTreeView::OnElementKeyTagDragDetected(const FRigElementKey& InDraggedTag)
+{
+	(void)Delegates.OnRigTreeElementKeyTagDragDetected.ExecuteIfBound(InDraggedTag);
+}
+
 TArray<FRigElementKey> SRigHierarchyTreeView::GetSelectedKeys() const
 {
 	TArray<FRigElementKey> Keys;
diff --git a/Engine/Plugins/Animation/ControlRig/Source/ControlRigEditor/Private/Editor/SRigHierarchyTreeView.h b/Engine/Plugins/Animation/ControlRig/Source/ControlRigEditor/Private/Editor/SRigHierarchyTreeView.h
--- a/Engine/Plugins/Animation/ControlRig/Source/ControlRigEditor/Private/Editor/SRigHierarchyTreeView.h
+++ b/Engine/Plugins/Animation/ControlRig/Source/ControlRigEditor/Private/Editor/SRigHierarchyTreeView.h
@@ -99,6 +99,7 @@ struct CONTROLRIGEDITOR_API FRigTreeDelegates
 	FOnRigTreeCompareKeys OnCompareKeys;
 	FOnRigTreeGetResolvedKey OnGetResolvedKey;
 	FOnRigTreeRequestDetailsInspection OnRequestDetailsInspection;
+	FOnRigTreeElementKeyTagDragDetected OnRigTreeElementKeyTagDragDetected;
 
 	FRigTreeDelegates()
 	{
@@ -314,6 +315,7 @@ public:
 	void SetExpansionRecursive(TSharedPtr<FRigTreeElement> InElement, bool bTowardsParent, bool bShouldBeExpanded);
 	TSharedRef<ITableRow> MakeTableRowWidget(TSharedPtr<FRigTreeElement> InItem, const TSharedRef<STableViewBase>& OwnerTable, bool bPinned);
 	void HandleGetChildrenForTree(TSharedPtr<FRigTreeElement> InItem, TArray<TSharedPtr<FRigTreeElement>>& OutChildren);
+	void OnElementKeyTagDragDetected(const FRigElementKey& InDraggedTag);
 
 	TArray<FRigElementKey> GetSelectedKeys() const;
 	const TArray<TSharedPtr<FRigTreeElement>>& GetRootElements() const { return RootElements; }
