Control Rig: Trigger connector event when filtering connector candidates
#rnx

diff --git a/Engine/Plugins/Animation/ControlRig/Source/ControlRig/Private/ControlRig.cpp b/Engine/Plugins/Animation/ControlRig/Source/ControlRig/Private/ControlRig.cpp
--- a/Engine/Plugins/Animation/ControlRig/Source/ControlRig/Private/ControlRig.cpp
+++ b/Engine/Plugins/Animation/ControlRig/Source/ControlRig/Private/ControlRig.cpp
@@ -37,6 +37,7 @@
 #include "UObject/UObjectIterator.h"
 #include "RigVMCore/RigVMAssetUserData.h"
 #include "ModularRig.h"
+#include "Units/Modules/RigUnit_ConnectorExecution.h"
 
 #include UE_INLINE_GENERATED_CPP_BY_NAME(ControlRig)
 
@@ -1477,7 +1478,7 @@ bool UControlRig::Execute_Internal(const FName& InEventName)
 		}
 	}
 	
-	if(IsRigModule())
+	if(IsRigModule() && InEventName != FRigUnit_ConnectorExecution::EventName)
 	{
 		FString ConnectorWarning;
 		if(!AllConnectorsAreResolved(&ConnectorWarning))
diff --git a/Engine/Plugins/Animation/ControlRig/Source/ControlRig/Private/ModularRig.cpp b/Engine/Plugins/Animation/ControlRig/Source/ControlRig/Private/ModularRig.cpp
--- a/Engine/Plugins/Animation/ControlRig/Source/ControlRig/Private/ModularRig.cpp
+++ b/Engine/Plugins/Animation/ControlRig/Source/ControlRig/Private/ModularRig.cpp
@@ -7,6 +7,7 @@
 #include "Rigs/RigHierarchyController.h"
 #include "ControlRigComponent.h"
 #include "RigVMCore/RigVMExecuteContext.h"
+#include "Units/Modules/RigUnit_ConnectorExecution.h"
 
 #define LOCTEXT_NAMESPACE "ModularRig"
 
@@ -296,6 +297,7 @@ void UModularRig::Evaluate_AnyThread()
 void UModularRig::ExecuteQueue()
 {
 	FRigVMExtendedExecuteContext& Context = GetRigVMExtendedExecuteContext();
+	FControlRigExecuteContext& PublicContext = Context.GetPublicDataSafe<FControlRigExecuteContext>();
 	URigHierarchy* Hierarchy = GetHierarchy();
 	
 	while(ExecutionQueue.IsValidIndex(ExecutionQueueFront))
@@ -317,7 +319,6 @@ void UModularRig::ExecuteQueue()
 			// Make sure the hierarchy has the correct execute context with the rig module namespace
 			FRigHierarchyExecuteContextBracket ExecuteContextBracket(Hierarchy, &RigExtendedExecuteContext);
 
-			FControlRigExecuteContext& PublicContext = Context.GetPublicDataSafe<FControlRigExecuteContext>();
 			FControlRigExecuteContext& RigPublicContext = RigExtendedExecuteContext.GetPublicDataSafe<FControlRigExecuteContext>();
 			FRigUnitContext& RigUnitContext = RigPublicContext.UnitContext;
 			RigUnitContext = PublicContext.UnitContext;
@@ -396,6 +397,12 @@ void UModularRig::ExecuteQueue()
 			
 			ModuleRig->Execute_Internal(ExecutionElement.EventName);
 			ExecutionElement.bExecuted = true;
+
+			// Copy result of Connection event to the ModularRig's unit context
+			if (ExecutionElement.EventName == FRigUnit_ConnectorExecution::EventName)
+			{
+				PublicContext.UnitContext.ConnectionResolve = RigPublicContext.UnitContext.ConnectionResolve;
+			}
 		}
 		
 		ExecutionQueueFront++;
@@ -730,4 +737,47 @@ void UModularRig::ForEachModule(TFunctionRef<bool(const FRigModuleInstance*)> Pe
 	}
 }
 
+void UModularRig::ExecuteConnectorEvent(const FRigElementKey& InConnector, const FRigModuleInstance* InModuleInstance, const FRigElementKeyRedirector* InRedirector, TArray<FRigElementResolveResult>& InOutCandidates)
+{
+	if (!InModuleInstance)
+	{
+		InOutCandidates.Reset();
+		return;
+	}
+
+	if (!InRedirector)
+	{
+		InOutCandidates.Reset();
+		return;
+	}
+	
+	FRigModuleInstance* Module = Modules.FindByPredicate([InModuleInstance](FRigModuleInstance& Instance)
+	{
+		return &Instance == InModuleInstance;
+	});
+	if (!Module)
+	{
+		InOutCandidates.Reset();
+		return;
+	}
+
+	TArray<FRigElementResolveResult> Candidates = InOutCandidates;
+	
+	FControlRigExecuteContext& PublicContext = GetRigVMExtendedExecuteContext().GetPublicDataSafe<FControlRigExecuteContext>();
+
+	FString ShortConnectorName = InConnector.Name.ToString();
+	ShortConnectorName.RemoveFromStart(Module->GetNamespace());
+	TGuardValue<FRigElementKey> ConnectorGuard(PublicContext.UnitContext.ConnectionResolve.Connector, FRigElementKey(*ShortConnectorName, InConnector.Type));
+	TGuardValue<TArray<FRigElementResolveResult>> CandidatesGuard(PublicContext.UnitContext.ConnectionResolve.Matches, Candidates);
+	TGuardValue<TArray<FRigElementResolveResult>> MatchesGuard(PublicContext.UnitContext.ConnectionResolve.Excluded, {});
+	
+	FRigModuleExecutionElement ExecutionElement(Module, FRigUnit_ConnectorExecution::EventName);
+	TGuardValue<TArray<FRigModuleExecutionElement>> ExecutionGuard(ExecutionQueue, {ExecutionElement});
+	TGuardValue<int32> ExecutionFrontGuard(ExecutionQueueFront, 0);
+	TGuardValue<FRigElementKeyRedirector> RedirectorGuard(ElementKeyRedirector, *InRedirector);
+	ExecuteQueue();
+
+	InOutCandidates = PublicContext.UnitContext.ConnectionResolve.Matches;
+}
+
 #undef LOCTEXT_NAMESPACE
diff --git a/Engine/Plugins/Animation/ControlRig/Source/ControlRig/Private/ModularRigRuleManager.cpp b/Engine/Plugins/Animation/ControlRig/Source/ControlRig/Private/ModularRigRuleManager.cpp
--- a/Engine/Plugins/Animation/ControlRig/Source/ControlRig/Private/ModularRigRuleManager.cpp
+++ b/Engine/Plugins/Animation/ControlRig/Source/ControlRig/Private/ModularRigRuleManager.cpp
@@ -232,8 +232,7 @@ void UModularRigRuleManager::FilterByConnectorEvent(FWorkData& InOutWorkData)
 	
 	FModularRigResolveResult* Result = InOutWorkData.Result;
 
-	// todo: run the VM event to perform filtering
-	// ModularRig->ExecuteConnectorEvent(InOutWorkData->Connector->GetKey(), Module, Redirector)
+	ModularRig->ExecuteConnectorEvent(InOutWorkData.Connector->GetKey(), InOutWorkData.Module, InOutWorkData.ResolvedConnectors, InOutWorkData.Result->Matches);
 	
 	// move the default match to the front of the list
 	if(!Result->Matches.IsEmpty())
diff --git a/Engine/Plugins/Animation/ControlRig/Source/ControlRig/Private/Units/Modules/RigUnit_ConnectionCandidates.cpp b/Engine/Plugins/Animation/ControlRig/Source/ControlRig/Private/Units/Modules/RigUnit_ConnectionCandidates.cpp
new file mode 100644
--- /dev/null
+++ b/Engine/Plugins/Animation/ControlRig/Source/ControlRig/Private/Units/Modules/RigUnit_ConnectionCandidates.cpp
@@ -0,0 +1,59 @@
+// Copyright Epic Games, Inc. All Rights Reserved.
+
+#include "Units/Modules/RigUnit_ConnectionCandidates.h"
+#include "Units/Modules/RigUnit_ConnectorExecution.h"
+#include "Units/RigUnitContext.h"
+
+#include UE_INLINE_GENERATED_CPP_BY_NAME(RigUnit_ConnectionCandidates)
+
+FRigUnit_GetCandidates_Execute()
+{
+    DECLARE_SCOPE_HIERARCHICAL_COUNTER_RIGUNIT()
+
+	// only allow these nodes during connector resolval event
+	if(ExecuteContext.GetEventName() != FRigUnit_ConnectorExecution::EventName)
+	{
+		UE_CONTROLRIG_RIGUNIT_REPORT_ERROR(TEXT("Can only use GetCandidates during %s event."), *FRigUnit_ConnectorExecution::EventName.ToString());
+	}
+
+	Algo::Transform(ExecuteContext.UnitContext.ConnectionResolve.Matches, Candidates, [](const FRigElementResolveResult& Match)
+	{
+		return Match.GetKey();
+	});
+	Connector = ExecuteContext.UnitContext.ConnectionResolve.Connector;
+}
+
+FRigUnit_DiscardMatches_Execute()
+{
+	DECLARE_SCOPE_HIERARCHICAL_COUNTER_RIGUNIT()
+
+	// only allow these nodes during connector resolval event
+	if(ExecuteContext.GetEventName() != FRigUnit_ConnectorExecution::EventName)
+	{
+		UE_CONTROLRIG_RIGUNIT_REPORT_ERROR(TEXT("Can only use AddMatches during %s event."), *FRigUnit_ConnectorExecution::EventName.ToString());
+	}
+
+	// Filter any items that are already excluded
+	TArray<FRigElementKey> FilterExcluded = Excluded.FilterByPredicate([ExecuteContext](const FRigElementKey& Exclude)
+	{
+		return !ExecuteContext.UnitContext.ConnectionResolve.Excluded.ContainsByPredicate([Exclude](const FRigElementResolveResult& Existing)
+		{
+			return Existing.GetKey() == Exclude;
+		});
+	});
+
+	// Remove elements from matches, and add to excluded
+	for (FRigElementKey& Exclude : FilterExcluded)
+	{
+		int32 Index = ExecuteContext.UnitContext.ConnectionResolve.Matches.IndexOfByPredicate([Exclude](const FRigElementResolveResult& Match)
+		{
+			return Exclude == Match.GetKey();
+		});
+		if (Index != INDEX_NONE)
+		{
+			ExecuteContext.UnitContext.ConnectionResolve.Matches.RemoveAt(Index);
+		}
+
+		ExecuteContext.UnitContext.ConnectionResolve.Excluded.Add(FRigElementResolveResult(Exclude, ERigElementResolveState::InvalidTarget, FText::FromString(Message)));
+	}
+}
\ No newline at end of file
diff --git a/Engine/Plugins/Animation/ControlRig/Source/ControlRig/Private/Units/Modules/RigUnit_ConnectionCandidates.h b/Engine/Plugins/Animation/ControlRig/Source/ControlRig/Private/Units/Modules/RigUnit_ConnectionCandidates.h
new file mode 100644
--- /dev/null
+++ b/Engine/Plugins/Animation/ControlRig/Source/ControlRig/Private/Units/Modules/RigUnit_ConnectionCandidates.h
@@ -0,0 +1,46 @@
+// Copyright Epic Games, Inc. All Rights Reserved.
+
+#pragma once
+
+#include "Units/RigUnit.h"
+#include "ControlRigDefines.h"
+#include "RigUnit_ConnectionCandidates.generated.h"
+
+/**
+ * Returns the connection candidates for one connector
+ */
+USTRUCT(meta=(DisplayName="Get Candidates", Category="Modules", TitleColor="1 0 0", NodeColor="1 1 1", Keywords="Connection,Resolve", Varying))
+struct CONTROLRIG_API FRigUnit_GetCandidates : public FRigUnit
+{
+	GENERATED_BODY()
+
+	RIGVM_METHOD()
+	virtual void Execute() override;
+
+	// The connector being resolved
+	UPROPERTY(EditAnywhere, Transient, Category = "Modules", meta = (Output))
+	FRigElementKey Connector;
+
+	// The items being interacted on
+	UPROPERTY(EditAnywhere, Transient, Category = "Modules", meta = (Output))
+	TArray<FRigElementKey> Candidates;
+};
+
+/**
+ * Adds possible matches during a connector event
+ */
+USTRUCT(meta=(DisplayName="Discard Matches", Category="Modules", TitleColor="1 0 0", NodeColor="1 1 1", Keywords="Connection,Resolve,Match", Varying))
+struct CONTROLRIG_API FRigUnit_DiscardMatches : public FRigUnitMutable
+{
+	GENERATED_BODY()
+
+	RIGVM_METHOD()
+	virtual void Execute() override;
+
+	// The items being interacted on
+	UPROPERTY(EditAnywhere, Transient, Category = "Modules", meta = (Input))
+	TArray<FRigElementKey> Excluded;
+
+	UPROPERTY(EditAnywhere, Transient, Category = "Modules", meta = (Input))
+	FString Message;
+};
diff --git a/Engine/Plugins/Animation/ControlRig/Source/ControlRig/Private/Units/Modules/RigUnit_ConnectorExecution.cpp b/Engine/Plugins/Animation/ControlRig/Source/ControlRig/Private/Units/Modules/RigUnit_ConnectorExecution.cpp
new file mode 100644
--- /dev/null
+++ b/Engine/Plugins/Animation/ControlRig/Source/ControlRig/Private/Units/Modules/RigUnit_ConnectorExecution.cpp
@@ -0,0 +1,15 @@
+// Copyright Epic Games, Inc. All Rights Reserved.
+
+#include "Units/Modules/RigUnit_ConnectorExecution.h"
+#include "Units/RigUnitContext.h"
+
+#include UE_INLINE_GENERATED_CPP_BY_NAME(RigUnit_ConnectorExecution)
+
+FName FRigUnit_ConnectorExecution::EventName = TEXT("Connector");
+
+FRigUnit_ConnectorExecution_Execute()
+{
+    DECLARE_SCOPE_HIERARCHICAL_COUNTER_RIGUNIT()
+	ExecuteContext.SetEventName(FRigUnit_ConnectorExecution::EventName);
+}
+
diff --git a/Engine/Plugins/Animation/ControlRig/Source/ControlRig/Private/Units/Modules/RigUnit_ConnectorExecution.h b/Engine/Plugins/Animation/ControlRig/Source/ControlRig/Private/Units/Modules/RigUnit_ConnectorExecution.h
new file mode 100644
--- /dev/null
+++ b/Engine/Plugins/Animation/ControlRig/Source/ControlRig/Private/Units/Modules/RigUnit_ConnectorExecution.h
@@ -0,0 +1,27 @@
+// Copyright Epic Games, Inc. All Rights Reserved.
+
+#pragma once
+
+#include "Units/RigUnit.h"
+#include "ControlRigDefines.h"
+#include "RigUnit_ConnectorExecution.generated.h"
+
+/**
+ * Event for filtering connection candidates
+ */
+USTRUCT(meta=(DisplayName="Connector", Category="Events", NodeColor="1, 0, 0", Keywords="Event,During,Resolve,Connect,Filter"))
+struct CONTROLRIG_API FRigUnit_ConnectorExecution : public FRigUnit
+{
+	GENERATED_BODY()
+
+	RIGVM_METHOD()
+	virtual void Execute() override;
+
+	virtual FName GetEventName() const override { return EventName; }
+
+	// The execution result
+	UPROPERTY(EditAnywhere, Transient, DisplayName = "Execute", Category = "BeginExecution", meta = (Output))
+	FControlRigExecuteContext ExecuteContext;
+
+	static FName EventName;
+};
diff --git a/Engine/Plugins/Animation/ControlRig/Source/ControlRig/Public/ModularRig.h b/Engine/Plugins/Animation/ControlRig/Source/ControlRig/Public/ModularRig.h
--- a/Engine/Plugins/Animation/ControlRig/Source/ControlRig/Public/ModularRig.h
+++ b/Engine/Plugins/Animation/ControlRig/Source/ControlRig/Public/ModularRig.h
@@ -167,6 +167,8 @@ public:
 	void ForEachModule(TFunctionRef<bool(FRigModuleInstance*)> PerModuleFunction);
 	void ForEachModule(TFunctionRef<bool(const FRigModuleInstance*)> PerModuleFunction) const;
 
+	void ExecuteConnectorEvent(const FRigElementKey& InConnector, const FRigModuleInstance* InModuleInstance, const FRigElementKeyRedirector* InRedirector, TArray<FRigElementResolveResult>& InOutCandidates);
+
 	/**
 	 * Returns a handle to an existing module
 	 * @param InPath The path of the module to retrieve a handle for.
diff --git a/Engine/Plugins/Animation/ControlRig/Source/ControlRig/Public/Rigs/RigHierarchyDefines.h b/Engine/Plugins/Animation/ControlRig/Source/ControlRig/Public/Rigs/RigHierarchyDefines.h
--- a/Engine/Plugins/Animation/ControlRig/Source/ControlRig/Public/Rigs/RigHierarchyDefines.h
+++ b/Engine/Plugins/Animation/ControlRig/Source/ControlRig/Public/Rigs/RigHierarchyDefines.h
@@ -1931,4 +1931,7 @@ private:
 	FText Message;
 
 	friend class UModularRigRuleManager;
+	friend class UModularRig;
+	friend struct FRigUnit_GetCandidates;
+	friend struct FRigUnit_DiscardMatches;
 };
diff --git a/Engine/Plugins/Animation/ControlRig/Source/ControlRig/Public/Units/RigUnitContext.h b/Engine/Plugins/Animation/ControlRig/Source/ControlRig/Public/Units/RigUnitContext.h
--- a/Engine/Plugins/Animation/ControlRig/Source/ControlRig/Public/Units/RigUnitContext.h
+++ b/Engine/Plugins/Animation/ControlRig/Source/ControlRig/Public/Units/RigUnitContext.h
@@ -70,6 +70,9 @@ struct FRigUnitContext
 	/** The elements being interacted with. */
 	TArray<FRigElementKey> ElementsBeingInteracted;
 
+	/** Acceptable subset of connection matches */
+	FModularRigResolveResult ConnectionResolve;
+
 	/**
 	 * Returns a given data source and cast it to the expected class.
 	 *
diff --git a/Engine/Plugins/Animation/ControlRig/Source/ControlRigDeveloper/Private/ControlRigSchema.cpp b/Engine/Plugins/Animation/ControlRig/Source/ControlRigDeveloper/Private/ControlRigSchema.cpp
--- a/Engine/Plugins/Animation/ControlRig/Source/ControlRigDeveloper/Private/ControlRigSchema.cpp
+++ b/Engine/Plugins/Animation/ControlRig/Source/ControlRigDeveloper/Private/ControlRigSchema.cpp
@@ -1,10 +1,13 @@
 // Copyright Epic Games, Inc. All Rights Reserved.
 
 #include "ControlRigSchema.h"
+#include "ControlRigBlueprint.h"
 #include "RigVMModel/RigVMController.h"
 #include "Rigs/RigHierarchyPose.h"
 #include "Curves/CurveFloat.h"
 #include "Units/RigUnitContext.h"
+#include "Units/Modules/RigUnit_ConnectorExecution.h"
+#include "Units/Modules/RigUnit_ConnectionCandidates.h"
 
 UControlRigSchema::UControlRigSchema(const FObjectInitializer& ObjectInitializer)
 	: Super(ObjectInitializer)
@@ -35,3 +38,21 @@ bool UControlRigSchema::ShouldUnfoldStruct(URigVMController* InController, const
 	
 	return true;
 }
+
+bool UControlRigSchema::SupportsUnitFunction(URigVMController* InController, const FRigVMFunction* InUnitFunction) const
+{
+	if (InUnitFunction->Struct == FRigUnit_ConnectorExecution::StaticStruct() ||
+		InUnitFunction->Struct == FRigUnit_GetCandidates::StaticStruct() ||
+		InUnitFunction->Struct == FRigUnit_DiscardMatches::StaticStruct())
+	{
+		if (UControlRigBlueprint* Blueprint = Cast<UControlRigBlueprint>(InController->GetOuter()))
+		{
+			if (Blueprint->IsControlRigModule())
+			{
+				return true;
+			}
+			return false;
+		}
+	}
+	return Super::SupportsUnitFunction(InController, InUnitFunction);
+}
diff --git a/Engine/Plugins/Animation/ControlRig/Source/ControlRigDeveloper/Private/Graph/ControlRigGraphSchema.cpp b/Engine/Plugins/Animation/ControlRig/Source/ControlRigDeveloper/Private/Graph/ControlRigGraphSchema.cpp
--- a/Engine/Plugins/Animation/ControlRig/Source/ControlRigDeveloper/Private/Graph/ControlRigGraphSchema.cpp
+++ b/Engine/Plugins/Animation/ControlRig/Source/ControlRigDeveloper/Private/Graph/ControlRigGraphSchema.cpp
@@ -6,6 +6,7 @@
 #include "Units/Execution/RigUnit_PrepareForExecution.h"
 #include "Units/Execution/RigUnit_InverseExecution.h"
 #include "Units/Execution/RigUnit_InteractionExecution.h"
+#include "Units/Modules/RigUnit_ConnectorExecution.h"
 
 
 #include UE_INLINE_GENERATED_CPP_BY_NAME(ControlRigGraphSchema)
@@ -54,7 +55,8 @@ bool UControlRigGraphSchema::IsRigVMDefaultEvent(const FName& InEventName) const
 		InEventName == FRigUnit_PostBeginExecution::EventName ||
 		InEventName == FRigUnit_InverseExecution::EventName ||
 		InEventName == FRigUnit_PrepareForExecution::EventName ||
-		InEventName == FRigUnit_InteractionExecution::EventName;
+		InEventName == FRigUnit_InteractionExecution::EventName ||
+		InEventName == FRigUnit_ConnectorExecution::EventName;
 }
 
 #undef LOCTEXT_NAMESPACE
diff --git a/Engine/Plugins/Animation/ControlRig/Source/ControlRigDeveloper/Public/ControlRigSchema.h b/Engine/Plugins/Animation/ControlRig/Source/ControlRigDeveloper/Public/ControlRigSchema.h
--- a/Engine/Plugins/Animation/ControlRig/Source/ControlRigDeveloper/Public/ControlRigSchema.h
+++ b/Engine/Plugins/Animation/ControlRig/Source/ControlRigDeveloper/Public/ControlRigSchema.h
@@ -13,4 +13,6 @@ class CONTROLRIGDEVELOPER_API UControlRigSchema : public URigVMSchema
 public:
 
 	virtual bool ShouldUnfoldStruct(URigVMController* InController, const UStruct* InStruct) const override;
+
+	virtual bool SupportsUnitFunction(URigVMController* InController, const FRigVMFunction* InUnitFunction) const override;
 };
\ No newline at end of file
