Control Rig: Fix compilation error because of private headers
#rnx

diff --git a/Engine/Plugins/Animation/ControlRig/Source/ControlRig/Private/Units/Modules/RigUnit_ConnectionCandidates.h b/Engine/Plugins/Animation/ControlRig/Source/ControlRig/Private/Units/Modules/RigUnit_ConnectionCandidates.h
deleted file mode 100644
diff --git a/Engine/Plugins/Animation/ControlRig/Source/ControlRig/Private/Units/Modules/RigUnit_ConnectorExecution.h b/Engine/Plugins/Animation/ControlRig/Source/ControlRig/Private/Units/Modules/RigUnit_ConnectorExecution.h
deleted file mode 100644
diff --git a/Engine/Plugins/Animation/ControlRig/Source/ControlRig/Public/Units/Modules/RigUnit_ConnectionCandidates.h b/Engine/Plugins/Animation/ControlRig/Source/ControlRig/Public/Units/Modules/RigUnit_ConnectionCandidates.h
new file mode 100644
--- /dev/null
+++ b/Engine/Plugins/Animation/ControlRig/Source/ControlRig/Public/Units/Modules/RigUnit_ConnectionCandidates.h
@@ -0,0 +1,46 @@
+// Copyright Epic Games, Inc. All Rights Reserved.
+
+#pragma once
+
+#include "Units/RigUnit.h"
+#include "ControlRigDefines.h"
+#include "RigUnit_ConnectionCandidates.generated.h"
+
+/**
+ * Returns the connection candidates for one connector
+ */
+USTRUCT(meta=(DisplayName="Get Candidates", Category="Modules", TitleColor="1 0 0", NodeColor="1 1 1", Keywords="Connection,Resolve", Varying))
+struct CONTROLRIG_API FRigUnit_GetCandidates : public FRigUnit
+{
+	GENERATED_BODY()
+
+	RIGVM_METHOD()
+	virtual void Execute() override;
+
+	// The connector being resolved
+	UPROPERTY(EditAnywhere, Transient, Category = "Modules", meta = (Output))
+	FRigElementKey Connector;
+
+	// The items being interacted on
+	UPROPERTY(EditAnywhere, Transient, Category = "Modules", meta = (Output))
+	TArray<FRigElementKey> Candidates;
+};
+
+/**
+ * Adds possible matches during a connector event
+ */
+USTRUCT(meta=(DisplayName="Discard Matches", Category="Modules", TitleColor="1 0 0", NodeColor="1 1 1", Keywords="Connection,Resolve,Match", Varying))
+struct CONTROLRIG_API FRigUnit_DiscardMatches : public FRigUnitMutable
+{
+	GENERATED_BODY()
+
+	RIGVM_METHOD()
+	virtual void Execute() override;
+
+	// The items being interacted on
+	UPROPERTY(EditAnywhere, Transient, Category = "Modules", meta = (Input))
+	TArray<FRigElementKey> Excluded;
+
+	UPROPERTY(EditAnywhere, Transient, Category = "Modules", meta = (Input))
+	FString Message;
+};
diff --git a/Engine/Plugins/Animation/ControlRig/Source/ControlRig/Public/Units/Modules/RigUnit_ConnectorExecution.h b/Engine/Plugins/Animation/ControlRig/Source/ControlRig/Public/Units/Modules/RigUnit_ConnectorExecution.h
new file mode 100644
--- /dev/null
+++ b/Engine/Plugins/Animation/ControlRig/Source/ControlRig/Public/Units/Modules/RigUnit_ConnectorExecution.h
@@ -0,0 +1,27 @@
+// Copyright Epic Games, Inc. All Rights Reserved.
+
+#pragma once
+
+#include "Units/RigUnit.h"
+#include "ControlRigDefines.h"
+#include "RigUnit_ConnectorExecution.generated.h"
+
+/**
+ * Event for filtering connection candidates
+ */
+USTRUCT(meta=(DisplayName="Connector", Category="Events", NodeColor="1, 0, 0", Keywords="Event,During,Resolve,Connect,Filter"))
+struct CONTROLRIG_API FRigUnit_ConnectorExecution : public FRigUnit
+{
+	GENERATED_BODY()
+
+	RIGVM_METHOD()
+	virtual void Execute() override;
+
+	virtual FName GetEventName() const override { return EventName; }
+
+	// The execution result
+	UPROPERTY(EditAnywhere, Transient, DisplayName = "Execute", Category = "BeginExecution", meta = (Output))
+	FControlRigExecuteContext ExecuteContext;
+
+	static FName EventName;
+};
