Control Rig: Enable import of namespaced bones

#rnx

diff --git a/Engine/Plugins/Animation/ControlRig/Source/ControlRig/Private/Rigs/RigHierarchy.cpp b/Engine/Plugins/Animation/ControlRig/Source/ControlRig/Private/Rigs/RigHierarchy.cpp
--- a/Engine/Plugins/Animation/ControlRig/Source/ControlRig/Private/Rigs/RigHierarchy.cpp
+++ b/Engine/Plugins/Animation/ControlRig/Source/ControlRig/Private/Rigs/RigHierarchy.cpp
@@ -1346,6 +1346,10 @@ bool URigHierarchy::IsNameAvailable(const FRigName& InPotentialNewName, ERigElem
 			bAllowNameSpaces = true;
 		}
 	}
+	else
+	{
+		bAllowNameSpaces = true;
+	}
 
 	FRigName SanitizedName = UnsanitizedName;
 	SanitizeName(SanitizedName, bAllowNameSpaces);
@@ -1422,12 +1426,13 @@ bool URigHierarchy::IsDisplayNameAvailable(const FRigElementKey& InParentElement
 	return true;
 }
 
-FRigName URigHierarchy::GetSafeNewName(const FRigName& InPotentialNewName, ERigElementType InType) const
+FRigName URigHierarchy::GetSafeNewName(const FRigName& InPotentialNewName, ERigElementType InType, bool bAllowNameSpace) const
 {
 	FRigName SanitizedName = InPotentialNewName;
-	SanitizeName(SanitizedName, true);
+	SanitizeName(SanitizedName, bAllowNameSpace);
 
-	if(ExecuteContext)
+	bAllowNameSpaceWhenSanitizingName = bAllowNameSpace;
+	if(!bAllowNameSpaceWhenSanitizingName && ExecuteContext)
 	{
 		const FControlRigExecuteContext& CRContext = ExecuteContext->GetPublicData<FControlRigExecuteContext>();
 		if(CRContext.IsRigModule())
diff --git a/Engine/Plugins/Animation/ControlRig/Source/ControlRig/Private/Rigs/RigHierarchyController.cpp b/Engine/Plugins/Animation/ControlRig/Source/ControlRig/Private/Rigs/RigHierarchyController.cpp
--- a/Engine/Plugins/Animation/ControlRig/Source/ControlRig/Private/Rigs/RigHierarchyController.cpp
+++ b/Engine/Plugins/Animation/ControlRig/Source/ControlRig/Private/Rigs/RigHierarchyController.cpp
@@ -2004,18 +2004,21 @@ bool URigHierarchyController::IsValid() const
 	return WeakHierarchy.IsValid(bPendingKillAcceptable);
 }
 
-FName URigHierarchyController::GetSafeNewName(const FName& InDesiredName, ERigElementType InElementType) const
+FName URigHierarchyController::GetSafeNewName(const FName& InDesiredName, ERigElementType InElementType, bool bAllowNameSpace) const
 {
 	FRigName Name(InDesiredName);
 
 	// remove potential namespaces from it
-	int32 Index = INDEX_NONE;
-	if(Name.GetName().FindLastChar(TEXT(':'), Index))
+	if(!bAllowNameSpace)
 	{
-		Name.SetName(Name.GetName().Mid(Index + 1));
+		int32 Index = INDEX_NONE;
+		if(Name.GetName().FindLastChar(TEXT(':'), Index))
+		{
+			Name.SetName(Name.GetName().Mid(Index + 1));
+		}
 	}
 	
-	return GetHierarchy()->GetSafeNewName(Name, InElementType).GetFName();
+	return GetHierarchy()->GetSafeNewName(Name, InElementType, bAllowNameSpace).GetFName();
 }
 
 int32 URigHierarchyController::AddElement(FRigBaseElement* InElementToAdd, FRigBaseElement* InFirstParent, bool bMaintainGlobalTransform, const FName& InDesiredName)
diff --git a/Engine/Plugins/Animation/ControlRig/Source/ControlRig/Public/Rigs/RigHierarchy.h b/Engine/Plugins/Animation/ControlRig/Source/ControlRig/Public/Rigs/RigHierarchy.h
--- a/Engine/Plugins/Animation/ControlRig/Source/ControlRig/Public/Rigs/RigHierarchy.h
+++ b/Engine/Plugins/Animation/ControlRig/Source/ControlRig/Public/Rigs/RigHierarchy.h
@@ -1706,9 +1706,10 @@ public:
 	 * Returns a valid new name for a to-be-added element.
 	 * @param InPotentialNewName The name to be sanitized and adjusted for availability
 	 * @param InType The type of the to-be-added element
+	 * @param bAllowNameSpace If true the name will be allowed to contain namespaces
 	 * @return Returns the name to use for the to-be-added element.
 	 */
-	FRigName GetSafeNewName(const FRigName& InPotentialNewName, ERigElementType InType) const;
+	FRigName GetSafeNewName(const FRigName& InPotentialNewName, ERigElementType InType, bool bAllowNameSpace = false) const;
 
 	/**
 	 * Returns a valid new display name for a control
@@ -4718,7 +4719,6 @@ private:
 	void QueueNotification(ERigHierarchyNotification InNotification, const FRigBaseElement* InElement);
 	void SendQueuedNotifications();
 	void Reset_Impl(bool bResetElements);
-
 #if WITH_EDITOR
 	void ForEachListeningHierarchy(TFunctionRef<void(const FRigHierarchyListener&)> PerListeningHierarchyFunction);
 #endif
diff --git a/Engine/Plugins/Animation/ControlRig/Source/ControlRig/Public/Rigs/RigHierarchyController.h b/Engine/Plugins/Animation/ControlRig/Source/ControlRig/Public/Rigs/RigHierarchyController.h
--- a/Engine/Plugins/Animation/ControlRig/Source/ControlRig/Public/Rigs/RigHierarchyController.h
+++ b/Engine/Plugins/Animation/ControlRig/Source/ControlRig/Public/Rigs/RigHierarchyController.h
@@ -651,9 +651,10 @@ private:
 	 * elements with the same short name (for example "MyModuleA:Control" and "MyModuleB:Control").
 	 * @param InDesiredName The name provided by the user
 	 * @param InElementType The kind of element we are about to create
+	 * @param bAllowNameSpace If true the name won't be changed for namespaces
 	 * @return The safe name of the element to create.
 	 */
-	FName GetSafeNewName(const FName& InDesiredName, ERigElementType InElementType) const;
+	FName GetSafeNewName(const FName& InDesiredName, ERigElementType InElementType, bool bAllowNameSpace = true) const;
 	
 	/**
 	 * Adds a new element to the hierarchy
