Control Rig: Simplify use of shape library user data

#rnx

diff --git a/Engine/Plugins/Animation/ControlRig/Source/ControlRig/Private/ControlRig.cpp b/Engine/Plugins/Animation/ControlRig/Source/ControlRig/Private/ControlRig.cpp
--- a/Engine/Plugins/Animation/ControlRig/Source/ControlRig/Private/ControlRig.cpp
+++ b/Engine/Plugins/Animation/ControlRig/Source/ControlRig/Private/ControlRig.cpp
@@ -220,7 +220,7 @@ void UControlRig::RestoreShapeLibrariesFromCDO()
 	}
 }
 
-void UControlRig::OnAddShapeLibrary(const FControlRigExecuteContext* InContext, const FString& InLibraryName, UControlRigShapeLibrary* InShapeLibrary, bool bReplaceExisting, bool bLogResults)
+void UControlRig::OnAddShapeLibrary(const FControlRigExecuteContext* InContext, const FString& InLibraryName, UControlRigShapeLibrary* InShapeLibrary, bool bLogResults)
 {
 	// don't ever change the CDO
 	if (HasAnyFlags(RF_ClassDefaultObject))
@@ -250,47 +250,11 @@ void UControlRig::OnAddShapeLibrary(const FControlRigExecuteContext* InContext,
 
 	// if we are supposed to replace the library and the library name is empty
 	FString LibraryName = InLibraryName;
-	if(LibraryName.IsEmpty() && bReplaceExisting)
-	{
-		if(ShapeLibraries.Num() == 1)
-		{
-			LibraryName = ShapeLibraries[0]->GetName();
-		}
-	}
-
 	if(LibraryName.IsEmpty())
 	{
 		LibraryName = InShapeLibrary->GetName();
 	}
 
-	// if we need to replace all existing shape libraries - we'll remove any shape library matching the library
-	// name which doesn't map to something else
-	if(bReplaceExisting)
-	{
-		TArray<FString> KeysToRemove;
-		for(const TPair<FString, FString>& Pair : ShapeLibraryNameMap)
-		{
-			if(Pair.Value.Equals(LibraryName, ESearchCase::CaseSensitive))
-			{
-				KeysToRemove.Add(Pair.Key);
-			}
-		}
-		
-		if(KeysToRemove.IsEmpty())
-		{
-			KeysToRemove.Add(InShapeLibrary->GetName());
-		}
-		
-		for(const FString& KeyToRemove : KeysToRemove)
-		{
-			ShapeLibraryNameMap.Remove(KeyToRemove);
-			ShapeLibraries.RemoveAll([KeyToRemove](const TSoftObjectPtr<UControlRigShapeLibrary>& ShapeLibrary) -> bool
-			{
-				return ShapeLibrary->GetName().Equals(KeyToRemove, ESearchCase::CaseSensitive);
-			});
-		}
-	}
-
 	if(LibraryName != InShapeLibrary->GetName())
 	{
 		ShapeLibraryNameMap.FindOrAdd(InShapeLibrary->GetName()) = LibraryName;
@@ -304,6 +268,7 @@ void UControlRig::OnAddShapeLibrary(const FControlRigExecuteContext* InContext,
 		static constexpr TCHAR LibraryFormat[] = TEXT("Control Rig '%s': Shape Library '%s' uses asset '%s'");
 		static constexpr TCHAR DefaultShapeFormat[] = TEXT("Control Rig '%s': Shape Library '%s' has default shape '%s'");
 		static constexpr TCHAR ShapeFormat[] = TEXT("Control Rig '%s': Shape Library '%s' contains shape %03d: '%s'");
+		static constexpr TCHAR ResolvedShapeFormat[] = TEXT("Control Rig '%s': ShapeName '%s' resolved to '%s.%s'");
 
 		const FString PathName = GetPathName();
 
@@ -313,7 +278,8 @@ void UControlRig::OnAddShapeLibrary(const FControlRigExecuteContext* InContext,
 		}
 
 		const int32 NumShapeLibraries = ShapeLibraries.Num();
-		
+
+		TMap<FString, const UControlRigShapeLibrary*> ShapeNameToShapeLibrary;
 		for(int32 Index = 0; Index <NumShapeLibraries; Index++)
 		{
 			const TSoftObjectPtr<UControlRigShapeLibrary>& ShapeLibrary = ShapeLibraries[Index];
@@ -332,8 +298,16 @@ void UControlRig::OnAddShapeLibrary(const FControlRigExecuteContext* InContext,
 			{
 				const FString ShapeName = UControlRigShapeLibrary::GetShapeName(ShapeLibrary.Get(), bUseNameSpace, ShapeLibraryNameMap, ShapeLibrary->Shapes[ShapeIndex]);
 				UE_LOG(LogControlRig, Display, ShapeFormat, *PathName, *ShapeLibrary->GetName(), ShapeIndex, *ShapeName);
+				ShapeNameToShapeLibrary.FindOrAdd(ShapeName, nullptr) = ShapeLibrary.Get();
 			}
 		}
+
+		for(const TPair<FString, const UControlRigShapeLibrary*>& Pair : ShapeNameToShapeLibrary)
+		{
+			FString Left, Right = Pair.Key;
+			(void)Pair.Key.Split(TEXT("."), &Left, &Right, ESearchCase::CaseSensitive, ESearchDir::FromEnd);
+			UE_LOG(LogControlRig, Display, ResolvedShapeFormat, *PathName, *Pair.Key, *Pair.Value->GetName(), *Right);
+		}
 	}
 #endif
 }
diff --git a/Engine/Plugins/Animation/ControlRig/Source/ControlRig/Private/ControlRigGizmoLibrary.cpp b/Engine/Plugins/Animation/ControlRig/Source/ControlRig/Private/ControlRigGizmoLibrary.cpp
--- a/Engine/Plugins/Animation/ControlRig/Source/ControlRig/Private/ControlRigGizmoLibrary.cpp
+++ b/Engine/Plugins/Animation/ControlRig/Source/ControlRig/Private/ControlRigGizmoLibrary.cpp
@@ -148,7 +148,9 @@ const FControlRigShapeDefinition* UControlRigShapeLibrary::GetShapeByName(const
 
 			if(ShapeLibraryName.Equals(Left) || Left.IsEmpty())
 			{
-				if(const FControlRigShapeDefinition* Shape = ShapeLibrary->GetShapeByName(RightName))
+				// only fall back on the default shape for the very last shape library
+				const bool bUseDefaultShapeIfNotFound = (Pass > 0) && (Index == 0);
+				if(const FControlRigShapeDefinition* Shape = ShapeLibrary->GetShapeByName(RightName, bUseDefaultShapeIfNotFound))
 				{
 					return Shape;
 				}
diff --git a/Engine/Plugins/Animation/ControlRig/Source/ControlRig/Private/Units/Core/RigUnit_UserData.cpp b/Engine/Plugins/Animation/ControlRig/Source/ControlRig/Private/Units/Core/RigUnit_UserData.cpp
--- a/Engine/Plugins/Animation/ControlRig/Source/ControlRig/Private/Units/Core/RigUnit_UserData.cpp
+++ b/Engine/Plugins/Animation/ControlRig/Source/ControlRig/Private/Units/Core/RigUnit_UserData.cpp
@@ -206,7 +206,7 @@ FRigUnit_SetupShapeLibraryFromUserData_Execute()
 				{
 					if(ExecuteContext.OnAddShapeLibraryDelegate.IsBound())
 					{
-						ExecuteContext.OnAddShapeLibraryDelegate.Execute(&ExecuteContext, LibraryName, (UControlRigShapeLibrary*)UserData->GetMemory(), ReplaceExisting, LogShapeLibraries);
+						ExecuteContext.OnAddShapeLibraryDelegate.Execute(&ExecuteContext, LibraryName, (UControlRigShapeLibrary*)UserData->GetMemory(), LogShapeLibraries);
 					}
 					return;
 				}
diff --git a/Engine/Plugins/Animation/ControlRig/Source/ControlRig/Public/ControlRig.h b/Engine/Plugins/Animation/ControlRig/Source/ControlRig/Public/ControlRig.h
--- a/Engine/Plugins/Animation/ControlRig/Source/ControlRig/Public/ControlRig.h
+++ b/Engine/Plugins/Animation/ControlRig/Source/ControlRig/Public/ControlRig.h
@@ -530,7 +530,7 @@ protected:
 	FRigEventDelegate RigEventDelegate;
 
 	void RestoreShapeLibrariesFromCDO();
-	void OnAddShapeLibrary(const FControlRigExecuteContext* InContext, const FString& InLibraryName, UControlRigShapeLibrary* InShapeLibrary, bool bReplaceExisting, bool bLogResults);
+	void OnAddShapeLibrary(const FControlRigExecuteContext* InContext, const FString& InLibraryName, UControlRigShapeLibrary* InShapeLibrary, bool bLogResults);
 	bool OnShapeExists(const FName& InShapeName) const;
 	virtual void InitializeVMsFromCDO() { Super::InitializeFromCDO(); }
 	virtual void InitializeFromCDO() override;
diff --git a/Engine/Plugins/Animation/ControlRig/Source/ControlRig/Public/Units/Core/RigUnit_UserData.h b/Engine/Plugins/Animation/ControlRig/Source/ControlRig/Public/Units/Core/RigUnit_UserData.h
--- a/Engine/Plugins/Animation/ControlRig/Source/ControlRig/Public/Units/Core/RigUnit_UserData.h
+++ b/Engine/Plugins/Animation/ControlRig/Source/ControlRig/Public/Units/Core/RigUnit_UserData.h
@@ -48,7 +48,6 @@ struct CONTROLRIG_API FRigUnit_SetupShapeLibraryFromUserData : public FRigUnitMu
 
 	FRigUnit_SetupShapeLibraryFromUserData()
 	{
-		ReplaceExisting = true;
 		LogShapeLibraries = false;
 	}
 
@@ -68,18 +67,10 @@ struct CONTROLRIG_API FRigUnit_SetupShapeLibraryFromUserData : public FRigUnitMu
 	 * Optionally provide the namespace of the shape library to use.
 	 * This is only useful if you have multiple shape libraries and you
 	 * want to override a specific one.
-	 * If this is empty and ReplaceExisting is checked we'll replace the first
-	 * shape library on this rig. 
 	 */
 	UPROPERTY(meta = (Input))
 	FString LibraryName;
 
-	/*
-	 * If true we'll try to replace an existing library.
-	 */
-	UPROPERTY(meta = (Input))
-	bool ReplaceExisting;
-
 	/*
 	 * If this is checked we'll output the resulting shape libraries to the log for debugging.
 	 */
diff --git a/Engine/Plugins/Animation/ControlRig/Source/ControlRig/Public/Units/RigUnitContext.h b/Engine/Plugins/Animation/ControlRig/Source/ControlRig/Public/Units/RigUnitContext.h
--- a/Engine/Plugins/Animation/ControlRig/Source/ControlRig/Public/Units/RigUnitContext.h
+++ b/Engine/Plugins/Animation/ControlRig/Source/ControlRig/Public/Units/RigUnitContext.h
@@ -189,7 +189,7 @@ public:
 	/** The list of available asset user data object */
 	TArray<const UAssetUserData*> AssetUserData;
 
-	DECLARE_DELEGATE_FiveParams(FOnAddShapeLibrary, const FControlRigExecuteContext* InContext, const FString&, UControlRigShapeLibrary*, bool /* replace? */, bool /* log results */);
+	DECLARE_DELEGATE_FourParams(FOnAddShapeLibrary, const FControlRigExecuteContext* InContext, const FString&, UControlRigShapeLibrary*, bool /* log results */);
 	FOnAddShapeLibrary OnAddShapeLibraryDelegate;
 
 	DECLARE_DELEGATE_RetVal_OneParam(bool, FOnShapeExists, const FName&);
