ControlRig: Follow up to 30290002

#rnx

diff --git a/Engine/Plugins/Animation/ControlRig/Source/ControlRig/Private/Rigs/RigHierarchy.cpp b/Engine/Plugins/Animation/ControlRig/Source/ControlRig/Private/Rigs/RigHierarchy.cpp
--- a/Engine/Plugins/Animation/ControlRig/Source/ControlRig/Private/Rigs/RigHierarchy.cpp
+++ b/Engine/Plugins/Animation/ControlRig/Source/ControlRig/Private/Rigs/RigHierarchy.cpp
@@ -1432,7 +1432,7 @@ FRigName URigHierarchy::GetSafeNewName(const FRigName& InPotentialNewName, ERigE
 	SanitizeName(SanitizedName, bAllowNameSpace);
 
 	bAllowNameSpaceWhenSanitizingName = bAllowNameSpace;
-	if(!bAllowNameSpaceWhenSanitizingName && ExecuteContext)
+	if(ExecuteContext)
 	{
 		const FControlRigExecuteContext& CRContext = ExecuteContext->GetPublicData<FControlRigExecuteContext>();
 		if(CRContext.IsRigModule())
