Control Rig: Update rig hierarchy widget with resolval messages

#rnx

diff --git a/Engine/Plugins/Animation/ControlRig/Source/ControlRig/Private/Rigs/RigHierarchyDefines.cpp b/Engine/Plugins/Animation/ControlRig/Source/ControlRig/Private/Rigs/RigHierarchyDefines.cpp
--- a/Engine/Plugins/Animation/ControlRig/Source/ControlRig/Private/Rigs/RigHierarchyDefines.cpp
+++ b/Engine/Plugins/Animation/ControlRig/Source/ControlRig/Private/Rigs/RigHierarchyDefines.cpp
@@ -417,7 +417,6 @@ bool FModularRigResolveResult::ContainsMatch(const FRigElementKey& InKey, FStrin
 		}))
 		{
 			*OutErrorMessage = Mismatch->GetMessage().ToString();
-			return true;
 		}
 	}
 	
diff --git a/Engine/Plugins/Animation/ControlRig/Source/ControlRigEditor/Private/Editor/SRigHierarchy.cpp b/Engine/Plugins/Animation/ControlRig/Source/ControlRigEditor/Private/Editor/SRigHierarchy.cpp
--- a/Engine/Plugins/Animation/ControlRig/Source/ControlRigEditor/Private/Editor/SRigHierarchy.cpp
+++ b/Engine/Plugins/Animation/ControlRig/Source/ControlRigEditor/Private/Editor/SRigHierarchy.cpp
@@ -158,6 +158,7 @@ void SRigHierarchy::Construct(const FArguments& InArgs, TSharedRef<FControlRigEd
 	Delegates.OnGetResolvedKey = FOnRigTreeGetResolvedKey::CreateSP(this, &SRigHierarchy::OnGetResolvedKey);
 	Delegates.OnRequestDetailsInspection = FOnRigTreeRequestDetailsInspection::CreateSP(this, &SRigHierarchy::OnRequestDetailsInspection);
 	Delegates.OnRigTreeElementKeyTagDragDetected = FOnRigTreeRequestDetailsInspection::CreateSP(this, &SRigHierarchy::OnElementKeyTagDragDetected);
+	Delegates.OnRigTreeGetItemToolTip = FOnRigTreeItemGetToolTip::CreateSP(this, &SRigHierarchy::OnGetItemTooltip);
 
 	ChildSlot
 	[
@@ -1153,6 +1154,28 @@ void SRigHierarchy::OnSetExpansionRecursive(TSharedPtr<FRigTreeElement> InItem,
 	TreeView->SetExpansionRecursive(InItem, false, bShouldBeExpanded);
 }
 
+TOptional<FText> SRigHierarchy::OnGetItemTooltip(const FRigElementKey& InKey) const
+{
+	if(!DragRigResolveResults.IsEmpty() && FSlateApplication::Get().IsDragDropping())
+	{
+		FString Message;
+		for(const TPair<FRigElementKey, FModularRigResolveResult>& DragRigResolveResult : DragRigResolveResults)
+		{
+			if(DragRigResolveResult.Key != InKey)
+			{
+				if(!DragRigResolveResult.Value.ContainsMatch(InKey, &Message))
+				{
+					if(!Message.IsEmpty())
+					{
+						return FText::FromString(Message);
+					}
+				}
+			}
+		}
+	}
+	return TOptional<FText>();
+}
+
 void SRigHierarchy::CreateDragDropMenu()
 {
 	static bool bCreatedMenu = false;
@@ -2749,6 +2772,7 @@ FReply SRigHierarchy::OnAcceptDrop(const FDragDropEvent& DragDropEvent, EItemDro
 	bool bSummonDragDropMenu = DragDropEvent.GetModifierKeys().IsAltDown() && DragDropEvent.GetModifierKeys().IsShiftDown(); 
 	bool bMatchTransforms = DragDropEvent.GetModifierKeys().IsAltDown();
 	bool bReparentItems = !bMatchTransforms;
+	UpdateConnectorMatchesOnDrag({});
 
 	TSharedPtr<FRigElementHierarchyDragDropOp> RigDragDropOp = DragDropEvent.GetOperationAs<FRigElementHierarchyDragDropOp>();
 	if (RigDragDropOp.IsValid())
@@ -2880,26 +2904,35 @@ void SRigHierarchy::OnElementKeyTagDragDetected(const FRigElementKey& InDraggedT
 void SRigHierarchy::UpdateConnectorMatchesOnDrag(const TArray<FRigElementKey>& InDraggedKeys)
 {
 	DragRigResolveResults.Reset();
+
+	// fade in all items
+	for(const TPair<FRigElementKey, TSharedPtr<FRigTreeElement>>& Pair : TreeView->ElementMap)
+	{
+		Pair.Value->bFadedOutDuringDragDrop = false;
+	}
+	
 	if(ControlRigBeingDebuggedPtr.IsValid())
 	{
-		const UModularRig* ControlRig = Cast<UModularRig>(ControlRigBeingDebuggedPtr.Get());
-		if(URigHierarchy* Hierarchy = ControlRig->GetHierarchy())
+		if(const UModularRig* ControlRig = Cast<UModularRig>(ControlRigBeingDebuggedPtr.Get()))
 		{
-			if(const UModularRigRuleManager* RuleManager = Hierarchy->GetRuleManager())
+			if(URigHierarchy* Hierarchy = ControlRig->GetHierarchy())
 			{
-				for(const FRigElementKey& DraggedElement : InDraggedKeys)
+				if(const UModularRigRuleManager* RuleManager = Hierarchy->GetRuleManager())
 				{
-					if(DraggedElement.Type == ERigElementType::Connector)
+					for(const FRigElementKey& DraggedElement : InDraggedKeys)
 					{
-						if(const FRigConnectorElement* Connector = Hierarchy->Find<FRigConnectorElement>(DraggedElement))
+						if(DraggedElement.Type == ERigElementType::Connector)
 						{
-							const FName NameSpace = Hierarchy->GetNameMetadata(Connector->GetKey(), URigHierarchy::NameSpaceMetadataName, NAME_None);
-							if(!NameSpace.IsNone())
+							if(const FRigConnectorElement* Connector = Hierarchy->Find<FRigConnectorElement>(DraggedElement))
 							{
-								if(const FRigModuleInstance* Module = ControlRig->FindModule(NameSpace.ToString()))
+								const FName NameSpace = Hierarchy->GetNameMetadata(Connector->GetKey(), URigHierarchy::NameSpaceMetadataName, NAME_None);
+								if(!NameSpace.IsNone())
 								{
-									const FModularRigResolveResult ResolveResult = RuleManager->FindMatches(Connector, Module, ControlRig->ElementKeyRedirector); 
-									DragRigResolveResults.Add(DraggedElement, ResolveResult);
+									if(const FRigModuleInstance* Module = ControlRig->FindModule(NameSpace.ToString()))
+									{
+										const FModularRigResolveResult ResolveResult = RuleManager->FindMatches(Connector, Module, ControlRig->ElementKeyRedirector); 
+										DragRigResolveResults.Add(DraggedElement, ResolveResult);
+									}
 								}
 							}
 						}
@@ -2908,6 +2941,18 @@ void SRigHierarchy::UpdateConnectorMatchesOnDrag(const TArray<FRigElementKey>& I
 			}
 		}
 	}
+
+	// fade out anything that's on an excluded list
+	for(const TPair<FRigElementKey, FModularRigResolveResult>& Pair: DragRigResolveResults)
+	{
+		for(const FRigElementResolveResult& ExcludedElement : Pair.Value.GetExcluded())
+		{
+			if(const TSharedPtr<FRigTreeElement>* TreeElementPtr = TreeView->ElementMap.Find(ExcludedElement.GetKey()))
+			{
+				TreeElementPtr->Get()->bFadedOutDuringDragDrop = true;
+			}
+		}
+	}
 }
 
 FName SRigHierarchy::HandleRenameElement(const FRigElementKey& OldKey, const FString& NewName)
diff --git a/Engine/Plugins/Animation/ControlRig/Source/ControlRigEditor/Private/Editor/SRigHierarchy.h b/Engine/Plugins/Animation/ControlRig/Source/ControlRigEditor/Private/Editor/SRigHierarchy.h
--- a/Engine/Plugins/Animation/ControlRig/Source/ControlRigEditor/Private/Editor/SRigHierarchy.h
+++ b/Engine/Plugins/Animation/ControlRig/Source/ControlRigEditor/Private/Editor/SRigHierarchy.h
@@ -144,6 +144,7 @@ private:
 	void OnItemClicked(TSharedPtr<FRigTreeElement> InItem);
 	void OnItemDoubleClicked(TSharedPtr<FRigTreeElement> InItem);
 	void OnSetExpansionRecursive(TSharedPtr<FRigTreeElement> InItem, bool bShouldBeExpanded);
+	TOptional<FText> OnGetItemTooltip(const FRigElementKey& InKey) const;
 
 	// FEditorUndoClient
 	virtual void PostUndo(bool bSuccess) override;
diff --git a/Engine/Plugins/Animation/ControlRig/Source/ControlRigEditor/Private/Editor/SRigHierarchyTreeView.cpp b/Engine/Plugins/Animation/ControlRig/Source/ControlRigEditor/Private/Editor/SRigHierarchyTreeView.cpp
--- a/Engine/Plugins/Animation/ControlRig/Source/ControlRigEditor/Private/Editor/SRigHierarchyTreeView.cpp
+++ b/Engine/Plugins/Animation/ControlRig/Source/ControlRigEditor/Private/Editor/SRigHierarchyTreeView.cpp
@@ -41,6 +41,7 @@ FRigTreeElement::FRigTreeElement(const FRigElementKey& InKey, TWeakPtr<SRigHiera
 	bIsProcedural = false;
 	bSupportsRename = InSupportsRename;
 	FilterResult = InFilterResult;
+	bFadedOutDuringDragDrop = false;
 
 	if(InTreeView.IsValid())
 	{
@@ -103,6 +104,30 @@ void FRigTreeElement::RefreshDisplaySettings(const URigHierarchy* InHierarchy, c
 		(InHierarchy->IsProcedural(Key) ? FSlateColor(FLinearColor(0.9f, 0.8f, 0.4f) * 0.5f) : FSlateColor(FLinearColor::Gray * 0.5f));
 }
 
+FSlateColor FRigTreeElement::GetIconColor() const
+{
+	if(bFadedOutDuringDragDrop)
+	{
+		if(FSlateApplication::Get().IsDragDropping())
+		{
+			return IconColor.GetColor(FWidgetStyle()) * 0.3f;
+		}
+	}
+	return IconColor;
+}
+
+FSlateColor FRigTreeElement::GetTextColor() const
+{
+	if(bFadedOutDuringDragDrop)
+	{
+		if(FSlateApplication::Get().IsDragDropping())
+		{
+			return TextColor.GetColor(FWidgetStyle()) * 0.3f;
+		}
+	}
+	return TextColor;
+}
+
 //////////////////////////////////////////////////////////////
 /// SRigHierarchyItem
 ///////////////////////////////////////////////////////////
@@ -163,7 +188,7 @@ void SRigHierarchyItem::Construct(const FArguments& InArgs, const TSharedRef<STa
 				{
 					if(WeakRigTreeElement.IsValid())
 					{
-						return WeakRigTreeElement.Pin()->IconColor;
+						return WeakRigTreeElement.Pin()->GetIconColor();
 					}
 					return FSlateColor::UseForeground();
 				})
@@ -182,7 +207,7 @@ void SRigHierarchyItem::Construct(const FArguments& InArgs, const TSharedRef<STa
 				{
 					if(WeakRigTreeElement.IsValid())
 					{
-						return WeakRigTreeElement.Pin()->TextColor;
+						return WeakRigTreeElement.Pin()->GetTextColor();
 					}
 					return FSlateColor::UseForeground();
 				})
@@ -224,6 +249,14 @@ FText SRigHierarchyItem::GetName(bool bUseShortName) const
 
 FText SRigHierarchyItem::GetItemTooltip() const
 {
+	if(Delegates.OnRigTreeGetItemToolTip.IsBound())
+	{
+		const TOptional<FText> ToolTip = Delegates.OnRigTreeGetItemToolTip.Execute(WeakRigTreeElement.Pin()->Key);
+		if(ToolTip.IsSet())
+		{
+			return ToolTip.GetValue();
+		}
+	}
 	const FText FullName = GetName(false);
 	const FText ShortName = GetName(true);
 	if(FullName.EqualTo(ShortName))
diff --git a/Engine/Plugins/Animation/ControlRig/Source/ControlRigEditor/Private/Editor/SRigHierarchyTreeView.h b/Engine/Plugins/Animation/ControlRig/Source/ControlRigEditor/Private/Editor/SRigHierarchyTreeView.h
--- a/Engine/Plugins/Animation/ControlRig/Source/ControlRigEditor/Private/Editor/SRigHierarchyTreeView.h
+++ b/Engine/Plugins/Animation/ControlRig/Source/ControlRigEditor/Private/Editor/SRigHierarchyTreeView.h
@@ -74,6 +74,7 @@ DECLARE_DELEGATE_RetVal_ThreeParams(bool, FOnRigTreeVerifyElementNameChanged, co
 DECLARE_DELEGATE_RetVal_TwoParams(bool, FOnRigTreeCompareKeys, const FRigElementKey& /*A*/, const FRigElementKey& /*B*/);
 DECLARE_DELEGATE_RetVal_OneParam(FRigElementKey, FOnRigTreeGetResolvedKey, const FRigElementKey&);
 DECLARE_DELEGATE_OneParam(FOnRigTreeRequestDetailsInspection, const FRigElementKey&);
+DECLARE_DELEGATE_RetVal_OneParam(TOptional<FText>, FOnRigTreeItemGetToolTip, const FRigElementKey&);
 
 typedef STableRow<TSharedPtr<FRigTreeElement>>::FOnCanAcceptDrop FOnRigTreeCanAcceptDrop;
 typedef STableRow<TSharedPtr<FRigTreeElement>>::FOnAcceptDrop FOnRigTreeAcceptDrop;
@@ -100,6 +101,7 @@ struct CONTROLRIGEDITOR_API FRigTreeDelegates
 	FOnRigTreeGetResolvedKey OnGetResolvedKey;
 	FOnRigTreeRequestDetailsInspection OnRequestDetailsInspection;
 	FOnRigTreeElementKeyTagDragDetected OnRigTreeElementKeyTagDragDetected;
+	FOnRigTreeItemGetToolTip OnRigTreeGetItemToolTip;
 
 	FRigTreeDelegates()
 	{
@@ -212,6 +214,9 @@ public:
 
 	void RefreshDisplaySettings(const URigHierarchy* InHierarchy, const FRigTreeDisplaySettings& InSettings);
 
+	FSlateColor GetIconColor() const;
+	FSlateColor GetTextColor() const;
+
 	/** Delegate for when the context menu requests a rename */
 	DECLARE_DELEGATE(FOnRenameRequested);
 	FOnRenameRequested OnRenameRequested;
@@ -228,6 +233,9 @@ public:
 	/** The color to use when rendering the label text */
 	FSlateColor TextColor;
 
+	/** If true the item is filtered out during a drag */
+	bool bFadedOutDuringDragDrop;
+
 	/** The tag arguments for this element */
 	TArray<SRigHierarchyTagWidget::FArguments> Tags;
 };
