Control Rig: Removed Additive checkbox for FKControlRig
#rnx

diff --git a/Engine/Plugins/Animation/ControlRig/Source/ControlRig/Private/Rigs/FKControlRig.cpp b/Engine/Plugins/Animation/ControlRig/Source/ControlRig/Private/Rigs/FKControlRig.cpp
--- a/Engine/Plugins/Animation/ControlRig/Source/ControlRig/Private/Rigs/FKControlRig.cpp
+++ b/Engine/Plugins/Animation/ControlRig/Source/ControlRig/Private/Rigs/FKControlRig.cpp
@@ -504,63 +504,74 @@ void UFKControlRig::CreateRigElements(const USkeletalMesh* InReferenceMesh)
 
 void UFKControlRig::SetApplyMode(EControlRigFKRigExecuteMode InMode)
 {
+	if (ApplyMode == InMode)
+	{
+		return;
+	}
+	
 	ApplyMode = InMode;
+
+	if (URigHierarchy* Hierarchy = GetHierarchy())
+	{
+		if (ApplyMode == EControlRigFKRigExecuteMode::Additive)
+		{
+			FRigControlModifiedContext Context;
+			Context.SetKey = EControlRigSetKey::Never;
+			const bool bSetupUndo = false;
+
+			Hierarchy->ForEach<FRigControlElement>([&](FRigControlElement* ControlElement) -> bool
+			{
+				if (ControlElement->Settings.ControlType == ERigControlType::EulerTransform)
+				{
+					SetControlValue<FEulerTransform>(ControlElement->GetFName(), FEulerTransform::Identity, true, Context, bSetupUndo);
+				}
+				else if (ControlElement->Settings.ControlType == ERigControlType::Float)
+				{
+					SetControlValue<float>(ControlElement->GetFName(), 0.f, true, Context, bSetupUndo);
+				}
+
+				return true;
+			});
+		}
+		else
+		{
+			FRigControlModifiedContext Context;
+			Context.SetKey = EControlRigSetKey::Never;
+			const bool bSetupUndo = false;
+
+			Hierarchy->ForEach<FRigControlElement>([&](FRigControlElement* ControlElement) -> bool
+			{
+				if (ControlElement->Settings.ControlType == ERigControlType::EulerTransform)
+				{
+					const FRigControlValue::FEulerTransform_Float InitValue =
+						GetHierarchy()->GetControlValue(ControlElement, ERigControlValueType::Initial).Get<FRigControlValue::FEulerTransform_Float>();
+					SetControlValue<FRigControlValue::FEulerTransform_Float>(ControlElement->GetFName(), InitValue, true, Context, bSetupUndo);
+				}
+				else if (ControlElement->Settings.ControlType == ERigControlType::Float)
+				{
+					const float InitValue = GetHierarchy()->GetControlValue(ControlElement, ERigControlValueType::Initial).Get<float>();
+					SetControlValue<float>(ControlElement->GetFName(), InitValue, true, Context, bSetupUndo);
+				}
+
+				return true;
+			});
+		}
+	}
 }
 
 void UFKControlRig::ToggleApplyMode()
 {
+	EControlRigFKRigExecuteMode ModeToApply = ApplyMode;
 	if (ApplyMode == EControlRigFKRigExecuteMode::Additive)
 	{
-		ApplyMode = CachedToggleApplyMode;
+		ModeToApply = CachedToggleApplyMode;
 	}
 	else
 	{
 		CachedToggleApplyMode = ApplyMode;
-		ApplyMode = EControlRigFKRigExecuteMode::Additive;
-	}
-	if (ApplyMode == EControlRigFKRigExecuteMode::Additive)
-	{
-		FRigControlModifiedContext Context;
-		Context.SetKey = EControlRigSetKey::Never;
-		const bool bSetupUndo = false;
-
-		GetHierarchy()->ForEach<FRigControlElement>([&](FRigControlElement* ControlElement) -> bool
-        {
-            if (ControlElement->Settings.ControlType == ERigControlType::EulerTransform)
-			{
-				SetControlValue<FEulerTransform>(ControlElement->GetFName(), FEulerTransform::Identity, true, Context, bSetupUndo);
-			}
-			else if (ControlElement->Settings.ControlType == ERigControlType::Float)
-			{
-				SetControlValue<float>(ControlElement->GetFName(), 0.f, true, Context, bSetupUndo);
-			}
-
-			return true;
-		});
-	}
-	else
-	{
-		FRigControlModifiedContext Context;
-		Context.SetKey = EControlRigSetKey::Never;
-		const bool bSetupUndo = false;
-
-		GetHierarchy()->ForEach<FRigControlElement>([&](FRigControlElement* ControlElement) -> bool
-        {
-            if (ControlElement->Settings.ControlType == ERigControlType::EulerTransform)
-			{
-				const FRigControlValue::FEulerTransform_Float InitValue =
-					GetHierarchy()->GetControlValue(ControlElement, ERigControlValueType::Initial).Get<FRigControlValue::FEulerTransform_Float>();
-				SetControlValue<FRigControlValue::FEulerTransform_Float>(ControlElement->GetFName(), InitValue, true, Context, bSetupUndo);
-			}
-			else if (ControlElement->Settings.ControlType == ERigControlType::Float)
-			{
-				const float InitValue = GetHierarchy()->GetControlValue(ControlElement, ERigControlValueType::Initial).Get<float>();
-				SetControlValue<float>(ControlElement->GetFName(), InitValue, true, Context, bSetupUndo);
-			}
-
-			return true;
-		});
+		ModeToApply = EControlRigFKRigExecuteMode::Additive;
 	}
+	SetApplyMode(ModeToApply);
 }
 
 #undef LOCTEXT_NAMESPACE
diff --git a/Engine/Plugins/Animation/ControlRig/Source/ControlRig/Public/ControlRig.h b/Engine/Plugins/Animation/ControlRig/Source/ControlRig/Public/ControlRig.h
--- a/Engine/Plugins/Animation/ControlRig/Source/ControlRig/Public/ControlRig.h
+++ b/Engine/Plugins/Animation/ControlRig/Source/ControlRig/Public/ControlRig.h
@@ -698,7 +698,7 @@ public:
 
 #endif
 	
-	bool IsAdditive() const { return bIsAdditive; }
+	virtual bool IsAdditive() const { return bIsAdditive; }
 	void SetIsAdditive(const bool bInIsAdditive)
 	{
 		bIsAdditive = bInIsAdditive;
diff --git a/Engine/Plugins/Animation/ControlRig/Source/ControlRig/Public/Rigs/FKControlRig.h b/Engine/Plugins/Animation/ControlRig/Source/ControlRig/Public/Rigs/FKControlRig.h
--- a/Engine/Plugins/Animation/ControlRig/Source/ControlRig/Public/Rigs/FKControlRig.h
+++ b/Engine/Plugins/Animation/ControlRig/Source/ControlRig/Public/Rigs/FKControlRig.h
@@ -68,7 +68,7 @@ public:
 	void SetApplyMode(EControlRigFKRigExecuteMode InMode);
 	void ToggleApplyMode();
 	bool CanToggleApplyMode() const { return true; }
-	bool IsApplyModeAdditive() const { return ApplyMode == EControlRigFKRigExecuteMode::Additive; }
+	bool IsAdditive() const override { return ApplyMode == EControlRigFKRigExecuteMode::Additive; }
 	EControlRigFKRigExecuteMode GetApplyMode() const { return ApplyMode; }
 
 	// Ensures that controls mask is updated according to contained ControlRig (control) elements
diff --git a/Engine/Plugins/Animation/ControlRig/Source/ControlRigEditor/Private/ControlRigSequencerEditorLibrary.cpp b/Engine/Plugins/Animation/ControlRig/Source/ControlRigEditor/Private/ControlRigSequencerEditorLibrary.cpp
--- a/Engine/Plugins/Animation/ControlRig/Source/ControlRigEditor/Private/ControlRigSequencerEditorLibrary.cpp
+++ b/Engine/Plugins/Animation/ControlRig/Source/ControlRigEditor/Private/ControlRigSequencerEditorLibrary.cpp
@@ -3106,11 +3106,18 @@ bool UControlRigSequencerEditorLibrary::SetControlRigLayeredMode(UMovieSceneCont
 	InTrack->Modify();
 	ControlRig->Modify();
 
-	ControlRig->ClearPoseBeforeBackwardsSolve();
-	ControlRig->ResetControlValues();
-	ControlRig->SetIsAdditive(bSetIsLayered);
+	if (UFKControlRig* FKRig = Cast<UFKControlRig>(ControlRig))
+	{
+		FKRig->SetApplyMode(bSetIsLayered ? EControlRigFKRigExecuteMode::Additive : EControlRigFKRigExecuteMode::Replace);
+	}
+	else
+	{
+		ControlRig->ClearPoseBeforeBackwardsSolve();
+		ControlRig->ResetControlValues();
+		ControlRig->SetIsAdditive(bSetIsLayered);
 
-	ControlRig->Evaluate_AnyThread();
+		ControlRig->Evaluate_AnyThread();
+	}
 
 	FString ObjectName = ControlRig->GetClass()->GetName(); //GetDisplayNameText().ToString();
 	ObjectName.RemoveFromEnd(TEXT("_C"));
diff --git a/Engine/Plugins/Animation/ControlRig/Source/ControlRigEditor/Private/Sequencer/ControlRigParameterTrackEditor.cpp b/Engine/Plugins/Animation/ControlRig/Source/ControlRigEditor/Private/Sequencer/ControlRigParameterTrackEditor.cpp
--- a/Engine/Plugins/Animation/ControlRig/Source/ControlRigEditor/Private/Sequencer/ControlRigParameterTrackEditor.cpp
+++ b/Engine/Plugins/Animation/ControlRig/Source/ControlRigEditor/Private/Sequencer/ControlRigParameterTrackEditor.cpp
@@ -1241,6 +1241,16 @@ void FControlRigParameterTrackEditor::BakeInvertedPose(UControlRig* InControlRig
 	ParentSequencer->NotifyMovieSceneDataChanged(EMovieSceneDataChangeType::TrackValueChanged);
 }
 
+bool FControlRigParameterTrackEditor::IsLayered(UMovieSceneControlRigParameterTrack* Track) const
+{
+	UControlRig* ControlRig = Track->GetControlRig();
+	if (!ControlRig)
+	{
+		return false;
+	}
+	return ControlRig->IsAdditive();
+}
+
 void FControlRigParameterTrackEditor::ConvertIsLayered(UMovieSceneControlRigParameterTrack* Track)
 {
 	UControlRig* ControlRig = Track->GetControlRig();
@@ -3925,7 +3935,7 @@ void FControlRigParameterTrackEditor::BuildTrackContextMenu(FMenuBuilder& MenuBu
 				   FUIAction(
 					   FExecuteAction::CreateRaw(this, &FControlRigParameterTrackEditor::ConvertIsLayered, Track),
 					   FCanExecuteAction(),
-					   FIsActionChecked::CreateUObject(Track->GetControlRig(), &UControlRig::IsAdditive)
+					   FIsActionChecked::CreateRaw(this, &FControlRigParameterTrackEditor::IsLayered, Track)
 				   ),
 				   NAME_None,
 				   EUserInterfaceActionType::ToggleButton);
@@ -3946,18 +3956,6 @@ void FControlRigParameterTrackEditor::BuildTrackContextMenu(FMenuBuilder& MenuBu
 				FSlateIcon(),
 				FUIAction(
 					FExecuteAction::CreateRaw(this, &FControlRigParameterTrackEditor::SelectFKBonesToAnimate, AutoRig, Track)));
-
-			MenuBuilder.AddMenuEntry(
-				LOCTEXT("FKRigApplyMode", "Additive"),
-				LOCTEXT("FKRigApplyModeToolTip", "Toggles the apply mode between Replace and Additive"),
-				FSlateIcon(),
-				FUIAction(
-					FExecuteAction::CreateRaw(this, &FControlRigParameterTrackEditor::ToggleFKControlRig, Track, AutoRig),
-					FCanExecuteAction::CreateUObject(AutoRig, &UFKControlRig::CanToggleApplyMode),
-					FIsActionChecked::CreateUObject(AutoRig, &UFKControlRig::IsApplyModeAdditive)
-				),
-				NAME_None,
-				EUserInterfaceActionType::ToggleButton);
 		}
 		MenuBuilder.EndSection();
 
@@ -4051,43 +4049,6 @@ bool FControlRigParameterTrackEditor::HandleAssetAdded(UObject* Asset, const FGu
 	return true;
 }
 
-
-void FControlRigParameterTrackEditor::ToggleFKControlRig(UMovieSceneControlRigParameterTrack* Track, UFKControlRig* FKControlRig)
-{
-	FScopedTransaction Transaction(LOCTEXT("ToggleFKControlRig", "Toggle FK Control Rig"));
-	FKControlRig->Modify();
-	Track->Modify();
-	FKControlRig->ToggleApplyMode();
-	if (FKControlRig->GetApplyMode() == EControlRigFKRigExecuteMode::Additive)
-	{
-		const FString LayeredObjectName = Track->GetTrackName().ToString() + TEXT(" (Layered)");
-		Track->SetDisplayName(FText::FromString(LayeredObjectName));
-		Track->SetColorTint(UMovieSceneControlRigParameterTrack::LayeredRigTrackColor);
-	}
-	else
-	{
-		Track->SetDisplayName(FText::FromName(Track->GetTrackName()));
-		Track->SetColorTint(UMovieSceneControlRigParameterTrack::LayeredRigTrackColor);
-		Track->SetColorTint(UMovieSceneControlRigParameterTrack::AbsoluteRigTrackColor);
-	}
-	for (UMovieSceneSection* Section : Track->GetAllSections())
-	{
-		if (Section)
-		{
-			UMovieSceneControlRigParameterSection* CRSection = Cast<UMovieSceneControlRigParameterSection>(Section);
-			if (CRSection)
-			{
-				Section->Modify();
-				CRSection->ClearAllParameters();
-				CRSection->RecreateWithThisControlRig(CRSection->GetControlRig(), true);
-			}
-		}
-	}
-	GetSequencer()->NotifyMovieSceneDataChanged(EMovieSceneDataChangeType::MovieSceneStructureItemsChanged);
-
-}
-
-
 void FControlRigParameterTrackEditor::ImportFBX(UMovieSceneControlRigParameterTrack* InTrack, UMovieSceneControlRigParameterSection* InSection,
 	TArray<FRigControlFBXNodeAndChannels>* NodeAndChannels)
 {
diff --git a/Engine/Plugins/Animation/ControlRig/Source/ControlRigEditor/Private/Sequencer/ControlRigParameterTrackEditor.h b/Engine/Plugins/Animation/ControlRig/Source/ControlRigEditor/Private/Sequencer/ControlRigParameterTrackEditor.h
--- a/Engine/Plugins/Animation/ControlRig/Source/ControlRigEditor/Private/Sequencer/ControlRigParameterTrackEditor.h
+++ b/Engine/Plugins/Animation/ControlRig/Source/ControlRigEditor/Private/Sequencer/ControlRigParameterTrackEditor.h
@@ -187,9 +187,6 @@ private:
 	/** Select Bones to Animate on FK Rig*/
 	void SelectFKBonesToAnimate(UFKControlRig* FKControlRig, UMovieSceneControlRigParameterTrack* Track);
 
-	/** Toggle FK Control Rig*/
-	void ToggleFKControlRig(UMovieSceneControlRigParameterTrack* Track, UFKControlRig* FKControlRig);
-
 	/** Bake To Control Rig Sub Menu*/
 	void BakeToControlRigSubMenu(FMenuBuilder& MenuBuilder, FGuid ObjectBinding, UObject* BoundObject, USkeletalMeshComponent* SkelMeshComp,USkeleton* Skeleton);
 	
@@ -199,6 +196,9 @@ private:
 	/** Bake Inersion of Additive Control Rig to Rest Pose*/
 	void BakeInvertedPose(UControlRig* InControlRig, UMovieSceneControlRigParameterTrack* Track);
 
+	/** Return true if the control rig is in layered mode */
+	bool IsLayered(UMovieSceneControlRigParameterTrack* Track) const;
+
 	/** Convert an absolute control rit to a layered control rig */
 	void ConvertIsLayered(UMovieSceneControlRigParameterTrack* Track);
 
