Control Rig: Fix additive rigs still using preferred euler angles in certain cases
#rnx

diff --git a/Engine/Plugins/Animation/ControlRig/Source/ControlRig/Private/ControlRig.cpp b/Engine/Plugins/Animation/ControlRig/Source/ControlRig/Private/ControlRig.cpp
--- a/Engine/Plugins/Animation/ControlRig/Source/ControlRig/Private/ControlRig.cpp
+++ b/Engine/Plugins/Animation/ControlRig/Source/ControlRig/Private/ControlRig.cpp
@@ -1965,7 +1965,7 @@ FTransform UControlRig::GetControlGlobalTransform(const FName& InControlName) co
 	return DynamicHierarchy->GetGlobalTransform(FRigElementKey(InControlName, ERigElementType::Control), false);
 }
 
-FRigControlValue UControlRig::GetControlValue(FRigControlElement* InControl, const ERigControlValueType& InValueType)
+FRigControlValue UControlRig::GetControlValue(FRigControlElement* InControl, const ERigControlValueType& InValueType) const
 {
 	if (bIsAdditive && InValueType == ERigControlValueType::Current)
 	{
@@ -2173,6 +2173,21 @@ FTransform UControlRig::GetControlLocalTransform(const FName& InControlName)
 	return DynamicHierarchy->GetLocalTransform(FRigElementKey(InControlName, ERigElementType::Control));
 }
 
+FVector UControlRig::GetControlSpecifiedEulerAngle(const FRigControlElement* InControlElement, bool bIsInitial) const
+{
+	if (bIsAdditive)
+	{
+		const ERigControlValueType Type = bIsInitial ? ERigControlValueType::Initial : ERigControlValueType::Current;
+		const FRigControlValue Value = GetControlValue(InControlElement->GetKey().Name);
+		FRotator Rotator = Value.GetAsTransform(InControlElement->Settings.ControlType, InControlElement->Settings.PrimaryAxis).Rotator();
+		return FVector(Rotator.Roll, Rotator.Pitch, Rotator.Yaw);
+	}
+	else
+	{
+		return DynamicHierarchy->GetControlSpecifiedEulerAngle(InControlElement, bIsInitial);
+	}
+}
+
 const TArray<TSoftObjectPtr<UControlRigShapeLibrary>>& UControlRig::GetShapeLibraries() const
 {
 	const TArray<TSoftObjectPtr<UControlRigShapeLibrary>>* LibrariesPtr = &ShapeLibraries;
diff --git a/Engine/Plugins/Animation/ControlRig/Source/ControlRig/Private/Rigs/RigHierarchy.cpp b/Engine/Plugins/Animation/ControlRig/Source/ControlRig/Private/Rigs/RigHierarchy.cpp
--- a/Engine/Plugins/Animation/ControlRig/Source/ControlRig/Private/Rigs/RigHierarchy.cpp
+++ b/Engine/Plugins/Animation/ControlRig/Source/ControlRig/Private/Rigs/RigHierarchy.cpp
@@ -2867,8 +2867,11 @@ FRigPose URigHierarchy::GetPose(bool bInitial, ERigElementType InElementType,
 
 			if(FRigControlElement* ControlElement = Cast<FRigControlElement>(Element))
 			{
-				PoseElement.PreferredEulerAngle = GetControlPreferredEulerAngles(ControlElement,
-					GetControlPreferredEulerRotationOrder(ControlElement), bInitial);
+				if (bUsePreferredEulerAngles)
+				{
+					PoseElement.PreferredEulerAngle = GetControlPreferredEulerAngles(ControlElement,
+					   GetControlPreferredEulerRotationOrder(ControlElement), bInitial);
+				}
 
 				if(!bIncludeTransientControls && ControlElement->Settings.bIsTransientControl)
 				{
diff --git a/Engine/Plugins/Animation/ControlRig/Source/ControlRig/Private/Sequencer/MovieSceneControlRigParameterSection.cpp b/Engine/Plugins/Animation/ControlRig/Source/ControlRig/Private/Sequencer/MovieSceneControlRigParameterSection.cpp
--- a/Engine/Plugins/Animation/ControlRig/Source/ControlRig/Private/Sequencer/MovieSceneControlRigParameterSection.cpp
+++ b/Engine/Plugins/Animation/ControlRig/Source/ControlRig/Private/Sequencer/MovieSceneControlRigParameterSection.cpp
@@ -635,7 +635,7 @@ struct FParameterTransformChannelEditorData
 					}
 				}
 				
-				FVector Vector = ControlRig->GetHierarchy()->GetControlSpecifiedEulerAngle(ControlElement);
+				FVector Vector = ControlRig->GetControlSpecifiedEulerAngle(ControlElement);
 				FRotator Rotation = FRotator(Vector.Y, Vector.Z, Vector.X);
 				return Rotation;
 			}
diff --git a/Engine/Plugins/Animation/ControlRig/Source/ControlRig/Private/Sequencer/MovieSceneControlRigParameterTemplate.cpp b/Engine/Plugins/Animation/ControlRig/Source/ControlRig/Private/Sequencer/MovieSceneControlRigParameterTemplate.cpp
--- a/Engine/Plugins/Animation/ControlRig/Source/ControlRig/Private/Sequencer/MovieSceneControlRigParameterTemplate.cpp
+++ b/Engine/Plugins/Animation/ControlRig/Source/ControlRig/Private/Sequencer/MovieSceneControlRigParameterTemplate.cpp
@@ -986,7 +986,7 @@ struct FControlRigParameterPreAnimatedTokenProducer : IMovieScenePreAnimatedToke
 						FVector3f Val = ControlRig->GetControlValue(ControlElement, ERigControlValueType::Current).Get<FVector3f>();
 						if (ControlElement->Settings.ControlType == ERigControlType::Rotator)
 						{
-							FVector Vector = RigHierarchy->GetControlSpecifiedEulerAngle(ControlElement);
+							FVector Vector = ControlRig->GetControlSpecifiedEulerAngle(ControlElement);
 							Val = FVector3f(Vector.X, Vector.Y, Vector.Z);
 						}
 						Token.VectorValues.Add(TNameAndValue<FVector>{ ControlElement->GetFName(), (FVector)Val });
@@ -1004,7 +1004,7 @@ struct FControlRigParameterPreAnimatedTokenProducer : IMovieScenePreAnimatedToke
 						Token.SpaceValues.Add(FControlSpaceAndValue(ControlElement->GetFName(), SpaceValue));
 						const FTransform Val = ControlRig->GetControlValue(ControlElement, ERigControlValueType::Current).Get<FRigControlValue::FTransform_Float>().ToTransform();
 						FEulerTransform EulerTransform(Val);
-						FVector Vector = RigHierarchy->GetControlSpecifiedEulerAngle(ControlElement);
+						FVector Vector = ControlRig->GetControlSpecifiedEulerAngle(ControlElement);
 						EulerTransform.Rotation = FRotator(Vector.Y, Vector.Z, Vector.X);
 						Token.TransformValues.Add(TNameAndValue<FEulerTransform>{ ControlElement->GetFName(), EulerTransform });
 						break;
@@ -1015,7 +1015,7 @@ struct FControlRigParameterPreAnimatedTokenProducer : IMovieScenePreAnimatedToke
 							ControlRig
 							->GetControlValue(ControlElement, ERigControlValueType::Current).Get<FRigControlValue::FTransformNoScale_Float>().ToTransform();
 						FEulerTransform EulerTransform(NoScale.ToFTransform());
-						FVector Vector = RigHierarchy->GetControlSpecifiedEulerAngle(ControlElement);
+						FVector Vector = ControlRig->GetControlSpecifiedEulerAngle(ControlElement);
 						EulerTransform.Rotation = FRotator(Vector.Y, Vector.Z, Vector.X);
 						Token.TransformValues.Add(TNameAndValue<FEulerTransform>{ ControlElement->GetFName(), EulerTransform });
 						break;
@@ -1025,7 +1025,7 @@ struct FControlRigParameterPreAnimatedTokenProducer : IMovieScenePreAnimatedToke
 						FEulerTransform EulerTransform = 
 							ControlRig
 							->GetControlValue(ControlElement, ERigControlValueType::Current).Get<FRigControlValue::FEulerTransform_Float>().ToTransform();
-						FVector Vector = RigHierarchy->GetControlSpecifiedEulerAngle(ControlElement);
+						FVector Vector = ControlRig->GetControlSpecifiedEulerAngle(ControlElement);
 						EulerTransform.Rotation = FRotator(Vector.Y, Vector.Z, Vector.X);
 						Token.TransformValues.Add(TNameAndValue<FEulerTransform>{ ControlElement->GetFName(), EulerTransform });
 						break;
@@ -1629,7 +1629,7 @@ struct TControlRigParameterActuatorVector : TMovieSceneBlendingActuator<FControl
 				FVector3f Val = ControlRig->GetControlValue(ControlElement, ERigControlValueType::Current).Get<FVector3f>();
 				if (ControlElement->Settings.ControlType == ERigControlType::Rotator)
 				{
-					FVector Vector = ControlRig->GetHierarchy()->GetControlSpecifiedEulerAngle(ControlElement);
+					FVector Vector = ControlRig->GetControlSpecifiedEulerAngle(ControlElement);
 					Val = FVector3f(Vector.X, Vector.Y, Vector.Z);
 				}
 				return FControlRigTrackTokenVector((FVector)Val);
@@ -1709,7 +1709,7 @@ struct TControlRigParameterActuatorTransform : TMovieSceneBlendingActuator<FCont
 				{
 					const FTransform Val = ControlRig->GetControlValue(ControlElement, ERigControlValueType::Current).Get<FRigControlValue::FTransform_Float>().ToTransform();
 					FEulerTransform EulerTransform(Val);
-					FVector Vector = Hierarchy->GetControlSpecifiedEulerAngle(ControlElement);
+					FVector Vector = ControlRig->GetControlSpecifiedEulerAngle(ControlElement);
 					EulerTransform.Rotation = FRotator(Vector.Y, Vector.Z, Vector.X);
 					return FControlRigTrackTokenTransform(EulerTransform);
 				}
@@ -1718,14 +1718,14 @@ struct TControlRigParameterActuatorTransform : TMovieSceneBlendingActuator<FCont
 					FTransformNoScale ValNoScale = ControlRig->GetControlValue(ControlElement, ERigControlValueType::Current).Get<FRigControlValue::FTransformNoScale_Float>().ToTransform();
 					FTransform Val = ValNoScale;
 					FEulerTransform EulerTransform(Val);
-					FVector Vector = Hierarchy->GetControlSpecifiedEulerAngle(ControlElement);
+					FVector Vector = ControlRig->GetControlSpecifiedEulerAngle(ControlElement);
 					EulerTransform.Rotation = FRotator(Vector.Y, Vector.Z, Vector.X);
 					return FControlRigTrackTokenTransform(EulerTransform);
 				}
 				else if (ControlElement && ControlElement->Settings.ControlType == ERigControlType::EulerTransform)
 				{
 					FEulerTransform EulerTransform = ControlRig->GetControlValue(ControlElement, ERigControlValueType::Current).Get<FRigControlValue::FEulerTransform_Float>().ToTransform();
-					FVector Vector = Hierarchy->GetControlSpecifiedEulerAngle(ControlElement);
+					FVector Vector = ControlRig->GetControlSpecifiedEulerAngle(ControlElement);
 					EulerTransform.Rotation = FRotator(Vector.Y, Vector.Z, Vector.X);
 					return FControlRigTrackTokenTransform(EulerTransform);
 				}
diff --git a/Engine/Plugins/Animation/ControlRig/Source/ControlRig/Public/ControlRig.h b/Engine/Plugins/Animation/ControlRig/Source/ControlRig/Public/ControlRig.h
--- a/Engine/Plugins/Animation/ControlRig/Source/ControlRig/Public/ControlRig.h
+++ b/Engine/Plugins/Animation/ControlRig/Source/ControlRig/Public/ControlRig.h
@@ -252,7 +252,7 @@ public:
 	}
 
 	// Returns the value of a Control
-	FRigControlValue GetControlValue(const FName& InControlName)
+	FRigControlValue GetControlValue(const FName& InControlName) const
 	{
 		const FRigElementKey Key(InControlName, ERigElementType::Control);
 		if (FRigBaseElement* Element = DynamicHierarchy->Find(Key))
@@ -265,7 +265,7 @@ public:
 		return DynamicHierarchy->GetControlValue(Key);
 	}
 
-	FRigControlValue GetControlValue(FRigControlElement* InControl, const ERigControlValueType& InValueType);
+	FRigControlValue GetControlValue(FRigControlElement* InControl, const ERigControlValueType& InValueType) const;
 
 	// Sets the relative value of a Control
 	virtual void SetControlValueImpl(const FName& InControlName, const FRigControlValue& InValue, bool bNotify = true,
@@ -290,6 +290,8 @@ public:
 	virtual void SetControlLocalTransform(const FName& InControlName, const FTransform& InLocalTransform, bool bNotify = true, const FRigControlModifiedContext& Context = FRigControlModifiedContext(), bool bSetupUndo = true, bool bFixEulerFlips = false);
 	virtual FTransform GetControlLocalTransform(const FName& InControlName) ;
 
+	FVector GetControlSpecifiedEulerAngle(const FRigControlElement* InControlElement, bool bIsInitial = false) const;
+
 	virtual const TArray<TSoftObjectPtr<UControlRigShapeLibrary>>& GetShapeLibraries() const;
 	virtual void CreateRigControlsForCurveContainer();
 	virtual void GetControlsInOrder(TArray<FRigControlElement*>& SortedControls) const;
diff --git a/Engine/Plugins/Animation/ControlRig/Source/ControlRig/Public/Rigs/RigHierarchy.h b/Engine/Plugins/Animation/ControlRig/Source/ControlRig/Public/Rigs/RigHierarchy.h
--- a/Engine/Plugins/Animation/ControlRig/Source/ControlRig/Public/Rigs/RigHierarchy.h
+++ b/Engine/Plugins/Animation/ControlRig/Source/ControlRig/Public/Rigs/RigHierarchy.h
@@ -2420,7 +2420,7 @@ public:
 
 	bool GetUsePreferredRotationOrder(const FRigControlElement* InControlElement) const
 	{
-		if (InControlElement)
+		if (InControlElement && bUsePreferredEulerAngles)
 		{
 			{
 				return InControlElement->Settings.bUsePreferredRotationOrder;
diff --git a/Engine/Plugins/Animation/ControlRig/Source/ControlRigEditor/Private/ControlRigElementDetails.cpp b/Engine/Plugins/Animation/ControlRig/Source/ControlRigEditor/Private/ControlRigElementDetails.cpp
--- a/Engine/Plugins/Animation/ControlRig/Source/ControlRigEditor/Private/ControlRigElementDetails.cpp
+++ b/Engine/Plugins/Animation/ControlRig/Source/ControlRigEditor/Private/ControlRigElementDetails.cpp
@@ -2024,7 +2024,15 @@ FDetailWidgetRow& FRigTransformElementDetails::CreateEulerTransformValueWidgetRo
 					case ERigControlType::Transform:
 					case ERigControlType::TransformNoScale:
 					{
-						FVector Vector = Hierarchy->GetControlSpecifiedEulerAngle(ControlElement, bInitial);
+						FVector Vector;
+						if(const UControlRig* ControlRig = Hierarchy->GetTypedOuter<UControlRig>())
+						{
+							Vector = ControlRig->GetControlSpecifiedEulerAngle(ControlElement, bInitial);
+						}
+						else
+						{
+							Vector = Hierarchy->GetControlSpecifiedEulerAngle(ControlElement, bInitial);
+						}
 						RelativeTransform.Rotation =  FRotator(Vector.Y, Vector.Z, Vector.X);
 						break;
 					}
diff --git a/Engine/Plugins/Animation/ControlRig/Source/ControlRigEditor/Private/EditMode/ControlRigControlsProxy.cpp b/Engine/Plugins/Animation/ControlRig/Source/ControlRigEditor/Private/EditMode/ControlRigControlsProxy.cpp
--- a/Engine/Plugins/Animation/ControlRig/Source/ControlRigEditor/Private/EditMode/ControlRigControlsProxy.cpp
+++ b/Engine/Plugins/Animation/ControlRig/Source/ControlRigEditor/Private/EditMode/ControlRigControlsProxy.cpp
@@ -612,7 +612,7 @@ void UControlRigVectorControlProxy::ValueChanged()
 		FVector3f Val = ControlRig.Get()->GetControlValue(ControlElement, ERigControlValueType::Current).Get<FVector3f>();
 		if (ControlElement->Settings.ControlType == ERigControlType::Rotator)
 		{ 
-			FVector DVector = ControlRig->GetHierarchy()->GetControlSpecifiedEulerAngle(ControlElement);
+			FVector DVector = ControlRig->GetControlSpecifiedEulerAngle(ControlElement);
 			Val = FVector3f(DVector.X, DVector.Y, DVector.Z);
 		}
 		Binding.CallFunction<FVector3f>(*this, Val);
diff --git a/Engine/Plugins/Animation/ControlRig/Source/ControlRigEditor/Private/Sequencer/ControlRigParameterTrackEditor.cpp b/Engine/Plugins/Animation/ControlRig/Source/ControlRigEditor/Private/Sequencer/ControlRigParameterTrackEditor.cpp
--- a/Engine/Plugins/Animation/ControlRig/Source/ControlRigEditor/Private/Sequencer/ControlRigParameterTrackEditor.cpp
+++ b/Engine/Plugins/Animation/ControlRig/Source/ControlRigEditor/Private/Sequencer/ControlRigParameterTrackEditor.cpp
@@ -3418,7 +3418,7 @@ void FControlRigParameterTrackEditor::GetControlRigKeys(
 			case ERigControlType::EulerTransform:
 			{
 				FVector Translation, Scale(1.0f, 1.0f, 1.0f);
-				FVector Vector = Hierarchy->GetControlSpecifiedEulerAngle(ControlElement);
+				FVector Vector = InControlRig->GetControlSpecifiedEulerAngle(ControlElement);
 				FRotator Rotation = FRotator(Vector.Y, Vector.Z, Vector.X);
 				if (ControlElement->Settings.ControlType == ERigControlType::TransformNoScale)
 				{
@@ -3451,7 +3451,7 @@ void FControlRigParameterTrackEditor::GetControlRigKeys(
 							FVector Angle = Rotation.Euler();
 							//need to wind rotators still
 							ControlElement->PreferredEulerAngles.SetAngles(Angle, false, ControlElement->PreferredEulerAngles.RotationOrder, true);
-							Angle = Hierarchy->GetControlSpecifiedEulerAngle(ControlElement);
+							Angle = InControlRig->GetControlSpecifiedEulerAngle(ControlElement);
 							Rotation = FRotator(Vector.Y, Vector.Z, Vector.X);
 						}
 						else
