Sequencer: Control Rig: Clear out delegates in outliner, but need to use cached control rig's

diff --git a/Engine/Plugins/Animation/ControlRig/Source/ControlRigEditor/Private/EditMode/SControlRigOutliner.cpp b/Engine/Plugins/Animation/ControlRig/Source/ControlRigEditor/Private/EditMode/SControlRigOutliner.cpp
--- a/Engine/Plugins/Animation/ControlRig/Source/ControlRigEditor/Private/EditMode/SControlRigOutliner.cpp
+++ b/Engine/Plugins/Animation/ControlRig/Source/ControlRigEditor/Private/EditMode/SControlRigOutliner.cpp
@@ -1027,22 +1027,19 @@ SControlRigOutliner::SControlRigOutliner()
 SControlRigOutliner::~SControlRigOutliner()
 {
 	FCoreUObjectDelegates::OnObjectsReplaced.RemoveAll(this);
-	if (FControlRigEditMode* EditMode = static_cast<FControlRigEditMode*>(ModeTools->GetActiveMode(FControlRigEditMode::ModeName)))
-	{
-		TArrayView<TWeakObjectPtr<UControlRig>> ControlRigs = EditMode->GetControlRigs();
-		for (TWeakObjectPtr<UControlRig>& ControlRig : ControlRigs)
+	for(TWeakObjectPtr<UControlRig>& ControlRig: BoundControlRigs)
+	{ 
+		if (ControlRig.IsValid())
 		{
-			if (ControlRig.IsValid())
+			ControlRig.Get()->ControlRigBound().RemoveAll(this);
+			const TSharedPtr<IControlRigObjectBinding> Binding = ControlRig.Get()->GetObjectBinding();
+			if (Binding)
 			{
-				ControlRig.Get()->ControlRigBound().RemoveAll(this);
-				const TSharedPtr<IControlRigObjectBinding> Binding = ControlRig.Get()->GetObjectBinding();
-				if (Binding)
-				{
-					Binding->OnControlRigBind().RemoveAll(this);
-				}
+				Binding->OnControlRigBind().RemoveAll(this);
 			}
 		}
 	}
+	BoundControlRigs.SetNum(0);
 }
 
 void SControlRigOutliner::HandleControlSelected(UControlRig* Subject, FRigControlElement* ControlElement, bool bSelected)
@@ -1155,6 +1152,7 @@ void SControlRigOutliner::SetEditMode(FControlRigEditMode& InEditMode)
 				if (!ControlRig.Get()->ControlRigBound().IsBoundToObject(this))
 				{
 					ControlRig.Get()->ControlRigBound().AddRaw(this, &SControlRigOutliner::HandleOnControlRigBound);
+					BoundControlRigs.Add(ControlRig);
 				}
 				const TSharedPtr<IControlRigObjectBinding> Binding = ControlRig.Get()->GetObjectBinding();
 				if (Binding && !Binding->OnControlRigBind().IsBoundToObject(this))
@@ -1177,6 +1175,7 @@ void SControlRigOutliner::HandleControlAdded(UControlRig* ControlRig, bool bIsAd
 			if (!ControlRig->ControlRigBound().IsBoundToObject(this))
 			{
 				ControlRig->ControlRigBound().AddRaw(this, &SControlRigOutliner::HandleOnControlRigBound);
+				BoundControlRigs.Add(ControlRig);
 			}
 			const TSharedPtr<IControlRigObjectBinding> Binding = ControlRig->GetObjectBinding();
 			if (Binding && !Binding->OnControlRigBind().IsBoundToObject(this))
@@ -1186,6 +1185,7 @@ void SControlRigOutliner::HandleControlAdded(UControlRig* ControlRig, bool bIsAd
 		}
 		else
 		{
+			BoundControlRigs.Remove(ControlRig);
 			ControlRig->ControlRigBound().RemoveAll(this);
 			const TSharedPtr<IControlRigObjectBinding> Binding = ControlRig->GetObjectBinding();
 			if (Binding)
diff --git a/Engine/Plugins/Animation/ControlRig/Source/ControlRigEditor/Private/EditMode/SControlRigOutliner.h b/Engine/Plugins/Animation/ControlRig/Source/ControlRigEditor/Private/EditMode/SControlRigOutliner.h
--- a/Engine/Plugins/Animation/ControlRig/Source/ControlRigEditor/Private/EditMode/SControlRigOutliner.h
+++ b/Engine/Plugins/Animation/ControlRig/Source/ControlRigEditor/Private/EditMode/SControlRigOutliner.h
@@ -267,7 +267,8 @@ private:
 	bool bIsChangingRigHierarchy = false;
 	TSharedPtr<SExpandableArea> PickerExpander;
 
-
+	//set of control rigs we are bound too and need to clear delegates from
+	TArray<TWeakObjectPtr<UControlRig>> BoundControlRigs;
 };
 
 
