Control Rig: Fix order and convert to layered options in contextual menu being shown when right clicking section
#rnx

diff --git a/Engine/Plugins/Animation/ControlRig/Source/ControlRigEditor/Private/Sequencer/ControlRigParameterTrackEditor.cpp b/Engine/Plugins/Animation/ControlRig/Source/ControlRigEditor/Private/Sequencer/ControlRigParameterTrackEditor.cpp
--- a/Engine/Plugins/Animation/ControlRig/Source/ControlRigEditor/Private/Sequencer/ControlRigParameterTrackEditor.cpp
+++ b/Engine/Plugins/Animation/ControlRig/Source/ControlRigEditor/Private/Sequencer/ControlRigParameterTrackEditor.cpp
@@ -93,6 +93,7 @@
 #include "Widgets/Layout/SSpacer.h"
 #include "ControlRigSequencerEditorLibrary.h"
 #include "LevelSequence.h"
+#include "MVVM/Extensions/ITrackExtension.h"
 
 #define LOCTEXT_NAMESPACE "FControlRigParameterTrackEditor"
 
@@ -3893,6 +3894,22 @@ void FControlRigParameterTrackEditor::BuildTrackContextMenu(FMenuBuilder& MenuBu
 		return;
 	}
 
+	
+	// Check if the selected element is a section of the track
+	bool bIsSection = Track->GetAllSections().Num() > 1;
+	if (bIsSection)
+	{
+		TArray<TWeakObjectPtr<UObject>> TrackSections;
+		for (UE::Sequencer::TViewModelPtr<UE::Sequencer::ITrackExtension> TrackExtension : GetSequencer()->GetViewModel()->GetSelection()->Outliner.Filter<UE::Sequencer::ITrackExtension>())
+		{
+			for (UMovieSceneSection* Section : TrackExtension->GetSections())
+			{
+				TrackSections.Add(Section);
+			}
+		}
+		bIsSection = TrackSections.Num() > 0;
+	}
+	
 	TArray<FRigControlFBXNodeAndChannels>* NodeAndChannels = Track->GetNodeAndChannelMappings(SectionToKey);
 
 	MenuBuilder.BeginSection("Control Rig IO", LOCTEXT("ControlRigIO", "Control Rig I/O"));
@@ -3913,35 +3930,38 @@ void FControlRigParameterTrackEditor::BuildTrackContextMenu(FMenuBuilder& MenuBu
 	}
 	MenuBuilder.EndSection();
 
-	MenuBuilder.BeginSection("Control Rig", LOCTEXT("ControlRig", "Control Rig"));
+	if (!bIsSection)
 	{
-		MenuBuilder.AddWidget(
-			SNew(SSpinBox<int32>)
-			.MinValue(0)
-			.Font(FAppStyle::GetFontStyle(TEXT("MenuItem.Font")))
-			.ToolTipText(LOCTEXT("OrderTooltip", "Order for this Control Rig to evaluate compared to others on the same binding"))
-			.Value_Lambda([Track]() { return Track->GetPriorityOrder(); })
-			.OnValueChanged_Lambda([Track](int32 InValue) { Track->SetPriorityOrder(InValue); })
-			,
-			LOCTEXT("Order", "Order")
-		);
-
-		if (CVarEnableAdditiveControlRigs->GetBool())
+		MenuBuilder.BeginSection("Control Rig", LOCTEXT("ControlRig", "Control Rig"));
 		{
-			MenuBuilder.AddMenuEntry(
-				   LOCTEXT("ConvertIsLayeredControlRig", "Convert To Layered"),
-				   LOCTEXT("ConvertIsLayeredControlRigToolTip", "Converts the Control Rig from an Absolute rig to a Layered rig"),
-				   FSlateIcon(),
-				   FUIAction(
-					   FExecuteAction::CreateRaw(this, &FControlRigParameterTrackEditor::ConvertIsLayered, Track),
-					   FCanExecuteAction(),
-					   FIsActionChecked::CreateRaw(this, &FControlRigParameterTrackEditor::IsLayered, Track)
-				   ),
-				   NAME_None,
-				   EUserInterfaceActionType::ToggleButton);
+			MenuBuilder.AddWidget(
+				SNew(SSpinBox<int32>)
+				.MinValue(0)
+				.Font(FAppStyle::GetFontStyle(TEXT("MenuItem.Font")))
+				.ToolTipText(LOCTEXT("OrderTooltip", "Order for this Control Rig to evaluate compared to others on the same binding"))
+				.Value_Lambda([Track]() { return Track->GetPriorityOrder(); })
+				.OnValueChanged_Lambda([Track](int32 InValue) { Track->SetPriorityOrder(InValue); })
+				,
+				LOCTEXT("Order", "Order")
+			);
+
+			if (CVarEnableAdditiveControlRigs->GetBool())
+			{
+				MenuBuilder.AddMenuEntry(
+					   LOCTEXT("ConvertIsLayeredControlRig", "Convert To Layered"),
+					   LOCTEXT("ConvertIsLayeredControlRigToolTip", "Converts the Control Rig from an Absolute rig to a Layered rig"),
+					   FSlateIcon(),
+					   FUIAction(
+						   FExecuteAction::CreateRaw(this, &FControlRigParameterTrackEditor::ConvertIsLayered, Track),
+						   FCanExecuteAction(),
+						   FIsActionChecked::CreateRaw(this, &FControlRigParameterTrackEditor::IsLayered, Track)
+					   ),
+					   NAME_None,
+					   EUserInterfaceActionType::ToggleButton);
+			}
 		}
+		MenuBuilder.EndSection();
 	}
-	MenuBuilder.EndSection();
 
 	MenuBuilder.AddMenuSeparator();
 
