Control Rig: Fix possible deadlock when restoring state on Sequencer
#rnx

diff --git a/Engine/Plugins/Animation/ControlRig/Source/ControlRig/Private/Sequencer/MovieSceneControlRigParameterTemplate.cpp b/Engine/Plugins/Animation/ControlRig/Source/ControlRig/Private/Sequencer/MovieSceneControlRigParameterTemplate.cpp
--- a/Engine/Plugins/Animation/ControlRig/Source/ControlRig/Private/Sequencer/MovieSceneControlRigParameterTemplate.cpp
+++ b/Engine/Plugins/Animation/ControlRig/Source/ControlRig/Private/Sequencer/MovieSceneControlRigParameterTemplate.cpp
@@ -767,131 +767,139 @@ struct FControlRigParameterPreAnimatedTokenProducer : IMovieScenePreAnimatedToke
 				{
 					if (ControlRig->GetObjectBinding())
 					{
-						// control rig evaluate critical section: when restoring the state, we can be poking into running instances of Control Rigs
-						// on the anim thread so using a lock here to avoid this thread and anim thread both touching the rig
-						// at the same time, which can lead to various issues like double-freeing some random array when doing SetControlValue
-						// note: the critical section accepts recursive locking so it is ok that we
-						// are calling evaluate_anythread later within the same scope
-						FScopeLock EvaluateLock(&ControlRig->GetEvaluateMutex());
-						
-						//Restore control rig first
-						const bool bSetupUndo = false;
-						if (URigHierarchy* RigHierarchy = ControlRig->GetHierarchy())
 						{
+							// Reduce the EvaluationMutex lock to only the absolutely necessary
+							// After locking, the call to UnBindFromSequencerInstance might wait for parallel evaluations to finish
+							// and some control rigs might have been queued to evaluate (but haven's started yet), so they could get stuck
+							// waiting on the evaluation lock.
 
-							FRigElementKey ControlKey;
-							ControlKey.Type = ERigElementType::Control;
-							for (FControlSpaceAndValue& SpaceNameAndValue : SpaceValues)
+							
+							// control rig evaluate critical section: when restoring the state, we can be poking into running instances of Control Rigs
+							// on the anim thread so using a lock here to avoid this thread and anim thread both touching the rig
+							// at the same time, which can lead to various issues like double-freeing some random array when doing SetControlValue
+							// note: the critical section accepts recursive locking so it is ok that we
+							// are calling evaluate_anythread later within the same scope
+							FScopeLock EvaluateLock(&ControlRig->GetEvaluateMutex());
+						
+							//Restore control rig first
+							const bool bSetupUndo = false;
+							if (URigHierarchy* RigHierarchy = ControlRig->GetHierarchy())
 							{
-								ControlKey.Name = SpaceNameAndValue.ControlName;
-								switch (SpaceNameAndValue.Value.SpaceType)
-								{
-								case EMovieSceneControlRigSpaceType::Parent:
-									RigHierarchy->SwitchToDefaultParent(ControlKey);
-									break;
-								case EMovieSceneControlRigSpaceType::World:
-									RigHierarchy->SwitchToWorldSpace(ControlKey);
-									break;
-								case EMovieSceneControlRigSpaceType::ControlRig:
+
+								FRigElementKey ControlKey;
+								ControlKey.Type = ERigElementType::Control;
+								for (FControlSpaceAndValue& SpaceNameAndValue : SpaceValues)
 								{
+									ControlKey.Name = SpaceNameAndValue.ControlName;
+									switch (SpaceNameAndValue.Value.SpaceType)
+									{
+									case EMovieSceneControlRigSpaceType::Parent:
+										RigHierarchy->SwitchToDefaultParent(ControlKey);
+										break;
+									case EMovieSceneControlRigSpaceType::World:
+										RigHierarchy->SwitchToWorldSpace(ControlKey);
+										break;
+									case EMovieSceneControlRigSpaceType::ControlRig:
+										{
 #if WITH_EDITOR
-									ControlRig->SwitchToParent(ControlKey, SpaceNameAndValue.Value.ControlRigElement, false, true);
+											ControlRig->SwitchToParent(ControlKey, SpaceNameAndValue.Value.ControlRigElement, false, true);
 #else
-									RigHierarchy->SwitchToParent(ControlKey, SpaceNameAndValue.Value.ControlRigElement);
+											RigHierarchy->SwitchToParent(ControlKey, SpaceNameAndValue.Value.ControlRigElement);
 #endif
+										}
+										break;
+									}
 								}
-								break;
-								}
-							}
 							
 
-							for (TNameAndValue<float>& Value : ScalarValues)
-							{
-								if (ControlRig->FindControl(Value.Name))
+								for (TNameAndValue<float>& Value : ScalarValues)
 								{
-									ControlRig->SetControlValue<float>(Value.Name, Value.Value, true, FRigControlModifiedContext(EControlRigSetKey::Never), bSetupUndo);
+									if (ControlRig->FindControl(Value.Name))
+									{
+										ControlRig->SetControlValue<float>(Value.Name, Value.Value, true, FRigControlModifiedContext(EControlRigSetKey::Never), bSetupUndo);
+									}
 								}
-							}
 
-							for (TNameAndValue<bool>& Value : BoolValues)
-							{
-								if (ControlRig->FindControl(Value.Name))
+								for (TNameAndValue<bool>& Value : BoolValues)
 								{
-									ControlRig->SetControlValue<bool>(Value.Name, Value.Value, true, FRigControlModifiedContext(EControlRigSetKey::Never), bSetupUndo);
+									if (ControlRig->FindControl(Value.Name))
+									{
+										ControlRig->SetControlValue<bool>(Value.Name, Value.Value, true, FRigControlModifiedContext(EControlRigSetKey::Never), bSetupUndo);
+									}
 								}
-							}
 
-							for (TNameAndValue<int32>& Value : IntegerValues)
-							{
-								if (ControlRig->FindControl(Value.Name))
-								{
-									ControlRig->SetControlValue<int32>(Value.Name, Value.Value, true, FRigControlModifiedContext(EControlRigSetKey::Never), bSetupUndo);
-								}
-							}
-							for (int32 TwiceHack = 0; TwiceHack < 2; ++TwiceHack)
-							{
-								for (TNameAndValue<FVector2D>& Value : Vector2DValues)
+								for (TNameAndValue<int32>& Value : IntegerValues)
 								{
 									if (ControlRig->FindControl(Value.Name))
 									{
-										const FVector3f Vector3(Value.Value.X, Value.Value.Y, 0.f);
-										//okay to use vector3 for 2d here
-										ControlRig->SetControlValue<FVector3f>(Value.Name, Vector3, true, FRigControlModifiedContext(EControlRigSetKey::Never), bSetupUndo);
+										ControlRig->SetControlValue<int32>(Value.Name, Value.Value, true, FRigControlModifiedContext(EControlRigSetKey::Never), bSetupUndo);
 									}
 								}
-
-								for (TNameAndValue<FVector>& Value : VectorValues)
+								for (int32 TwiceHack = 0; TwiceHack < 2; ++TwiceHack)
 								{
-									if (FRigControlElement* ControlElement = ControlRig->FindControl(Value.Name))
+									for (TNameAndValue<FVector2D>& Value : Vector2DValues)
 									{
-										if (ControlElement->Settings.ControlType == ERigControlType::Rotator)
+										if (ControlRig->FindControl(Value.Name))
 										{
-											RigHierarchy->SetControlSpecifiedEulerAngle(ControlElement, Value.Value);
+											const FVector3f Vector3(Value.Value.X, Value.Value.Y, 0.f);
+											//okay to use vector3 for 2d here
+											ControlRig->SetControlValue<FVector3f>(Value.Name, Vector3, true, FRigControlModifiedContext(EControlRigSetKey::Never), bSetupUndo);
 										}
-										ControlRig->SetControlValue<FVector3f>(Value.Name, (FVector3f)Value.Value, true, FRigControlModifiedContext(EControlRigSetKey::Never), bSetupUndo);
 									}
-								}
 
-								for (TNameAndValue<FEulerTransform>& Value : TransformValues)
-								{
-									if (FRigControlElement* ControlElement = ControlRig->FindControl(Value.Name))
+									for (TNameAndValue<FVector>& Value : VectorValues)
 									{
-										switch (ControlElement->Settings.ControlType)
-										{
-										case ERigControlType::Transform:
-										{
-											FVector EulerAngle(Value.Value.Rotation.Roll, Value.Value.Rotation.Pitch, Value.Value.Rotation.Yaw);
-											RigHierarchy->SetControlSpecifiedEulerAngle(ControlElement, EulerAngle);
-											ControlRig->SetControlValue<FRigControlValue::FTransform_Float>(Value.Name, Value.Value.ToFTransform(), true, FRigControlModifiedContext(EControlRigSetKey::Never), bSetupUndo);
-											break;
-										}
-										case ERigControlType::TransformNoScale:
-										{
-											FTransformNoScale NoScale = Value.Value.ToFTransform();
-											FVector EulerAngle(Value.Value.Rotation.Roll, Value.Value.Rotation.Pitch, Value.Value.Rotation.Yaw);
-											RigHierarchy->SetControlSpecifiedEulerAngle(ControlElement, EulerAngle);
-											ControlRig->SetControlValue<FRigControlValue::FTransformNoScale_Float>(Value.Name, NoScale, true, FRigControlModifiedContext(EControlRigSetKey::Never), bSetupUndo);
-											break;
-										}
-										case ERigControlType::EulerTransform:
+										if (FRigControlElement* ControlElement = ControlRig->FindControl(Value.Name))
 										{
-											FEulerTransform EulerTransform = Value.Value;
-											FVector EulerAngle(Value.Value.Rotation.Roll, Value.Value.Rotation.Pitch, Value.Value.Rotation.Yaw);
-											RigHierarchy->SetControlSpecifiedEulerAngle(ControlElement, EulerAngle);
-											ControlRig->SetControlValue<FRigControlValue::FEulerTransform_Float>(Value.Name, EulerTransform, true, FRigControlModifiedContext(EControlRigSetKey::Never), bSetupUndo);
-											break;
+											if (ControlElement->Settings.ControlType == ERigControlType::Rotator)
+											{
+												RigHierarchy->SetControlSpecifiedEulerAngle(ControlElement, Value.Value);
+											}
+											ControlRig->SetControlValue<FVector3f>(Value.Name, (FVector3f)Value.Value, true, FRigControlModifiedContext(EControlRigSetKey::Never), bSetupUndo);
 										}
-										default:
+									}
+
+									for (TNameAndValue<FEulerTransform>& Value : TransformValues)
+									{
+										if (FRigControlElement* ControlElement = ControlRig->FindControl(Value.Name))
 										{
-											break;
-										}
+											switch (ControlElement->Settings.ControlType)
+											{
+											case ERigControlType::Transform:
+												{
+													FVector EulerAngle(Value.Value.Rotation.Roll, Value.Value.Rotation.Pitch, Value.Value.Rotation.Yaw);
+													RigHierarchy->SetControlSpecifiedEulerAngle(ControlElement, EulerAngle);
+													ControlRig->SetControlValue<FRigControlValue::FTransform_Float>(Value.Name, Value.Value.ToFTransform(), true, FRigControlModifiedContext(EControlRigSetKey::Never), bSetupUndo);
+													break;
+												}
+											case ERigControlType::TransformNoScale:
+												{
+													FTransformNoScale NoScale = Value.Value.ToFTransform();
+													FVector EulerAngle(Value.Value.Rotation.Roll, Value.Value.Rotation.Pitch, Value.Value.Rotation.Yaw);
+													RigHierarchy->SetControlSpecifiedEulerAngle(ControlElement, EulerAngle);
+													ControlRig->SetControlValue<FRigControlValue::FTransformNoScale_Float>(Value.Name, NoScale, true, FRigControlModifiedContext(EControlRigSetKey::Never), bSetupUndo);
+													break;
+												}
+											case ERigControlType::EulerTransform:
+												{
+													FEulerTransform EulerTransform = Value.Value;
+													FVector EulerAngle(Value.Value.Rotation.Roll, Value.Value.Rotation.Pitch, Value.Value.Rotation.Yaw);
+													RigHierarchy->SetControlSpecifiedEulerAngle(ControlElement, EulerAngle);
+													ControlRig->SetControlValue<FRigControlValue::FEulerTransform_Float>(Value.Name, EulerTransform, true, FRigControlModifiedContext(EControlRigSetKey::Never), bSetupUndo);
+													break;
+												}
+											default:
+												{
+													break;
+												}
+											}
 										}
 									}
 								}
 							}
+							//make sure to evaluate the control rig
+							ControlRig->Evaluate_AnyThread();
 						}
-						//make sure to evaluate the control rig
-						ControlRig->Evaluate_AnyThread();
 
 						//unbind instances and reset animbp
 						FControlRigBindingHelper::UnBindFromSequencerInstance(ControlRig);
