Control Rig: FIx crash when destroy shape actors with invalid attached parents
#rnx

diff --git a/Engine/Plugins/Animation/ControlRig/Source/ControlRigEditor/Private/EditMode/ControlRigEditMode.cpp b/Engine/Plugins/Animation/ControlRig/Source/ControlRigEditor/Private/EditMode/ControlRigEditMode.cpp
--- a/Engine/Plugins/Animation/ControlRig/Source/ControlRigEditor/Private/EditMode/ControlRigEditMode.cpp
+++ b/Engine/Plugins/Animation/ControlRig/Source/ControlRigEditor/Private/EditMode/ControlRigEditMode.cpp
@@ -5108,6 +5108,7 @@ void FControlRigEditMode::DestroyShapesActors(UControlRig* ControlRig)
 				UWorld* World = ShapeActor->GetWorld();
 				if (World)
 				{
+					ShapeActor->UnregisterAllComponents();
 					if (ShapeActor->GetAttachParentActor())
 					{
 						ShapeActor->DetachFromActor(FDetachmentTransformRules::KeepRelativeTransform);
