Control Rig: Fix issues when a rig is running from multiple threads
#rnx

diff --git a/Engine/Plugins/Animation/ControlRig/Source/ControlRig/Private/AnimNode_ControlRigBase.cpp b/Engine/Plugins/Animation/ControlRig/Source/ControlRig/Private/AnimNode_ControlRigBase.cpp
--- a/Engine/Plugins/Animation/ControlRig/Source/ControlRig/Private/AnimNode_ControlRigBase.cpp
+++ b/Engine/Plugins/Animation/ControlRig/Source/ControlRig/Private/AnimNode_ControlRigBase.cpp
@@ -439,6 +439,10 @@ void FAnimNode_ControlRigBase::ExecuteControlRig(FPoseContext& InOutput)
 
 	if (UControlRig* ControlRig = GetControlRig())
 	{
+		// Before we start modifying the RigHierarchy, we need to lock the rig to avoid
+		// corrupted state
+		FScopeLock LockRig(&ControlRig->GetEvaluateMutex());
+		
 		// temporarily give control rig access to the stack allocated attribute container
 		// control rig may have rig units that can add/get attributes to/from this container
 		UControlRig::FAnimAttributeContainerPtrScope AttributeScope(ControlRig, InOutput.CustomAttributes);
@@ -831,6 +835,11 @@ void FAnimNode_ControlRigBase::QueueControlRigDrawInstructions(UControlRig* Cont
 
 void FAnimNode_ControlRigBase::UpdateGetAssetUserDataDelegate(UControlRig* InControlRig) const
 {
+	if (!IsInGameThread())
+	{
+		return;
+	}
+	
 	if(GetAssetUserData().IsEmpty() || !WeakAnimInstanceObject.IsValid())
 	{
 		InControlRig->GetExternalAssetUserDataDelegate.Unbind();
