Control Rig: Fix unit test breakage

#rnx

diff --git a/Engine/Plugins/Animation/ControlRig/Source/ControlRig/Private/ModularRigRuleManager.cpp b/Engine/Plugins/Animation/ControlRig/Source/ControlRig/Private/ModularRigRuleManager.cpp
--- a/Engine/Plugins/Animation/ControlRig/Source/ControlRig/Private/ModularRigRuleManager.cpp
+++ b/Engine/Plugins/Animation/ControlRig/Source/ControlRig/Private/ModularRigRuleManager.cpp
@@ -172,8 +172,7 @@ void UModularRigRuleManager::FilterInvalidNameSpaces(FWorkData& InOutWorkData)
 	}
 	
 	const FString NameSpaceString = NameSpace.ToString(); 
-	const FString NameSpacePrefix = NameSpaceString + UModularRig::NamespaceSeparator; 
-	InOutWorkData.Filter([NameSpaceString, NameSpacePrefix, InOutWorkData](FRigElementResolveResult& Result)
+	InOutWorkData.Filter([NameSpaceString, InOutWorkData](FRigElementResolveResult& Result)
 	{
 		const FName MatchNameSpace = InOutWorkData.Hierarchy->GetNameMetadata(Result.GetKey(), URigHierarchy::NameSpaceMetadataName, NAME_None);
 		if(!MatchNameSpace.IsNone())
@@ -184,7 +183,7 @@ void UModularRigRuleManager::FilterInvalidNameSpaces(FWorkData& InOutWorkData)
 				static const FText CannotConnectWithinNameSpaceMessage = LOCTEXT("CannotConnectWithinNameSpace", "Cannot connect within the same namespace.");
 				Result.SetInvalidTarget(CannotConnectWithinNameSpaceMessage);
 			}
-			else if(MatchNameSpaceString.StartsWith(NameSpacePrefix, ESearchCase::CaseSensitive))
+			else if(MatchNameSpaceString.StartsWith(NameSpaceString, ESearchCase::CaseSensitive))
 			{
 				static const FText CannotConnectBelowNameSpaceMessage = LOCTEXT("CannotConnectBelowNameSpace", "Cannot connect to element below the connector's namespace.");
 				Result.SetInvalidTarget(CannotConnectBelowNameSpaceMessage);
