Control RIg: Details panel connector chooser as a tree
#rnx

diff --git a/Engine/Plugins/Animation/ControlRig/Source/ControlRig/Private/ModularRigController.cpp b/Engine/Plugins/Animation/ControlRig/Source/ControlRig/Private/ModularRigController.cpp
--- a/Engine/Plugins/Animation/ControlRig/Source/ControlRig/Private/ModularRigController.cpp
+++ b/Engine/Plugins/Animation/ControlRig/Source/ControlRig/Private/ModularRigController.cpp
@@ -173,6 +173,13 @@ bool UModularRigController::ConnectConnectorToElement(const FRigElementKey& InCo
 		return false;
 	}
 
+	const FRigElementKey ConnectorKey(*ConnectorName, ERigElementType::Connector);
+	FRigElementKey* CurrentTarget = Module->Connections.Find(ConnectorKey);
+	if (CurrentTarget && InTargetKey == *CurrentTarget)
+	{
+		return true; // Nothing to do
+	}
+
 	/*
 	FText ErrorMessage;
 	if (!CanConnectConnectorToElement(*ModuleConnector, InTargetKey, ErrorMessage))
@@ -201,9 +208,13 @@ bool UModularRigController::ConnectConnectorToElement(const FRigElementKey& InCo
 	}
 #endif 
 
-	const FRigElementKey ConnectorKey(*ConnectorName, ERigElementType::Connector);
-	FRigElementKey& TargetKey = Module->Connections.FindOrAdd(ConnectorKey);
-	TargetKey = InTargetKey;
+	// First disconnect before connecting to anything else. This might disconnect other secondary/optional connectors.
+	if (CurrentTarget && CurrentTarget->IsValid())
+	{
+		DisconnectConnector(InConnectorKey, bSetupUndo);
+	}
+
+	Module->Connections.FindOrAdd(ConnectorKey) = InTargetKey;
 
 	Notify(EModularRigNotification::ConnectionChanged, Module);
 
@@ -269,6 +280,15 @@ bool UModularRigController::DisconnectConnector(const FRigElementKey& InConnecto
 	const FRigElementKey ConnectorKey(*ConnectorName, ERigElementType::Connector);
 	Module->Connections.Remove(ConnectorKey);
 
+	if (Connector->Settings.Type == EConnectorType::Primary)
+	{
+		Module->Connections.Reset();
+	}
+	else
+	{
+		// todo: Make sure all the rest of the connections are still valid
+	}
+
 	Notify(EModularRigNotification::ConnectionChanged, Module);
 
 #if WITH_EDITOR
diff --git a/Engine/Plugins/Animation/ControlRig/Source/ControlRigEditor/Private/ControlRigModuleDetails.cpp b/Engine/Plugins/Animation/ControlRig/Source/ControlRigEditor/Private/ControlRigModuleDetails.cpp
--- a/Engine/Plugins/Animation/ControlRig/Source/ControlRigEditor/Private/ControlRigModuleDetails.cpp
+++ b/Engine/Plugins/Animation/ControlRig/Source/ControlRigEditor/Private/ControlRigModuleDetails.cpp
@@ -11,6 +11,7 @@
 #include "Widgets/Input/SButton.h"
 #include "ModularRigController.h"
 #include "ControlRigBlueprint.h"
+#include "ControlRigEditorStyle.h"
 #include "ControlRigElementDetails.h"
 #include "Graph/ControlRigGraph.h"
 #include "PropertyCustomizationHelpers.h"
@@ -23,7 +24,9 @@
 #include "StructViewerModule.h"
 #include "Features/IModularFeatures.h"
 #include "IPropertyAccessEditor.h"
+#include "ModularRigRuleManager.h"
 #include "ScopedTransaction.h"
+#include "Editor/SRigHierarchyTreeView.h"
 
 #define LOCTEXT_NAMESPACE "ControlRigModuleDetails"
 
@@ -240,20 +243,66 @@ void FRigModuleInstanceDetails::CustomizeDetails(IDetailLayoutBuilder& DetailBui
 	IDetailCategoryBuilder& ConnectionsCategory = DetailBuilder.EditCategory(TEXT("Connections"), LOCTEXT("Connections", "Connections"));
 	{
 		TArray<FRigModuleConnector> Connectors = GetConnectors();
-		FRigElementKeyRedirector Redirector = GetConnections();
 		for(FRigModuleConnector& Connector : Connectors)
 		{
 			const FText Label = FText::FromString(Connector.Name);
-			FRigElementKey ConnectorKey(*Connector.Name, ERigElementType::Connector);
-			const FRigElementKey* TargetKey = Redirector.FindExternalKey(ConnectorKey);
-			if (TargetKey)
+			TSharedPtr<SVerticalBox> ButtonBox;
+			FString ConnectorPath = FString::Printf(TEXT("%s:%s"), *PerModuleInfos[0].Path, *Connector.Name);
+			FRigElementKey ConnectorKey(*ConnectorPath, ERigElementType::Connector);
+
+			FRigTreeDelegates TreeDelegates;
+			TreeDelegates.OnGetHierarchy = FOnGetRigTreeHierarchy::CreateLambda([this]()
 			{
-				Connections.Add(ConnectorKey, *TargetKey);
-			}
-			else
+				return PerModuleInfos[0].GetModularRig()->GetHierarchy();
+			});
+			TreeDelegates.OnRigTreeIsItemVisible = FOnRigTreeIsItemVisible::CreateLambda([this, ConnectorKey](const FRigElementKey& InTarget)
 			{
-				Connections.Add(ConnectorKey, FRigElementKey());
-			}
+				FPerModuleInfo Info = PerModuleInfos[0];
+				if (const FRigModuleInstance* Module = Info.GetModule())
+				{
+					if (UModularRig* ModularRig = Info.GetModularRig())
+					{
+						if (URigHierarchy* Hierarchy = ModularRig->GetHierarchy())
+						{
+							if (FRigConnectorElement* ConnectorElement = Cast<FRigConnectorElement>(Hierarchy->Find(ConnectorKey)))
+							{
+								const UModularRigRuleManager* RuleManager = ModularRig->GetHierarchy()->GetRuleManager();
+								const FModularRigResolveResult Result = RuleManager->FindMatches(ConnectorElement, Module, ModularRig->GetElementKeyRedirector());
+								return Result.ContainsMatch(InTarget);
+							}
+						}
+					}
+				}
+				return false;
+			});
+			TreeDelegates.OnGetSelection.BindLambda([this, ConnectorKey]() -> TArray<FRigElementKey>
+			{
+				FPerModuleInfo Info = PerModuleInfos[0];
+				FRigElementKeyRedirector Redirector = Info.GetModularRig()->GetElementKeyRedirector();
+				if (const FRigElementKey* Key = PerModuleInfos[0].GetModularRig()->GetElementKeyRedirector().FindExternalKey(ConnectorKey))
+				{
+					return {*Key};
+				}
+				return {};
+			});
+			TreeDelegates.OnSelectionChanged.BindLambda([this, ConnectorKey](TSharedPtr<FRigTreeElement> Selection, ESelectInfo::Type SelectInfo)
+			{
+				FPerModuleInfo Info = PerModuleInfos[0];
+				TSharedPtr<FRigTreeElement> NewSelection = Selection;
+				if (UModularRigController* Controller = Info.GetBlueprint()->GetModularRigController())
+				{
+					if (NewSelection.IsValid())
+					{
+						Controller->ConnectConnectorToElement(ConnectorKey, NewSelection->Key);
+					}
+					else
+					{
+						Controller->DisconnectConnector(ConnectorKey);
+					}
+				}
+			});
+			
+
 			ConnectionsCategory.AddCustomRow(Label)
 				.NameContent()
 				[
@@ -264,20 +313,161 @@ void FRigModuleInstanceDetails::CustomizeDetails(IDetailLayoutBuilder& DetailBui
 				]
 				.ValueContent()
 				[
-					SAssignNew(RigElementKeyWidget, SRigElementKeyWidget)
-					.Blueprint(PerModuleInfos[0].GetBlueprint())
-					.IsEnabled_Lambda([this](){ return true; })
-					.ActiveBackgroundColor(FSlateColor(FLinearColor(1.f, 1.f, 1.f, FRigElementKeyDetailsDefs::ActivePinBackgroundAlpha)))
-					.ActiveForegroundColor(FSlateColor(FLinearColor(1.f, 1.f, 1.f, FRigElementKeyDetailsDefs::ActivePinForegroundAlpha)))
-					.InactiveBackgroundColor(FSlateColor(FLinearColor(1.f, 1.f, 1.f, FRigElementKeyDetailsDefs::InactivePinBackgroundAlpha)))
-					.InactiveForegroundColor(FSlateColor(FLinearColor(1.f, 1.f, 1.f, FRigElementKeyDetailsDefs::InactivePinForegroundAlpha)))
-					.OnElementNameChanged(this, &FRigModuleInstanceDetails::OnElementNameChanged, ConnectorKey)
-					.OnGetSelectedClicked(this, &FRigModuleInstanceDetails::OnGetSelectedClicked, ConnectorKey)
-					.OnSelectInHierarchyClicked(this, &FRigModuleInstanceDetails::OnSelectInHierarchyClicked, ConnectorKey)
-					.OnGetElementNameAsText_Raw(this, &FRigModuleInstanceDetails::GetElementNameAsText, ConnectorKey)
-					.OnGetElementType(this, &FRigModuleInstanceDetails::GetElementType, ConnectorKey)
-					.OnElementTypeChanged(this, &FRigModuleInstanceDetails::OnElementTypeChanged, ConnectorKey)
+					SNew(SHorizontalBox)
+
+					// Combo button
+					+SHorizontalBox::Slot()
+					.AutoWidth()
+					.Padding(0.f, 0.f, 0.f, 0.f)
+					[
+						SNew( SComboButton )
+						.ContentPadding(3)
+						.MenuPlacement(MenuPlacement_BelowAnchor)
+						.ButtonContent()
+						[
+							// Wrap in configurable box to restrain height/width of menu
+							SNew(SBox)
+							.MinDesiredWidth(150.0f)
+							[
+								SAssignNew(ButtonBox, SVerticalBox)
+							]
+						]
+						.MenuContent()
+						[
+							SNew(SBorder)
+							.Visibility(EVisibility::Visible)
+							.BorderImage(FAppStyle::GetBrush("Menu.Background"))
+							[
+								SAssignNew(ConnectionListBox.FindOrAdd(ConnectorKey), SSearchableRigHierarchyTreeView)
+									.RigTreeDelegates(TreeDelegates)
+							]
+						]
+					]
+
+					// Reset button
+					+SHorizontalBox::Slot()
+					.AutoWidth()
+					.Padding(0.f, 0.f, 0.f, 0.f)
+					[
+						SAssignNew(ResetConnectorButton.FindOrAdd(ConnectorKey), SButton)
+						.ButtonStyle( FAppStyle::Get(), "NoBorder" )
+						.ButtonColorAndOpacity_Lambda([this, ConnectorKey]()
+						{
+							const TSharedPtr<SButton>& Button = ResetConnectorButton.FindRef(ConnectorKey);
+							return Button.IsValid() && Button->IsHovered()
+								? FSlateColor(FLinearColor(1.f, 1.f, 1.f, 0.8))
+								: FSlateColor(FLinearColor(1.f, 1.f, 1.f, 0.4));
+						})
+						.OnClicked_Lambda([this, ConnectorKey]()
+						{
+							PerModuleInfos[0].GetBlueprint()->GetModularRigController()->DisconnectConnector(ConnectorKey);
+							return FReply::Handled();
+						})
+						.ContentPadding(1.f)
+						.ToolTipText(NSLOCTEXT("ControlRigModuleDetails", "Reset_Connector", "Reset Connector"))
+						[
+							SNew(SImage)
+							.ColorAndOpacity_Lambda( [this, ConnectorKey]()
+							{
+								const TSharedPtr<SButton>& Button = ResetConnectorButton.FindRef(ConnectorKey);
+								return Button.IsValid() && Button->IsHovered()
+								? FSlateColor(FLinearColor(1.f, 1.f, 1.f, 0.8))
+								: FSlateColor(FLinearColor(1.f, 1.f, 1.f, 0.4));
+							})
+							.Image(FSlateIcon(FAppStyle::Get().GetStyleSetName(), "PropertyWindow.DiffersFromDefault").GetIcon())
+						]
+					]
+
+					// Use button
+					+SHorizontalBox::Slot()
+					.AutoWidth()
+					.Padding(0.f, 0.f, 0.f, 0.f)
+					[
+						SAssignNew(UseSelectedButton.FindOrAdd(ConnectorKey), SButton)
+						.ButtonStyle( FAppStyle::Get(), "NoBorder" )
+						.ButtonColorAndOpacity_Lambda([this, ConnectorKey]()
+						{
+							const TSharedPtr<SButton>& Button = UseSelectedButton.FindRef(ConnectorKey);
+							return Button.IsValid() && Button->IsHovered()
+								? FSlateColor(FLinearColor(1.f, 1.f, 1.f, 0.8))
+								: FSlateColor(FLinearColor(1.f, 1.f, 1.f, 0.4));
+						})
+						.OnClicked_Lambda([this, ConnectorKey]()
+						{
+							if (UModularRig* ModularRig = PerModuleInfos[0].GetModularRig())
+							{
+								const TArray<FRigElementKey>& Selected = ModularRig->GetHierarchy()->GetSelectedKeys();
+								if (Selected.Num() > 0)
+								{
+									PerModuleInfos[0].GetBlueprint()->GetModularRigController()->ConnectConnectorToElement(ConnectorKey, Selected[0]);
+								}
+							}
+							return FReply::Handled();
+						})
+						.ContentPadding(1.f)
+						.ToolTipText(NSLOCTEXT("ControlRigModuleDetails", "Use_Selected", "Use Selected"))
+						[
+							SNew(SImage)
+							.ColorAndOpacity_Lambda( [this, ConnectorKey]()
+							{
+								const TSharedPtr<SButton>& Button = UseSelectedButton.FindRef(ConnectorKey);
+								return Button.IsValid() && Button->IsHovered()
+								? FSlateColor(FLinearColor(1.f, 1.f, 1.f, 0.8))
+								: FSlateColor(FLinearColor(1.f, 1.f, 1.f, 0.4));
+							})
+							.Image(FAppStyle::GetBrush("Icons.CircleArrowLeft"))
+						]
+					]
+
+					// Select in hierarchy button
+					+SHorizontalBox::Slot()
+					.AutoWidth()
+					.Padding(0.f, 0.f, 0.f, 0.f)
+					[
+						SAssignNew(SelectElementButton.FindOrAdd(ConnectorKey), SButton)
+						.ButtonStyle( FAppStyle::Get(), "NoBorder" )
+						.ButtonColorAndOpacity_Lambda([this, ConnectorKey]()
+						{
+							const TSharedPtr<SButton>& Button = SelectElementButton.FindRef(ConnectorKey);
+							return Button.IsValid() && Button->IsHovered()
+								? FSlateColor(FLinearColor(1.f, 1.f, 1.f, 0.8))
+								: FSlateColor(FLinearColor(1.f, 1.f, 1.f, 0.4));
+						})
+						.OnClicked_Lambda([this, ConnectorKey]()
+						{
+							if (UModularRig* ModularRig = PerModuleInfos[0].GetModularRig())
+							{
+								if (const FRigElementKey* TargetKey = ModularRig->GetElementKeyRedirector().FindExternalKey(ConnectorKey))
+								{
+									ModularRig->GetHierarchy()->GetController()->SelectElement(*TargetKey);
+								}
+							}
+							return FReply::Handled();
+						})
+						.ContentPadding(1.f)
+						.ToolTipText(NSLOCTEXT("ControlRigModuleDetails", "Select_Element", "Select Element"))
+						[
+							SNew(SImage)
+							.ColorAndOpacity_Lambda( [this, ConnectorKey]()
+							{
+								const TSharedPtr<SButton>& Button = SelectElementButton.FindRef(ConnectorKey);
+								return Button.IsValid() && Button->IsHovered()
+								? FSlateColor(FLinearColor(1.f, 1.f, 1.f, 0.8))
+								: FSlateColor(FLinearColor(1.f, 1.f, 1.f, 0.4));
+							})
+							.Image(FAppStyle::GetBrush("Icons.Search"))
+						]
+					]
 				];
+
+			FRigElementKey CurrentTargetKey;
+			if (const FRigElementKey* Key = PerModuleInfos[0].GetModularRig()->GetElementKeyRedirector().FindExternalKey(ConnectorKey))
+			{
+				CurrentTargetKey = *Key;
+			}
+			TPair<const FSlateBrush*, FSlateColor> IconAndColor = SRigHierarchyItem::GetBrushForElementType(PerModuleInfos[0].GetModularRig()->GetHierarchy(), CurrentTargetKey);
+			PopulateConnectorCurrentTarget(ButtonBox, ConnectorKey, CurrentTargetKey, IconAndColor.Key, IconAndColor.Value, FText::FromName(CurrentTargetKey.Name));
+			PopulateConnectorTargetList(ConnectorKey);
 		}
 	}
 
@@ -631,6 +821,98 @@ FRigElementKeyRedirector FRigModuleInstanceDetails::GetConnections() const
 	return FRigElementKeyRedirector();
 }
 
+void FRigModuleInstanceDetails::PopulateConnectorTargetList(const FRigElementKey& InConnectorKey)
+{
+	if (PerModuleInfos.IsEmpty() || PerModuleInfos.Num() > 1)
+	{
+		return;
+	}
+	FPerModuleInfo& Info = PerModuleInfos[0];
+	if (const FRigModuleInstance* Module = Info.GetModule())
+	{
+		if (UModularRig* ModularRig = Info.GetModularRig())
+		{
+			if (URigHierarchy* Hierarchy = ModularRig->GetHierarchy())
+			{
+				if (FRigConnectorElement* ConnectorElement = Cast<FRigConnectorElement>(Hierarchy->Find(InConnectorKey)))
+				{
+					ConnectionListBox.FindRef(InConnectorKey)->GetTreeView()->RefreshTreeView(true);
+				}
+			}
+		}
+	}
+}
+
+void FRigModuleInstanceDetails::PopulateConnectorCurrentTarget(TSharedPtr<SVerticalBox> InListBox, const FRigElementKey& InConnectorKey, const FRigElementKey& InTargetKey,
+	const FSlateBrush* InBrush, const FSlateColor& InColor, const FText& InTitle)
+{
+	static const FSlateBrush* RoundedBoxBrush = FControlRigEditorStyle::Get().GetBrush(TEXT("ControlRig.SpacePicker.RoundedRect"));
+	FPerModuleInfo Info = PerModuleInfos[0];
+	UModularRigController* Controller = Info.GetBlueprint()->GetModularRigController();
+
+	TSharedPtr<SHorizontalBox> RowBox, ButtonBox;
+	InListBox->AddSlot()
+	.AutoHeight()
+	.VAlign(VAlign_Top)
+	.HAlign(HAlign_Fill)
+	.Padding(4.0, 0.0, 4.0, 0.0)
+	[
+		SNew( SButton )
+		.ButtonStyle(FAppStyle::Get(), "SimpleButton")
+		.ContentPadding(FMargin(0.0))
+		.OnClicked_Lambda([Controller, InConnectorKey, InTargetKey]()
+		{
+			Controller->ConnectConnectorToElement(InConnectorKey, InTargetKey);
+			return FReply::Handled();
+		})
+		[
+			SAssignNew(RowBox, SHorizontalBox)
+			+ SHorizontalBox::Slot()
+			.AutoWidth()
+			.VAlign(VAlign_Center)
+			.HAlign(HAlign_Fill)
+			.Padding(0)
+			[
+				SNew(SBorder)
+				.Padding(FMargin(5.0, 2.0, 5.0, 2.0))
+				.BorderImage(RoundedBoxBrush)
+				.BorderBackgroundColor(FSlateColor(FLinearColor::Transparent))
+				.Content()
+				[
+					SAssignNew(ButtonBox, SHorizontalBox)
+					+ SHorizontalBox::Slot()
+					.AutoWidth()
+					.VAlign(VAlign_Center)
+					.HAlign(HAlign_Left)
+					.Padding(FMargin(0.f, 0.f, 3.f, 0.f))
+					[
+						SNew(SImage)
+						.Image(InBrush)
+						.ColorAndOpacity(InColor)
+					]
+
+					+ SHorizontalBox::Slot()
+					.AutoWidth()
+					.VAlign(VAlign_Center)
+					.HAlign(HAlign_Left)
+					.Padding(0)
+					[
+						SNew( STextBlock )
+						.Text( InTitle )
+						.Font( IDetailLayoutBuilder::GetDetailFont() )
+					]
+
+					+ SHorizontalBox::Slot()
+					.FillWidth(1.f)
+					[
+						SNew(SSpacer)
+					]
+				]
+			]
+		]
+	];
+}
+
 void FRigModuleInstanceDetails::OnConfigValueChanged(const FName InVariableName)
 {
 	for(const FPerModuleInfo& Info : PerModuleInfos)
@@ -686,119 +968,6 @@ void FRigModuleInstanceDetails::RegisterSectionMappings(FPropertyEditorModule& P
 	MetadataSection->AddCategory("Metadata");
 }
 
-void FRigModuleInstanceDetails::OnElementNameChanged(TSharedPtr<FString> InItem, ESelectInfo::Type InSelectionInfo, FRigElementKey Connector)
-{
-	for (FPerModuleInfo& Info : PerModuleInfos)
-	{
-		if (UControlRigBlueprint* Blueprint = RigModuleDetails_GetBlueprintFromRig(Info.GetModularRig()))
-		{
-			UModularRigController* Controller = Blueprint->GetModularRigController();
-			FRigElementKey Key;
-			FRigElementKey* TargetKey = Connections.Find(Connector);
-			if (TargetKey)
-			{
-				if (InItem)
-				{
-					TargetKey->Name = **InItem;
-				}
-
-				if(const FRigModuleInstance* Module = Info.GetModule())
-				{
-					FRigElementKey NamespacedConnector(*URigHierarchy::JoinNameSpace(Module->GetPath(), Connector.Name.ToString()), ERigElementType::Connector);
-					if(!Blueprint->GetDebuggedControlRig()->GetHierarchy()->Contains(*TargetKey))
-					{
-						Controller->DisconnectConnector(NamespacedConnector);
-					}
-					else
-					{
-						Controller->ConnectConnectorToElement(NamespacedConnector, *TargetKey);
-					}
-				}
-			}
-		}
-	}
-}
-
-void FRigModuleInstanceDetails::OnElementTypeChanged(ERigElementType InElementType, FRigElementKey Connector)
-{
-	for (FPerModuleInfo& Info : PerModuleInfos)
-	{
-		if (UControlRigBlueprint* Blueprint = RigModuleDetails_GetBlueprintFromRig(Info.GetModularRig()))
-		{
-			UModularRigController* Controller = Blueprint->GetModularRigController();
-			FRigElementKey Key;
-			FRigElementKey* TargetKey = Connections.Find(Connector);
-			if (TargetKey)
-			{
-				TargetKey->Type = InElementType;
-
-				if(const FRigModuleInstance* Module = Info.GetModule())
-				{
-					FRigElementKey NamespacedConnector(*URigHierarchy::JoinNameSpace(Module->GetPath(), Connector.Name.ToString()), ERigElementType::Connector);
-					if(!Blueprint->GetDebuggedControlRig()->GetHierarchy()->Contains(*TargetKey))
-					{
-						Controller->DisconnectConnector(NamespacedConnector);
-					}
-					else
-					{
-						Controller->ConnectConnectorToElement(NamespacedConnector, *TargetKey);
-					}
-				}
-			}
-		}
-	}
-}
-
-FText FRigModuleInstanceDetails::GetElementNameAsText(FRigElementKey Connector) const
-{
-	if (const FRigElementKey* TargetKey = Connections.Find(Connector))
-	{
-		return FText::FromName(TargetKey->Name);
-	}
-	return FText();
-}
-
-FReply FRigModuleInstanceDetails::OnGetSelectedClicked(FRigElementKey Connector)
-{
-	if (UControlRigBlueprint* Blueprint = RigModuleDetails_GetBlueprintFromRig(PerModuleInfos[0].GetModularRig()))
-	{
-		const TArray<FRigElementKey>& Selected = Blueprint->Hierarchy->GetSelectedKeys();
-		if (Selected.Num() > 0)
-		{
-			if(const FRigModuleInstance* Module = PerModuleInfos[0].GetModule())
-			{
-				FRigElementKey NamespacedConnector(*URigHierarchy::JoinNameSpace(Module->GetPath(), Connector.Name.ToString()), ERigElementType::Connector);
-				Blueprint->GetModularRigController()->ConnectConnectorToElement(NamespacedConnector, Selected[0]);
-			}
-		}
-	}
-	return FReply::Handled();
-}
-
-FReply FRigModuleInstanceDetails::OnSelectInHierarchyClicked(FRigElementKey Connector)
-{
-	if (UControlRigBlueprint* Blueprint = RigModuleDetails_GetBlueprintFromRig(PerModuleInfos[0].GetModularRig()))
-	{
-		if (FRigElementKey* Target = Connections.Find(Connector))
-		{
-			if (Target->IsValid())
-			{
-				Blueprint->GetHierarchyController()->SetSelection({*Target});
-			}
-		}
-	}
-	return FReply::Handled();
-}
-
-ERigElementType FRigModuleInstanceDetails::GetElementType(FRigElementKey Connector) const
-{
-	if (const FRigElementKey* TargetKey = Connections.Find(Connector))
-	{
-		return TargetKey->Type;
-	}
-	return ERigElementType::None;
-}
-
 FText FRigModuleInstanceDetails::GetBindingText(const FProperty* InProperty) const
 {
 	const FName VariableName = InProperty->GetFName();
diff --git a/Engine/Plugins/Animation/ControlRig/Source/ControlRigEditor/Private/ControlRigModuleDetails.h b/Engine/Plugins/Animation/ControlRig/Source/ControlRigEditor/Private/ControlRigModuleDetails.h
--- a/Engine/Plugins/Animation/ControlRig/Source/ControlRigEditor/Private/ControlRigModuleDetails.h
+++ b/Engine/Plugins/Animation/ControlRig/Source/ControlRigEditor/Private/ControlRigModuleDetails.h
@@ -23,6 +23,7 @@
 #include "Styling/AppStyle.h"
 #include "Algo/Transform.h"
 #include "IPropertyUtilities.h"
+#include "Editor/SRigHierarchyTreeView.h"
 
 class IPropertyHandle;
 
@@ -49,6 +50,14 @@ public:
 	FText GetRigClassPath() const;
 	TArray<FRigModuleConnector> GetConnectors() const;
 	FRigElementKeyRedirector GetConnections() const;
+	void PopulateConnectorTargetList(const FRigElementKey& InConnectorKey);
+	void PopulateConnectorCurrentTarget(
+		TSharedPtr<SVerticalBox> InListBox,
+		const FRigElementKey& InConnectorKey,
+		const FRigElementKey& InTargetKey,
+		const FSlateBrush* InBrush,
+		const FSlateColor& InColor,
+		const FText& InTitle);
 
 	void OnConfigValueChanged(const FName InVariableName);
 	
@@ -125,13 +134,6 @@ public:
 
 	virtual void RegisterSectionMappings(FPropertyEditorModule& PropertyEditorModule, UClass* InClass);
 
-	void OnElementNameChanged(TSharedPtr<FString> InItem, ESelectInfo::Type InSelectionInfo, FRigElementKey Connector);
-	void OnElementTypeChanged(ERigElementType InElementType, FRigElementKey Connector);
-	FText GetElementNameAsText(FRigElementKey Connector) const;
-	FReply OnGetSelectedClicked(FRigElementKey Connector);
-	FReply OnSelectInHierarchyClicked(FRigElementKey Connector);
-	ERigElementType GetElementType(FRigElementKey Connector) const;
-
 protected:
 
 	FText GetBindingText(const FProperty* InProperty) const;
@@ -143,6 +145,10 @@ protected:
 	void HandleChangeBinding(const FProperty* InProperty, const FString& InNewVariablePath) const;
 
 	TArray<FPerModuleInfo> PerModuleInfos;
-	TSharedPtr<SRigElementKeyWidget> RigElementKeyWidget;
-	TMap<FRigElementKey, FRigElementKey> Connections;
+	TMap<FRigElementKey, TSharedPtr<SSearchableRigHierarchyTreeView>> ConnectionListBox;
+
+	/** Helper buttons. */
+	TMap<FRigElementKey, TSharedPtr<SButton>> UseSelectedButton;
+	TMap<FRigElementKey, TSharedPtr<SButton>> SelectElementButton;
+	TMap<FRigElementKey, TSharedPtr<SButton>> ResetConnectorButton;
 };
diff --git a/Engine/Plugins/Animation/ControlRig/Source/ControlRigEditor/Private/Editor/SRigHierarchyTreeView.cpp b/Engine/Plugins/Animation/ControlRig/Source/ControlRigEditor/Private/Editor/SRigHierarchyTreeView.cpp
--- a/Engine/Plugins/Animation/ControlRig/Source/ControlRigEditor/Private/Editor/SRigHierarchyTreeView.cpp
+++ b/Engine/Plugins/Animation/ControlRig/Source/ControlRigEditor/Private/Editor/SRigHierarchyTreeView.cpp
@@ -400,7 +400,13 @@ bool SRigHierarchyTreeView::AddElement(FRigElementKey InKey, FRigElementKey InPa
 	const bool bSupportsRename = Delegates.OnRenameElement.IsBound();
 
 	const FString FilteredString = Settings.FilterText.ToString();
-	if (FilteredString.IsEmpty() || !InKey.IsValid())
+	bool bAnyFilteredOut = Delegates.OnRigTreeIsItemVisible.IsBound();
+	if (!bAnyFilteredOut)
+	{
+		bAnyFilteredOut = FilteredString.IsEmpty() || !InKey.IsValid();
+	}
+
+	if (!bAnyFilteredOut)
 	{
 		TSharedPtr<FRigTreeElement> NewItem = MakeShared<FRigTreeElement>(InKey, SharedThis(this), bSupportsRename, ERigTreeFilterResult::Shown);
 
@@ -430,8 +436,14 @@ bool SRigHierarchyTreeView::AddElement(FRigElementKey InKey, FRigElementKey InPa
 	}
 	else
 	{
+		bool bIsFilteredOut = false;
+		if (Delegates.OnRigTreeIsItemVisible.IsBound())
+		{
+			bIsFilteredOut = !Delegates.OnRigTreeIsItemVisible.Execute(InKey);
+		}
+		
 		FString FilteredStringUnderScores = FilteredString.Replace(TEXT(" "), TEXT("_"));
-		if (InKey.Name.ToString().Contains(FilteredString) || InKey.Name.ToString().Contains(FilteredStringUnderScores))	
+		if (!bIsFilteredOut && (InKey.Name.ToString().Contains(FilteredString) || InKey.Name.ToString().Contains(FilteredStringUnderScores)))
 		{
 			TSharedPtr<FRigTreeElement> NewItem = MakeShared<FRigTreeElement>(InKey, SharedThis(this), bSupportsRename, ERigTreeFilterResult::Shown);
 			ElementMap.Add(InKey, NewItem);
@@ -844,7 +856,7 @@ void SRigHierarchyTreeView::RefreshTreeView(bool bRebuildContent)
 
 		if(const URigHierarchy* Hierarchy = Delegates.GetHierarchy())
 		{
-			TArray<FRigElementKey> Selection = Hierarchy->GetSelectedKeys();
+			TArray<FRigElementKey> Selection = Delegates.GetSelection();
 			for (const FRigElementKey& Key : Selection)
 			{
 				for (int32 RootIndex = 0; RootIndex < RootElements.Num(); ++RootIndex)
diff --git a/Engine/Plugins/Animation/ControlRig/Source/ControlRigEditor/Private/Editor/SRigHierarchyTreeView.h b/Engine/Plugins/Animation/ControlRig/Source/ControlRigEditor/Private/Editor/SRigHierarchyTreeView.h
--- a/Engine/Plugins/Animation/ControlRig/Source/ControlRigEditor/Private/Editor/SRigHierarchyTreeView.h
+++ b/Engine/Plugins/Animation/ControlRig/Source/ControlRigEditor/Private/Editor/SRigHierarchyTreeView.h
@@ -69,12 +69,14 @@ struct CONTROLRIGEDITOR_API FRigTreeDisplaySettings
 
 DECLARE_DELEGATE_RetVal(const URigHierarchy*, FOnGetRigTreeHierarchy);
 DECLARE_DELEGATE_RetVal(const FRigTreeDisplaySettings&, FOnGetRigTreeDisplaySettings);
+DECLARE_DELEGATE_RetVal(const TArray<FRigElementKey>, FOnRigTreeGetSelection);
 DECLARE_DELEGATE_RetVal_TwoParams(FName, FOnRigTreeRenameElement, const FRigElementKey& /*OldKey*/, const FString& /*NewName*/);
 DECLARE_DELEGATE_RetVal_ThreeParams(bool, FOnRigTreeVerifyElementNameChanged, const FRigElementKey& /*OldKey*/, const FString& /*NewName*/, FText& /*OutErrorMessage*/);
 DECLARE_DELEGATE_RetVal_TwoParams(bool, FOnRigTreeCompareKeys, const FRigElementKey& /*A*/, const FRigElementKey& /*B*/);
 DECLARE_DELEGATE_RetVal_OneParam(FRigElementKey, FOnRigTreeGetResolvedKey, const FRigElementKey&);
 DECLARE_DELEGATE_OneParam(FOnRigTreeRequestDetailsInspection, const FRigElementKey&);
 DECLARE_DELEGATE_RetVal_OneParam(TOptional<FText>, FOnRigTreeItemGetToolTip, const FRigElementKey&);
+DECLARE_DELEGATE_RetVal_OneParam(bool, FOnRigTreeIsItemVisible, const FRigElementKey&);
 
 typedef STableRow<TSharedPtr<FRigTreeElement>>::FOnCanAcceptDrop FOnRigTreeCanAcceptDrop;
 typedef STableRow<TSharedPtr<FRigTreeElement>>::FOnAcceptDrop FOnRigTreeAcceptDrop;
@@ -92,6 +94,7 @@ struct CONTROLRIGEDITOR_API FRigTreeDelegates
 	FOnDragDetected OnDragDetected;
 	FOnRigTreeCanAcceptDrop OnCanAcceptDrop;
 	FOnRigTreeAcceptDrop OnAcceptDrop;
+	FOnRigTreeGetSelection OnGetSelection;
 	FOnRigTreeSelectionChanged OnSelectionChanged;
 	FOnContextMenuOpening OnContextMenuOpening;
 	FOnRigTreeMouseButtonClick OnMouseButtonClick;
@@ -102,6 +105,7 @@ struct CONTROLRIGEDITOR_API FRigTreeDelegates
 	FOnRigTreeRequestDetailsInspection OnRequestDetailsInspection;
 	FOnRigTreeElementKeyTagDragDetected OnRigTreeElementKeyTagDragDetected;
 	FOnRigTreeItemGetToolTip OnRigTreeGetItemToolTip;
+	FOnRigTreeIsItemVisible OnRigTreeIsItemVisible;
 
 	FRigTreeDelegates()
 	{
@@ -125,6 +129,19 @@ struct CONTROLRIGEDITOR_API FRigTreeDelegates
 		}
 		return DefaultDisplaySettings;
 	}
+
+	TArray<FRigElementKey> GetSelection() const
+	{
+		if (OnGetSelection.IsBound())
+		{
+			return OnGetSelection.Execute();
+		}
+		if (const URigHierarchy* Hierarchy = GetHierarchy())
+		{
+			return Hierarchy->GetSelectedKeys();
+		}
+		return {};
+	}
 	
 	FName HandleRenameElement(const FRigElementKey& OldKey, const FString& NewName) const
 	{
