Control Rig: Adjust vertical scale of connection buttons

#rnx

diff --git a/Engine/Plugins/Animation/ControlRig/Source/ControlRigEditor/Private/ControlRigModuleDetails.cpp b/Engine/Plugins/Animation/ControlRig/Source/ControlRigEditor/Private/ControlRigModuleDetails.cpp
--- a/Engine/Plugins/Animation/ControlRig/Source/ControlRigEditor/Private/ControlRigModuleDetails.cpp
+++ b/Engine/Plugins/Animation/ControlRig/Source/ControlRigEditor/Private/ControlRigModuleDetails.cpp
@@ -344,118 +344,130 @@ void FRigModuleInstanceDetails::CustomizeDetails(IDetailLayoutBuilder& DetailBui
 						]
 					]
 
-					// Reset button
 					+SHorizontalBox::Slot()
 					.AutoWidth()
+					.VAlign(VAlign_Center)
 					.Padding(0.f, 0.f, 0.f, 0.f)
 					[
-						SAssignNew(ResetConnectorButton.FindOrAdd(ConnectorKey), SButton)
-						.ButtonStyle( FAppStyle::Get(), "NoBorder" )
-						.ButtonColorAndOpacity_Lambda([this, ConnectorKey]()
-						{
-							const TSharedPtr<SButton>& Button = ResetConnectorButton.FindRef(ConnectorKey);
-							return Button.IsValid() && Button->IsHovered()
-								? FSlateColor(FLinearColor(1.f, 1.f, 1.f, 0.8))
-								: FSlateColor(FLinearColor(1.f, 1.f, 1.f, 0.4));
-						})
-						.OnClicked_Lambda([this, ConnectorKey]()
-						{
-							PerModuleInfos[0].GetBlueprint()->GetModularRigController()->DisconnectConnector(ConnectorKey);
-							return FReply::Handled();
-						})
-						.ContentPadding(1.f)
-						.ToolTipText(NSLOCTEXT("ControlRigModuleDetails", "Reset_Connector", "Reset Connector"))
+						SNew(SVerticalBox)
+						+SVerticalBox::Slot()
+						.VAlign(VAlign_Center)
+						.AutoHeight()
 						[
-							SNew(SImage)
-							.ColorAndOpacity_Lambda( [this, ConnectorKey]()
-							{
-								const TSharedPtr<SButton>& Button = ResetConnectorButton.FindRef(ConnectorKey);
-								return Button.IsValid() && Button->IsHovered()
-								? FSlateColor(FLinearColor(1.f, 1.f, 1.f, 0.8))
-								: FSlateColor(FLinearColor(1.f, 1.f, 1.f, 0.4));
-							})
-							.Image(FSlateIcon(FAppStyle::Get().GetStyleSetName(), "PropertyWindow.DiffersFromDefault").GetIcon())
-						]
-					]
+							SNew(SHorizontalBox)
+							+SHorizontalBox::Slot()
+							.AutoWidth()
+							.Padding(0.f, 0.f, 0.f, 0.f)
+							[
+								SAssignNew(ResetConnectorButton.FindOrAdd(ConnectorKey), SButton)
+								.ButtonStyle( FAppStyle::Get(), "NoBorder" )
+								.ButtonColorAndOpacity_Lambda([this, ConnectorKey]()
+								{
+									const TSharedPtr<SButton>& Button = ResetConnectorButton.FindRef(ConnectorKey);
+									return Button.IsValid() && Button->IsHovered()
+										? FSlateColor(FLinearColor(1.f, 1.f, 1.f, 0.8))
+										: FSlateColor(FLinearColor(1.f, 1.f, 1.f, 0.4));
+								})
+								.OnClicked_Lambda([this, ConnectorKey]()
+								{
+									PerModuleInfos[0].GetBlueprint()->GetModularRigController()->DisconnectConnector(ConnectorKey);
+									return FReply::Handled();
+								})
+								.ContentPadding(1.f)
+								.ToolTipText(NSLOCTEXT("ControlRigModuleDetails", "Reset_Connector", "Reset Connector"))
+								[
+									SNew(SImage)
+									.ColorAndOpacity_Lambda( [this, ConnectorKey]()
+									{
+										const TSharedPtr<SButton>& Button = ResetConnectorButton.FindRef(ConnectorKey);
+										return Button.IsValid() && Button->IsHovered()
+										? FSlateColor(FLinearColor(1.f, 1.f, 1.f, 0.8))
+										: FSlateColor(FLinearColor(1.f, 1.f, 1.f, 0.4));
+									})
+									.Image(FSlateIcon(FAppStyle::Get().GetStyleSetName(), "PropertyWindow.DiffersFromDefault").GetIcon())
+								]
+							]
 
-					// Use button
-					+SHorizontalBox::Slot()
-					.AutoWidth()
-					.Padding(0.f, 0.f, 0.f, 0.f)
-					[
-						SAssignNew(UseSelectedButton.FindOrAdd(ConnectorKey), SButton)
-						.ButtonStyle( FAppStyle::Get(), "NoBorder" )
-						.ButtonColorAndOpacity_Lambda([this, ConnectorKey]()
-						{
-							const TSharedPtr<SButton>& Button = UseSelectedButton.FindRef(ConnectorKey);
-							return Button.IsValid() && Button->IsHovered()
-								? FSlateColor(FLinearColor(1.f, 1.f, 1.f, 0.8))
-								: FSlateColor(FLinearColor(1.f, 1.f, 1.f, 0.4));
-						})
-						.OnClicked_Lambda([this, ConnectorKey]()
-						{
-							if (UModularRig* ModularRig = PerModuleInfos[0].GetModularRig())
-							{
-								const TArray<FRigElementKey>& Selected = ModularRig->GetHierarchy()->GetSelectedKeys();
-								if (Selected.Num() > 0)
+							// Use button
+							+SHorizontalBox::Slot()
+							.AutoWidth()
+							.Padding(0.f, 0.f, 0.f, 0.f)
+							[
+								SAssignNew(UseSelectedButton.FindOrAdd(ConnectorKey), SButton)
+								.ButtonStyle( FAppStyle::Get(), "NoBorder" )
+								.ButtonColorAndOpacity_Lambda([this, ConnectorKey]()
 								{
-									PerModuleInfos[0].GetBlueprint()->GetModularRigController()->ConnectConnectorToElement(ConnectorKey, Selected[0]);
-								}
-							}
-							return FReply::Handled();
-						})
-						.ContentPadding(1.f)
-						.ToolTipText(NSLOCTEXT("ControlRigModuleDetails", "Use_Selected", "Use Selected"))
-						[
-							SNew(SImage)
-							.ColorAndOpacity_Lambda( [this, ConnectorKey]()
-							{
-								const TSharedPtr<SButton>& Button = UseSelectedButton.FindRef(ConnectorKey);
-								return Button.IsValid() && Button->IsHovered()
-								? FSlateColor(FLinearColor(1.f, 1.f, 1.f, 0.8))
-								: FSlateColor(FLinearColor(1.f, 1.f, 1.f, 0.4));
-							})
-							.Image(FAppStyle::GetBrush("Icons.CircleArrowLeft"))
-						]
-					]
+									const TSharedPtr<SButton>& Button = UseSelectedButton.FindRef(ConnectorKey);
+									return Button.IsValid() && Button->IsHovered()
+										? FSlateColor(FLinearColor(1.f, 1.f, 1.f, 0.8))
+										: FSlateColor(FLinearColor(1.f, 1.f, 1.f, 0.4));
+								})
+								.OnClicked_Lambda([this, ConnectorKey]()
+								{
+									if (UModularRig* ModularRig = PerModuleInfos[0].GetModularRig())
+									{
+										const TArray<FRigElementKey>& Selected = ModularRig->GetHierarchy()->GetSelectedKeys();
+										if (Selected.Num() > 0)
+										{
+											PerModuleInfos[0].GetBlueprint()->GetModularRigController()->ConnectConnectorToElement(ConnectorKey, Selected[0]);
+										}
+									}
+									return FReply::Handled();
+								})
+								.ContentPadding(1.f)
+								.ToolTipText(NSLOCTEXT("ControlRigModuleDetails", "Use_Selected", "Use Selected"))
+								[
+									SNew(SImage)
+									.ColorAndOpacity_Lambda( [this, ConnectorKey]()
+									{
+										const TSharedPtr<SButton>& Button = UseSelectedButton.FindRef(ConnectorKey);
+										return Button.IsValid() && Button->IsHovered()
+										? FSlateColor(FLinearColor(1.f, 1.f, 1.f, 0.8))
+										: FSlateColor(FLinearColor(1.f, 1.f, 1.f, 0.4));
+									})
+									.Image(FAppStyle::GetBrush("Icons.CircleArrowLeft"))
+								]
+							]
 
-					// Select in hierarchy button
-					+SHorizontalBox::Slot()
-					.AutoWidth()
-					.Padding(0.f, 0.f, 0.f, 0.f)
-					[
-						SAssignNew(SelectElementButton.FindOrAdd(ConnectorKey), SButton)
-						.ButtonStyle( FAppStyle::Get(), "NoBorder" )
-						.ButtonColorAndOpacity_Lambda([this, ConnectorKey]()
-						{
-							const TSharedPtr<SButton>& Button = SelectElementButton.FindRef(ConnectorKey);
-							return Button.IsValid() && Button->IsHovered()
-								? FSlateColor(FLinearColor(1.f, 1.f, 1.f, 0.8))
-								: FSlateColor(FLinearColor(1.f, 1.f, 1.f, 0.4));
-						})
-						.OnClicked_Lambda([this, ConnectorKey]()
-						{
-							if (UModularRig* ModularRig = PerModuleInfos[0].GetModularRig())
-							{
-								if (const FRigElementKey* TargetKey = ModularRig->GetElementKeyRedirector().FindExternalKey(ConnectorKey))
+							// Select in hierarchy button
+							+SHorizontalBox::Slot()
+							.AutoWidth()
+							.Padding(0.f, 0.f, 0.f, 0.f)
+							[
+								SAssignNew(SelectElementButton.FindOrAdd(ConnectorKey), SButton)
+								.ButtonStyle( FAppStyle::Get(), "NoBorder" )
+								.ButtonColorAndOpacity_Lambda([this, ConnectorKey]()
 								{
-									ModularRig->GetHierarchy()->GetController()->SelectElement(*TargetKey);
-								}
-							}
-							return FReply::Handled();
-						})
-						.ContentPadding(1.f)
-						.ToolTipText(NSLOCTEXT("ControlRigModuleDetails", "Select_Element", "Select Element"))
-						[
-							SNew(SImage)
-							.ColorAndOpacity_Lambda( [this, ConnectorKey]()
-							{
-								const TSharedPtr<SButton>& Button = SelectElementButton.FindRef(ConnectorKey);
-								return Button.IsValid() && Button->IsHovered()
-								? FSlateColor(FLinearColor(1.f, 1.f, 1.f, 0.8))
-								: FSlateColor(FLinearColor(1.f, 1.f, 1.f, 0.4));
-							})
-							.Image(FAppStyle::GetBrush("Icons.Search"))
+									const TSharedPtr<SButton>& Button = SelectElementButton.FindRef(ConnectorKey);
+									return Button.IsValid() && Button->IsHovered()
+										? FSlateColor(FLinearColor(1.f, 1.f, 1.f, 0.8))
+										: FSlateColor(FLinearColor(1.f, 1.f, 1.f, 0.4));
+								})
+								.OnClicked_Lambda([this, ConnectorKey]()
+								{
+									if (UModularRig* ModularRig = PerModuleInfos[0].GetModularRig())
+									{
+										if (const FRigElementKey* TargetKey = ModularRig->GetElementKeyRedirector().FindExternalKey(ConnectorKey))
+										{
+											ModularRig->GetHierarchy()->GetController()->SelectElement(*TargetKey);
+										}
+									}
+									return FReply::Handled();
+								})
+								.ContentPadding(1.f)
+								.ToolTipText(NSLOCTEXT("ControlRigModuleDetails", "Select_Element", "Select Element"))
+								[
+									SNew(SImage)
+									.ColorAndOpacity_Lambda( [this, ConnectorKey]()
+									{
+										const TSharedPtr<SButton>& Button = SelectElementButton.FindRef(ConnectorKey);
+										return Button.IsValid() && Button->IsHovered()
+										? FSlateColor(FLinearColor(1.f, 1.f, 1.f, 0.8))
+										: FSlateColor(FLinearColor(1.f, 1.f, 1.f, 0.4));
+									})
+									.Image(FAppStyle::GetBrush("Icons.Search"))
+								]
+							]
 						]
 					]
 				];
