Control Rig: Make sure connections are valid when resolving connector to target element
#rnx

diff --git a/Engine/Plugins/Animation/ControlRig/Source/ControlRig/Private/ModularRigController.cpp b/Engine/Plugins/Animation/ControlRig/Source/ControlRig/Private/ModularRigController.cpp
--- a/Engine/Plugins/Animation/ControlRig/Source/ControlRig/Private/ModularRigController.cpp
+++ b/Engine/Plugins/Animation/ControlRig/Source/ControlRig/Private/ModularRigController.cpp
@@ -5,6 +5,7 @@
 #include "ControlRig.h"
 #include "ModularRig.h"
 #include "ModularRigModel.h"
+#include "ModularRigRuleManager.h"
 #include "Misc/DefaultValueHelper.h"
 #include "Rigs/RigHierarchyController.h"
 
@@ -127,26 +128,26 @@ FRigModuleReference* UModularRigController::FindModule(const FString& InPath)
 	return *Cur;
 }
 
-bool UModularRigController::ConnectConnectorToElement(const FRigElementKey& InConnectorKey, const FRigElementKey& InTargetKey, bool bSetupUndo)
+bool UModularRigController::CanConnectConnectorToElement(const FRigElementKey& InConnectorKey, const FRigElementKey& InTargetKey, FText& OutErrorMessage)
 {
 	FString ConnectorParentPath, ConnectorName;
 	if (!InConnectorKey.Name.ToString().Split(UModularRig::NamespaceSeparator, &ConnectorParentPath, &ConnectorName, ESearchCase::CaseSensitive, ESearchDir::FromEnd))
 	{
-		UE_LOG(LogControlRig, Error, TEXT("Connector %s does not contain a namespace"), *InConnectorKey.ToString());
+		OutErrorMessage = FText::FromString(FString::Printf(TEXT("Connector %s does not contain a namespace"), *InConnectorKey.ToString()));
 		return false;
 	}
 	
 	FRigModuleReference* Module = FindModule(ConnectorParentPath);
 	if (!Module)
 	{
-		UE_LOG(LogControlRig, Error, TEXT("Could not find module %s"), *ConnectorParentPath);
+		OutErrorMessage = FText::FromString(FString::Printf(TEXT("Could not find module %s"), *ConnectorParentPath));
 		return false;
 	}
 
 	UControlRig* RigCDO = Module->Class->GetDefaultObject<UControlRig>();
 	if (!RigCDO)
 	{
-		UE_LOG(LogControlRig, Error, TEXT("Invalid rig module class %s"), *Module->Class->GetPathName());
+		OutErrorMessage = FText::FromString(FString::Printf(TEXT("Invalid rig module class %s"), *Module->Class->GetPathName()));
 		return false;
 	}
 
@@ -157,19 +158,19 @@ bool UModularRigController::ConnectConnectorToElement(const FRigElementKey& InCo
 		});
 	if (!ModuleConnector)
 	{
-		UE_LOG(LogControlRig, Error, TEXT("Could not find connector %s in class %s"), *ConnectorName, *Module->Class->GetPathName());
+		OutErrorMessage = FText::FromString(FString::Printf(TEXT("Could not find connector %s in class %s"), *ConnectorName, *Module->Class->GetPathName()));
 		return false;
 	}
 
 	if (!InTargetKey.IsValid())
 	{
-		UE_LOG(LogControlRig, Error, TEXT("Invalid target %s in class %s"), *InTargetKey.ToString(), *Module->Class->GetPathName());
+		OutErrorMessage = FText::FromString(FString::Printf(TEXT("Invalid target %s in class %s"), *InTargetKey.ToString(), *Module->Class->GetPathName()));
 		return false;
 	}
 
 	if (InTargetKey == InConnectorKey)
 	{
-		UE_LOG(LogControlRig, Error, TEXT("Cannot resolve connector %s to itself in class %s"), *InTargetKey.ToString(), *Module->Class->GetPathName());
+		OutErrorMessage = FText::FromString(FString::Printf(TEXT("Cannot resolve connector %s to itself in class %s"), *InTargetKey.ToString(), *Module->Class->GetPathName()));
 		return false;
 	}
 
@@ -180,24 +181,67 @@ bool UModularRigController::ConnectConnectorToElement(const FRigElementKey& InCo
 		return true; // Nothing to do
 	}
 
-	/*
-	FText ErrorMessage;
-	if (!CanConnectConnectorToElement(*ModuleConnector, InTargetKey, ErrorMessage))
+	UBlueprint* Blueprint = Cast<UBlueprint>(GetOuter());
+
+	// Make sure the connection is valid
 	{
-		UE_LOG(LogControlRig, Error, TEXT("Cannot connect connector %s to target %s: %s"),
-			*InConnectorKey.Name.ToString(), *InTargetKey.ToString(), *ErrorMessage.ToString());
-		return false;
+		UModularRig* ModularRig = Cast<UModularRig>(Blueprint->GetObjectBeingDebugged());
+		if (!ModularRig)
+		{
+			OutErrorMessage = FText::FromString(FString::Printf(TEXT("Could not find debugged modular rig in %s"), *Blueprint->GetPathName()));
+			return false;
+		}
+	
+		URigHierarchy* Hierarchy = ModularRig->GetHierarchy();
+		if (!Hierarchy)
+		{
+			OutErrorMessage = FText::FromString(FString::Printf(TEXT("Could not find hierarchy in %s"), *ModularRig->GetPathName()));
+			return false;
+		}
+	
+		const FRigConnectorElement* Connector = Cast<FRigConnectorElement>(Hierarchy->Find(InConnectorKey));
+		if (!Connector)
+		{
+			OutErrorMessage = FText::FromString(FString::Printf(TEXT("Could not find connector %s"), *InConnectorKey.ToString()));
+			return false;
+		}
+
+		UModularRigRuleManager* RuleManager = Hierarchy->GetRuleManager();
+		if (!RuleManager)
+		{
+			OutErrorMessage = FText::FromString(FString::Printf(TEXT("Could not get rule manager")));
+			return false;
+		}
+
+		const FRigModuleInstance* ModuleInstance = ModularRig->FindModule(Module->GetPath());
+		FModularRigResolveResult RuleResults = RuleManager->FindMatches(Connector, ModuleInstance, ModularRig->GetElementKeyRedirector());
+		if (!RuleResults.ContainsMatch(InTargetKey))
+		{
+			OutErrorMessage = FText::FromString(FString::Printf(TEXT("The target %s is not a valid match for connector %s"), *InTargetKey.ToString(), *InConnectorKey.ToString()));
+			return false;
+		}
 	}
-	*/
 
-	UBlueprint* Blueprint = Cast<UBlueprint>(GetOuter());
-	const IRigHierarchyProvider* HierarchyProvider = CastChecked<IRigHierarchyProvider>(Blueprint);
-	const FRigConnectorElement* Connector = Cast<FRigConnectorElement>(HierarchyProvider->GetHierarchy()->Find(InConnectorKey));
-	if (!Connector)
+	return true;
+}
+
+bool UModularRigController::ConnectConnectorToElement(const FRigElementKey& InConnectorKey, const FRigElementKey& InTargetKey, bool bSetupUndo)
+{
+	FText ErrorMessage;
+	if (!CanConnectConnectorToElement(InConnectorKey, InTargetKey, ErrorMessage))
 	{
-		UE_LOG(LogControlRig, Error, TEXT("Could not find connector %s"), *InConnectorKey.ToString());
+		UE_LOG(LogControlRig, Error, TEXT("Could not connect %s to %s: %s"), *InConnectorKey.ToString(), *InTargetKey.ToString(), *ErrorMessage.ToString());
 		return false;
 	}
+	
+	FString ConnectorParentPath, ConnectorName;
+	InConnectorKey.Name.ToString().Split(UModularRig::NamespaceSeparator, &ConnectorParentPath, &ConnectorName, ESearchCase::CaseSensitive, ESearchDir::FromEnd);
+	FRigModuleReference* Module = FindModule(ConnectorParentPath);
+
+	const FRigElementKey ConnectorKey(*ConnectorName, ERigElementType::Connector);
+	FRigElementKey* CurrentTarget = Module->Connections.Find(ConnectorKey);
+
+	UBlueprint* Blueprint = Cast<UBlueprint>(GetOuter());
 
 #if WITH_EDITOR
 	TSharedPtr<FScopedTransaction> TransactionPtr;
diff --git a/Engine/Plugins/Animation/ControlRig/Source/ControlRig/Public/ModularRigController.h b/Engine/Plugins/Animation/ControlRig/Source/ControlRig/Public/ModularRigController.h
--- a/Engine/Plugins/Animation/ControlRig/Source/ControlRig/Public/ModularRigController.h
+++ b/Engine/Plugins/Animation/ControlRig/Source/ControlRig/Public/ModularRigController.h
@@ -27,6 +27,9 @@ class CONTROLRIG_API UModularRigController : public UObject
 	UFUNCTION(BlueprintCallable, Category = "Control Rig | Modules")
 	FString AddModule(const FName& InModuleName, TSubclassOf<UControlRig> InClass, const FString& InParentModulePath, bool bSetupUndo = true);
 
+	UFUNCTION(BlueprintCallable, Category = "Control Rig | Modules")
+	bool CanConnectConnectorToElement(const FRigElementKey& InConnectorKey, const FRigElementKey& InTargetKey, FText& OutErrorMessage);
+
 	UFUNCTION(BlueprintCallable, Category = "Control Rig | Modules")
 	bool ConnectConnectorToElement(const FRigElementKey& InConnectorKey, const FRigElementKey& InTargetKey, bool bSetupUndo = true);
 
