Control Rig: Fix compilation error
#rnx

diff --git a/Engine/Plugins/Animation/ControlRig/Source/ControlRig/Private/ModularRigController.cpp b/Engine/Plugins/Animation/ControlRig/Source/ControlRig/Private/ModularRigController.cpp
--- a/Engine/Plugins/Animation/ControlRig/Source/ControlRig/Private/ModularRigController.cpp
+++ b/Engine/Plugins/Animation/ControlRig/Source/ControlRig/Private/ModularRigController.cpp
@@ -181,6 +181,7 @@ bool UModularRigController::CanConnectConnectorToElement(const FRigElementKey& I
 		return true; // Nothing to do
 	}
 
+#if WITH_EDITOR
 	UBlueprint* Blueprint = Cast<UBlueprint>(GetOuter());
 
 	// Make sure the connection is valid
@@ -221,6 +222,7 @@ bool UModularRigController::CanConnectConnectorToElement(const FRigElementKey& I
 			return false;
 		}
 	}
+#endif
 
 	return true;
 }
