Control Rig: Attempt to autoresolve connectors when another connector is resolved
#rnx

diff --git a/Engine/Plugins/Animation/ControlRig/Source/ControlRig/Private/ModularRigController.cpp b/Engine/Plugins/Animation/ControlRig/Source/ControlRig/Private/ModularRigController.cpp
--- a/Engine/Plugins/Animation/ControlRig/Source/ControlRig/Private/ModularRigController.cpp
+++ b/Engine/Plugins/Animation/ControlRig/Source/ControlRig/Private/ModularRigController.cpp
@@ -181,6 +181,22 @@ bool UModularRigController::CanConnectConnectorToElement(const FRigElementKey& I
 		return true; // Nothing to do
 	}
 
+	if (ModuleConnector->Settings.Type != EConnectorType::Primary)
+	{
+		const FRigModuleConnector* PrimaryMdouleConnector = RigCDO->GetRigModuleSettings().ExposedConnectors.FindByPredicate(
+		[ConnectorName](FRigModuleConnector& Connector)
+		{
+			return Connector.Settings.Type == EConnectorType::Primary;
+		});
+		const FRigElementKey PrimaryConnectorKey(*PrimaryMdouleConnector->Name, ERigElementType::Connector);
+		FRigElementKey* PrimaryTarget = Module->Connections.Find(PrimaryConnectorKey);
+		if (!PrimaryTarget || !PrimaryTarget->IsValid())
+		{
+			OutErrorMessage = FText::FromString(FString::Printf(TEXT("Cannot resolve connector %s because primary connector is not resolved"), *InConnectorKey.ToString()));
+			return false;
+		}
+	}
+
 #if WITH_EDITOR
 	UBlueprint* Blueprint = Cast<UBlueprint>(GetOuter());
 
@@ -227,7 +243,7 @@ bool UModularRigController::CanConnectConnectorToElement(const FRigElementKey& I
 	return true;
 }
 
-bool UModularRigController::ConnectConnectorToElement(const FRigElementKey& InConnectorKey, const FRigElementKey& InTargetKey, bool bSetupUndo)
+bool UModularRigController::ConnectConnectorToElement(const FRigElementKey& InConnectorKey, const FRigElementKey& InTargetKey, bool bSetupUndo, bool bAutoResolveOtherConnectors)
 {
 	FText ErrorMessage;
 	if (!CanConnectConnectorToElement(InConnectorKey, InTargetKey, ErrorMessage))
@@ -261,9 +277,59 @@ bool UModularRigController::ConnectConnectorToElement(const FRigElementKey& InCo
 	}
 
 	Module->Connections.FindOrAdd(ConnectorKey) = InTargetKey;
-
 	Notify(EModularRigNotification::ConnectionChanged, Module);
 
+#if WITH_EDITOR
+	if (bAutoResolveOtherConnectors)
+	{
+		if (UControlRig* RigCDO = Module->Class->GetDefaultObject<UControlRig>())
+		{
+			if (UModularRig* ModularRig = Cast<UModularRig>(Blueprint->GetObjectBeingDebugged()))
+			{
+				if (URigHierarchy* Hierarchy = ModularRig->GetHierarchy())
+				{
+					UModularRigRuleManager* RuleManager = Hierarchy->GetRuleManager();
+					const FRigModuleInstance* ModuleInstance = ModularRig->FindModule(Module->GetPath());
+
+					for (const FRigModuleConnector& OtherConnector : RigCDO->GetRigModuleSettings().ExposedConnectors)
+					{
+						FRigElementKey OtherConnectorShortKey;
+						OtherConnectorShortKey.Name = *OtherConnector.Name;
+						OtherConnectorShortKey.Type = ERigElementType::Connector;
+						if (!Module->Connections.Contains(OtherConnectorShortKey))
+						{
+							FRigElementKey OtherConnectorKey;
+							OtherConnectorKey.Name = *FString::Printf(TEXT("%s%s"), *Module->GetNamespace(), *OtherConnector.Name);
+							OtherConnectorKey.Type = ERigElementType::Connector;
+							
+							const FRigConnectorElement* OtherConnectorElement = Cast<FRigConnectorElement>(Hierarchy->Find(OtherConnectorKey));
+							FModularRigResolveResult RuleResults = RuleManager->FindMatches(OtherConnectorElement, ModuleInstance, ModularRig->GetElementKeyRedirector());
+
+							if (RuleResults.GetMatches().Num() == 1)
+							{
+								Module->Connections.FindOrAdd(OtherConnectorShortKey) = RuleResults.GetMatches()[0].GetKey();
+								Notify(EModularRigNotification::ConnectionChanged, Module);
+							}
+							else
+							{
+								for (const FRigElementResolveResult& Result : RuleResults.GetMatches())
+								{
+									if (Result.GetState() == ERigElementResolveState::DefaultTarget)
+									{
+										Module->Connections.FindOrAdd(OtherConnectorShortKey) = Result.GetKey();
+										Notify(EModularRigNotification::ConnectionChanged, Module);
+										break;
+									}
+								}
+							}
+						}
+					}
+				}
+			}
+		}
+	}
+#endif
+
 #if WITH_EDITOR
 	TransactionPtr.Reset();
 #endif
diff --git a/Engine/Plugins/Animation/ControlRig/Source/ControlRig/Public/ModularRigController.h b/Engine/Plugins/Animation/ControlRig/Source/ControlRig/Public/ModularRigController.h
--- a/Engine/Plugins/Animation/ControlRig/Source/ControlRig/Public/ModularRigController.h
+++ b/Engine/Plugins/Animation/ControlRig/Source/ControlRig/Public/ModularRigController.h
@@ -31,7 +31,7 @@ class CONTROLRIG_API UModularRigController : public UObject
 	bool CanConnectConnectorToElement(const FRigElementKey& InConnectorKey, const FRigElementKey& InTargetKey, FText& OutErrorMessage);
 
 	UFUNCTION(BlueprintCallable, Category = "Control Rig | Modules")
-	bool ConnectConnectorToElement(const FRigElementKey& InConnectorKey, const FRigElementKey& InTargetKey, bool bSetupUndo = true);
+	bool ConnectConnectorToElement(const FRigElementKey& InConnectorKey, const FRigElementKey& InTargetKey, bool bSetupUndo = true, bool bAutoResolveOtherConnectors = true);
 
 	UFUNCTION(BlueprintCallable, Category = "Control Rig | Modules")
 	bool DisconnectConnector(const FRigElementKey& InConnectorKey, bool bSetupUndo = true);
diff --git a/Engine/Plugins/Animation/ControlRig/Source/ControlRig/Public/ModularRigModel.h b/Engine/Plugins/Animation/ControlRig/Source/ControlRig/Public/ModularRigModel.h
--- a/Engine/Plugins/Animation/ControlRig/Source/ControlRig/Public/ModularRigModel.h
+++ b/Engine/Plugins/Animation/ControlRig/Source/ControlRig/Public/ModularRigModel.h
@@ -109,6 +109,22 @@ public:
 
 	GENERATED_BODY()
 
+	FModularRigModel(){}
+	FModularRigModel(const FModularRigModel& Other)
+	{
+		Modules = Other.Modules;
+		UpdateCachedChildren();
+	}
+	
+	FModularRigModel& operator=(const FModularRigModel& Other)
+	{
+		Modules = Other.Modules;
+		UpdateCachedChildren();
+		return *this;
+	}
+	FModularRigModel(FModularRigModel&&) = delete;
+	FModularRigModel& operator=(FModularRigModel&&) = delete;
+
 	UPROPERTY()
 	TArray<FRigModuleReference> Modules;
 	TArray<FRigModuleReference*> RootModules;
diff --git a/Engine/Plugins/Animation/ControlRig/Source/ControlRigDeveloper/Private/ControlRigBlueprint.cpp b/Engine/Plugins/Animation/ControlRig/Source/ControlRigDeveloper/Private/ControlRigBlueprint.cpp
--- a/Engine/Plugins/Animation/ControlRig/Source/ControlRigDeveloper/Private/ControlRigBlueprint.cpp
+++ b/Engine/Plugins/Animation/ControlRig/Source/ControlRigDeveloper/Private/ControlRigBlueprint.cpp
@@ -991,7 +991,7 @@ bool UControlRigBlueprint::ResolveConnector(const FRigElementKey& DraggedKey, co
 		// Add connection to the model
 		if (UModularRigController* Controller = GetModularRigController())
 		{
-			Controller->ConnectConnectorToElement(DraggedKey, TargetKey, bSetupUndoRedo);
+			Controller->ConnectConnectorToElement(DraggedKey, TargetKey, bSetupUndoRedo, ModularRigSettings.bAutoResolve);
 		}
 	}
 	else
@@ -1600,6 +1600,8 @@ void UControlRigBlueprint::RecompileModularRig()
 			// behaviors in the blueprint re-instancer - which is disregarding any
 			// object under a CDO.
 			DefaultObject->ModularRigModel = ModularRigModel;
+			DefaultObject->ModularRigModel.SetOuterClientHost(DefaultObject);
+			DefaultObject->ModularRigSettings = ModularRigSettings;
 			PropagateModuleHierarchyFromBPToInstances();
 		}
 	}
diff --git a/Engine/Plugins/Animation/ControlRig/Source/ControlRigEditor/Private/ControlRigModuleDetails.cpp b/Engine/Plugins/Animation/ControlRig/Source/ControlRigEditor/Private/ControlRigModuleDetails.cpp
--- a/Engine/Plugins/Animation/ControlRig/Source/ControlRigEditor/Private/ControlRigModuleDetails.cpp
+++ b/Engine/Plugins/Animation/ControlRig/Source/ControlRigEditor/Private/ControlRigModuleDetails.cpp
@@ -285,22 +285,7 @@ void FRigModuleInstanceDetails::CustomizeDetails(IDetailLayoutBuilder& DetailBui
 				}
 				return {};
 			});
-			TreeDelegates.OnSelectionChanged.BindLambda([this, ConnectorKey](TSharedPtr<FRigTreeElement> Selection, ESelectInfo::Type SelectInfo)
-			{
-				FPerModuleInfo Info = PerModuleInfos[0];
-				TSharedPtr<FRigTreeElement> NewSelection = Selection;
-				if (UModularRigController* Controller = Info.GetBlueprint()->GetModularRigController())
-				{
-					if (NewSelection.IsValid())
-					{
-						Controller->ConnectConnectorToElement(ConnectorKey, NewSelection->Key);
-					}
-					else
-					{
-						Controller->DisconnectConnector(ConnectorKey);
-					}
-				}
-			});
+			TreeDelegates.OnSelectionChanged.BindSP(this, &FRigModuleInstanceDetails::OnConnectorTargetChanged, ConnectorKey);
 			
 
 			ConnectionsCategory.AddCustomRow(Label)
@@ -323,6 +308,7 @@ void FRigModuleInstanceDetails::CustomizeDetails(IDetailLayoutBuilder& DetailBui
 						SNew( SComboButton )
 						.ContentPadding(3)
 						.MenuPlacement(MenuPlacement_BelowAnchor)
+						.OnComboBoxOpened(this, &FRigModuleInstanceDetails::PopulateConnectorTargetList, ConnectorKey)
 						.ButtonContent()
 						[
 							// Wrap in configurable box to restrain height/width of menu
@@ -354,6 +340,7 @@ void FRigModuleInstanceDetails::CustomizeDetails(IDetailLayoutBuilder& DetailBui
 						.VAlign(VAlign_Center)
 						.AutoHeight()
 						[
+							// Reset button
 							SNew(SHorizontalBox)
 							+SHorizontalBox::Slot()
 							.AutoWidth()
@@ -409,7 +396,7 @@ void FRigModuleInstanceDetails::CustomizeDetails(IDetailLayoutBuilder& DetailBui
 										const TArray<FRigElementKey>& Selected = ModularRig->GetHierarchy()->GetSelectedKeys();
 										if (Selected.Num() > 0)
 										{
-											PerModuleInfos[0].GetBlueprint()->GetModularRigController()->ConnectConnectorToElement(ConnectorKey, Selected[0]);
+											PerModuleInfos[0].GetBlueprint()->GetModularRigController()->ConnectConnectorToElement(ConnectorKey, Selected[0], true, ModularRig->GetModularRigSettings().bAutoResolve);
 										}
 									}
 									return FReply::Handled();
@@ -479,7 +466,6 @@ void FRigModuleInstanceDetails::CustomizeDetails(IDetailLayoutBuilder& DetailBui
 			}
 			TPair<const FSlateBrush*, FSlateColor> IconAndColor = SRigHierarchyItem::GetBrushForElementType(PerModuleInfos[0].GetModularRig()->GetHierarchy(), CurrentTargetKey);
 			PopulateConnectorCurrentTarget(ButtonBox, ConnectorKey, CurrentTargetKey, IconAndColor.Key, IconAndColor.Value, FText::FromName(CurrentTargetKey.Name));
-			PopulateConnectorTargetList(ConnectorKey);
 		}
 	}
 
@@ -833,7 +819,7 @@ FRigElementKeyRedirector FRigModuleInstanceDetails::GetConnections() const
 	return FRigElementKeyRedirector();
 }
 
-void FRigModuleInstanceDetails::PopulateConnectorTargetList(const FRigElementKey& InConnectorKey)
+void FRigModuleInstanceDetails::PopulateConnectorTargetList(const FRigElementKey InConnectorKey)
 {
 	if (PerModuleInfos.IsEmpty() || PerModuleInfos.Num() > 1)
 	{
@@ -872,9 +858,9 @@ void FRigModuleInstanceDetails::PopulateConnectorCurrentTarget(TSharedPtr<SVerti
 		SNew( SButton )
 		.ButtonStyle(FAppStyle::Get(), "SimpleButton")
 		.ContentPadding(FMargin(0.0))
-		.OnClicked_Lambda([Controller, InConnectorKey, InTargetKey]()
+		.OnClicked_Lambda([Info, Controller, InConnectorKey, InTargetKey]()
 		{
-			Controller->ConnectConnectorToElement(InConnectorKey, InTargetKey);
+			Controller->ConnectConnectorToElement(InConnectorKey, InTargetKey, true, Info.GetModularRig()->GetModularRigSettings().bAutoResolve);
 			return FReply::Handled();
 		})
 		[
@@ -944,6 +930,30 @@ void FRigModuleInstanceDetails::OnConfigValueChanged(const FName InVariableName)
 	}
 }
 
+void FRigModuleInstanceDetails::OnConnectorTargetChanged(TSharedPtr<FRigTreeElement> Selection, ESelectInfo::Type SelectInfo, const FRigElementKey InConnectorKey)
+{
+	if (SelectInfo == ESelectInfo::OnNavigation)
+	{
+		return;
+	}
+	
+	FScopedTransaction Transaction(LOCTEXT("ModuleHierarchyResolveConnector", "Resolve Connector"));
+	FPerModuleInfo Info = PerModuleInfos[0];
+	TSharedPtr<FRigTreeElement> NewSelection = Selection;
+	if (UModularRigController* Controller = Info.GetBlueprint()->GetModularRigController())
+	{
+		if (NewSelection.IsValid())
+		{
+			const FModularRigSettings& Settings = Info.GetModularRig()->GetModularRigSettings();
+			Controller->ConnectConnectorToElement(InConnectorKey, NewSelection->Key, true, Settings.bAutoResolve);
+		}
+		else
+		{
+			Controller->DisconnectConnector(InConnectorKey);
+		}
+	}
+}
+
 const FRigModuleInstanceDetails::FPerModuleInfo& FRigModuleInstanceDetails::FindModule(const FString& InPath) const
 {
 	const FPerModuleInfo* Info = FindModuleByPredicate([InPath](const FPerModuleInfo& Info)
diff --git a/Engine/Plugins/Animation/ControlRig/Source/ControlRigEditor/Private/ControlRigModuleDetails.h b/Engine/Plugins/Animation/ControlRig/Source/ControlRigEditor/Private/ControlRigModuleDetails.h
--- a/Engine/Plugins/Animation/ControlRig/Source/ControlRigEditor/Private/ControlRigModuleDetails.h
+++ b/Engine/Plugins/Animation/ControlRig/Source/ControlRigEditor/Private/ControlRigModuleDetails.h
@@ -50,7 +50,7 @@ public:
 	FText GetRigClassPath() const;
 	TArray<FRigModuleConnector> GetConnectors() const;
 	FRigElementKeyRedirector GetConnections() const;
-	void PopulateConnectorTargetList(const FRigElementKey& InConnectorKey);
+	void PopulateConnectorTargetList(const FRigElementKey InConnectorKey);
 	void PopulateConnectorCurrentTarget(
 		TSharedPtr<SVerticalBox> InListBox,
 		const FRigElementKey& InConnectorKey,
@@ -60,6 +60,7 @@ public:
 		const FText& InTitle);
 
 	void OnConfigValueChanged(const FName InVariableName);
+	void OnConnectorTargetChanged(TSharedPtr<FRigTreeElement> Selection, ESelectInfo::Type SelectInfo, const FRigElementKey InConnectorKey);
 	
 	struct FPerModuleInfo
 	{
