Control Rig: Add ModularRigSettings to ModularRigs and the ControlRigBlueprint
#rnx

diff --git a/Engine/Plugins/Animation/ControlRig/Source/ControlRig/Private/ModularRig.cpp b/Engine/Plugins/Animation/ControlRig/Source/ControlRig/Private/ModularRig.cpp
--- a/Engine/Plugins/Animation/ControlRig/Source/ControlRig/Private/ModularRig.cpp
+++ b/Engine/Plugins/Animation/ControlRig/Source/ControlRig/Private/ModularRig.cpp
@@ -294,6 +294,19 @@ void UModularRig::Evaluate_AnyThread()
 	Super::Evaluate_AnyThread();
 }
 
+const FModularRigSettings& UModularRig::GetModularRigSettings() const
+{
+	if(HasAnyFlags(RF_ClassDefaultObject))
+	{
+		return ModularRigSettings;
+	}
+	if (const UModularRig* CDO = Cast<UModularRig>(GetClass()->GetDefaultObject()))
+	{
+		return CDO->GetModularRigSettings();
+	}
+	return ModularRigSettings;
+}
+
 void UModularRig::ExecuteQueue()
 {
 	FRigVMExtendedExecuteContext& Context = GetRigVMExtendedExecuteContext();
@@ -303,105 +316,108 @@ void UModularRig::ExecuteQueue()
 	while(ExecutionQueue.IsValidIndex(ExecutionQueueFront))
 	{
 		FRigModuleExecutionElement& ExecutionElement = ExecutionQueue[ExecutionQueueFront];
-		if (UControlRig* ModuleRig = ExecutionElement.ModuleInstance->GetRig())
+		if (FRigModuleInstance* ModuleInstance = ExecutionElement.ModuleInstance)
 		{
-			if (!ModuleRig->SupportsEvent(ExecutionElement.EventName))
+			if (UControlRig* ModuleRig = ModuleInstance->GetRig())
 			{
-				ExecutionQueueFront++;
-				continue;
-			}
+				if (!ModuleRig->SupportsEvent(ExecutionElement.EventName))
+				{
+					ExecutionQueueFront++;
+					continue;
+				}
 
-			// Make sure the hierarchy has the correct element redirector from this module rig
-			FRigHierarchyRedirectorGuard ElementRedirectorGuard(ModuleRig);
+				// Make sure the hierarchy has the correct element redirector from this module rig
+				FRigHierarchyRedirectorGuard ElementRedirectorGuard(ModuleRig);
 
-			FRigVMExtendedExecuteContext& RigExtendedExecuteContext= ModuleRig->GetRigVMExtendedExecuteContext();
+				FRigVMExtendedExecuteContext& RigExtendedExecuteContext= ModuleRig->GetRigVMExtendedExecuteContext();
 
-			// Make sure the hierarchy has the correct execute context with the rig module namespace
-			FRigHierarchyExecuteContextBracket ExecuteContextBracket(Hierarchy, &RigExtendedExecuteContext);
+				// Make sure the hierarchy has the correct execute context with the rig module namespace
+				FRigHierarchyExecuteContextBracket ExecuteContextBracket(Hierarchy, &RigExtendedExecuteContext);
 
-			FControlRigExecuteContext& RigPublicContext = RigExtendedExecuteContext.GetPublicDataSafe<FControlRigExecuteContext>();
-			FRigUnitContext& RigUnitContext = RigPublicContext.UnitContext;
-			RigUnitContext = PublicContext.UnitContext;
+				FControlRigExecuteContext& RigPublicContext = RigExtendedExecuteContext.GetPublicDataSafe<FControlRigExecuteContext>();
+				FRigUnitContext& RigUnitContext = RigPublicContext.UnitContext;
+				RigUnitContext = PublicContext.UnitContext;
 
-			// forward the draw interface to each module
-			RigPublicContext.SetDrawInterface(PublicContext.GetDrawInterface());
-			RigPublicContext.SetDrawContainer(PublicContext.GetDrawContainer());
-			RigPublicContext.RigModuleInstance = ExecutionElement.ModuleInstance;
+				// forward the draw interface to each module
+				RigPublicContext.SetDrawInterface(PublicContext.GetDrawInterface());
+				RigPublicContext.SetDrawContainer(PublicContext.GetDrawContainer());
+				RigPublicContext.RigModuleInstance = ExecutionElement.ModuleInstance;
 
-			// re-initialize the module in case only the VM side got recompiled.
-			// this happens when the user relies on auto recompilation when editing the
-			// module (dependency) graph - by changing a value, add / remove nodes or links.
-			if(ModuleRig->IsInitRequired())
-			{
-				const TGuardValue<float> AbsoluteTimeGuard(ModuleRig->AbsoluteTime, ModuleRig->AbsoluteTime);
-				const TGuardValue<float> DeltaTimeGuard(ModuleRig->DeltaTime, ModuleRig->DeltaTime);
-				if(!ModuleRig->InitializeVM(ExecutionElement.EventName))
+				// re-initialize the module in case only the VM side got recompiled.
+				// this happens when the user relies on auto recompilation when editing the
+				// module (dependency) graph - by changing a value, add / remove nodes or links.
+				if(ModuleRig->IsInitRequired())
 				{
-					ExecutionQueueFront++;
-					continue;
-				}
+					const TGuardValue<float> AbsoluteTimeGuard(ModuleRig->AbsoluteTime, ModuleRig->AbsoluteTime);
+					const TGuardValue<float> DeltaTimeGuard(ModuleRig->DeltaTime, ModuleRig->DeltaTime);
+					if(!ModuleRig->InitializeVM(ExecutionElement.EventName))
+					{
+						ExecutionQueueFront++;
+						continue;
+					}
 
-				// put the variable defaults back
-				if(const FRigModuleReference* ModuleReference = GetModularRigModel().FindModule(ExecutionElement.ModulePath))
-				{
-					for (const TPair<FName, FString>& Variable : ModuleReference->ConfigValues)
+					// put the variable defaults back
+					if(const FRigModuleReference* ModuleReference = GetModularRigModel().FindModule(ExecutionElement.ModulePath))
 					{
-						ModuleRig->SetVariableFromString(Variable.Key, Variable.Value);
+						for (const TPair<FName, FString>& Variable : ModuleReference->ConfigValues)
+						{
+							ModuleRig->SetVariableFromString(Variable.Key, Variable.Value);
+						}
 					}
 				}
-			}
 
-			// Update the interaction elements to show only the ones belonging to this module
-			const FString ModuleNamespace = FString::Printf(TEXT("%s:"), *ExecutionElement.ModulePath);
-			RigUnitContext.ElementsBeingInteracted = RigUnitContext.ElementsBeingInteracted.FilterByPredicate(
-				[ModuleNamespace, Hierarchy](const FRigElementKey& Key)
-			{
-				return ModuleNamespace == Hierarchy->GetNameMetadata(Key, URigHierarchy::NameSpaceMetadataName, NAME_None);
-			});
-			RigUnitContext.InteractionType = RigUnitContext.ElementsBeingInteracted.IsEmpty() ?
-				(uint8) EControlRigInteractionType::None
-				: RigUnitContext.InteractionType;
-
-			// Make sure the module's rig has the corrct user data
-			// The rig will combine the user data of the
-			// - skeleton
-			// - skeletalmesh
-			// - SkeletalMeshComponent
-			// - default control rig module
-			// - outer modular rig
-			// - external variables
-			{
-				RigPublicContext.AssetUserData.Reset();
-				if(const TArray<UAssetUserData*>* ControlRigUserDataArray = ModuleRig->GetAssetUserDataArray())
+				// Update the interaction elements to show only the ones belonging to this module
+				const FString ModuleNamespace = FString::Printf(TEXT("%s:"), *ExecutionElement.ModulePath);
+				RigUnitContext.ElementsBeingInteracted = RigUnitContext.ElementsBeingInteracted.FilterByPredicate(
+					[ModuleNamespace, Hierarchy](const FRigElementKey& Key)
+				{
+					return ModuleNamespace == Hierarchy->GetNameMetadata(Key, URigHierarchy::NameSpaceMetadataName, NAME_None);
+				});
+				RigUnitContext.InteractionType = RigUnitContext.ElementsBeingInteracted.IsEmpty() ?
+					(uint8) EControlRigInteractionType::None
+					: RigUnitContext.InteractionType;
+
+				// Make sure the module's rig has the corrct user data
+				// The rig will combine the user data of the
+				// - skeleton
+				// - skeletalmesh
+				// - SkeletalMeshComponent
+				// - default control rig module
+				// - outer modular rig
+				// - external variables
 				{
-					for(const UAssetUserData* ControlRigUserData : *ControlRigUserDataArray)
+					RigPublicContext.AssetUserData.Reset();
+					if(const TArray<UAssetUserData*>* ControlRigUserDataArray = ModuleRig->GetAssetUserDataArray())
 					{
-						RigPublicContext.AssetUserData.Add(ControlRigUserData);
+						for(const UAssetUserData* ControlRigUserData : *ControlRigUserDataArray)
+						{
+							RigPublicContext.AssetUserData.Add(ControlRigUserData);
+						}
 					}
+					RigPublicContext.AssetUserData.Remove(nullptr);
 				}
-				RigPublicContext.AssetUserData.Remove(nullptr);
-			}
 
-			// Copy variable bindings
-			for (TPair<FName, FRigVMExternalVariable>& Pair : ExecutionElement.ModuleInstance->VariableBindings)
-			{
-				const FRigVMExternalVariable TargetVariable = ExecutionElement.ModuleInstance->GetRig()->GetPublicVariableByName(Pair.Key);
-				if(ensure(TargetVariable.Property))
+				// Copy variable bindings
+				for (TPair<FName, FRigVMExternalVariable>& Pair : ExecutionElement.ModuleInstance->VariableBindings)
 				{
-					if (RigVMTypeUtils::AreCompatible(Pair.Value.Property, TargetVariable.Property))
+					const FRigVMExternalVariable TargetVariable = ExecutionElement.ModuleInstance->GetRig()->GetPublicVariableByName(Pair.Key);
+					if(ensure(TargetVariable.Property))
 					{
-						Pair.Value.Property->CopyCompleteValue(TargetVariable.Memory, Pair.Value.Memory);
+						if (RigVMTypeUtils::AreCompatible(Pair.Value.Property, TargetVariable.Property))
+						{
+							Pair.Value.Property->CopyCompleteValue(TargetVariable.Memory, Pair.Value.Memory);
+						}
 					}
 				}
-			}
 			
-			ModuleRig->Execute_Internal(ExecutionElement.EventName);
-			ExecutionElement.bExecuted = true;
+				ModuleRig->Execute_Internal(ExecutionElement.EventName);
+				ExecutionElement.bExecuted = true;
 
-			// Copy result of Connection event to the ModularRig's unit context
-			if (ExecutionElement.EventName == FRigUnit_ConnectorExecution::EventName)
-			{
-				PublicContext.UnitContext.ConnectionResolve = RigPublicContext.UnitContext.ConnectionResolve;
+				// Copy result of Connection event to the ModularRig's unit context
+				if (ExecutionElement.EventName == FRigUnit_ConnectorExecution::EventName)
+				{
+					PublicContext.UnitContext.ConnectionResolve = RigPublicContext.UnitContext.ConnectionResolve;
+				}
 			}
 		}
 		
diff --git a/Engine/Plugins/Animation/ControlRig/Source/ControlRig/Public/ModularRig.h b/Engine/Plugins/Animation/ControlRig/Source/ControlRig/Public/ModularRig.h
--- a/Engine/Plugins/Animation/ControlRig/Source/ControlRig/Public/ModularRig.h
+++ b/Engine/Plugins/Animation/ControlRig/Source/ControlRig/Public/ModularRig.h
@@ -138,6 +138,12 @@ public:
 	virtual const TArray<FName>& GetSupportedEvents() const override { return SupportedEvents; }
 	// END ControlRig
 
+	UPROPERTY()
+	FModularRigSettings ModularRigSettings;
+
+	// Returns the settings of the modular rig
+	const FModularRigSettings& GetModularRigSettings() const;
+
 	UPROPERTY()
 	FModularRigModel ModularRigModel;
 
diff --git a/Engine/Plugins/Animation/ControlRig/Source/ControlRig/Public/Rigs/RigModuleDefines.h b/Engine/Plugins/Animation/ControlRig/Source/ControlRig/Public/Rigs/RigModuleDefines.h
--- a/Engine/Plugins/Animation/ControlRig/Source/ControlRig/Public/Rigs/RigModuleDefines.h
+++ b/Engine/Plugins/Animation/ControlRig/Source/ControlRig/Public/Rigs/RigModuleDefines.h
@@ -6,6 +6,16 @@
 #include "RigHierarchyElements.h"
 #include "RigModuleDefines.generated.h"
 
+USTRUCT(BlueprintType)
+struct FModularRigSettings
+{
+	GENERATED_BODY()
+
+	// Whether or not to autoresolve secondary connectors once the primary connector is resolved
+	UPROPERTY(BlueprintReadWrite, EditAnywhere, Category = ModularRig)
+	bool bAutoResolve = true;
+};
+
 USTRUCT(BlueprintType)
 struct CONTROLRIG_API FRigModuleIdentifier
 {
diff --git a/Engine/Plugins/Animation/ControlRig/Source/ControlRigDeveloper/Public/ControlRigBlueprint.h b/Engine/Plugins/Animation/ControlRig/Source/ControlRigDeveloper/Public/ControlRigBlueprint.h
--- a/Engine/Plugins/Animation/ControlRig/Source/ControlRigDeveloper/Public/ControlRigBlueprint.h
+++ b/Engine/Plugins/Animation/ControlRig/Source/ControlRigDeveloper/Public/ControlRigBlueprint.h
@@ -153,6 +153,9 @@ public:
 		return Hierarchy;
 	}
 
+	UPROPERTY(EditAnywhere, Category = "Modular Rig")
+	FModularRigSettings ModularRigSettings;
+
 	UPROPERTY(EditAnywhere, Category = "Hierarchy")
 	FRigHierarchySettings HierarchySettings;
 
diff --git a/Engine/Plugins/Animation/ControlRig/Source/ControlRigEditor/Private/Editor/ControlRigSchematicModel.cpp b/Engine/Plugins/Animation/ControlRig/Source/ControlRigEditor/Private/Editor/ControlRigSchematicModel.cpp
--- a/Engine/Plugins/Animation/ControlRig/Source/ControlRigEditor/Private/Editor/ControlRigSchematicModel.cpp
+++ b/Engine/Plugins/Animation/ControlRig/Source/ControlRigEditor/Private/Editor/ControlRigSchematicModel.cpp
@@ -673,7 +673,7 @@ void FControlRigSchematicModel::HandleSchematicDrop(SSchematicGraphPanel* InPane
 				{
 					if (UModularRigController* Controller = ControlRigBlueprint->GetModularRigController())
 					{
-						UControlRig* ControlRig = ControlRigBlueprint->GetDebuggedControlRig();
+						UModularRig* ControlRig = Cast<UModularRig>(ControlRigBlueprint->GetDebuggedControlRig());
 						if (!ControlRig)
 						{
 							return;
@@ -710,7 +710,7 @@ void FControlRigSchematicModel::HandleSchematicDrop(SSchematicGraphPanel* InPane
 							{
 								(void)Controller->ReparentModule(ModulePath, TargetModulePath.ToString());
 							}
-							Controller->ConnectConnectorToElement(PrimaryConnectorKey, TargetKey);
+							Controller->ConnectConnectorToElement(PrimaryConnectorKey, TargetKey, true, ControlRig->GetModularRigSettings().bAutoResolve);
 						}
 					}
 				}, TStatId(), NULL, ENamedThreads::GameThread);
@@ -722,7 +722,7 @@ void FControlRigSchematicModel::HandleSchematicDrop(SSchematicGraphPanel* InPane
 		const TArray<FGuid> Sources = SchematicDragDropOp->GetElements();
 		FFunctionGraphTask::CreateAndDispatchWhenReady([this, Sources, TargetKey]()
 		{
-			UControlRig* ControlRig = ControlRigBlueprint->GetDebuggedControlRig();
+			UModularRig* ControlRig = Cast<UModularRig>(ControlRigBlueprint->GetDebuggedControlRig());
 			if (!ControlRig)
 			{
 				return;
@@ -742,7 +742,7 @@ void FControlRigSchematicModel::HandleSchematicDrop(SSchematicGraphPanel* InPane
 					{
 						if (Hierarchy->Find<FRigConnectorElement>(Pair.Key))
 						{
-							Controller->ConnectConnectorToElement(Pair.Key, TargetKey);
+							Controller->ConnectConnectorToElement(Pair.Key, TargetKey, true, ControlRig->GetModularRigSettings().bAutoResolve);
 							break;
 						}
 					}
