Control Rig: Add SetDefaultMatch node for use in the Connector event
#rnx

diff --git a/Engine/Plugins/Animation/ControlRig/Source/ControlRig/Private/Rigs/RigHierarchyDefines.cpp b/Engine/Plugins/Animation/ControlRig/Source/ControlRig/Private/Rigs/RigHierarchyDefines.cpp
--- a/Engine/Plugins/Animation/ControlRig/Source/ControlRig/Private/Rigs/RigHierarchyDefines.cpp
+++ b/Engine/Plugins/Animation/ControlRig/Source/ControlRig/Private/Rigs/RigHierarchyDefines.cpp
@@ -422,3 +422,19 @@ bool FModularRigResolveResult::ContainsMatch(const FRigElementKey& InKey, FStrin
 	
 	return false;
 }
+
+const FRigElementResolveResult* FModularRigResolveResult::FindMatch(const FRigElementKey& InKey) const
+{
+	return Matches.FindByPredicate([InKey](const FRigElementResolveResult& InMatch) -> bool
+	{
+		return InMatch.GetKey() == InKey;
+	});
+}
+
+const FRigElementResolveResult* FModularRigResolveResult::GetDefaultMatch() const
+{
+	return Matches.FindByPredicate([](const FRigElementResolveResult& Match)
+	{
+		return Match.GetState() == ERigElementResolveState::DefaultTarget;
+	});
+}
diff --git a/Engine/Plugins/Animation/ControlRig/Source/ControlRig/Private/Units/Modules/RigUnit_ConnectionCandidates.cpp b/Engine/Plugins/Animation/ControlRig/Source/ControlRig/Private/Units/Modules/RigUnit_ConnectionCandidates.cpp
--- a/Engine/Plugins/Animation/ControlRig/Source/ControlRig/Private/Units/Modules/RigUnit_ConnectionCandidates.cpp
+++ b/Engine/Plugins/Animation/ControlRig/Source/ControlRig/Private/Units/Modules/RigUnit_ConnectionCandidates.cpp
@@ -14,6 +14,7 @@ FRigUnit_GetCandidates_Execute()
 	if(ExecuteContext.GetEventName() != FRigUnit_ConnectorExecution::EventName)
 	{
 		UE_CONTROLRIG_RIGUNIT_REPORT_ERROR(TEXT("Can only use GetCandidates during %s event."), *FRigUnit_ConnectorExecution::EventName.ToString());
+		return;
 	}
 
 	Algo::Transform(ExecuteContext.UnitContext.ConnectionResolve.Matches, Candidates, [](const FRigElementResolveResult& Match)
@@ -30,7 +31,8 @@ FRigUnit_DiscardMatches_Execute()
 	// only allow these nodes during connector resolval event
 	if(ExecuteContext.GetEventName() != FRigUnit_ConnectorExecution::EventName)
 	{
-		UE_CONTROLRIG_RIGUNIT_REPORT_ERROR(TEXT("Can only use AddMatches during %s event."), *FRigUnit_ConnectorExecution::EventName.ToString());
+		UE_CONTROLRIG_RIGUNIT_REPORT_ERROR(TEXT("Can only use DiscardMatches during %s event."), *FRigUnit_ConnectorExecution::EventName.ToString());
+		return;
 	}
 
 	// Filter any items that are already excluded
@@ -56,4 +58,39 @@ FRigUnit_DiscardMatches_Execute()
 
 		ExecuteContext.UnitContext.ConnectionResolve.Excluded.Add(FRigElementResolveResult(Exclude, ERigElementResolveState::InvalidTarget, FText::FromString(Message)));
 	}
+}
+
+FRigUnit_SetDefaultMatch_Execute()
+{
+	DECLARE_SCOPE_HIERARCHICAL_COUNTER_RIGUNIT()
+
+	// only allow these nodes during connector resolval event
+	if(ExecuteContext.GetEventName() != FRigUnit_ConnectorExecution::EventName)
+	{
+		UE_CONTROLRIG_RIGUNIT_REPORT_ERROR(TEXT("Can only use SetDefaultMatch during %s event."), *FRigUnit_ConnectorExecution::EventName.ToString());
+		return;
+	}
+
+	FModularRigResolveResult& Resolve = ExecuteContext.UnitContext.ConnectionResolve;
+	TArray<FRigElementResolveResult>& Matches = Resolve.Matches;
+
+	// Make sure the Default element is contained in the matches array
+	FRigElementResolveResult* NewDefault = const_cast<FRigElementResolveResult*>(Resolve.FindMatch(Default));
+	if (!NewDefault)
+	{
+		UE_CONTROLRIG_RIGUNIT_REPORT_ERROR(TEXT("Cannot set default match %s because it is not part of %s matches."), *Default.ToString(), *Resolve.GetConnectorKey().ToString());
+		return;
+	}
+
+	// Remove the old default if it exists (set to just possible target)
+	FRigElementResolveResult* OldDefault = const_cast<FRigElementResolveResult*>(Resolve.GetDefaultMatch());
+	if (OldDefault)
+	{
+		OldDefault->SetPossibleTarget();
+	}
+
+	if (NewDefault)
+	{
+		NewDefault->SetDefaultTarget();
+	}
 }
\ No newline at end of file
diff --git a/Engine/Plugins/Animation/ControlRig/Source/ControlRig/Public/Rigs/RigHierarchyDefines.h b/Engine/Plugins/Animation/ControlRig/Source/ControlRig/Public/Rigs/RigHierarchyDefines.h
--- a/Engine/Plugins/Animation/ControlRig/Source/ControlRig/Public/Rigs/RigHierarchyDefines.h
+++ b/Engine/Plugins/Animation/ControlRig/Source/ControlRig/Public/Rigs/RigHierarchyDefines.h
@@ -1913,6 +1913,8 @@ public:
 	}
 
 	bool ContainsMatch(const FRigElementKey& InKey, FString* OutErrorMessage = nullptr) const;
+	const FRigElementResolveResult* FindMatch(const FRigElementKey& InKey) const;
+	const FRigElementResolveResult* GetDefaultMatch() const;
 
 private:
 
@@ -1935,4 +1937,5 @@ private:
 	friend class UModularRig;
 	friend struct FRigUnit_GetCandidates;
 	friend struct FRigUnit_DiscardMatches;
+	friend struct FRigUnit_SetDefaultMatch;
 };
diff --git a/Engine/Plugins/Animation/ControlRig/Source/ControlRig/Public/Units/Modules/RigUnit_ConnectionCandidates.h b/Engine/Plugins/Animation/ControlRig/Source/ControlRig/Public/Units/Modules/RigUnit_ConnectionCandidates.h
--- a/Engine/Plugins/Animation/ControlRig/Source/ControlRig/Public/Units/Modules/RigUnit_ConnectionCandidates.h
+++ b/Engine/Plugins/Animation/ControlRig/Source/ControlRig/Public/Units/Modules/RigUnit_ConnectionCandidates.h
@@ -27,7 +27,7 @@ struct CONTROLRIG_API FRigUnit_GetCandidates : public FRigUnit
 };
 
 /**
- * Adds possible matches during a connector event
+ * Discards matches during a connector event
  */
 USTRUCT(meta=(DisplayName="Discard Matches", Category="Modules", TitleColor="1 0 0", NodeColor="1 1 1", Keywords="Connection,Resolve,Match", Varying))
 struct CONTROLRIG_API FRigUnit_DiscardMatches : public FRigUnitMutable
@@ -44,3 +44,20 @@ struct CONTROLRIG_API FRigUnit_DiscardMatches : public FRigUnitMutable
 	UPROPERTY(EditAnywhere, Transient, Category = "Modules", meta = (Input))
 	FString Message;
 };
+
+/**
+ * Set default match during a connector event
+ */
+USTRUCT(meta=(DisplayName="Set Default Match", Category="Modules", TitleColor="1 0 0", NodeColor="1 1 1", Keywords="Connection,Resolve,Match,Default", Varying))
+struct CONTROLRIG_API FRigUnit_SetDefaultMatch : public FRigUnitMutable
+{
+	GENERATED_BODY()
+
+	RIGVM_METHOD()
+	virtual void Execute() override;
+
+	// The items being interacted on
+	UPROPERTY(EditAnywhere, Transient, Category = "Modules", meta = (Input))
+	FRigElementKey Default;
+};
+
