Sequencer:Control Rig; Add check for CR object binding, for some reason it's not valid.

diff --git a/Engine/Plugins/Animation/ControlRig/Source/ControlRigEditor/Private/EditMode/ControlRigEditMode.cpp b/Engine/Plugins/Animation/ControlRig/Source/ControlRigEditor/Private/EditMode/ControlRigEditMode.cpp
--- a/Engine/Plugins/Animation/ControlRig/Source/ControlRigEditor/Private/EditMode/ControlRigEditMode.cpp
+++ b/Engine/Plugins/Animation/ControlRig/Source/ControlRigEditor/Private/EditMode/ControlRigEditMode.cpp
@@ -1426,7 +1426,7 @@ bool FControlRigEditMode::HandleClick(FEditorViewportClient* InViewportClient, H
 						UAdditiveControlRig* AdditiveControlRig = Cast<UAdditiveControlRig>(ControlRig);
 						UFKControlRig* FKControlRig = Cast<UFKControlRig>(ControlRig);
 
-						if (AdditiveControlRig || FKControlRig)
+						if ((AdditiveControlRig || FKControlRig) && ControlRig->GetObjectBinding().IsValid())
 						{
 							if (USkeletalMeshComponent* RigMeshComp = Cast<USkeletalMeshComponent>(ControlRig->GetObjectBinding()->GetBoundObject()))
 							{
