Control Rig: Change display order of sockets in rig hierarchy

#rnx

diff --git a/Engine/Plugins/Animation/ControlRig/Source/ControlRigEditor/Private/Editor/SRigHierarchyTreeView.cpp b/Engine/Plugins/Animation/ControlRig/Source/ControlRigEditor/Private/Editor/SRigHierarchyTreeView.cpp
--- a/Engine/Plugins/Animation/ControlRig/Source/ControlRigEditor/Private/Editor/SRigHierarchyTreeView.cpp
+++ b/Engine/Plugins/Animation/ControlRig/Source/ControlRigEditor/Private/Editor/SRigHierarchyTreeView.cpp
@@ -20,6 +20,7 @@
 #include "Rigs/AdditiveControlRig.h"
 #include "Rigs/RigHierarchyController.h"
 #include "Styling/AppStyle.h"
+#include "Algo/Sort.h"
 
 #define LOCTEXT_NAMESPACE "SRigHierarchyTreeView"
 
@@ -778,22 +779,55 @@ void SRigHierarchyTreeView::RefreshTreeView(bool bRebuildContent)
 		const URigHierarchy* Hierarchy = Delegates.GetHierarchy();
 		if(Hierarchy)
 		{
-			TArray<FRigConnectorElement*> Connectors;
+			TArray<const FRigSocketElement*> Sockets;
+			TArray<const FRigConnectorElement*> Connectors;
+			TArray<const FRigBaseElement*> EverythingElse;
+			TMap<const FRigBaseElement*, int32> ElementDepth;
+			Sockets.Reserve(Hierarchy->Num(ERigElementType::Socket));
+			Connectors.Reserve(Hierarchy->Num(ERigElementType::Connector));
+			EverythingElse.Reserve(Hierarchy->Num() - Hierarchy->Num(ERigElementType::Socket) - Hierarchy->Num(ERigElementType::Connector));
+			
 			Hierarchy->Traverse([&](FRigBaseElement* Element, bool& bContinue)
 			{
-				if(FRigConnectorElement* Connector = Cast<FRigConnectorElement>(Element))
+				int32& Depth = ElementDepth.Add(Element, 0);
+				if(const FRigBaseElement* ParentElement = Hierarchy->GetFirstParent(Element))
+				{
+					Depth = ElementDepth.FindChecked(ParentElement) + 1;
+				}
+				
+				if(const FRigSocketElement* Socket = Cast<FRigSocketElement>(Element))
+				{
+					Sockets.Add(Socket);
+				}
+				else if(const FRigConnectorElement* Connector = Cast<FRigConnectorElement>(Element))
 				{
 					Connectors.Add(Connector);
 				}
 				else
 				{
-					AddElement(Element);
+					EverythingElse.Add(Element);
 				}
 				bContinue = true;
 			});
 
+			// sort the sockets by depth
+			Algo::SortBy(Sockets, [ElementDepth](const FRigSocketElement* Socket) -> int32
+			{
+				return ElementDepth.FindChecked(Socket);
+			});
+			for(const FRigSocketElement* Socket : Sockets)
+			{
+				AddElement(Socket);
+			}
+
+			// add everything but connectors and sockets
+			for(const FRigBaseElement* Element : EverythingElse)
+			{
+				AddElement(Element);
+			}
+
 			// add all of the connectors. their parent relationship in the tree represents resolve
-			for(FRigConnectorElement* Connector : Connectors)
+			for(const FRigConnectorElement* Connector : Connectors)
 			{
 				AddElement(Connector);
 			}
diff --git a/Engine/Source/Editor/AnimationEditorWidgets/Private/SchematicGraphPanel/SSchematicGraphPanel.cpp b/Engine/Source/Editor/AnimationEditorWidgets/Private/SchematicGraphPanel/SSchematicGraphPanel.cpp
--- a/Engine/Source/Editor/AnimationEditorWidgets/Private/SchematicGraphPanel/SSchematicGraphPanel.cpp
+++ b/Engine/Source/Editor/AnimationEditorWidgets/Private/SchematicGraphPanel/SSchematicGraphPanel.cpp
@@ -482,7 +482,7 @@ void SSchematicGraphPanel::AddNode(const FSchematicGraphNode* InNodeToAdd)
 	TArray<TSharedPtr<FLinearColorAttribute>> Colors;
 	for(int32 LayerIndex = 0; LayerIndex < NumLayers; LayerIndex++)
 	{
-		Colors.Add(FLinearColorAttribute::CreateWithGetter(ColorInterpolationSettings, FLinearColorAttribute::FGetter::CreateRaw(GraphData, &FSchematicGraphModel::GetColorForNode, InNodeToAdd, LayerIndex)));
+		Colors.Add(FLinearColorAttribute::CreateWithGetter(ColorInterpolationSettings, FLinearColorAttribute::FGetter::CreateSP(this, &SSchematicGraphPanel::GetColorForNode, Guid, LayerIndex)));
 	}
 
 	const TFunction<const FSlateBrush*(const FGuid&, int32)> BrushGetter = [this](const FGuid& InGuid, int32 InLayerIndex) -> const FSlateBrush*
@@ -993,6 +993,15 @@ FVector2d SSchematicGraphPanel::GetPositionForNode(FGuid InNodeGuid) const
 	return FVector2d::ZeroVector;
 }
 
+FLinearColor SSchematicGraphPanel::GetColorForNode(FGuid InNodeGuid, int32 InLayerIndex) const
+{
+	if(GraphData)
+	{
+		return GraphData->GetColorForNode(InNodeGuid, InLayerIndex);
+	}
+	return FLinearColor::White;
+}
+
 float SSchematicGraphPanel::GetScaleForNode(FGuid InNodeGuid, bool bIncludeScaleOffset) const
 {
 	// check if the node may be auto scaled
diff --git a/Engine/Source/Editor/AnimationEditorWidgets/Public/SchematicGraphPanel/SSchematicGraphPanel.h b/Engine/Source/Editor/AnimationEditorWidgets/Public/SchematicGraphPanel/SSchematicGraphPanel.h
--- a/Engine/Source/Editor/AnimationEditorWidgets/Public/SchematicGraphPanel/SSchematicGraphPanel.h
+++ b/Engine/Source/Editor/AnimationEditorWidgets/Public/SchematicGraphPanel/SSchematicGraphPanel.h
@@ -208,6 +208,7 @@ public:
 	FReply HandleNodeDragDetected(FGuid Guid, const FGeometry& MyGeometry, const FPointerEvent& MouseEvent);
 
 	virtual FVector2d GetPositionForNode(FGuid InNodeGuid) const;
+	virtual FLinearColor GetColorForNode(FGuid InNodeGuid, int32 InLayerIndex) const;
 	virtual float GetScaleForNode(FGuid InNodeGuid, bool bIncludeScaleOffset) const;
 	virtual bool IsAutoScaleEnabledForNode(FGuid InNodeGuid) const;
 	virtual float GetMinimumLinkDistanceForNode(FGuid InLinkGuid, bool bIncludeScale = true) const;
