Control Rig: Avoid merging AssetUserDataEditorOnly during cook
#rnx

diff --git a/Engine/Plugins/Animation/ControlRig/Source/ControlRig/Private/ControlRig.cpp b/Engine/Plugins/Animation/ControlRig/Source/ControlRig/Private/ControlRig.cpp
--- a/Engine/Plugins/Animation/ControlRig/Source/ControlRig/Private/ControlRig.cpp
+++ b/Engine/Plugins/Animation/ControlRig/Source/ControlRig/Private/ControlRig.cpp
@@ -2975,11 +2975,18 @@ const TArray<UAssetUserData*>* UControlRig::GetAssetUserDataArray() const
 	if(HasAnyFlags(RF_ClassDefaultObject))
 	{
 #if WITH_EDITOR
-		static thread_local TArray<TObjectPtr<UAssetUserData>> CachedAssetUserData;
-		CachedAssetUserData.Reset();
-		CachedAssetUserData.Append(AssetUserData);
-		CachedAssetUserData.Append(AssetUserDataEditorOnly);
-		return &ToRawPtrTArrayUnsafe(CachedAssetUserData);
+		if (GIsCookerLoadingPackage)
+		{
+			return &ToRawPtrTArrayUnsafe(AssetUserData);
+		}
+		else
+		{
+			static thread_local TArray<TObjectPtr<UAssetUserData>> CachedAssetUserData;
+			CachedAssetUserData.Reset();
+			CachedAssetUserData.Append(AssetUserData);
+			CachedAssetUserData.Append(AssetUserDataEditorOnly);
+			return &ToRawPtrTArrayUnsafe(CachedAssetUserData);
+		}
 #else
 		return &ToRawPtrTArrayUnsafe(AssetUserData);
 #endif
diff --git a/Engine/Plugins/Runtime/RigVM/Source/RigVM/Private/RigVMHost.cpp b/Engine/Plugins/Runtime/RigVM/Source/RigVM/Private/RigVMHost.cpp
--- a/Engine/Plugins/Runtime/RigVM/Source/RigVM/Private/RigVMHost.cpp
+++ b/Engine/Plugins/Runtime/RigVM/Source/RigVM/Private/RigVMHost.cpp
@@ -1420,11 +1420,18 @@ void URigVMHost::RemoveUserDataOfClass(TSubclassOf<UAssetUserData> InUserDataCla
 const TArray<UAssetUserData*>* URigVMHost::GetAssetUserDataArray() const
 {
 #if WITH_EDITOR
-	static thread_local TArray<TObjectPtr<UAssetUserData>> CachedAssetUserData;
-	CachedAssetUserData.Reset();
-	CachedAssetUserData.Append(AssetUserData);
-	CachedAssetUserData.Append(AssetUserDataEditorOnly);
-	return &ToRawPtrTArrayUnsafe(CachedAssetUserData);
+	if (GIsCookerLoadingPackage)
+	{
+		return &ToRawPtrTArrayUnsafe(AssetUserData);
+	}
+	else
+	{
+		static thread_local TArray<TObjectPtr<UAssetUserData>> CachedAssetUserData;
+		CachedAssetUserData.Reset();
+		CachedAssetUserData.Append(AssetUserData);
+		CachedAssetUserData.Append(AssetUserDataEditorOnly);
+		return &ToRawPtrTArrayUnsafe(CachedAssetUserData);
+	}
 #else
 	return &ToRawPtrTArrayUnsafe(AssetUserData);
 #endif
diff --git a/Engine/Source/Runtime/Engine/Private/Animation/Skeleton.cpp b/Engine/Source/Runtime/Engine/Private/Animation/Skeleton.cpp
--- a/Engine/Source/Runtime/Engine/Private/Animation/Skeleton.cpp
+++ b/Engine/Source/Runtime/Engine/Private/Animation/Skeleton.cpp
@@ -2150,11 +2150,18 @@ void USkeleton::RemoveUserDataOfClass(TSubclassOf<UAssetUserData> InUserDataClas
 const TArray<UAssetUserData*>* USkeleton::GetAssetUserDataArray() const
 {
 #if WITH_EDITOR
-	static thread_local TArray<TObjectPtr<UAssetUserData>> CachedAssetUserData;
-	CachedAssetUserData.Reset();
-	CachedAssetUserData.Append(AssetUserData);
-	CachedAssetUserData.Append(AssetUserDataEditorOnly);
-	return &ToRawPtrTArrayUnsafe(CachedAssetUserData);
+	if (GIsCookerLoadingPackage)
+	{
+		return &ToRawPtrTArrayUnsafe(AssetUserData);
+	}
+	else
+	{
+		static thread_local TArray<TObjectPtr<UAssetUserData>> CachedAssetUserData;
+		CachedAssetUserData.Reset();
+		CachedAssetUserData.Append(AssetUserData);
+		CachedAssetUserData.Append(AssetUserDataEditorOnly);
+		return &ToRawPtrTArrayUnsafe(CachedAssetUserData);
+	}
 #else
 	return &ToRawPtrTArrayUnsafe(AssetUserData);
 #endif
diff --git a/Engine/Source/Runtime/Engine/Private/Components/ActorComponent.cpp b/Engine/Source/Runtime/Engine/Private/Components/ActorComponent.cpp
--- a/Engine/Source/Runtime/Engine/Private/Components/ActorComponent.cpp
+++ b/Engine/Source/Runtime/Engine/Private/Components/ActorComponent.cpp
@@ -2123,11 +2123,18 @@ UAssetUserData* UActorComponent::GetAssetUserDataOfClass(TSubclassOf<UAssetUserD
 const TArray<UAssetUserData*>* UActorComponent::GetAssetUserDataArray() const
 {
 #if WITH_EDITOR
-	static thread_local TArray<TObjectPtr<UAssetUserData>> CachedAssetUserData;
-	CachedAssetUserData.Reset();
-	CachedAssetUserData.Append(AssetUserData);
-	CachedAssetUserData.Append(AssetUserDataEditorOnly);
-	return &ToRawPtrTArrayUnsafe(CachedAssetUserData);
+	if (GIsCookerLoadingPackage)
+	{
+		return &ToRawPtrTArrayUnsafe(AssetUserData);
+	}
+	else
+	{
+		static thread_local TArray<TObjectPtr<UAssetUserData>> CachedAssetUserData;
+		CachedAssetUserData.Reset();
+		CachedAssetUserData.Append(AssetUserData);
+		CachedAssetUserData.Append(AssetUserDataEditorOnly);
+		return &ToRawPtrTArrayUnsafe(CachedAssetUserData);
+	}
 #else
 	return &ToRawPtrTArrayUnsafe(AssetUserData);
 #endif
diff --git a/Engine/Source/Runtime/Engine/Private/SkeletalMesh.cpp b/Engine/Source/Runtime/Engine/Private/SkeletalMesh.cpp
--- a/Engine/Source/Runtime/Engine/Private/SkeletalMesh.cpp
+++ b/Engine/Source/Runtime/Engine/Private/SkeletalMesh.cpp
@@ -4824,11 +4824,18 @@ void USkeletalMesh::RemoveUserDataOfClass(TSubclassOf<UAssetUserData> InUserData
 const TArray<UAssetUserData*>* USkeletalMesh::GetAssetUserDataArray() const
 {
 #if WITH_EDITOR
-	static thread_local TArray<TObjectPtr<UAssetUserData>> CachedAssetUserData;
-	CachedAssetUserData.Reset();
-	CachedAssetUserData.Append(AssetUserData);
-	CachedAssetUserData.Append(AssetUserDataEditorOnly);
-	return &ToRawPtrTArrayUnsafe(CachedAssetUserData);
+	if (GIsCookerLoadingPackage)
+	{
+		return &ToRawPtrTArrayUnsafe(AssetUserData);
+	}
+	else
+	{
+		static thread_local TArray<TObjectPtr<UAssetUserData>> CachedAssetUserData;
+		CachedAssetUserData.Reset();
+		CachedAssetUserData.Append(AssetUserData);
+		CachedAssetUserData.Append(AssetUserDataEditorOnly);
+		return &ToRawPtrTArrayUnsafe(CachedAssetUserData);
+	}
 #else
 	return &ToRawPtrTArrayUnsafe(AssetUserData);
 #endif
