Control Rig: Avoid merging AssetUserDataEditorOnly during all cook, not only during the load of a package.
#rnx

diff --git a/Engine/Plugins/Animation/ControlRig/Source/ControlRig/Private/ControlRig.cpp b/Engine/Plugins/Animation/ControlRig/Source/ControlRig/Private/ControlRig.cpp
--- a/Engine/Plugins/Animation/ControlRig/Source/ControlRig/Private/ControlRig.cpp
+++ b/Engine/Plugins/Animation/ControlRig/Source/ControlRig/Private/ControlRig.cpp
@@ -2975,7 +2975,7 @@ const TArray<UAssetUserData*>* UControlRig::GetAssetUserDataArray() const
 	if(HasAnyFlags(RF_ClassDefaultObject))
 	{
 #if WITH_EDITOR
-		if (GIsCookerLoadingPackage)
+		if (IsRunningCookCommandlet())
 		{
 			return &ToRawPtrTArrayUnsafe(AssetUserData);
 		}
@@ -3002,7 +3002,10 @@ const TArray<UAssetUserData*>* UControlRig::GetAssetUserDataArray() const
 	{
 		CombinedAssetUserData.Append(AssetUserData);
 #if WITH_EDITOR
-		CombinedAssetUserData.Append(AssetUserDataEditorOnly);
+		if (!IsRunningCookCommandlet())
+		{
+			CombinedAssetUserData.Append(AssetUserDataEditorOnly);
+		}
 #endif
 	}
 	
diff --git a/Engine/Plugins/Runtime/RigVM/Source/RigVM/Private/RigVMHost.cpp b/Engine/Plugins/Runtime/RigVM/Source/RigVM/Private/RigVMHost.cpp
--- a/Engine/Plugins/Runtime/RigVM/Source/RigVM/Private/RigVMHost.cpp
+++ b/Engine/Plugins/Runtime/RigVM/Source/RigVM/Private/RigVMHost.cpp
@@ -1420,7 +1420,7 @@ void URigVMHost::RemoveUserDataOfClass(TSubclassOf<UAssetUserData> InUserDataCla
 const TArray<UAssetUserData*>* URigVMHost::GetAssetUserDataArray() const
 {
 #if WITH_EDITOR
-	if (GIsCookerLoadingPackage)
+	if (IsRunningCookCommandlet())
 	{
 		return &ToRawPtrTArrayUnsafe(AssetUserData);
 	}
diff --git a/Engine/Source/Runtime/Engine/Private/Animation/Skeleton.cpp b/Engine/Source/Runtime/Engine/Private/Animation/Skeleton.cpp
--- a/Engine/Source/Runtime/Engine/Private/Animation/Skeleton.cpp
+++ b/Engine/Source/Runtime/Engine/Private/Animation/Skeleton.cpp
@@ -2150,7 +2150,7 @@ void USkeleton::RemoveUserDataOfClass(TSubclassOf<UAssetUserData> InUserDataClas
 const TArray<UAssetUserData*>* USkeleton::GetAssetUserDataArray() const
 {
 #if WITH_EDITOR
-	if (GIsCookerLoadingPackage)
+	if (IsRunningCookCommandlet())
 	{
 		return &ToRawPtrTArrayUnsafe(AssetUserData);
 	}
diff --git a/Engine/Source/Runtime/Engine/Private/Components/ActorComponent.cpp b/Engine/Source/Runtime/Engine/Private/Components/ActorComponent.cpp
--- a/Engine/Source/Runtime/Engine/Private/Components/ActorComponent.cpp
+++ b/Engine/Source/Runtime/Engine/Private/Components/ActorComponent.cpp
@@ -2123,7 +2123,7 @@ UAssetUserData* UActorComponent::GetAssetUserDataOfClass(TSubclassOf<UAssetUserD
 const TArray<UAssetUserData*>* UActorComponent::GetAssetUserDataArray() const
 {
 #if WITH_EDITOR
-	if (GIsCookerLoadingPackage)
+	if (IsRunningCookCommandlet())
 	{
 		return &ToRawPtrTArrayUnsafe(AssetUserData);
 	}
diff --git a/Engine/Source/Runtime/Engine/Private/SkeletalMesh.cpp b/Engine/Source/Runtime/Engine/Private/SkeletalMesh.cpp
--- a/Engine/Source/Runtime/Engine/Private/SkeletalMesh.cpp
+++ b/Engine/Source/Runtime/Engine/Private/SkeletalMesh.cpp
@@ -4824,7 +4824,7 @@ void USkeletalMesh::RemoveUserDataOfClass(TSubclassOf<UAssetUserData> InUserData
 const TArray<UAssetUserData*>* USkeletalMesh::GetAssetUserDataArray() const
 {
 #if WITH_EDITOR
-	if (GIsCookerLoadingPackage)
+	if (IsRunningCookCommandlet())
 	{
 		return &ToRawPtrTArrayUnsafe(AssetUserData);
 	}
