Control RIg: Add connections dropdown in the modular rig hierarchy view
#rnx

diff --git a/Engine/Plugins/Animation/ControlRig/Source/ControlRig/Public/ControlRig.h b/Engine/Plugins/Animation/ControlRig/Source/ControlRig/Public/ControlRig.h
--- a/Engine/Plugins/Animation/ControlRig/Source/ControlRig/Public/ControlRig.h
+++ b/Engine/Plugins/Animation/ControlRig/Source/ControlRig/Public/ControlRig.h
@@ -112,7 +112,8 @@ public:
 
 	// Returns the redirector from key to key for this rig
 	virtual FRigElementKeyRedirector& GetElementKeyRedirector();
-
+	virtual FRigElementKeyRedirector GetElementKeyRedirector() const { return ElementKeyRedirector; }
+	
 	// Returns the redirector from key to key for this rig
 	virtual void SetElementKeyRedirector(const FRigElementKeyRedirector InElementRedirector);
 
diff --git a/Engine/Plugins/Animation/ControlRig/Source/ControlRig/Public/ModularRig.h b/Engine/Plugins/Animation/ControlRig/Source/ControlRig/Public/ModularRig.h
--- a/Engine/Plugins/Animation/ControlRig/Source/ControlRig/Public/ModularRig.h
+++ b/Engine/Plugins/Animation/ControlRig/Source/ControlRig/Public/ModularRig.h
@@ -134,6 +134,7 @@ public:
 	virtual bool Execute_Internal(const FName& InEventName) override;
 	virtual void Evaluate_AnyThread() override;
 	virtual FRigElementKeyRedirector& GetElementKeyRedirector() override { return ElementKeyRedirector; }
+	virtual FRigElementKeyRedirector GetElementKeyRedirector() const override { return ElementKeyRedirector; }
 	virtual bool SupportsEvent(const FName& InEventName) const override { return SupportedEvents.Contains(InEventName); }
 	virtual const TArray<FName>& GetSupportedEvents() const override { return SupportedEvents; }
 	// END ControlRig
diff --git a/Engine/Plugins/Animation/ControlRig/Source/ControlRig/Public/Rigs/RigModuleDefines.h b/Engine/Plugins/Animation/ControlRig/Source/ControlRig/Public/Rigs/RigModuleDefines.h
--- a/Engine/Plugins/Animation/ControlRig/Source/ControlRig/Public/Rigs/RigModuleDefines.h
+++ b/Engine/Plugins/Animation/ControlRig/Source/ControlRig/Public/Rigs/RigModuleDefines.h
@@ -67,6 +67,14 @@ struct CONTROLRIG_API FRigModuleSettings
 			(!bRequireExposedConnectors || !ExposedConnectors.IsEmpty());
 	}
 
+	const FRigModuleConnector* FindPrimaryConnector() const
+	{
+		return ExposedConnectors.FindByPredicate([](const FRigModuleConnector& Connector)
+		{
+			return Connector.Settings.Type == EConnectorType::Primary;
+		});
+	}
+
 	// The identifier used to retrieve the module in the module library
 	UPROPERTY(BlueprintReadWrite, EditAnywhere, Category = Module)
 	FRigModuleIdentifier Identifier;
diff --git a/Engine/Plugins/Animation/ControlRig/Source/ControlRigEditor/Private/ControlRigEditorStyle.h b/Engine/Plugins/Animation/ControlRig/Source/ControlRigEditor/Private/ControlRigEditorStyle.h
--- a/Engine/Plugins/Animation/ControlRig/Source/ControlRigEditor/Private/ControlRigEditorStyle.h
+++ b/Engine/Plugins/Animation/ControlRig/Source/ControlRigEditor/Private/ControlRigEditorStyle.h
@@ -155,6 +155,7 @@ public:
 			Set("ControlRig.Schematic.ConnectorPrimary", new IMAGE_BRUSH_SVG("Slate/Connector_Primary", Icon128x128));
 			Set("ControlRig.Schematic.ConnectorSecondary", new IMAGE_BRUSH_SVG("Slate/Connector_Secondary", Icon128x128));
 			Set("ControlRig.Schematic.ConnectorOptional", new IMAGE_BRUSH_SVG("Slate/Connector_Optional", Icon128x128));
+			Set("ControlRig.Schematic.ConnectorWarning", new IMAGE_BRUSH_SVG("Slate/Connector_Warning", Icon128x128));
 			Set("ControlRig.Schematic.Bone", new IMAGE_BRUSH_SVG("Slate/Bone_Schematic", Icon128x128));
 			Set("ControlRig.Schematic.Control", new IMAGE_BRUSH_SVG("Slate/Control_Schematic", Icon128x128));
 			Set("ControlRig.Schematic.Null", new IMAGE_BRUSH_SVG("Slate/Null_Schematic", Icon128x128));
diff --git a/Engine/Plugins/Animation/ControlRig/Source/ControlRigEditor/Private/Editor/SModularRigModel.cpp b/Engine/Plugins/Animation/ControlRig/Source/ControlRigEditor/Private/Editor/SModularRigModel.cpp
--- a/Engine/Plugins/Animation/ControlRig/Source/ControlRigEditor/Private/Editor/SModularRigModel.cpp
+++ b/Engine/Plugins/Animation/ControlRig/Source/ControlRigEditor/Private/Editor/SModularRigModel.cpp
@@ -134,6 +134,8 @@ void SModularRigModel::Construct(const FArguments& InArgs, TSharedRef<FControlRi
 	Delegates.OnRequestDetailsInspection = FOnModularRigTreeRequestDetailsInspection::CreateSP(this, &SModularRigModel::OnRequestDetailsInspection);
 	Delegates.OnRenameElement = FOnModularRigTreeRenameElement::CreateSP(this, &SModularRigModel::HandleRenameModule);
 	Delegates.OnVerifyModuleNameChanged = FOnModularRigTreeVerifyElementNameChanged::CreateSP(this, &SModularRigModel::HandleVerifyNameChanged);
+	Delegates.OnResolveConnector = FOnModularRigTreeResolveConnector::CreateSP(this, &SModularRigModel::HandleConnectorResolved);
+	Delegates.OnDisconnectConnector = FOnModularRigTreeDisconnectConnector::CreateSP(this, &SModularRigModel::HandleConnectorDisconnect);
 	
 	ChildSlot
 	[
@@ -313,7 +315,7 @@ void SModularRigModel::OnItemClicked(TSharedPtr<FModularRigTreeElement> InItem)
 
 	if (ControlRigEditor.IsValid() && InItem.IsValid())
 	{
-		ControlRigEditor.Pin()->SetDetailViewForRigModules({InItem->Key});
+		ControlRigEditor.Pin()->SetDetailViewForRigModules({InItem->ModulePath});
 	}
 }
 
@@ -527,7 +529,7 @@ void SModularRigModel::HandleNewItem(UClass* InClass, const FString &InParentPat
 
 bool SModularRigModel::CanRenameModule() const
 {
-	return IsSingleSelected();
+	return IsSingleSelected() && TreeView->FindElement(GetSelectedKeys()[0])->bIsPrimary;
 }
 
 void SModularRigModel::HandleRenameModule()
@@ -611,7 +613,7 @@ void SModularRigModel::HandleDeleteModules()
 		{
 			if (Element.IsValid())
 			{
-				return Element->Key;
+				return Element->ModulePath;
 			}
 			return FString();
 		});
@@ -656,6 +658,35 @@ void SModularRigModel::HandleReparentModules(const TArray<FString>& InPaths, con
 	}
 }
 
+void SModularRigModel::HandleConnectorResolved(const FRigElementKey& InConnector, const FRigElementKey& InTarget)
+{
+	if (ControlRigBlueprint.IsValid())
+	{
+		FScopedTransaction Transaction(LOCTEXT("ModularRigModelResolveConnector", "Resolve Connector"));
+
+		UModularRigController* Controller = ControlRigBlueprint->GetModularRigController();
+		check(Controller);
+
+		if (ControlRigBeingDebuggedPtr.IsValid())
+		{
+			Controller->ConnectConnectorToElement(InConnector, InTarget, true, ControlRigBeingDebuggedPtr->GetModularRigSettings().bAutoResolve);
+		}
+	}
+}
+
+void SModularRigModel::HandleConnectorDisconnect(const FRigElementKey& InConnector)
+{
+	if (ControlRigBlueprint.IsValid())
+	{
+		FScopedTransaction Transaction(LOCTEXT("ModularRigModelDisconnectConnector", "Disconnect Connector"));
+
+		UModularRigController* Controller = ControlRigBlueprint->GetModularRigController();
+		check(Controller);
+
+		Controller->DisconnectConnector(InConnector, true);
+	}
+}
+
 class SModularRigModelPasteTransformsErrorPipe : public FOutputDevice
 {
 public:
@@ -792,7 +823,7 @@ FReply SModularRigModel::OnAcceptDrop(const FDragDropEvent& DragDropEvent, EItem
 	FString ParentPath;
 	if (ItemAtMouse && ItemAtMouse->IsValid())
 	{
-		ParentPath = ItemAtMouse->Get()->Key;
+		ParentPath = ItemAtMouse->Get()->ModulePath;
 	}
 
 	TSharedPtr<FAssetDragDropOp> AssetDragDropOperation = DragDropEvent.GetOperationAs<FAssetDragDropOp>();
diff --git a/Engine/Plugins/Animation/ControlRig/Source/ControlRigEditor/Private/Editor/SModularRigModel.h b/Engine/Plugins/Animation/ControlRig/Source/ControlRigEditor/Private/Editor/SModularRigModel.h
--- a/Engine/Plugins/Animation/ControlRig/Source/ControlRigEditor/Private/Editor/SModularRigModel.h
+++ b/Engine/Plugins/Animation/ControlRig/Source/ControlRigEditor/Private/Editor/SModularRigModel.h
@@ -102,6 +102,12 @@ private:
 	/** Reparent items */
 	void HandleReparentModules(const TArray<FString>& InPaths, const FString& InParentPath);
 
+	/** Resolve connector */
+	void HandleConnectorResolved(const FRigElementKey& InConnector, const FRigElementKey& InTarget);
+
+	/** UnResolve connector */
+	void HandleConnectorDisconnect(const FRigElementKey& InConnector);
+
 	/** Set Selection Changed */
 	void OnSelectionChanged(TSharedPtr<FModularRigTreeElement> Selection, ESelectInfo::Type SelectInfo);
 
diff --git a/Engine/Plugins/Animation/ControlRig/Source/ControlRigEditor/Private/Editor/SModularRigTreeView.cpp b/Engine/Plugins/Animation/ControlRig/Source/ControlRigEditor/Private/Editor/SModularRigTreeView.cpp
--- a/Engine/Plugins/Animation/ControlRig/Source/ControlRigEditor/Private/Editor/SModularRigTreeView.cpp
+++ b/Engine/Plugins/Animation/ControlRig/Source/ControlRigEditor/Private/Editor/SModularRigTreeView.cpp
@@ -15,27 +15,56 @@
 #include "ControlRig.h"
 #include "ControlRigEditorStyle.h"
 #include "ModularRig.h"
+#include "ModularRigRuleManager.h"
+#include "SRigHierarchyTreeView.h"
 #include "Blueprint/WidgetBlueprintLibrary.h"
 #include "Editor/SModularRigModel.h"
+#include "Fonts/FontMeasure.h"
 #include "Settings/ControlRigSettings.h"
 #include "Graph/ControlRigGraphSchema.h"
 #include "Rigs/AdditiveControlRig.h"
 #include "Rigs/RigHierarchyController.h"
 #include "Styling/AppStyle.h"
+#include "Widgets/Input/SComboButton.h"
+#include "Widgets/Input/SButton.h"
+#include "ScopedTransaction.h"
 
 #define LOCTEXT_NAMESPACE "SModularRigTreeView"
 
-TMap<FSoftObjectPath, FSlateBrush> SModularRigModelItem::IconPathToBrush;
+TMap<FSoftObjectPath, FSlateBrush> FModularRigTreeElement::IconPathToBrush;
 
 //////////////////////////////////////////////////////////////
 /// FModularRigTreeElement
 ///////////////////////////////////////////////////////////
-FModularRigTreeElement::FModularRigTreeElement(const FString& InKey, TWeakPtr<SModularRigTreeView> InTreeView, bool InSupportsRename)
+FModularRigTreeElement::FModularRigTreeElement(const FString& InKey, TWeakPtr<SModularRigTreeView> InTreeView, bool bInIsPrimary)
 {
 	Key = InKey;
+	bIsPrimary = bInIsPrimary;
+
 	FString ShortNameStr = Key;
-	FString ParentPath;
-	ShortNameStr.Split(UModularRig::NamespaceSeparator, &ParentPath, &ShortNameStr, ESearchCase::Type::CaseSensitive, ESearchDir::FromEnd);
+	ShortNameStr.Split(UModularRig::NamespaceSeparator, &ModulePath, &ShortNameStr, ESearchCase::Type::CaseSensitive, ESearchDir::FromEnd);
+	if (bIsPrimary)
+	{
+		ModulePath = Key;
+		if(const UModularRig* ModularRig = InTreeView.Pin()->GetRigTreeDelegates().GetModularRig())
+		{
+			if (const FRigModuleInstance* Module = ModularRig->FindModule(ModulePath))
+			{
+				if (const UControlRig* Rig = Module->GetRig())
+				{
+					if (const FRigModuleConnector* PrimaryConnector = Rig->GetRigModuleSettings().FindPrimaryConnector())
+					{
+						ConnectorName = PrimaryConnector->Name;
+					}
+				}
+			}
+		}
+	}
+	else
+	{
+		ConnectorName = ShortNameStr;
+	}
+	
 	ShortName = *ShortNameStr;
 	
 	if(InTreeView.IsValid())
@@ -49,7 +78,7 @@ FModularRigTreeElement::FModularRigTreeElement(const FString& InKey, TWeakPtr<SM
 
 void FModularRigTreeElement::RefreshDisplaySettings(const UModularRig* InModularRig)
 {
-	const TPair<const FSlateBrush*, FSlateColor> Result = SModularRigModelItem::GetBrushForElementType(InModularRig, Key);
+	const TPair<const FSlateBrush*, FSlateColor> Result = GetBrushAndColor(InModularRig);
 
 	IconBrush = Result.Key;
 	IconColor = Result.Value;
@@ -93,8 +122,50 @@ void SModularRigModelItem::Construct(const FArguments& InArgs, const TSharedRef<
 		return;
 	}
 
+	const FString& ModulePath = InRigTreeElement->ModulePath;
+
+	FString ConnectorPath = FString::Printf(TEXT("%s:%s"), *ModulePath, *InRigTreeElement->ConnectorName);
+	FRigElementKey ConnectorKey(*ConnectorPath, ERigElementType::Connector);
+
 	TSharedPtr< SInlineEditableTextBlock > InlineWidget;
 	TSharedPtr< SHorizontalBox > HorizontalBox;
+	TSharedPtr<SVerticalBox> ComboButtonBox;
+
+	FRigTreeDelegates TreeDelegates;
+	TreeDelegates.OnGetHierarchy = FOnGetRigTreeHierarchy::CreateLambda([this]()
+	{
+		return Delegates.GetModularRig()->GetHierarchy();
+	});
+	TreeDelegates.OnRigTreeIsItemVisible = FOnRigTreeIsItemVisible::CreateLambda([this, ConnectorKey, ModulePath](const FRigElementKey& InTarget)
+	{
+		if (const UModularRig* ModularRig = Delegates.GetModularRig())
+		{
+			if (const FRigModuleInstance* Module = ModularRig->FindModule(ModulePath))
+			{
+				if (URigHierarchy* Hierarchy = ModularRig->GetHierarchy())
+				{
+					if (FRigConnectorElement* ConnectorElement = Cast<FRigConnectorElement>(Hierarchy->Find(ConnectorKey)))
+					{
+						const UModularRigRuleManager* RuleManager = ModularRig->GetHierarchy()->GetRuleManager();
+						const FModularRigResolveResult Result = RuleManager->FindMatches(ConnectorElement, Module, ModularRig->GetElementKeyRedirector());
+						return Result.ContainsMatch(InTarget);
+					}
+				}
+			}
+		}
+		return false;
+	});
+	TreeDelegates.OnGetSelection.BindLambda([this, ConnectorKey]() -> TArray<FRigElementKey>
+	{
+		FRigElementKeyRedirector Redirector = Delegates.GetModularRig()->GetElementKeyRedirector();
+		if (const FRigElementKey* Key = Redirector.FindExternalKey(ConnectorKey))
+		{
+			return {*Key};
+		}
+		return {};
+	});
+	TreeDelegates.OnSelectionChanged.BindSP(this, &SModularRigModelItem::OnConnectorTargetChanged, ConnectorKey);
+
 
 	STableRow<TSharedPtr<FModularRigTreeElement>>::Construct(
 		STableRow<TSharedPtr<FModularRigTreeElement>>::FArguments()
@@ -106,52 +177,319 @@ void SModularRigModelItem::Construct(const FArguments& InArgs, const TSharedRef<
 		[
 			SAssignNew(HorizontalBox, SHorizontalBox)
 			+ SHorizontalBox::Slot()
-			.MaxWidth(18)
-			.FillWidth(1.0)
-			.HAlign(HAlign_Left)
-			.VAlign(VAlign_Center)
-			.Padding(FMargin(0.f, 0.f, 3.f, 0.f))
 			[
-				SNew(SImage)
-				.Image_Lambda([this]() -> const FSlateBrush*
-				{
-					if(WeakRigTreeElement.IsValid())
-					{
-						return WeakRigTreeElement.Pin()->IconBrush;
-					}
-					return nullptr;
-				})
-				.ColorAndOpacity_Lambda([this]()
-				{
-					if(WeakRigTreeElement.IsValid())
+				SNew(SHorizontalBox)
+				+SHorizontalBox::Slot()
+				.MaxWidth(25)
+				.HAlign(HAlign_Left)
+				.VAlign(VAlign_Center)
+				.Padding(FMargin(0.f, 0.f, 3.f, 0.f))
+				[
+					SNew(SVerticalBox)
+					+SVerticalBox::Slot()
+					.MaxHeight(25)
+					[
+						SNew(SImage)
+						.Image_Lambda([this]() -> const FSlateBrush*
+						{
+							if(WeakRigTreeElement.IsValid())
+							{
+								return WeakRigTreeElement.Pin()->IconBrush;
+							}
+							return nullptr;
+						})
+						.ColorAndOpacity_Lambda([this]()
+						{
+							if(WeakRigTreeElement.IsValid())
+							{
+								return WeakRigTreeElement.Pin()->IconColor;
+							}
+							return FSlateColor::UseForeground();
+						})
+					]
+				]
+
+				+ SHorizontalBox::Slot()
+				.FillWidth(1.0)
+				.VAlign(VAlign_Center)
+				[
+					SAssignNew(InlineWidget, SInlineEditableTextBlock)
+					.Text(this, &SModularRigModelItem::GetName, true)
+					.OnVerifyTextChanged(this, &SModularRigModelItem::OnVerifyNameChanged)
+					.OnTextCommitted(this, &SModularRigModelItem::OnNameCommitted)
+					.ToolTipText(this, &SModularRigModelItem::GetItemTooltip)
+					.MultiLine(false)
+					.ColorAndOpacity_Lambda([this]()
 					{
-						return WeakRigTreeElement.Pin()->IconColor;
-					}
-					return FSlateColor::UseForeground();
-				})
+						if(WeakRigTreeElement.IsValid())
+						{
+							return WeakRigTreeElement.Pin()->TextColor;
+						}
+						return FSlateColor::UseForeground();
+					})
+				]
 			]
+
 			+ SHorizontalBox::Slot()
+			.HAlign(HAlign_Right)
+			.VAlign(VAlign_Center)
+			[
+				SNew( SComboButton )
+				.ContentPadding(3)
+				.MenuPlacement(MenuPlacement_BelowAnchor)
+				.OnComboBoxOpened(this, &SModularRigModelItem::PopulateConnectorTargetList, ConnectorKey)
+				.ButtonContent()
+				[
+					// Wrap in configurable box to restrain height/width of menu
+					SNew(SBox)
+					.MinDesiredWidth(200.0f)
+					[
+						SAssignNew(ComboButtonBox, SVerticalBox)
+					]
+				]
+				.MenuContent()
+				[
+					SNew(SBorder)
+					.Visibility(EVisibility::Visible)
+					.BorderImage(FAppStyle::GetBrush("Menu.Background"))
+					[
+						SAssignNew(ConnectorComboBox, SSearchableRigHierarchyTreeView)
+							.RigTreeDelegates(TreeDelegates)
+					]
+				]
+			]
+
+			// Connection buttons
+			+SHorizontalBox::Slot()
 			.AutoWidth()
 			.VAlign(VAlign_Center)
+			.Padding(0.f, 0.f, 0.f, 0.f)
 			[
-				SAssignNew(InlineWidget, SInlineEditableTextBlock)
-				.Text(this, &SModularRigModelItem::GetName, true)
-				.OnVerifyTextChanged(this, &SModularRigModelItem::OnVerifyNameChanged)
-				.OnTextCommitted(this, &SModularRigModelItem::OnNameCommitted)
-				.ToolTipText(this, &SModularRigModelItem::GetItemTooltip)
-				.MultiLine(false)
-				.ColorAndOpacity_Lambda([this]()
-				{
-					if(WeakRigTreeElement.IsValid())
-					{
-						return WeakRigTreeElement.Pin()->TextColor;
-					}
-					return FSlateColor::UseForeground();
-				})
+				SNew(SVerticalBox)
+				+SVerticalBox::Slot()
+				.VAlign(VAlign_Center)
+				.AutoHeight()
+				[
+					// Reset button
+					SNew(SHorizontalBox)
+					+SHorizontalBox::Slot()
+					.AutoWidth()
+					.Padding(0.f, 0.f, 0.f, 0.f)
+					[
+						SAssignNew(ResetConnectorButton, SButton)
+						.ButtonStyle( FAppStyle::Get(), "NoBorder" )
+						.ButtonColorAndOpacity_Lambda([this, ConnectorKey]()
+						{
+							return ResetConnectorButton.IsValid() && ResetConnectorButton->IsHovered()
+								? FSlateColor(FLinearColor(1.f, 1.f, 1.f, 0.8))
+								: FSlateColor(FLinearColor(1.f, 1.f, 1.f, 0.4));
+						})
+						.OnClicked_Lambda([this, ConnectorKey]()
+						{
+							Delegates.HandleDisconnectConnector(ConnectorKey);
+							return FReply::Handled();
+						})
+						.ContentPadding(1.f)
+						.ToolTipText(NSLOCTEXT("ControlRigModuleDetails", "Reset_Connector", "Reset Connector"))
+						[
+							SNew(SImage)
+							.ColorAndOpacity_Lambda( [this, ConnectorKey]()
+							{
+								return ResetConnectorButton.IsValid() && ResetConnectorButton->IsHovered()
+								? FSlateColor(FLinearColor(1.f, 1.f, 1.f, 0.8))
+								: FSlateColor(FLinearColor(1.f, 1.f, 1.f, 0.4));
+							})
+							.Image(FSlateIcon(FAppStyle::Get().GetStyleSetName(), "PropertyWindow.DiffersFromDefault").GetIcon())
+						]
+					]
+
+					// Use button
+					+SHorizontalBox::Slot()
+					.AutoWidth()
+					.Padding(0.f, 0.f, 0.f, 0.f)
+					[
+						SAssignNew(UseSelectedButton, SButton)
+						.ButtonStyle( FAppStyle::Get(), "NoBorder" )
+						.ButtonColorAndOpacity_Lambda([this]()
+						{
+							return UseSelectedButton.IsValid() && UseSelectedButton->IsHovered()
+								? FSlateColor(FLinearColor(1.f, 1.f, 1.f, 0.8))
+								: FSlateColor(FLinearColor(1.f, 1.f, 1.f, 0.4));
+						})
+						.OnClicked_Lambda([this, ConnectorKey]()
+						{
+							if (const UModularRig* ModularRig = Delegates.GetModularRig())
+							{
+								const TArray<FRigElementKey>& Selected = ModularRig->GetHierarchy()->GetSelectedKeys();
+								if (Selected.Num() > 0)
+								{
+									Delegates.HandleResolveConnector(ConnectorKey, Selected[0]);
+								}
+							}
+							return FReply::Handled();
+						})
+						.ContentPadding(1.f)
+						.ToolTipText(NSLOCTEXT("ControlRigModuleDetails", "Use_Selected", "Use Selected"))
+						[
+							SNew(SImage)
+							.ColorAndOpacity_Lambda( [this]()
+							{
+								return UseSelectedButton.IsValid() && UseSelectedButton->IsHovered()
+								? FSlateColor(FLinearColor(1.f, 1.f, 1.f, 0.8))
+								: FSlateColor(FLinearColor(1.f, 1.f, 1.f, 0.4));
+							})
+							.Image(FAppStyle::GetBrush("Icons.CircleArrowLeft"))
+						]
+					]
+					
+					// Select in hierarchy button
+					+SHorizontalBox::Slot()
+					.AutoWidth()
+					.Padding(0.f, 0.f, 0.f, 0.f)
+					[
+						SAssignNew(SelectElementButton, SButton)
+						.ButtonStyle( FAppStyle::Get(), "NoBorder" )
+						.ButtonColorAndOpacity_Lambda([this]()
+						{
+							return SelectElementButton.IsValid() && SelectElementButton->IsHovered()
+								? FSlateColor(FLinearColor(1.f, 1.f, 1.f, 0.8))
+								: FSlateColor(FLinearColor(1.f, 1.f, 1.f, 0.4));
+						})
+						.OnClicked_Lambda([this, ConnectorKey]()
+						{
+							if (const UModularRig* ModularRig = Delegates.GetModularRig())
+							{
+								if (const FRigElementKey* TargetKey = ModularRig->GetElementKeyRedirector().FindExternalKey(ConnectorKey))
+								{
+									ModularRig->GetHierarchy()->GetController()->SelectElement(*TargetKey);
+								}
+							}
+							return FReply::Handled();
+						})
+						.ContentPadding(1.f)
+						.ToolTipText(NSLOCTEXT("ControlRigModuleDetails", "Select_Element", "Select Element"))
+						[
+							SNew(SImage)
+							.ColorAndOpacity_Lambda( [this]()
+							{
+								return SelectElementButton.IsValid() && SelectElementButton->IsHovered()
+								? FSlateColor(FLinearColor(1.f, 1.f, 1.f, 0.8))
+								: FSlateColor(FLinearColor(1.f, 1.f, 1.f, 0.4));
+							})
+							.Image(FAppStyle::GetBrush("Icons.Search"))
+						]
+					]
+				]
 			]
 		], OwnerTable);
 
 	InRigTreeElement->OnRenameRequested.BindSP(InlineWidget.Get(), &SInlineEditableTextBlock::EnterEditingMode);
+	FRigElementKey CurrentTargetKey;
+	if (const FRigElementKey* Key = Delegates.GetModularRig()->GetElementKeyRedirector().FindExternalKey(ConnectorKey))
+	{
+		CurrentTargetKey = *Key;
+	}
+	TPair<const FSlateBrush*, FSlateColor> IconAndColor = SRigHierarchyItem::GetBrushForElementType(Delegates.GetModularRig()->GetHierarchy(), CurrentTargetKey);
+	PopulateConnectorCurrentTarget(ComboButtonBox, ConnectorKey, CurrentTargetKey, IconAndColor.Key, IconAndColor.Value, FText::FromName(CurrentTargetKey.Name));
+}
+
+void SModularRigModelItem::PopulateConnectorTargetList(const FRigElementKey InConnectorKey)
+{
+	if (!WeakRigTreeElement.IsValid())
+	{
+		return;
+	}
+
+	if (const UModularRig* ModularRig = Delegates.GetModularRig())
+	{
+		if (URigHierarchy* Hierarchy = ModularRig->GetHierarchy())
+		{
+			if (FRigConnectorElement* ConnectorElement = Cast<FRigConnectorElement>(Hierarchy->Find(InConnectorKey)))
+			{
+				ConnectorComboBox->GetTreeView()->RefreshTreeView(true);
+			}
+		}
+	}
+}
+
+void SModularRigModelItem::PopulateConnectorCurrentTarget(TSharedPtr<SVerticalBox> InListBox, const FRigElementKey& InConnectorKey, const FRigElementKey& InTargetKey, const FSlateBrush* InBrush, const FSlateColor& InColor, const FText& InTitle)
+{
+	static const FSlateBrush* RoundedBoxBrush = FControlRigEditorStyle::Get().GetBrush(TEXT("ControlRig.SpacePicker.RoundedRect"));
+
+	TSharedPtr<SHorizontalBox> RowBox, ButtonBox;
+	InListBox->AddSlot()
+	.AutoHeight()
+	.VAlign(VAlign_Top)
+	.HAlign(HAlign_Fill)
+	.Padding(4.0, 0.0, 4.0, 0.0)
+	[
+		SNew( SButton )
+		.ButtonStyle(FAppStyle::Get(), "SimpleButton")
+		.ContentPadding(FMargin(0.0))
+		.OnClicked_Lambda([this, InConnectorKey, InTargetKey]()
+		{
+			//Controller->ConnectConnectorToElement(InConnectorKey, InTargetKey, true, Info.GetModularRig()->GetModularRigSettings().bAutoResolve);
+			PopulateConnectorTargetList(InConnectorKey);
+			return FReply::Handled();
+		})
+		[
+			SAssignNew(RowBox, SHorizontalBox)
+			+ SHorizontalBox::Slot()
+			.AutoWidth()
+			.VAlign(VAlign_Center)
+			.HAlign(HAlign_Fill)
+			.Padding(0)
+			[
+				SNew(SBorder)
+				.Padding(FMargin(5.0, 2.0, 5.0, 2.0))
+				.BorderImage(RoundedBoxBrush)
+				.BorderBackgroundColor(FSlateColor(FLinearColor::Transparent))
+				.Content()
+				[
+					SAssignNew(ButtonBox, SHorizontalBox)
+					+ SHorizontalBox::Slot()
+					.AutoWidth()
+					.VAlign(VAlign_Center)
+					.HAlign(HAlign_Left)
+					.Padding(FMargin(0.f, 0.f, 3.f, 0.f))
+					[
+						SNew(SImage)
+						.Image(InBrush)
+						.ColorAndOpacity(InColor)
+					]
+
+					+ SHorizontalBox::Slot()
+					.AutoWidth()
+					.VAlign(VAlign_Center)
+					.HAlign(HAlign_Left)
+					.Padding(0)
+					[
+						SNew( STextBlock )
+						.Text( InTitle )
+						//.Font( IDetailLayoutBuilder::GetDetailFont() )
+					]
+
+					+ SHorizontalBox::Slot()
+					.FillWidth(1.f)
+					[
+						SNew(SSpacer)
+					]
+				]
+			]
+		]
+	];
+}
+
+void SModularRigModelItem::OnConnectorTargetChanged(TSharedPtr<FRigTreeElement> Selection, ESelectInfo::Type SelectInfo, const FRigElementKey InConnectorKey)
+{
+	if (!Selection.IsValid() || SelectInfo == ESelectInfo::OnNavigation)
+	{
+		return;
+	}
+	
+	FScopedTransaction Transaction(LOCTEXT("ModuleHierarchyResolveConnector", "Resolve Connector"));
+	TSharedPtr<FRigTreeElement> NewSelection = Selection;
+	Delegates.HandleResolveConnector(InConnectorKey, NewSelection->Key);
 }
 
 void SModularRigModelItem::OnNameCommitted(const FText& InText, ETextCommit::Type InCommitType) const
@@ -161,7 +499,7 @@ void SModularRigModelItem::OnNameCommitted(const FText& InText, ETextCommit::Typ
 	if (InCommitType == ETextCommit::OnEnter)
 	{
 		FString NewName = InText.ToString();
-		const FString OldKey = WeakRigTreeElement.Pin()->Key;
+		const FString OldKey = WeakRigTreeElement.Pin()->ModulePath;
 
 		Delegates.HandleRenameElement(OldKey, *NewName);
 	}
@@ -170,7 +508,7 @@ void SModularRigModelItem::OnNameCommitted(const FText& InText, ETextCommit::Typ
 bool SModularRigModelItem::OnVerifyNameChanged(const FText& InText, FText& OutErrorMessage)
 {
 	const FName NewName = *InText.ToString();
-	const FString OldPath = WeakRigTreeElement.Pin()->Key;
+	const FString OldPath = WeakRigTreeElement.Pin()->ModulePath;
 	return Delegates.HandleVerifyElementNameChanged(OldPath, NewName, OutErrorMessage);
 }
 
@@ -180,7 +518,7 @@ FText SModularRigModelItem::GetName(bool bUseShortName) const
 	{
 		return (FText::FromName(WeakRigTreeElement.Pin()->ShortName));
 	}
-	return (FText::FromString(WeakRigTreeElement.Pin()->Key));
+	return (FText::FromString(WeakRigTreeElement.Pin()->ModulePath));
 }
 
 FText SModularRigModelItem::GetItemTooltip() const
@@ -338,31 +676,46 @@ bool SModularRigTreeView::AddElement(FString InKey, FString InParentKey)
 
 	if (!InKey.IsEmpty())
 	{
-		TSharedPtr<FModularRigTreeElement> NewItem = MakeShared<FModularRigTreeElement>(InKey, SharedThis(this), false);
+		TSharedPtr<FModularRigTreeElement> NewItem = MakeShared<FModularRigTreeElement>(InKey, SharedThis(this), true);
 
-		if (!InKey.IsEmpty())
+		ElementMap.Add(InKey, NewItem);
+		if (!InParentKey.IsEmpty())
 		{
-			ElementMap.Add(InKey, NewItem);
-			if (!InParentKey.IsEmpty())
-			{
-				ParentMap.Add(InKey, InParentKey);
-			}
+			ParentMap.Add(InKey, InParentKey);
 
-			if (!InParentKey.IsEmpty())
-			{
-				TSharedPtr<FModularRigTreeElement>* FoundItem = ElementMap.Find(InParentKey);
-				check(FoundItem);
-				FoundItem->Get()->Children.Add(NewItem);
-			}
-			else
-			{
-				RootElements.Add(NewItem);
-			}
+			TSharedPtr<FModularRigTreeElement>* FoundItem = ElementMap.Find(InParentKey);
+			check(FoundItem);
+			FoundItem->Get()->Children.Add(NewItem);
 		}
 		else
 		{
 			RootElements.Add(NewItem);
 		}
+
+		if (const UModularRig* ModularRig = Delegates.GetModularRig())
+		{
+			if (const FRigModuleInstance* Module = ModularRig->FindModule(NewItem->ModulePath))
+			{
+				if (const UControlRig* ModuleRig = Module->GetRig())
+				{
+					const UControlRig* CDO = ModuleRig->GetClass()->GetDefaultObject<UControlRig>();
+					const TArray<FRigModuleConnector>& Connectors = CDO->GetRigModuleSettings().ExposedConnectors;
+					for (const FRigModuleConnector& Connector : Connectors)
+					{
+						if (Connector.Settings.Type == EConnectorType::Primary)
+						{
+							continue;
+						}
+
+						const FString Key = FString::Printf(TEXT("%s:%s"), *NewItem->ModulePath, *Connector.Name);
+						TSharedPtr<FModularRigTreeElement> ConnectorItem = MakeShared<FModularRigTreeElement>(Key, SharedThis(this), false);
+						NewItem.Get()->Children.Add(ConnectorItem);
+						ElementMap.Add(Key, ConnectorItem);
+						ParentMap.Add(Key, InKey);
+					}
+				}
+			}
+		}
 	}
 
 	return true;
@@ -379,7 +732,8 @@ bool SModularRigTreeView::AddElement(const FRigModuleInstance* InElement)
 
 	const UModularRig* ModularRig = Delegates.GetModularRig();
 
-	if(!AddElement(InElement->GetPath(), FString()))
+	FString ElementPath = InElement->GetPath();
+	if(!AddElement(ElementPath, FString()))
 	{
 		return false;
 	}
@@ -388,16 +742,16 @@ bool SModularRigTreeView::AddElement(const FRigModuleInstance* InElement)
 	{
 		if(ModularRig)
 		{
-			FString ParentKey = ModularRig->GetParentPath(InElement->GetPath());
-			if (!ParentKey.IsEmpty())
+			FString ParentPath = ModularRig->GetParentPath(ElementPath);
+			if (!ParentPath.IsEmpty())
 			{
-				if(const FRigModuleInstance* ParentElement = ModularRig->FindModule(ParentKey))
+				if(const FRigModuleInstance* ParentElement = ModularRig->FindModule(ParentPath))
 				{
 					AddElement(ParentElement);
 
-					if(ElementMap.Contains(ParentKey))
+					if(ElementMap.Contains(ParentPath))
 					{
-						ReparentElement(InElement->GetPath(), ParentKey);
+						ReparentElement(ElementPath, ParentPath);
 					}
 				}
 			}
@@ -596,22 +950,63 @@ const TSharedPtr<FModularRigTreeElement>* SModularRigTreeView::FindItemAtPositio
 	return nullptr;
 }
 
-TPair<const FSlateBrush*, FSlateColor> SModularRigModelItem::GetBrushForElementType(const UModularRig* InModularRig, const FString& InKey)
+TPair<const FSlateBrush*, FSlateColor> FModularRigTreeElement::GetBrushAndColor(const UModularRig* InModularRig)
 {
 	const FSlateBrush* Brush = nullptr;
 	FSlateColor Color = FSlateColor::UseForeground();
 
-	if (const FRigModuleInstance* Module = InModularRig->FindModule(InKey))
+	if (const FRigModuleInstance* ConnectorModule = InModularRig->FindModule(ModulePath))
 	{
-		if (UControlRig* ModuleRig = Module->GetRig())
+		bool bIsConnected = false;
+		if (const FRigModuleReference* Reference = ConnectorModule->GetModuleReference())
 		{
-			FSoftObjectPath IconPath = ModuleRig->GetRigModuleSettings().Icon;
-			Brush = IconPathToBrush.Find(IconPath);
-			if (!Brush)
+			FRigElementKey ConnectorKey(*ConnectorName, ERigElementType::Connector);
+			const FRigElementKey* Target = Reference->Connections.Find(ConnectorKey);
+
+			if (!Target || !Target->IsValid())
 			{
-				if(UTexture2D* Icon = Cast<UTexture2D>(IconPath.TryLoad()))
+				Color = FStyleColors::AccentYellow;
+			}
+			else
+			{
+				bIsConnected = true;
+			}
+		}
+		
+		if (const UControlRig* ModuleRig = ConnectorModule->GetRig())
+		{
+			const FRigModuleConnector* Connector = ModuleRig->GetRigModuleSettings().ExposedConnectors.FindByPredicate([this](FRigModuleConnector& Connector)
+			{
+				return Connector.Name == ConnectorName;
+			});
+			if (Connector)
+			{
+				if (Connector->Settings.Type == EConnectorType::Primary)
+				{
+					if (bIsConnected)
+					{
+						FSoftObjectPath IconPath = ModuleRig->GetRigModuleSettings().Icon;
+						Brush = IconPathToBrush.Find(IconPath);
+						if (!Brush)
+						{
+							if(UTexture2D* Icon = Cast<UTexture2D>(IconPath.TryLoad()))
+							{
+								Brush = &IconPathToBrush.Add(IconPath, UWidgetBlueprintLibrary::MakeBrushFromTexture(Icon, 16.0f, 16.0f));
+							}
+						}
+					}
+					else
+					{
+						Brush = FControlRigEditorStyle::Get().GetBrush("ControlRig.Schematic.ConnectorWarning");
+					}
+				}
+				else if (Connector->Settings.bOptional)
+				{
+					Brush = FControlRigEditorStyle::Get().GetBrush("ControlRig.Schematic.ConnectorOptional");
+				}
+				else
 				{
-					Brush = &IconPathToBrush.Add(IconPath, UWidgetBlueprintLibrary::MakeBrushFromTexture(Icon, 16.0f, 16.0f));
+					Brush = FControlRigEditorStyle::Get().GetBrush("ControlRig.Schematic.ConnectorSecondary");
 				}
 			}
 		}
@@ -623,67 +1018,6 @@ TPair<const FSlateBrush*, FSlateColor> SModularRigModelItem::GetBrushForElementT
 	return TPair<const FSlateBrush*, FSlateColor>(Brush, Color);
 }
 
-FLinearColor SModularRigModelItem::GetColorForControlType(ERigControlType InControlType, UEnum* InControlEnum)
-{
-	FEdGraphPinType PinType;
-	switch(InControlType)
-	{
-		case ERigControlType::Bool:
-		{
-			PinType = RigVMTypeUtils::PinTypeFromCPPType(RigVMTypeUtils::BoolTypeName, nullptr);
-			break;
-		}
-		case ERigControlType::Float:
-		case ERigControlType::ScaleFloat:
-		{
-			PinType = RigVMTypeUtils::PinTypeFromCPPType(RigVMTypeUtils::FloatTypeName, nullptr);
-			break;
-		}
-		case ERigControlType::Integer:
-		{
-			if(InControlEnum)
-			{
-				PinType = RigVMTypeUtils::PinTypeFromCPPType(NAME_None, InControlEnum);
-			}
-			else
-			{
-				PinType = RigVMTypeUtils::PinTypeFromCPPType(RigVMTypeUtils::Int32TypeName, nullptr);
-			}
-			break;
-		}
-		case ERigControlType::Vector2D:
-		{
-			UScriptStruct* Struct = TBaseStructure<FVector2D>::Get(); 
-			PinType = RigVMTypeUtils::PinTypeFromCPPType(*RigVMTypeUtils::GetUniqueStructTypeName(Struct), Struct);
-			break;
-		}
-		case ERigControlType::Position:
-		case ERigControlType::Scale:
-		{
-			UScriptStruct* Struct = TBaseStructure<FVector>::Get(); 
-			PinType = RigVMTypeUtils::PinTypeFromCPPType(*RigVMTypeUtils::GetUniqueStructTypeName(Struct), Struct);
-			break;
-		}
-		case ERigControlType::Rotator:
-		{
-			UScriptStruct* Struct = TBaseStructure<FRotator>::Get(); 
-			PinType = RigVMTypeUtils::PinTypeFromCPPType(*RigVMTypeUtils::GetUniqueStructTypeName(Struct), Struct);
-			break;
-		}
-		case ERigControlType::Transform:
-		case ERigControlType::TransformNoScale:
-		case ERigControlType::EulerTransform:
-		default:
-		{
-			UScriptStruct* Struct = TBaseStructure<FTransform>::Get(); 
-			PinType = RigVMTypeUtils::PinTypeFromCPPType(*RigVMTypeUtils::GetUniqueStructTypeName(Struct), Struct);
-			break;
-		}
-	}
-	const UControlRigGraphSchema* Schema = GetDefault<UControlRigGraphSchema>();
-	return Schema->GetPinTypeColor(PinType);
-}
-
 //////////////////////////////////////////////////////////////
 /// SSearchableModularRigTreeView
 ///////////////////////////////////////////////////////////
diff --git a/Engine/Plugins/Animation/ControlRig/Source/ControlRigEditor/Private/Editor/SModularRigTreeView.h b/Engine/Plugins/Animation/ControlRig/Source/ControlRigEditor/Private/Editor/SModularRigTreeView.h
--- a/Engine/Plugins/Animation/ControlRig/Source/ControlRigEditor/Private/Editor/SModularRigTreeView.h
+++ b/Engine/Plugins/Animation/ControlRig/Source/ControlRigEditor/Private/Editor/SModularRigTreeView.h
@@ -5,6 +5,7 @@
 #include "CoreMinimal.h"
 #include "Widgets/Views/STreeView.h"
 #include "ModularRig.h"
+#include "Editor/SRigHierarchyTreeView.h"
 
 class SSearchBox;
 class SModularRigTreeView;
@@ -15,6 +16,8 @@ class FModularRigTreeElement;
 DECLARE_DELEGATE_RetVal(const UModularRig*, FOnGetModularRigTreeRig);
 DECLARE_DELEGATE_OneParam(FOnModularRigTreeRequestDetailsInspection, const FString&);
 DECLARE_DELEGATE_RetVal_TwoParams(FName, FOnModularRigTreeRenameElement, const FString& /*OldPath*/, const FName& /*NewName*/);
+DECLARE_DELEGATE_TwoParams(FOnModularRigTreeResolveConnector, const FRigElementKey& /*Connector*/, const FRigElementKey& /*Target*/);
+DECLARE_DELEGATE_OneParam(FOnModularRigTreeDisconnectConnector, const FRigElementKey& /*Connector*/);
 DECLARE_DELEGATE_RetVal_ThreeParams(bool, FOnModularRigTreeVerifyElementNameChanged, const FString& /*OldPath*/, const FName& /*NewName*/, FText& /*OutErrorMessage*/);
 
 
@@ -35,6 +38,8 @@ struct CONTROLRIGEDITOR_API FModularRigTreeDelegates
 	FOnModularRigTreeRequestDetailsInspection OnRequestDetailsInspection;
 	FOnModularRigTreeRenameElement OnRenameElement;
 	FOnModularRigTreeVerifyElementNameChanged OnVerifyModuleNameChanged;
+	FOnModularRigTreeResolveConnector OnResolveConnector;
+	FOnModularRigTreeDisconnectConnector OnDisconnectConnector;
 	
 	FModularRigTreeDelegates()
 	{
@@ -66,6 +71,26 @@ struct CONTROLRIGEDITOR_API FModularRigTreeDelegates
 		}
 		return false;
 	}
+
+	bool HandleResolveConnector(const FRigElementKey& InConnector, const FRigElementKey& InTarget)
+	{
+		if(OnResolveConnector.IsBound())
+		{
+			OnResolveConnector.Execute(InConnector, InTarget);
+			return true;
+		}
+		return false;
+	}
+
+	bool HandleDisconnectConnector(const FRigElementKey& InConnector)
+	{
+		if(OnDisconnectConnector.IsBound())
+		{
+			OnDisconnectConnector.Execute(InConnector);
+			return true;
+		}
+		return false;
+	}
 };
 
 
@@ -73,11 +98,14 @@ struct CONTROLRIGEDITOR_API FModularRigTreeDelegates
 class FModularRigTreeElement : public TSharedFromThis<FModularRigTreeElement>
 {
 public:
-	FModularRigTreeElement(const FString& InKey, TWeakPtr<SModularRigTreeView> InTreeView, bool InSupportsRename);
+	FModularRigTreeElement(const FString& InKey, TWeakPtr<SModularRigTreeView> InTreeView, bool bInIsPrimary);
 
 public:
 	/** Element Data to display */
 	FString Key;
+	bool bIsPrimary;
+	FString ModulePath;
+	FString ConnectorName;
 	FName ShortName;
 	TArray<TSharedPtr<FModularRigTreeElement>> Children;
 
@@ -87,10 +115,14 @@ public:
 
 	void RefreshDisplaySettings(const UModularRig* InModularRig);
 
+	TPair<const FSlateBrush*, FSlateColor> GetBrushAndColor(const UModularRig* InModularRig);
+
 	/** Delegate for when the context menu requests a rename */
 	DECLARE_DELEGATE(FOnRenameRequested);
 	FOnRenameRequested OnRenameRequested;
 
+	static TMap<FSoftObjectPath, FSlateBrush> IconPathToBrush;
+
 	/** The brush to use when rendering an icon */
 	const FSlateBrush* IconBrush;
 
@@ -106,21 +138,30 @@ class SModularRigModelItem : public STableRow<TSharedPtr<FModularRigTreeElement>
 public:
 	
 	void Construct(const FArguments& InArgs, const TSharedRef<STableViewBase>& OwnerTable, TSharedRef<FModularRigTreeElement> InRigTreeElement, TSharedPtr<SModularRigTreeView> InTreeView, bool bPinned);
+	void PopulateConnectorTargetList(FRigElementKey InConnectorKey);
+	void PopulateConnectorCurrentTarget(
+		TSharedPtr<SVerticalBox> InListBox, 
+		const FRigElementKey& InConnectorKey,
+		const FRigElementKey& InTargetKey,
+		const FSlateBrush* InBrush,
+		const FSlateColor& InColor,
+		const FText& InTitle);
+	void OnConnectorTargetChanged(TSharedPtr<FRigTreeElement> Selection, ESelectInfo::Type SelectInfo, const FRigElementKey InConnectorKey);
 
 	void OnNameCommitted(const FText& InText, ETextCommit::Type InCommitType) const;
 	bool OnVerifyNameChanged(const FText& InText, FText& OutErrorMessage);
-	static TPair<const FSlateBrush*, FSlateColor> GetBrushForElementType(const UModularRig* InModularRig, const FString& InKey);
-	static FLinearColor GetColorForControlType(ERigControlType InControlType, UEnum* InControlEnum);
 
 private:
 	TWeakPtr<FModularRigTreeElement> WeakRigTreeElement;
  	FModularRigTreeDelegates Delegates;
+	TSharedPtr<SSearchableRigHierarchyTreeView> ConnectorComboBox;
+	TSharedPtr<SButton> ResetConnectorButton;
+	TSharedPtr<SButton> UseSelectedButton;
+	TSharedPtr<SButton> SelectElementButton;
 
 	FText GetName(bool bUseShortName) const;
 	FText GetItemTooltip() const;
 
-	static TMap<FSoftObjectPath, FSlateBrush> IconPathToBrush;
-
 	friend class SModularRigTreeView; 
 };
 
