Control Rig: Enable schematic by default, add hotkey

#rnx

diff --git a/Engine/Plugins/Animation/ControlRig/Source/ControlRigEditor/Private/ControlRigEditorStyle.h b/Engine/Plugins/Animation/ControlRig/Source/ControlRigEditor/Private/ControlRigEditorStyle.h
--- a/Engine/Plugins/Animation/ControlRig/Source/ControlRigEditor/Private/ControlRigEditorStyle.h
+++ b/Engine/Plugins/Animation/ControlRig/Source/ControlRigEditor/Private/ControlRigEditorStyle.h
@@ -146,12 +146,14 @@ public:
 			Set("ControlRig.TestData.Record", new IMAGE_BRUSH("Slate/RecordingIndicator", Icon32x32));
 		}
 
+		Set("ControlRig.ConnectorWarning", new IMAGE_BRUSH_SVG("Slate/Connector_Warning", Icon128x128));
+
 		// Schematic
 		{
 			Set("ControlRig.Schematic.ConnectorPrimary", new IMAGE_BRUSH_SVG("Slate/Connector_Primary", Icon128x128));
 			Set("ControlRig.Schematic.ConnectorSecondary", new IMAGE_BRUSH_SVG("Slate/Connector_Secondary", Icon128x128));
 			Set("ControlRig.Schematic.ConnectorOptional", new IMAGE_BRUSH_SVG("Slate/Connector_Optional", Icon128x128));
-			Set("ControlRig.Schematic.ConnectorWarning", new IMAGE_BRUSH_SVG("Slate/Connector_Warning", Icon128x128));
+			Set("ControlRig.Schematic.ConnectorWarning", new IMAGE_BRUSH_SVG("Slate/Connector_Warning_Schematic", Icon128x128));
 			Set("ControlRig.Schematic.Bone", new IMAGE_BRUSH_SVG("Slate/Bone_Schematic", Icon128x128));
 			Set("ControlRig.Schematic.Control", new IMAGE_BRUSH_SVG("Slate/Control_Schematic", Icon128x128));
 			Set("ControlRig.Schematic.Null", new IMAGE_BRUSH_SVG("Slate/Null_Schematic", Icon128x128));
diff --git a/Engine/Plugins/Animation/ControlRig/Source/ControlRigEditor/Private/Editor/ControlRigEditor.cpp b/Engine/Plugins/Animation/ControlRig/Source/ControlRigEditor/Private/Editor/ControlRigEditor.cpp
--- a/Engine/Plugins/Animation/ControlRig/Source/ControlRigEditor/Private/Editor/ControlRigEditor.cpp
+++ b/Engine/Plugins/Animation/ControlRig/Source/ControlRigEditor/Private/Editor/ControlRigEditor.cpp
@@ -114,7 +114,7 @@
 #define LOCTEXT_NAMESPACE "ControlRigEditor"
 
 TAutoConsoleVariable<bool> CVarControlRigShowTestingToolbar(TEXT("ControlRig.Test.EnableTestingToolbar"), false, TEXT("When true we'll show the testing toolbar in Control Rig Editor."));
-TAutoConsoleVariable<bool> CVarShowSchematicPanelOverlay(TEXT("ControlRig.Preview.ShowSchematicPanelOverlay"), false, TEXT("When true we'll add an overlay to the persona viewport to show modular rig information."));
+TAutoConsoleVariable<bool> CVarShowSchematicPanelOverlay(TEXT("ControlRig.Preview.ShowSchematicPanelOverlay"), true, TEXT("When true we'll add an overlay to the persona viewport to show modular rig information."));
 
 const FName FControlRigEditorModes::ControlRigEditorMode = TEXT("Rigging");
 const TArray<FName> FControlRigEditor::ForwardsSolveEventQueue = {FRigUnit_BeginExecution::EventName};
@@ -341,6 +341,11 @@ void FControlRigEditor::InitRigVMEditor(const EToolkitMode::Type Mode, const TSh
 	}
 
 	CreateRigHierarchyToGraphDragAndDropMenu();
+
+	if(SchematicViewport.IsValid())
+	{
+		SchematicModel.UpdateControlRigContent();
+	}
 }
 
 void FControlRigEditor::CreatePersonaToolKitIfRequired()
@@ -2285,7 +2290,15 @@ void FControlRigEditor::HandleViewportCreated(const TSharedRef<class IPersonaVie
 	InViewport->GetKeyDownDelegate().BindLambda([&](const FGeometry& MyGeometry, const FKeyEvent& InKeyEvent) -> FReply {
 		if (OnKeyDownDelegate.IsBound())
 		{
-			return OnKeyDownDelegate.Execute(MyGeometry, InKeyEvent);
+			FReply Reply = OnKeyDownDelegate.Execute(MyGeometry, InKeyEvent);
+			if(Reply.IsEventHandled())
+			{
+				return Reply;
+			}
+		}
+		if(GetToolkitCommands()->ProcessCommandBindings(InKeyEvent.GetKey(), InKeyEvent.GetModifierKeys(), false))
+		{
+			return FReply::Handled();
 		}
 		return FReply::Unhandled();
 	});
@@ -2332,6 +2345,20 @@ void FControlRigEditor::OnToolbarAxesScaleChanged(float InValue)
 	}
 }
 
+void FControlRigEditor::HandleToggleSchematicViewport()
+{
+	if(SchematicViewport.IsValid())
+	{
+		SchematicModel.UpdateControlRigContent();
+
+		const EVisibility PreviousVisibility = SchematicViewport.Pin()->GetVisibility();
+		SchematicViewport.Pin()->SetVisibility(
+			PreviousVisibility == EVisibility::Hidden ?
+			EVisibility::SelfHitTestInvisible :
+			EVisibility::Hidden);
+	}
+}
+
 ECheckBoxState FControlRigEditor::GetToolbarDrawAxesOnSelection() const
 {
 	if (const UControlRigEditModeSettings* Settings = GetDefault<UControlRigEditModeSettings>())
@@ -2571,6 +2598,7 @@ void FControlRigEditor::CacheNameLists()
 
 void FControlRigEditor::HandleSchematicViewportCreated(const TSharedRef<SSchematicGraphPanel>& InViewport)
 {
+	SchematicViewport = InViewport.ToWeakPtr();
 	InViewport->OnNodeClicked().BindRaw(&SchematicModel, &FControlRigSchematicModel::HandleSchematicNodeClicked);
 	InViewport->OnBeginDrag().BindRaw(&SchematicModel, &FControlRigSchematicModel::HandleSchematicBeginDrag);
 	InViewport->OnEndDrag().BindRaw(&SchematicModel, &FControlRigSchematicModel::HandleSchematicEndDrag);
@@ -3315,6 +3343,11 @@ void FControlRigEditor::BindCommands()
 		FControlRigEditorCommands::Get().BackwardsAndForwardsSolveEvent,
 		FExecuteAction::CreateSP(this, &FRigVMEditor::SetEventQueue, TArray<FName>(BackwardsAndForwardsSolveEventQueue)),
 		FCanExecuteAction());
+
+	GetToolkitCommands()->MapAction(
+		FControlRigEditorCommands::Get().ToggleSchematicViewportVisibility,
+		FExecuteAction::CreateSP(this, &FControlRigEditor::HandleToggleSchematicViewport),
+		FCanExecuteAction());
 }
 
 void FControlRigEditor::OnHierarchyChanged()
diff --git a/Engine/Plugins/Animation/ControlRig/Source/ControlRigEditor/Private/Editor/ControlRigEditor.h b/Engine/Plugins/Animation/ControlRig/Source/ControlRigEditor/Private/Editor/ControlRigEditor.h
--- a/Engine/Plugins/Animation/ControlRig/Source/ControlRigEditor/Private/Editor/ControlRigEditor.h
+++ b/Engine/Plugins/Animation/ControlRig/Source/ControlRigEditor/Private/Editor/ControlRigEditor.h
@@ -244,6 +244,7 @@ private:
 	void OnToolbarDrawAxesOnSelectionChanged(ECheckBoxState InNewValue);
 	TOptional<float> GetToolbarAxesScale() const;
 	void OnToolbarAxesScaleChanged(float InValue);
+	void HandleToggleSchematicViewport();
 
 	/** Handle schematic panel setup*/
 	void HandleSchematicViewportCreated(const TSharedRef<class SSchematicGraphPanel>& InViewport);
@@ -354,6 +355,7 @@ protected:
 	TSharedPtr<SComboBox<TSharedPtr<FString>>> DirectManipulationCombo;
 	bool bRefreshDirectionManipulationTargetsRequired;
 	FSimpleMulticastDelegate RequestNavigateToConnectorWarningDelegate;
+	TWeakPtr<SSchematicGraphPanel> SchematicViewport;
 
 	static const TArray<FName> ForwardsSolveEventQueue;
 	static const TArray<FName> BackwardsSolveEventQueue;
diff --git a/Engine/Plugins/Animation/ControlRig/Source/ControlRigEditor/Private/Editor/ControlRigEditorCommands.cpp b/Engine/Plugins/Animation/ControlRig/Source/ControlRigEditor/Private/Editor/ControlRigEditorCommands.cpp
--- a/Engine/Plugins/Animation/ControlRig/Source/ControlRigEditor/Private/Editor/ControlRigEditorCommands.cpp
+++ b/Engine/Plugins/Animation/ControlRig/Source/ControlRigEditor/Private/Editor/ControlRigEditorCommands.cpp
@@ -13,6 +13,7 @@ void FControlRigEditorCommands::RegisterCommands()
 	UI_COMMAND(RequestDirectManipulationPosition, "Request Direct Manipulation for Position", "Request per node direct manipulation on a position", EUserInterfaceActionType::Button, FInputChord(EKeys::W));
 	UI_COMMAND(RequestDirectManipulationRotation, "Request Direct Manipulation for Rotation", "Request per node direct manipulation on a rotation", EUserInterfaceActionType::Button, FInputChord(EKeys::E));
 	UI_COMMAND(RequestDirectManipulationScale, "Request Direct Manipulation for Scale", "Request per node direct manipulation on a scale", EUserInterfaceActionType::Button, FInputChord(EKeys::R));
+	UI_COMMAND(ToggleSchematicViewportVisibility, "Toggle the viewport schematic", "Toggles the visibility of the viewport schematic", EUserInterfaceActionType::Button, FInputChord(EKeys::Y));
 }
 
 #undef LOCTEXT_NAMESPACE
diff --git a/Engine/Plugins/Animation/ControlRig/Source/ControlRigEditor/Private/Editor/ControlRigEditorCommands.h b/Engine/Plugins/Animation/ControlRig/Source/ControlRigEditor/Private/Editor/ControlRigEditorCommands.h
--- a/Engine/Plugins/Animation/ControlRig/Source/ControlRigEditor/Private/Editor/ControlRigEditorCommands.h
+++ b/Engine/Plugins/Animation/ControlRig/Source/ControlRigEditor/Private/Editor/ControlRigEditorCommands.h
@@ -39,6 +39,9 @@ public:
 	/** Request per node direct manipulation on a scale */
 	TSharedPtr< FUICommandInfo > RequestDirectManipulationScale;
 
+	/** Toggle visibility of the schematic */
+	TSharedPtr< FUICommandInfo > ToggleSchematicViewportVisibility;
+
 	/**
 	 * Initialize commands
 	 */
diff --git a/Engine/Plugins/Animation/ControlRig/Source/ControlRigEditor/Private/Editor/ControlRigSchematicModel.cpp b/Engine/Plugins/Animation/ControlRig/Source/ControlRigEditor/Private/Editor/ControlRigSchematicModel.cpp
--- a/Engine/Plugins/Animation/ControlRig/Source/ControlRigEditor/Private/Editor/ControlRigSchematicModel.cpp
+++ b/Engine/Plugins/Animation/ControlRig/Source/ControlRigEditor/Private/Editor/ControlRigSchematicModel.cpp
@@ -12,6 +12,7 @@
 #include "Rigs/RigHierarchyController.h"
 #include "Async/TaskGraphInterfaces.h"
 #include <SchematicGraphPanel/SchematicGraphStyle.h>
+#include "Editor/SModularRigModel.h"
 
 #define LOCTEXT_NAMESPACE "ControlRigSchematicModel"
 
@@ -25,6 +26,23 @@ bool FControlRigSchematicRigElementKeyNode::IsDragSupported() const
 	return Key.Type == ERigElementType::Connector;
 }
 
+FControlRigSchematicConnectorTag::FControlRigSchematicConnectorTag()
+: FSchematicGraphTag()
+{
+	BackgroundColor = FLinearColor::Black;
+	ForegroundColor = FLinearColor::White;
+}
+
+FControlRigSchematicWarningTag::FControlRigSchematicWarningTag()
+: FSchematicGraphTag()
+{
+	BackgroundColor = FLinearColor::Black;
+	ForegroundColor = FLinearColor::Yellow;
+	
+	static const FSlateBrush* WarningBrush = FControlRigEditorStyle::Get().GetBrush("ControlRig.Schematic.ConnectorWarning");
+	ForegroundBrush = WarningBrush;
+}
+
 FControlRigSchematicModel::~FControlRigSchematicModel()
 {
 	if(ControlRigBeingDebuggedPtr.IsValid())
@@ -61,41 +79,15 @@ FControlRigSchematicRigElementKeyNode* FControlRigSchematicModel::AddElementKeyN
 	FControlRigSchematicRigElementKeyNode* Node = AddNode<FControlRigSchematicRigElementKeyNode>(false);
 	if(Node)
 	{
-		Node->Key = InKey;
+		ConfigureElementKeyNode(Node, InKey);
+		RigElementKeyToGuid.Add(Node->GetKey(), Node->GetGuid());
 
-		Node->SetLabel(FText::FromString(Node->GetKey().ToString())); 
-		switch(Node->GetKey().Type)
+		if(InKey.Type == ERigElementType::Connector)
 		{
-			case ERigElementType::Socket:
-			{
-				if (ControlRigBeingDebuggedPtr.IsValid())
-				{
-					if (const URigHierarchy* Hierarchy = ControlRigBeingDebuggedPtr->GetHierarchy())
-					{
-						if(const FRigSocketElement* Socket = Hierarchy->Find<FRigSocketElement>(Node->GetKey()))
-						{
-							const FString Description = Socket->GetDescription(Hierarchy);
-							if(!Description.IsEmpty())
-							{
-								Node->SetLabel(
-									FText::Format(LOCTEXT("NodeLabelFormat", "{0}\n{1}"), 
-									Node->GetLabel(), 
-									FText::FromString(Description))
-								);
-							}
-						}
-					}
-				}
-				break;
-			}
-			default:
-			{
-				break;
-			}
+			Node->AddTag<FControlRigSchematicConnectorTag>();
+			Node->AddTag<FControlRigSchematicWarningTag>();
 		}
 
-		RigElementKeyToGuid.Add(Node->GetKey(), Node->GetGuid());
-
 		FRigElementKey ResolvedKey;
 		if(IsConnectorResolved(InKey, &ResolvedKey))
 		{
@@ -115,6 +107,41 @@ FControlRigSchematicRigElementKeyNode* FControlRigSchematicModel::AddElementKeyN
 	return Node;
 }
 
+void FControlRigSchematicModel::ConfigureElementKeyNode(FControlRigSchematicRigElementKeyNode* InNode, const FRigElementKey& InKey)
+{
+	InNode->Key = InKey;
+	InNode->SetLabel(FText::FromString(InNode->GetKey().ToString())); 
+	switch(InNode->GetKey().Type)
+	{
+		case ERigElementType::Socket:
+		{
+			if (ControlRigBeingDebuggedPtr.IsValid())
+			{
+				if (const URigHierarchy* Hierarchy = ControlRigBeingDebuggedPtr->GetHierarchy())
+				{
+					if(const FRigSocketElement* Socket = Hierarchy->Find<FRigSocketElement>(InNode->GetKey()))
+					{
+						const FString Description = Socket->GetDescription(Hierarchy);
+						if(!Description.IsEmpty())
+						{
+							InNode->SetLabel(
+								FText::Format(LOCTEXT("NodeLabelFormat", "{0}\n{1}"), 
+								InNode->GetLabel(), 
+								FText::FromString(Description))
+							);
+						}
+					}
+				}
+			}
+			break;
+		}
+		default:
+		{
+			break;
+		}
+	}
+}
+
 const FControlRigSchematicRigElementKeyNode* FControlRigSchematicModel::FindElementKeyNode(const FRigElementKey& InKey) const
 {
 	if(const FGuid* FoundGuid = RigElementKeyToGuid.Find(InKey))
@@ -177,6 +204,17 @@ FControlRigSchematicRigElementKeyLink* FControlRigSchematicModel::AddElementKeyL
 	return nullptr;
 }
 
+void FControlRigSchematicModel::UpdateElementKeyNodes()
+{
+	if (ControlRigBeingDebuggedPtr.IsValid())
+	{
+		if (URigHierarchy* Hierarchy = ControlRigBeingDebuggedPtr->GetHierarchy())
+		{
+			OnHierarchyModified(ERigHierarchyNotification::HierarchyReset, Hierarchy, nullptr);
+		}
+	}
+}
+
 void FControlRigSchematicModel::UpdateElementKeyLinks()
 {
 	const TSharedPtr<FControlRigEditor> Editor = ControlRigEditor.Pin();
@@ -328,6 +366,12 @@ void FControlRigSchematicModel::UpdateElementKeyLinks()
 	// todo: also introduce links for module relationships
 }
 
+void FControlRigSchematicModel::UpdateControlRigContent()
+{
+	UpdateElementKeyNodes();
+	UpdateElementKeyLinks();
+}
+
 void FControlRigSchematicModel::OnSetObjectBeingDebugged(UObject* InObject)
 {
 	if(ControlRigBeingDebuggedPtr.Get() == InObject)
@@ -355,10 +399,8 @@ void FControlRigSchematicModel::OnSetObjectBeingDebugged(UObject* InObject)
 		{
 			Hierarchy->OnModified().RemoveAll(this);
 			Hierarchy->OnModified().AddRaw(this, &FControlRigSchematicModel::OnHierarchyModified);
+			UpdateControlRigContent();
 		}
-
-		// todo: react to the rig being constructed - but not when the rig hierarchy is being copied 
-		// ControlRig->OnPostConstruction_AnyThread().AddRaw(this, &FControlRigSchematicModel::OnPostConstruction);
 	}
 }
 
@@ -396,8 +438,31 @@ void FControlRigSchematicModel::OnHierarchyModified(ERigHierarchyNotification In
 				{
 					const FString OldNameStr = InHierarchy->GetPreviousName(InElement->GetKey()).ToString();
 					FRigElementKey OldKey(*OldNameStr, InElement->GetType());
-					// todooo
-					//RenameNode(OldKey.ToString(), InElement->GetKey().ToString());
+
+					// just have to update the elementkey and lookup - since the guid doesn't need to change
+					FControlRigSchematicRigElementKeyNode* Node = const_cast<FControlRigSchematicRigElementKeyNode*>(FindElementKeyNode(OldKey));
+					if(Node)
+					{
+						ConfigureElementKeyNode(Node, InElement->GetKey());
+						RigElementKeyToGuid.Remove(OldKey);
+						RigElementKeyToGuid.Add(InElement->GetKey(), Node->GetGuid());
+					}
+
+					// update all of the links - update they keys
+					for(const TSharedPtr<FSchematicGraphLink>& Link : GetLinks())
+					{
+						if(FControlRigSchematicRigElementKeyLink* ElementKeyLink = Cast<FControlRigSchematicRigElementKeyLink>(Link.Get()))
+						{
+							if(ElementKeyLink->SourceKey == OldKey)
+							{
+								ElementKeyLink->SourceKey = InElement->GetKey();
+							}
+							if(ElementKeyLink->TargetKey == OldKey)
+							{
+								ElementKeyLink->TargetKey = InElement->GetKey();
+							}
+						}
+					}
 				}
 			}
 			break;
@@ -487,8 +552,9 @@ void FControlRigSchematicModel::HandleModularRigModified(EModularRigNotification
 										if(const FSchematicGraphNode* ResolvedNode = FindElementKeyNode(*ResolvedKey))
 										{
 											(void)SetParentNode(ConnectorNode, ResolvedNode);
+											const_cast<FSchematicGraphNode*>(ResolvedNode)->OnMouseLeave();
 											const_cast<FSchematicGraphNode*>(ConnectorNode)->OnMouseLeave();
-											return;
+											continue;
 										}
 									}
 									if(const FSchematicGraphNode* ParentNode = ConnectorNode->GetParentNode())
@@ -496,7 +562,6 @@ void FControlRigSchematicModel::HandleModularRigModified(EModularRigNotification
 										(void)RemoveFromParentNode(ConnectorNode);
 										const_cast<FSchematicGraphNode*>(ConnectorNode)->OnMouseLeave();
 									}
-									return;
 								}
 							}
 						}
@@ -512,6 +577,13 @@ void FControlRigSchematicModel::HandleModularRigModified(EModularRigNotification
 	}
 }
 
+FSchematicGraphGroupNode* FControlRigSchematicModel::AddAutoGroupNode()
+{
+	FSchematicGraphGroupNode* GroupNode = Super::AddAutoGroupNode();
+	GroupNode->AddTag<FControlRigSchematicWarningTag>();
+	return GroupNode;
+}
+
 FVector2d FControlRigSchematicModel::GetPositionForNode(const FSchematicGraphNode* InNode) const
 {
 	if(const FControlRigSchematicRigElementKeyNode* Node = Cast<FControlRigSchematicRigElementKeyNode>(InNode))
@@ -568,119 +640,140 @@ int32 FControlRigSchematicModel::GetNumLayersForNode(const FSchematicGraphNode*
 	return Super::GetNumLayersForNode(InNode);
 }
 
-const FSlateBrush* FControlRigSchematicModel::GetBrushForNode(const FSchematicGraphNode* InNode, int32 InLayerIndex) const
+const FSlateBrush* FControlRigSchematicModel::GetBrushForKey(const FRigElementKey& InKey, bool bResolveConnectors) const
 {
-	if(const FControlRigSchematicRigElementKeyNode* Node = Cast<FControlRigSchematicRigElementKeyNode>(InNode))
+	switch(InKey.Type)
 	{
-		if(InLayerIndex == 0)
+		case ERigElementType::Socket:
 		{
-			static const FSlateBrush* BackgroundBrush = FSchematicGraphStyle::Get().GetBrush( "Schematic.Background");
-			return BackgroundBrush;
-		}
-		if(InLayerIndex == 1)
-		{
-			static const FSlateBrush* OutlineSingleBrush = FSchematicGraphStyle::Get().GetBrush( "Schematic.Outline.Single");
-			static const FSlateBrush* OutlineDoubleBrush = FSchematicGraphStyle::Get().GetBrush( "Schematic.Outline.Double");
+			static const FSlateBrush* UnresolvedSocketBrush = FSchematicGraphStyle::Get().GetBrush( "Schematic.Dot.Small");
+			static const FSlateBrush* ResolvedSocketBrush = FSchematicGraphStyle::Get().GetBrush( "Schematic.Dot.Large");
+			static const FSlateBrush* ResolvedMultipleSocketBrush = FSchematicGraphStyle::Get().GetBrush( "Schematic.Dot.Medium");
 
-			if(Node->GetKey().Type == ERigElementType::Socket)
+			if (ControlRigBlueprint.IsValid())
 			{
-				if (ControlRigBlueprint.IsValid())
+				int32 Count = 0;
+				for(const TPair<FRigElementKey, FRigElementKey>& Pair : ControlRigBlueprint->ConnectionMap)
 				{
-					int32 Count = 0;
-					for(const TPair<FRigElementKey, FRigElementKey>& Pair : ControlRigBlueprint->ConnectionMap)
+					if(Pair.Value == InKey)
 					{
-						if(Pair.Value == Node->GetKey())
+						Count++;
+						if(Count == 2)
 						{
-							Count++;
-							if(Count == 2)
-							{
-								return OutlineDoubleBrush;
-							}
+							break;
 						}
 					}
 				}
+
+				if(Count == 1)
+				{
+					return ResolvedSocketBrush;
+				}
+				if(Count == 2)
+				{
+					return ResolvedMultipleSocketBrush;
+				}
 			}
-			return OutlineSingleBrush;
+
+			return UnresolvedSocketBrush;
 		}
-		
-		switch(Node->GetKey().Type)
+		case ERigElementType::Connector:
 		{
-			case ERigElementType::Socket:
+			static const FSlateBrush* ConnectorPrimaryBrush = FControlRigEditorStyle::Get().GetBrush( "ControlRig.Schematic.ConnectorPrimary");
+			static const FSlateBrush* ConnectorOptionalBrush = FControlRigEditorStyle::Get().GetBrush( "ControlRig.Schematic.ConnectorOptional");
+			static const FSlateBrush* ConnectorSecondaryBrush = FControlRigEditorStyle::Get().GetBrush( "ControlRig.Schematic.ConnectorSecondary");
+
+			if(bResolveConnectors)
 			{
-				static const FSlateBrush* UnresolvedSocketBrush = FSchematicGraphStyle::Get().GetBrush( "Schematic.Dot.Small");
-				static const FSlateBrush* ResolvedSocketBrush = FSchematicGraphStyle::Get().GetBrush( "Schematic.Dot.Large");
-				static const FSlateBrush* ResolvedMultipleSocketBrush = FSchematicGraphStyle::Get().GetBrush( "Schematic.Dot.Medium");
+				FRigElementKey ResolvedKey;
+				if(IsConnectorResolved(InKey, &ResolvedKey))
+				{
+					if(ResolvedKey.Type != ERigElementType::Socket)
+					{
+						return GetBrushForKey(ResolvedKey, false);
+					}
+				}
+			}
 
-				if (ControlRigBlueprint.IsValid())
+			if (ControlRigBlueprint.IsValid())
+			{
+				if(const UControlRig* ControlRig = ControlRigBlueprint->GetDebuggedControlRig())
 				{
-					int32 Count = 0;
-					for(const TPair<FRigElementKey, FRigElementKey>& Pair : ControlRigBlueprint->ConnectionMap)
+					if(const URigHierarchy* Hierarchy = ControlRig->GetHierarchy())
 					{
-						if(Pair.Value == Node->GetKey())
+						if(const FRigConnectorElement* Connector = Hierarchy->Find<FRigConnectorElement>(InKey))
 						{
-							Count++;
-							if(Count == 2)
+							if(Connector->Settings.Type == EConnectorType::Secondary)
 							{
-								break;
+								return Connector->Settings.bOptional ? 	ConnectorOptionalBrush : ConnectorSecondaryBrush;
 							}
 						}
 					}
-
-					if(Count == 1)
-					{
-						return ResolvedSocketBrush;
-					}
-					if(Count == 2)
-					{
-						return ResolvedMultipleSocketBrush;
-					}
 				}
-
-				return UnresolvedSocketBrush;
 			}
-			case ERigElementType::Connector:
-			{
-				static const FSlateBrush* ConnectorPrimaryBrush = FControlRigEditorStyle::Get().GetBrush( "ControlRig.Schematic.ConnectorPrimary");
-				static const FSlateBrush* ConnectorOptionalBrush = FControlRigEditorStyle::Get().GetBrush( "ControlRig.Schematic.ConnectorOptional");
-				static const FSlateBrush* ConnectorSecondaryBrush = FControlRigEditorStyle::Get().GetBrush( "ControlRig.Schematic.ConnectorSecondary");
+			return ConnectorPrimaryBrush;
+		}
+		case ERigElementType::Bone:
+		{
+			static const FSlateBrush* BoneBrush = FControlRigEditorStyle::Get().GetBrush( "ControlRig.Schematic.Bone");
+			return BoneBrush;
+		}
+		case ERigElementType::Control:
+		{
+			static const FSlateBrush* ControlBrush = FControlRigEditorStyle::Get().GetBrush( "ControlRig.Schematic.Control");
+			return ControlBrush;
+		}
+		case ERigElementType::Null:
+		{
+			static const FSlateBrush* NullBrush = FControlRigEditorStyle::Get().GetBrush( "ControlRig.Schematic.Null");
+			return NullBrush;
+		}
+		default:
+		{
+			break;
+		}
+	}
+	return nullptr;
+}
+
+const FSlateBrush* FControlRigSchematicModel::GetBrushForNode(const FSchematicGraphNode* InNode, int32 InLayerIndex) const
+{
+	if(const FControlRigSchematicRigElementKeyNode* Node = Cast<FControlRigSchematicRigElementKeyNode>(InNode))
+	{
+		if(InLayerIndex == 0)
+		{
+			static const FSlateBrush* BackgroundBrush = FSchematicGraphStyle::Get().GetBrush( "Schematic.Background");
+			return BackgroundBrush;
+		}
+		if(InLayerIndex == 1)
+		{
+			static const FSlateBrush* OutlineSingleBrush = FSchematicGraphStyle::Get().GetBrush( "Schematic.Outline.Single");
+			static const FSlateBrush* OutlineDoubleBrush = FSchematicGraphStyle::Get().GetBrush( "Schematic.Outline.Double");
 
+			if(Node->GetKey().Type == ERigElementType::Socket)
+			{
 				if (ControlRigBlueprint.IsValid())
 				{
-					if(const UControlRig* ControlRig = ControlRigBlueprint->GetDebuggedControlRig())
+					int32 Count = 0;
+					for(const TPair<FRigElementKey, FRigElementKey>& Pair : ControlRigBlueprint->ConnectionMap)
 					{
-						if(const URigHierarchy* Hierarchy = ControlRig->GetHierarchy())
+						if(Pair.Value == Node->GetKey())
 						{
-							if(const FRigConnectorElement* Connector = Hierarchy->Find<FRigConnectorElement>(Node->GetKey()))
+							Count++;
+							if(Count == 2)
 							{
-								if(Connector->Settings.Type == EConnectorType::Secondary)
-								{
-									return Connector->Settings.bOptional ? 	ConnectorOptionalBrush : ConnectorSecondaryBrush;
-								}
+								return OutlineDoubleBrush;
 							}
 						}
 					}
 				}
-				return ConnectorPrimaryBrush;
-			}
-			case ERigElementType::Bone:
-			{
-				static const FSlateBrush* BoneBrush = FControlRigEditorStyle::Get().GetBrush( "ControlRig.Schematic.Bone");
-				return BoneBrush;
-			}
-			case ERigElementType::Control:
-			{
-				static const FSlateBrush* ControlBrush = FControlRigEditorStyle::Get().GetBrush( "ControlRig.Schematic.Control");
-				return ControlBrush;
-			}
-			case ERigElementType::Null:
-			{
-				static const FSlateBrush* NullBrush = FControlRigEditorStyle::Get().GetBrush( "ControlRig.Schematic.Null");
-				return NullBrush;
-			}
-			default:
-			{
-				break;
 			}
+			return OutlineSingleBrush;
+		}
+
+		if(const FSlateBrush* Brush = GetBrushForKey(Node->GetKey(), true))
+		{
+			return Brush;
 		}
 	}
 	return Super::GetBrushForNode(InNode, InLayerIndex);
@@ -696,6 +789,10 @@ FLinearColor FControlRigSchematicModel::GetColorForNode(const FSchematicGraphNod
 		}
 		if(InLayerIndex == 1) // outline
 		{
+			if(ControlRigBlueprint->GetHierarchy()->IsSelected(ElementKeyNode->GetKey()))
+			{
+				return FLinearColor(FColor(0, 112, 224));
+			}
 			return FLinearColor::White;
 		}
 		
@@ -817,9 +914,107 @@ ESchematicGraphVisibility::Type FControlRigSchematicModel::GetVisibilityForTag(c
 			}
 		}
 	}
+	if(InTag->IsA<FControlRigSchematicConnectorTag>())
+	{
+		if(const FControlRigSchematicRigElementKeyNode* Node = Cast<FControlRigSchematicRigElementKeyNode>(InTag->GetNode()))
+		{
+			if(Node->GetKey().Type == ERigElementType::Connector)
+			{
+				FRigElementKey ResolvedKey;
+				if(IsConnectorResolved(Node->GetKey(), &ResolvedKey))
+				{
+					if(ResolvedKey.Type != ERigElementType::Socket)
+					{
+						return ESchematicGraphVisibility::Visible;
+					}
+				}
+			}
+		}
+		return ESchematicGraphVisibility::Hidden;
+	}
+	if(InTag->IsA<FControlRigSchematicWarningTag>())
+	{
+		if(const FControlRigSchematicRigElementKeyNode* ConnectorNode = Cast<FControlRigSchematicRigElementKeyNode>(InTag->GetNode()))
+		{
+			if(ConnectorNode->GetKey().Type == ERigElementType::Connector)
+			{
+				FRigElementKey ResolvedKey;
+				if(!IsConnectorResolved(ConnectorNode->GetKey(), &ResolvedKey))
+				{
+					return ESchematicGraphVisibility::Visible;
+				}
+			}
+		}
+		else if(const FSchematicGraphAutoGroupNode* AutoGroupNode = Cast<FSchematicGraphAutoGroupNode>(InTag->GetNode()))
+		{
+			for(int32 Index = 0; Index < AutoGroupNode->GetNumChildNodes(); Index++)
+			{
+				if(const FSchematicGraphNode* ChildNode = AutoGroupNode->GetChildNode(Index))
+				{
+					if(const FControlRigSchematicWarningTag* ChildTag = ChildNode->FindTag<FControlRigSchematicWarningTag>())
+					{
+						const ESchematicGraphVisibility::Type ChildVisibility = GetVisibilityForTag(ChildTag); 
+						if(ChildVisibility != ESchematicGraphVisibility::Hidden)
+						{
+							return ChildVisibility;
+						}
+					}
+				}
+			}
+		}
+		return ESchematicGraphVisibility::Hidden;
+	}
 	return Super::GetVisibilityForTag(InTag);
 }
 
+const FSlateBrush* FControlRigSchematicModel::GetForegroundBrushForTag(const FSchematicGraphTag* InTag) const
+{
+	if(const FControlRigSchematicConnectorTag* ConnectorTag = Cast<FControlRigSchematicConnectorTag>(InTag))
+	{
+		if(const FControlRigSchematicRigElementKeyNode* Node = Cast<FControlRigSchematicRigElementKeyNode>(ConnectorTag->GetNode()))
+		{
+			return GetBrushForKey(Node->GetKey(), false);
+		}
+	}
+	return Super::GetForegroundBrushForTag(InTag);
+}
+
+const FText FControlRigSchematicModel::GetToolTipForTag(const FSchematicGraphTag* InTag) const
+{
+	if(InTag->IsA<FControlRigSchematicWarningTag>())
+	{
+		if(const FControlRigSchematicRigElementKeyNode* ConnectorNode = Cast<FControlRigSchematicRigElementKeyNode>(InTag->GetNode()))
+		{
+			if(ConnectorNode->GetKey().Type == ERigElementType::Connector)
+			{
+				if(!IsConnectorResolved(ConnectorNode->GetKey()))
+				{
+					static const FText UnresolvedConnectorToolTip = LOCTEXT("UnresolvedConnectorToolTip", "Connector is unresolved.");
+					return UnresolvedConnectorToolTip;
+				}
+			}
+		}
+		else if(const FSchematicGraphAutoGroupNode* AutoGroupNode = Cast<FSchematicGraphAutoGroupNode>(InTag->GetNode()))
+		{
+			for(int32 Index = 0; Index < AutoGroupNode->GetNumChildNodes(); Index++)
+			{
+				if(const FSchematicGraphNode* ChildNode = AutoGroupNode->GetChildNode(Index))
+				{
+					if(const FControlRigSchematicWarningTag* ChildTag = ChildNode->FindTag<FControlRigSchematicWarningTag>())
+					{
+						const FText ChildToolTip = GetToolTipForTag(ChildTag); 
+						if(!ChildToolTip.IsEmpty())
+						{
+							return ChildToolTip;
+						}
+					}
+				}
+			}
+		}
+	}
+	return Super::GetToolTipForTag(InTag);
+}
+
 void FControlRigSchematicModel::HandleSchematicNodeClicked(SSchematicGraphPanel* InPanel, SSchematicGraphNode* InNode, const FPointerEvent& InMouseEvent)
 {
 	for (const TSharedPtr<FSchematicGraphNode>& Node : Nodes)
@@ -836,8 +1031,16 @@ void FControlRigSchematicModel::HandleSchematicNodeClicked(SSchematicGraphPanel*
 			{
 				if (URigHierarchyController* Controller = Hierarchy->GetController())
 				{
-					const TArray<FRigElementKey> Selection = {Node->GetKey()};
-					Controller->SetSelection(Selection);
+					const bool bClearSelection = !InMouseEvent.IsShiftDown() && !InMouseEvent.IsControlDown();
+					bool bSelect = true;
+					if(InMouseEvent.IsControlDown())
+					{
+						if(Hierarchy->IsSelected(Node->GetKey()))
+						{
+							bSelect = false;
+						}
+					}
+					Controller->SelectElement(Node->GetKey(), bSelect, bClearSelection);
 				}
 			}
 		}
@@ -984,6 +1187,7 @@ void FControlRigSchematicModel::HandleSchematicDrop(SSchematicGraphPanel* InPane
 
 	TSharedPtr<FAssetDragDropOp> AssetDragDropOp = InDragDropEvent.GetOperationAs<FAssetDragDropOp>();
 	TSharedPtr<FSchematicGraphNodeDragDropOp> SchematicDragDropOp = InDragDropEvent.GetOperationAs<FSchematicGraphNodeDragDropOp>();
+	TSharedPtr<FModularRigModuleDragDropOp> ModuleDragDropOperation = InDragDropEvent.GetOperationAs<FModularRigModuleDragDropOp>();
 	if (AssetDragDropOp.IsValid())
 	{
 		for (const FAssetData& AssetData : AssetDragDropOp->GetAssets())
@@ -1077,6 +1281,53 @@ void FControlRigSchematicModel::HandleSchematicDrop(SSchematicGraphPanel* InPane
 			}
 		}, TStatId(), NULL, ENamedThreads::GameThread);
 	}
+	else if(ModuleDragDropOperation.IsValid())
+	{
+		const TArray<FString> Sources = ModuleDragDropOperation->GetElements();
+		FFunctionGraphTask::CreateAndDispatchWhenReady([this, Sources, TargetKey]()
+		{
+			UModularRig* ControlRig = Cast<UModularRig>(ControlRigBlueprint->GetDebuggedControlRig());
+			if (!ControlRig)
+			{
+				return;
+			}
+
+			URigHierarchy* Hierarchy = ControlRig->GetHierarchy();
+			if (!Hierarchy)
+			{
+				return;
+			}
+
+			if (UModularRigController* Controller = ControlRigBlueprint->GetModularRigController())
+			{
+				const TArray<FRigConnectorElement*> Connectors = Hierarchy->GetConnectors();
+				for (const FString& ModulePathOrConnectorName : Sources)
+				{
+					if(const FRigConnectorElement* Connector = Hierarchy->Find<FRigConnectorElement>(FRigElementKey(*ModulePathOrConnectorName, ERigElementType::Connector)))
+					{
+						Controller->ConnectConnectorToElement(Connector->GetKey(), TargetKey, true, ControlRig->GetModularRigSettings().bAutoResolve);
+						return;
+					}
+					if(const FRigModuleReference* Module = Controller->FindModule(ModulePathOrConnectorName))
+					{
+						const FString ModuleNameSpace = Module->GetNamespace();
+						for(const FRigConnectorElement* Connector : Connectors)
+						{
+							const FString ConnectorNameSpace = Hierarchy->GetNameMetadata(Connector->GetKey(), URigHierarchy::NameSpaceMetadataName, NAME_None).ToString();
+							if(ConnectorNameSpace.Equals(ModuleNameSpace))
+							{
+								if(Connector->Settings.Type == EConnectorType::Primary)
+								{
+									Controller->ConnectConnectorToElement(Connector->GetKey(), TargetKey, true, ControlRig->GetModularRigSettings().bAutoResolve);
+									return;
+								}
+							}
+						}
+					}
+				}
+			}
+		}, TStatId(), NULL, ENamedThreads::GameThread);
+	}
 }
 
 bool FControlRigSchematicModel::IsConnectorResolved(const FRigElementKey& InConnectorKey, FRigElementKey* OutKey) const
diff --git a/Engine/Plugins/Animation/ControlRig/Source/ControlRigEditor/Private/Editor/ControlRigSchematicModel.h b/Engine/Plugins/Animation/ControlRig/Source/ControlRigEditor/Private/Editor/ControlRigSchematicModel.h
--- a/Engine/Plugins/Animation/ControlRig/Source/ControlRigEditor/Private/Editor/ControlRigSchematicModel.h
+++ b/Engine/Plugins/Animation/ControlRig/Source/ControlRigEditor/Private/Editor/ControlRigSchematicModel.h
@@ -53,6 +53,28 @@ protected:
 	friend class FControlRigSchematicModel;
 };
 
+/** Tag to mark up connectors */
+class FControlRigSchematicConnectorTag : public FSchematicGraphTag
+{
+public:
+
+	SCHEMATICGRAPHTAG_BODY(FControlRigSchematicConnectorTag, FSchematicGraphTag)
+
+	FControlRigSchematicConnectorTag();
+	virtual ~FControlRigSchematicConnectorTag() override {}
+};
+
+class FControlRigSchematicWarningTag : public FSchematicGraphTag
+{
+public:
+
+	SCHEMATICGRAPHTAG_BODY(FControlRigSchematicWarningTag, FSchematicGraphTag)
+
+	FControlRigSchematicWarningTag();
+	virtual ~FControlRigSchematicWarningTag() override {}
+};
+
+
 /** Model for the schematic views */
 class FControlRigSchematicModel : public FSchematicGraphModel
 {
@@ -71,24 +93,32 @@ public:
 	virtual bool RemoveNode(const FGuid& InGuid) override;
 	bool RemoveElementKeyNode(const FRigElementKey& InKey);
 	FControlRigSchematicRigElementKeyLink* AddElementKeyLink(const FRigElementKey& InSourceKey, const FRigElementKey& InTargetKey, bool bNotify = true);
+	void UpdateElementKeyNodes();
 	void UpdateElementKeyLinks();
+	void UpdateControlRigContent();
 
 	void OnSetObjectBeingDebugged(UObject* InObject);
 	void OnHierarchyModified(ERigHierarchyNotification InNotif, URigHierarchy* InHierarchy, const FRigBaseElement* InElement);
 	void HandleModularRigModified(EModularRigNotification InNotification, const FRigModuleReference* InModule);
 
+	virtual FSchematicGraphGroupNode* AddAutoGroupNode() override;
 	virtual FVector2d GetPositionForNode(const FSchematicGraphNode* InNode) const override;
 	virtual bool GetPositionAnimationEnabledForNode(const FSchematicGraphNode* InNode) const override;
 	virtual int32 GetNumLayersForNode(const FSchematicGraphNode* InNode) const override;
+	const FSlateBrush* GetBrushForKey(const FRigElementKey& InKey, bool bResolveConnectors) const;
 	virtual const FSlateBrush* GetBrushForNode(const FSchematicGraphNode* InNode, int32 InLayerIndex) const override;
 	virtual FLinearColor GetColorForNode(const FSchematicGraphNode* InNode, int32 InLayerIndex) const override;
 	virtual ESchematicGraphPlacementConstraint::Type GetPlacementForNode(const FSchematicGraphNode* InNode) const override;
 	virtual const FSlateBrush* GetBrushForLink(const FSchematicGraphLink* InLink) const override;
 	virtual FLinearColor GetColorForLink(const FSchematicGraphLink* InLink) const override;
 	virtual ESchematicGraphVisibility::Type GetVisibilityForTag(const FSchematicGraphTag* InTag) const override;
+	virtual const FSlateBrush* GetForegroundBrushForTag(const FSchematicGraphTag* InTag) const override;
+	virtual const FText GetToolTipForTag(const FSchematicGraphTag* InTag) const override;
 
 private:
 
+	void ConfigureElementKeyNode(FControlRigSchematicRigElementKeyNode* InNode, const FRigElementKey& InKey);
+
 	void HandleSchematicNodeClicked(SSchematicGraphPanel* InPanel, SSchematicGraphNode* InNode, const FPointerEvent& InMouseEvent);
 	void HandleSchematicBeginDrag(SSchematicGraphPanel* InPanel, SSchematicGraphNode* InNode, const FDragDropOperation& InDragDropOperation);
 	void HandleSchematicEndDrag(SSchematicGraphPanel* InPanel, SSchematicGraphNode* InNode, const FDragDropOperation& InDragDropOperation);
diff --git a/Engine/Plugins/Animation/ControlRig/Source/ControlRigEditor/Private/Editor/SModularRigModel.cpp b/Engine/Plugins/Animation/ControlRig/Source/ControlRigEditor/Private/Editor/SModularRigModel.cpp
--- a/Engine/Plugins/Animation/ControlRig/Source/ControlRigEditor/Private/Editor/SModularRigModel.cpp
+++ b/Engine/Plugins/Animation/ControlRig/Source/ControlRigEditor/Private/Editor/SModularRigModel.cpp
@@ -172,6 +172,14 @@ void SModularRigModel::Construct(const FArguments& InArgs, TSharedRef<FControlRi
 	}
 	
 	CreateContextMenu();
+
+	if(const UModularRig* Rig = GetModularRigForTreeView())
+	{
+		if(URigHierarchy* Hierarchy = Rig->GetHierarchy())
+		{
+			Hierarchy->OnModified().AddSP(this, &SModularRigModel::OnHierarchyModified);
+		}
+	}
 }
 
 void SModularRigModel::OnEditorClose(const FRigVMEditor* InEditor, URigVMBlueprint* InBlueprint)
@@ -190,7 +198,15 @@ void SModularRigModel::OnEditorClose(const FRigVMEditor* InEditor, URigVMBluepri
 		BP->OnModularRigPreCompiled().RemoveAll(this);
 		BP->OnModularRigCompiled().RemoveAll(this);
 	}
-	
+
+	if(const UModularRig* Rig = GetModularRigForTreeView())
+	{
+		if(URigHierarchy* Hierarchy = Rig->GetHierarchy())
+		{
+			Hierarchy->OnModified().RemoveAll(this);
+		}
+	}
+
 	ControlRigEditor.Reset();
 	ControlRigBlueprint.Reset();
 }
@@ -286,11 +302,24 @@ void SModularRigModel::HandleSetObjectBeingDebugged(UObject* InObject)
 		return;
 	}
 
+	if(ControlRigBeingDebuggedPtr.IsValid())
+	{
+		if(URigHierarchy* Hierarchy = ControlRigBeingDebuggedPtr->GetHierarchy())
+		{
+			Hierarchy->OnModified().RemoveAll(this);
+		}
+	}
+
 	ControlRigBeingDebuggedPtr.Reset();
-	
+
 	if(UModularRig* ControlRig = Cast<UModularRig>(InObject))
 	{
 		ControlRigBeingDebuggedPtr = ControlRig;
+
+		if(URigHierarchy* Hierarchy = ControlRig->GetHierarchy())
+		{
+			Hierarchy->OnModified().AddSP(this, &SModularRigModel::OnHierarchyModified);
+		}
 	}
 
 	RefreshTreeView();
@@ -687,6 +716,64 @@ void SModularRigModel::HandleConnectorDisconnect(const FRigElementKey& InConnect
 	}
 }
 
+void SModularRigModel::OnSelectionChanged(TSharedPtr<FModularRigTreeElement> Selection, ESelectInfo::Type SelectInfo)
+{
+}
+
+void SModularRigModel::OnHierarchyModified(ERigHierarchyNotification InNotif, URigHierarchy* InHierarchy, const FRigBaseElement* InElement)
+{
+	if(!ControlRigBlueprint.IsValid())
+	{
+		return;
+	}
+	
+	switch(InNotif)
+	{
+		case ERigHierarchyNotification::ElementSelected:
+		case ERigHierarchyNotification::ElementDeselected:
+		{
+			const FRigConnectorElement* Connector = Cast<FRigConnectorElement>(InElement);
+			if(Connector == nullptr)
+			{
+				for(const TPair<FRigElementKey, FRigElementKey>& Pair : ControlRigBlueprint->ConnectionMap)
+				{
+					check(Pair.Key.Type == ERigElementType::Connector);
+					if(Pair.Value == InElement->GetKey())
+					{
+						if(const FRigConnectorElement* TargetConnector = InHierarchy->Find<FRigConnectorElement>(Pair.Key))
+						{
+							OnHierarchyModified(InNotif, InHierarchy, TargetConnector);
+						}
+					}
+				}
+				return;
+			}
+
+			FString ModulePathOrConnectorName;
+			if(Connector->Settings.Type == EConnectorType::Primary)
+			{
+				ModulePathOrConnectorName = InHierarchy->GetNameMetadata(Connector->GetKey(), URigHierarchy::NameSpaceMetadataName, NAME_None).ToString();
+				ModulePathOrConnectorName = ModulePathOrConnectorName.LeftChop(1);
+			}
+			else
+			{
+				ModulePathOrConnectorName = Connector->GetName();
+			}
+
+			TSharedPtr<FModularRigTreeElement> Item = TreeView->FindElement(ModulePathOrConnectorName);
+			if(Item.IsValid())
+			{
+				const bool bSelected = InNotif == ERigHierarchyNotification::ElementSelected;
+				TreeView->SetItemSelection(Item, bSelected);
+			}
+		}
+		default:
+		{
+			break;
+		}
+	}
+}
+
 class SModularRigModelPasteTransformsErrorPipe : public FOutputDevice
 {
 public:
diff --git a/Engine/Plugins/Animation/ControlRig/Source/ControlRigEditor/Private/Editor/SModularRigModel.h b/Engine/Plugins/Animation/ControlRig/Source/ControlRigEditor/Private/Editor/SModularRigModel.h
--- a/Engine/Plugins/Animation/ControlRig/Source/ControlRigEditor/Private/Editor/SModularRigModel.h
+++ b/Engine/Plugins/Animation/ControlRig/Source/ControlRigEditor/Private/Editor/SModularRigModel.h
@@ -154,7 +154,8 @@ private:
 	void OnRequestDetailsInspection(const FString& InKey);
 	void HandlePreCompileModularRigs(URigVMBlueprint* InBlueprint);
 	void HandlePostCompileModularRigs(URigVMBlueprint* InBlueprint);
-	
+	void OnHierarchyModified(ERigHierarchyNotification InNotif, URigHierarchy* InHierarchy, const FRigBaseElement* InElement);
+
 	void HandleRefreshEditorFromBlueprint(URigVMBlueprint* InBlueprint);
 	void HandleSetObjectBeingDebugged(UObject* InObject);
 
diff --git a/Engine/Plugins/Animation/ControlRig/Source/ControlRigEditor/Private/Editor/SModularRigTreeView.cpp b/Engine/Plugins/Animation/ControlRig/Source/ControlRigEditor/Private/Editor/SModularRigTreeView.cpp
--- a/Engine/Plugins/Animation/ControlRig/Source/ControlRigEditor/Private/Editor/SModularRigTreeView.cpp
+++ b/Engine/Plugins/Animation/ControlRig/Source/ControlRigEditor/Private/Editor/SModularRigTreeView.cpp
@@ -28,6 +28,7 @@
 #include "Widgets/Input/SComboButton.h"
 #include "Widgets/Input/SButton.h"
 #include "ScopedTransaction.h"
+#include "DetailLayoutBuilder.h"
 
 #define LOCTEXT_NAMESPACE "SModularRigTreeView"
 
@@ -243,6 +244,8 @@ void SModularRigModelItem::Construct(const FArguments& InArgs, const TSharedRef<
 					// Wrap in configurable box to restrain height/width of menu
 					SNew(SBox)
 					.MinDesiredWidth(200.0f)
+					.Padding(0, 0, 0, 0)
+					.HAlign(HAlign_Left)
 					[
 						SAssignNew(ComboButtonBox, SVerticalBox)
 					]
@@ -393,6 +396,7 @@ void SModularRigModelItem::Construct(const FArguments& InArgs, const TSharedRef<
 	PopulateConnectorCurrentTarget(ComboButtonBox, ConnectorKey, CurrentTargetKey, IconAndColor.Key, IconAndColor.Value, FText::FromName(CurrentTargetKey.Name));
 }
 
+
 void SModularRigModelItem::PopulateConnectorTargetList(const FRigElementKey InConnectorKey)
 {
 	if (!WeakRigTreeElement.IsValid())
@@ -420,59 +424,60 @@ void SModularRigModelItem::PopulateConnectorCurrentTarget(TSharedPtr<SVerticalBo
 	InListBox->AddSlot()
 	.AutoHeight()
 	.VAlign(VAlign_Top)
-	.HAlign(HAlign_Fill)
-	.Padding(4.0, 0.0, 4.0, 0.0)
+	.HAlign(HAlign_Left)
+	.Padding(0.0, 0.0, 4.0, 0.0)
 	[
-		SNew( SButton )
-		.ButtonStyle(FAppStyle::Get(), "SimpleButton")
-		.ContentPadding(FMargin(0.0))
-		.OnClicked_Lambda([this, InConnectorKey, InTargetKey]()
-		{
-			//Controller->ConnectConnectorToElement(InConnectorKey, InTargetKey, true, Info.GetModularRig()->GetModularRigSettings().bAutoResolve);
-			PopulateConnectorTargetList(InConnectorKey);
-			return FReply::Handled();
-		})
+		SNew(SHorizontalBox)
+		+ SHorizontalBox::Slot()
+		.AutoWidth()
+		.HAlign(HAlign_Left)
+		.Padding(0, 0, 0, 0)
 		[
-			SAssignNew(RowBox, SHorizontalBox)
-			+ SHorizontalBox::Slot()
-			.AutoWidth()
-			.VAlign(VAlign_Center)
-			.HAlign(HAlign_Fill)
-			.Padding(0)
+			SNew( SButton )
+			.ButtonStyle(FAppStyle::Get(), "SimpleButton")
+			.ContentPadding(FMargin(0.0))
+			.OnClicked_Lambda([this, InConnectorKey, InTargetKey]()
+			{
+				//Controller->ConnectConnectorToElement(InConnectorKey, InTargetKey, true, Info.GetModularRig()->GetModularRigSettings().bAutoResolve);
+				PopulateConnectorTargetList(InConnectorKey);
+				return FReply::Handled();
+			})
 			[
-				SNew(SBorder)
-				.Padding(FMargin(5.0, 2.0, 5.0, 2.0))
-				.BorderImage(RoundedBoxBrush)
-				.BorderBackgroundColor(FSlateColor(FLinearColor::Transparent))
-				.Content()
+				SAssignNew(RowBox, SHorizontalBox)
+				+ SHorizontalBox::Slot()
+				.AutoWidth()
+				.VAlign(VAlign_Center)
+				.HAlign(HAlign_Fill)
+				.Padding(0)
 				[
-					SAssignNew(ButtonBox, SHorizontalBox)
-					+ SHorizontalBox::Slot()
-					.AutoWidth()
-					.VAlign(VAlign_Center)
-					.HAlign(HAlign_Left)
-					.Padding(FMargin(0.f, 0.f, 3.f, 0.f))
-					[
-						SNew(SImage)
-						.Image(InBrush)
-						.ColorAndOpacity(InColor)
-					]
-
-					+ SHorizontalBox::Slot()
-					.AutoWidth()
-					.VAlign(VAlign_Center)
-					.HAlign(HAlign_Left)
-					.Padding(0)
+					SNew(SBorder)
+					.Padding(FMargin(2.0, 2.0, 5.0, 2.0))
+					.BorderImage(RoundedBoxBrush)
+					.BorderBackgroundColor(FSlateColor(FLinearColor::Transparent))
+					.Content()
 					[
-						SNew( STextBlock )
-						.Text( InTitle )
-						//.Font( IDetailLayoutBuilder::GetDetailFont() )
-					]
+						SAssignNew(ButtonBox, SHorizontalBox)
+						+ SHorizontalBox::Slot()
+						.AutoWidth()
+						.VAlign(VAlign_Center)
+						.HAlign(HAlign_Left)
+						.Padding(FMargin(0.f, 0.f, 3.f, 0.f))
+						[
+							SNew(SImage)
+							.Image(InBrush)
+							.ColorAndOpacity(InColor)
+						]
 
-					+ SHorizontalBox::Slot()
-					.FillWidth(1.f)
-					[
-						SNew(SSpacer)
+						+ SHorizontalBox::Slot()
+						.AutoWidth()
+						.VAlign(VAlign_Center)
+						.HAlign(HAlign_Left)
+						.Padding(0)
+						[
+							SNew( STextBlock )
+							.Text( InTitle )
+							.Font( IDetailLayoutBuilder::GetDetailFont() )
+						]
 					]
 				]
 			]
@@ -997,7 +1002,7 @@ TPair<const FSlateBrush*, FSlateColor> FModularRigTreeElement::GetBrushAndColor(
 					}
 					else
 					{
-						Brush = FControlRigEditorStyle::Get().GetBrush("ControlRig.Schematic.ConnectorWarning");
+						Brush = FControlRigEditorStyle::Get().GetBrush("ControlRig.ConnectorWarning");
 					}
 				}
 				else if (Connector->Settings.bOptional)
diff --git a/Engine/Source/Editor/AnimationEditorWidgets/Public/SchematicGraphPanel/SchematicGraphNode.h b/Engine/Source/Editor/AnimationEditorWidgets/Public/SchematicGraphPanel/SchematicGraphNode.h
--- a/Engine/Source/Editor/AnimationEditorWidgets/Public/SchematicGraphPanel/SchematicGraphNode.h
+++ b/Engine/Source/Editor/AnimationEditorWidgets/Public/SchematicGraphPanel/SchematicGraphNode.h
@@ -88,17 +88,27 @@ public:
 	}
 
 	template<typename TagType = FSchematicGraphTag>
-	const TagType* FindTag(const FGuid& InTagGuid) const
+	const TagType* FindTag(const FGuid& InTagGuid = FGuid()) const
 	{
 		if(const TSharedPtr<FSchematicGraphTag>* ExistingTag = TagByGuid.Find(InTagGuid))
 		{
 			return Cast<TagType>(ExistingTag->Get());
 		};
+		if(!InTagGuid.IsValid())
+		{
+			for(const TSharedPtr<FSchematicGraphTag>& Tag : Tags)
+			{
+				if(Tag->IsA<TagType>())
+				{
+					return Cast<TagType>(Tag.Get());
+				}
+			}
+		}
 		return nullptr;
 	}
 
 	template<typename TagType = FSchematicGraphTag>
-	TagType* FindTag(const FGuid& InTagGuid)
+	TagType* FindTag(const FGuid& InTagGuid = FGuid())
 	{
 		const FSchematicGraphNode* ConstThis = this;
 		return const_cast<TagType*>(ConstThis->FindTag<TagType>(InTagGuid));
