Control Rig: Schematic improvements

#rnx

diff --git a/Engine/Plugins/Animation/ControlRig/Source/ControlRig/Private/ModularRig.cpp b/Engine/Plugins/Animation/ControlRig/Source/ControlRig/Private/ModularRig.cpp
--- a/Engine/Plugins/Animation/ControlRig/Source/ControlRig/Private/ModularRig.cpp
+++ b/Engine/Plugins/Animation/ControlRig/Source/ControlRig/Private/ModularRig.cpp
@@ -135,6 +135,23 @@ const FRigModuleReference* FRigModuleInstance::GetModuleReference() const
 	return nullptr;
 }
 
+const FRigConnectorElement* FRigModuleInstance::FindPrimaryConnector() const
+{
+	if(const UControlRig* Rig = GetRig())
+	{
+		if(const FRigModuleConnector* ExposedConnector = Rig->GetRigModuleSettings().FindPrimaryConnector())
+		{
+			if(const URigHierarchy* Hierarchy = Rig->GetHierarchy())
+			{
+				const FString ConnectorName = URigHierarchy::JoinNameSpace(GetNamespace(), ExposedConnector->Name); 
+				const FRigElementKey ConnectorKey(*ConnectorName, ERigElementType::Connector);
+				return Hierarchy->Find<FRigConnectorElement>(ConnectorKey);
+			}
+		}
+	}
+	return nullptr;
+}
+
 void UModularRig::InitializeVMs(bool bRequestInit)
 {
 	URigVMHost::Initialize(bRequestInit);
diff --git a/Engine/Plugins/Animation/ControlRig/Source/ControlRig/Private/Rigs/RigHierarchyController.cpp b/Engine/Plugins/Animation/ControlRig/Source/ControlRig/Private/Rigs/RigHierarchyController.cpp
--- a/Engine/Plugins/Animation/ControlRig/Source/ControlRig/Private/Rigs/RigHierarchyController.cpp
+++ b/Engine/Plugins/Animation/ControlRig/Source/ControlRig/Private/Rigs/RigHierarchyController.cpp
@@ -2054,6 +2054,7 @@ int32 URigHierarchyController::AddElement(FRigBaseElement* InElementToAdd, FRigB
 			if(InElementToAdd->GetName().StartsWith(CRContext.GetRigModuleNameSpace(), ESearchCase::CaseSensitive))
 			{
 				Hierarchy->SetNameMetadata(InElementToAdd->Key, URigHierarchy::NameSpaceMetadataName, *CRContext.GetRigModuleNameSpace());
+				Hierarchy->SetNameMetadata(InElementToAdd->Key, URigHierarchy::ModuleMetadataName, *CRContext.GetRigModuleNameSpace().LeftChop(1));
 
 				if(Hierarchy->ElementKeyRedirector)
 				{
diff --git a/Engine/Plugins/Animation/ControlRig/Source/ControlRig/Public/ModularRig.h b/Engine/Plugins/Animation/ControlRig/Source/ControlRig/Public/ModularRig.h
--- a/Engine/Plugins/Animation/ControlRig/Source/ControlRig/Public/ModularRig.h
+++ b/Engine/Plugins/Animation/ControlRig/Source/ControlRig/Public/ModularRig.h
@@ -79,6 +79,7 @@ public:
 	void SetRig(UControlRig* InRig);
 	bool ContainsRig(const UControlRig* InRig) const;
 	const FRigModuleReference* GetModuleReference() const;
+	const FRigConnectorElement* FindPrimaryConnector() const;
 };
 
 USTRUCT(BlueprintType)
diff --git a/Engine/Plugins/Animation/ControlRig/Source/ControlRigEditor/Private/Editor/ControlRigEditor.cpp b/Engine/Plugins/Animation/ControlRig/Source/ControlRigEditor/Private/Editor/ControlRigEditor.cpp
--- a/Engine/Plugins/Animation/ControlRig/Source/ControlRigEditor/Private/Editor/ControlRigEditor.cpp
+++ b/Engine/Plugins/Animation/ControlRig/Source/ControlRigEditor/Private/Editor/ControlRigEditor.cpp
@@ -2350,12 +2350,7 @@ void FControlRigEditor::HandleToggleSchematicViewport()
 	if(SchematicViewport.IsValid())
 	{
 		SchematicModel.UpdateControlRigContent();
-
-		const EVisibility PreviousVisibility = SchematicViewport.Pin()->GetVisibility();
-		SchematicViewport.Pin()->SetVisibility(
-			PreviousVisibility == EVisibility::Hidden ?
-			EVisibility::SelfHitTestInvisible :
-			EVisibility::Hidden);
+		SchematicViewport.Pin()->ToggleVisibility();
 	}
 }
 
diff --git a/Engine/Plugins/Animation/ControlRig/Source/ControlRigEditor/Private/Editor/ControlRigSchematicModel.cpp b/Engine/Plugins/Animation/ControlRig/Source/ControlRigEditor/Private/Editor/ControlRigSchematicModel.cpp
--- a/Engine/Plugins/Animation/ControlRig/Source/ControlRigEditor/Private/Editor/ControlRigSchematicModel.cpp
+++ b/Engine/Plugins/Animation/ControlRig/Source/ControlRigEditor/Private/Editor/ControlRigSchematicModel.cpp
@@ -13,6 +13,7 @@
 #include "Async/TaskGraphInterfaces.h"
 #include <SchematicGraphPanel/SchematicGraphStyle.h>
 #include "Editor/SModularRigModel.h"
+#include "Blueprint/WidgetBlueprintLibrary.h"
 
 #define LOCTEXT_NAMESPACE "ControlRigSchematicModel"
 
@@ -26,13 +27,6 @@ bool FControlRigSchematicRigElementKeyNode::IsDragSupported() const
 	return Key.Type == ERigElementType::Connector;
 }
 
-FControlRigSchematicConnectorTag::FControlRigSchematicConnectorTag()
-: FSchematicGraphTag()
-{
-	BackgroundColor = FLinearColor::Black;
-	ForegroundColor = FLinearColor::White;
-}
-
 FControlRigSchematicWarningTag::FControlRigSchematicWarningTag()
 : FSchematicGraphTag()
 {
@@ -74,6 +68,40 @@ void FControlRigSchematicModel::Reset()
 	RigElementKeyToGuid.Reset();
 }
 
+void FControlRigSchematicModel::Tick(float InDeltaTime)
+{
+	Super::Tick(InDeltaTime);
+
+	if(bUpdateConnectorTargetsOnTick)
+	{
+		if(ControlRigBlueprint.IsValid())
+		{
+			TArray<FRigElementKey> KeysWithConnectors;
+			ControlRigBlueprint->ConnectionMap.GenerateValueArray(KeysWithConnectors);
+			
+			TArray<FRigElementKey> KeysToRemove;
+			for(const TPair<FRigElementKey,FGuid>& Pair : RigElementKeyToGuid)
+			{
+				if((Pair.Key.Type == ERigElementType::Socket) ||
+					(Pair.Key.Type == ERigElementType::Connector))
+				{
+					continue;
+				}
+				if(KeysWithConnectors.Contains(Pair.Key))
+				{
+					continue;
+				}
+				KeysToRemove.Add(Pair.Key);
+			}
+			for(const FRigElementKey& KeyToRemove : KeysToRemove)
+			{
+				RemoveElementKeyNode(KeyToRemove);
+			}
+		}
+		bUpdateConnectorTargetsOnTick = true;
+	}
+}
+
 FControlRigSchematicRigElementKeyNode* FControlRigSchematicModel::AddElementKeyNode(const FRigElementKey& InKey, bool bNotify)
 {
 	FControlRigSchematicRigElementKeyNode* Node = AddNode<FControlRigSchematicRigElementKeyNode>(false);
@@ -82,19 +110,20 @@ FControlRigSchematicRigElementKeyNode* FControlRigSchematicModel::AddElementKeyN
 		ConfigureElementKeyNode(Node, InKey);
 		RigElementKeyToGuid.Add(Node->GetKey(), Node->GetGuid());
 
-		if(InKey.Type == ERigElementType::Connector)
-		{
-			Node->AddTag<FControlRigSchematicConnectorTag>();
-			Node->AddTag<FControlRigSchematicWarningTag>();
-		}
-
 		FRigElementKey ResolvedKey;
 		if(IsConnectorResolved(InKey, &ResolvedKey))
 		{
-			if(const FControlRigSchematicRigElementKeyNode* ResolvedNode = FindElementKeyNode(ResolvedKey))
+			const FControlRigSchematicRigElementKeyNode* ResolvedNode = FindElementKeyNode(ResolvedKey);
+			if(ResolvedNode == nullptr)
 			{
-				SetParentNode(Node->GetGuid(), ResolvedNode->GetGuid());
+				ResolvedNode = AddElementKeyNode(ResolvedKey);
 			}
+			SetParentNode(Node->GetGuid(), ResolvedNode->GetGuid());
+		}
+
+		if(InKey.Type == ERigElementType::Connector)
+		{
+			Node->AddTag<FControlRigSchematicWarningTag>();
 		}
 
 		if (bNotify && OnNodeAddedDelegate.IsBound())
@@ -198,7 +227,7 @@ FControlRigSchematicRigElementKeyLink* FControlRigSchematicModel::AddElementKeyL
 		FControlRigSchematicRigElementKeyLink* Link = AddLink<FControlRigSchematicRigElementKeyLink>(SourceNode->GetGuid(), TargetNode->GetGuid());
 		Link->SourceKey = InSourceKey;
 		Link->TargetKey = InTargetKey;
-		Link->Thickness = 12.f;
+		Link->Thickness = 16.f;
 		return Link;
 	}
 	return nullptr;
@@ -527,6 +556,8 @@ void FControlRigSchematicModel::OnHierarchyModified(ERigHierarchyNotification In
 
 void FControlRigSchematicModel::HandleModularRigModified(EModularRigNotification InNotification, const FRigModuleReference* InModule)
 {
+	bool bUnresolvedConnectors = false;
+	
 	switch (InNotification)
 	{
 		case EModularRigNotification::ConnectionChanged:
@@ -549,18 +580,21 @@ void FControlRigSchematicModel::HandleModularRigModified(EModularRigNotification
 									const FRigElementKey LocalKey(*LocalName, ConnectorKey.Type);
 									if(const FRigElementKey* ResolvedKey = InModule->Connections.Find(LocalKey))
 									{
-										if(const FSchematicGraphNode* ResolvedNode = FindElementKeyNode(*ResolvedKey))
+										const FControlRigSchematicRigElementKeyNode* ResolvedNode = FindElementKeyNode(*ResolvedKey);
+										if(ResolvedNode == nullptr)
 										{
-											(void)SetParentNode(ConnectorNode, ResolvedNode);
-											const_cast<FSchematicGraphNode*>(ResolvedNode)->OnMouseLeave();
-											const_cast<FSchematicGraphNode*>(ConnectorNode)->OnMouseLeave();
-											continue;
+											ResolvedNode = AddElementKeyNode(*ResolvedKey);
 										}
+										(void)SetParentNode(ConnectorNode, ResolvedNode);
+										const_cast<FControlRigSchematicRigElementKeyNode*>(ResolvedNode)->OnMouseLeave();
+										const_cast<FSchematicGraphNode*>(ConnectorNode)->OnMouseLeave();
+										continue;
 									}
-									if(const FSchematicGraphNode* ParentNode = ConnectorNode->GetParentNode())
+									if(ConnectorNode->GetParentNode())
 									{
 										(void)RemoveFromParentNode(ConnectorNode);
 										const_cast<FSchematicGraphNode*>(ConnectorNode)->OnMouseLeave();
+										bUnresolvedConnectors = true;
 									}
 								}
 							}
@@ -575,6 +609,12 @@ void FControlRigSchematicModel::HandleModularRigModified(EModularRigNotification
 			break;
 		}
 	}
+
+	// if we have unresolved connectors we may want to remove nodes which no longer have any connectors resolved to them
+	if(bUnresolvedConnectors && ControlRigBlueprint.IsValid())
+	{
+		bUpdateConnectorTargetsOnTick = true;
+	}
 }
 
 FSchematicGraphGroupNode* FControlRigSchematicModel::AddAutoGroupNode()
@@ -633,23 +673,29 @@ bool FControlRigSchematicModel::GetPositionAnimationEnabledForNode(const FSchema
 
 int32 FControlRigSchematicModel::GetNumLayersForNode(const FSchematicGraphNode* InNode) const
 {
-	if(const FControlRigSchematicRigElementKeyNode* Node = Cast<FControlRigSchematicRigElementKeyNode>(InNode))
+	if(Cast<FControlRigSchematicRigElementKeyNode>(InNode))
 	{
 		return 3;
 	}
 	return Super::GetNumLayersForNode(InNode);
 }
 
-const FSlateBrush* FControlRigSchematicModel::GetBrushForKey(const FRigElementKey& InKey, bool bResolveConnectors) const
+const FSlateBrush* FControlRigSchematicModel::GetBrushForKey(const FRigElementKey& InKey, const FSchematicGraphNode* InNode) const
 {
+	static const FSlateBrush* UnresolvedSocketBrush = FSchematicGraphStyle::Get().GetBrush( "Schematic.Dot.Small");
+	static const FSlateBrush* ResolvedSocketBrush = FSchematicGraphStyle::Get().GetBrush( "Schematic.Dot.Large");
+	static const FSlateBrush* ResolvedMultipleSocketBrush = FSchematicGraphStyle::Get().GetBrush( "Schematic.Dot.Medium");
+	static const FSlateBrush* ConnectorPrimaryBrush = FControlRigEditorStyle::Get().GetBrush( "ControlRig.Schematic.ConnectorPrimary");
+	static const FSlateBrush* ConnectorOptionalBrush = FControlRigEditorStyle::Get().GetBrush( "ControlRig.Schematic.ConnectorOptional");
+	static const FSlateBrush* ConnectorSecondaryBrush = FControlRigEditorStyle::Get().GetBrush( "ControlRig.Schematic.ConnectorSecondary");
+	static const FSlateBrush* BoneBrush = FControlRigEditorStyle::Get().GetBrush( "ControlRig.Schematic.Bone");
+	static const FSlateBrush* ControlBrush = FControlRigEditorStyle::Get().GetBrush( "ControlRig.Schematic.Control");
+	static const FSlateBrush* NullBrush = FControlRigEditorStyle::Get().GetBrush( "ControlRig.Schematic.Null");
+
 	switch(InKey.Type)
 	{
 		case ERigElementType::Socket:
 		{
-			static const FSlateBrush* UnresolvedSocketBrush = FSchematicGraphStyle::Get().GetBrush( "Schematic.Dot.Small");
-			static const FSlateBrush* ResolvedSocketBrush = FSchematicGraphStyle::Get().GetBrush( "Schematic.Dot.Large");
-			static const FSlateBrush* ResolvedMultipleSocketBrush = FSchematicGraphStyle::Get().GetBrush( "Schematic.Dot.Medium");
-
 			if (ControlRigBlueprint.IsValid())
 			{
 				int32 Count = 0;
@@ -679,30 +725,38 @@ const FSlateBrush* FControlRigSchematicModel::GetBrushForKey(const FRigElementKe
 		}
 		case ERigElementType::Connector:
 		{
-			static const FSlateBrush* ConnectorPrimaryBrush = FControlRigEditorStyle::Get().GetBrush( "ControlRig.Schematic.ConnectorPrimary");
-			static const FSlateBrush* ConnectorOptionalBrush = FControlRigEditorStyle::Get().GetBrush( "ControlRig.Schematic.ConnectorOptional");
-			static const FSlateBrush* ConnectorSecondaryBrush = FControlRigEditorStyle::Get().GetBrush( "ControlRig.Schematic.ConnectorSecondary");
-
-			if(bResolveConnectors)
-			{
-				FRigElementKey ResolvedKey;
-				if(IsConnectorResolved(InKey, &ResolvedKey))
-				{
-					if(ResolvedKey.Type != ERigElementType::Socket)
-					{
-						return GetBrushForKey(ResolvedKey, false);
-					}
-				}
-			}
-
 			if (ControlRigBlueprint.IsValid())
 			{
-				if(const UControlRig* ControlRig = ControlRigBlueprint->GetDebuggedControlRig())
+				if(const UModularRig* ControlRig = Cast<UModularRig>(ControlRigBlueprint->GetDebuggedControlRig()))
 				{
 					if(const URigHierarchy* Hierarchy = ControlRig->GetHierarchy())
 					{
 						if(const FRigConnectorElement* Connector = Hierarchy->Find<FRigConnectorElement>(InKey))
 						{
+							if((Connector->Settings.Type == EConnectorType::Primary) && !IsConnectorResolved(InKey))
+							{
+								const FName ModulePathName = Hierarchy->GetNameMetadata(Connector->GetKey(), URigHierarchy::ModuleMetadataName, NAME_None);
+								if(!ModulePathName.IsNone())
+								{
+									if(const FRigModuleInstance* ModuleInstance = ControlRig->FindModule(ModulePathName.ToString()))
+									{
+										const FSoftObjectPath& IconPath = ModuleInstance->GetRig()->GetRigModuleSettings().Icon;
+										const FSlateBrush* IconBrush = ModuleIcons.Find(IconPath);
+										if (!IconBrush)
+										{
+											if(UTexture2D* Icon = Cast<UTexture2D>(IconPath.TryLoad()))
+											{
+												IconBrush = &ModuleIcons.Add(IconPath, UWidgetBlueprintLibrary::MakeBrushFromTexture(Icon, 16.0f, 16.0f));
+											}
+										}
+										if(IconBrush)
+										{
+											return IconBrush;
+										}
+									}
+								}
+							}
+
 							if(Connector->Settings.Type == EConnectorType::Secondary)
 							{
 								return Connector->Settings.bOptional ? 	ConnectorOptionalBrush : ConnectorSecondaryBrush;
@@ -713,23 +767,54 @@ const FSlateBrush* FControlRigSchematicModel::GetBrushForKey(const FRigElementKe
 			}
 			return ConnectorPrimaryBrush;
 		}
-		case ERigElementType::Bone:
-		{
-			static const FSlateBrush* BoneBrush = FControlRigEditorStyle::Get().GetBrush( "ControlRig.Schematic.Bone");
-			return BoneBrush;
-		}
-		case ERigElementType::Control:
-		{
-			static const FSlateBrush* ControlBrush = FControlRigEditorStyle::Get().GetBrush( "ControlRig.Schematic.Control");
-			return ControlBrush;
-		}
-		case ERigElementType::Null:
-		{
-			static const FSlateBrush* NullBrush = FControlRigEditorStyle::Get().GetBrush( "ControlRig.Schematic.Null");
-			return NullBrush;
-		}
 		default:
 		{
+			// check if the key has a resolved connector
+			if(const FSchematicGraphGroupNode* GroupNode = Cast<FSchematicGraphGroupNode>(InNode))
+			{
+				if(!GroupNode->IsExpanded() && !GroupNode->IsExpanding())
+				{
+					if (ControlRigBlueprint.IsValid())
+					{
+						int32 Count = 0;
+						for(const TPair<FRigElementKey, FRigElementKey>& Pair : ControlRigBlueprint->ConnectionMap)
+						{
+							if(Pair.Value == InKey)
+							{
+								Count++;
+								if(Count == 2)
+								{
+									return ResolvedMultipleSocketBrush;
+								}
+							}
+						}
+						if(Count > 0)
+						{
+							return ResolvedSocketBrush;
+						}
+					}
+				}
+			}
+			
+			switch(InKey.Type)
+			{
+				case ERigElementType::Bone:
+				{
+					return BoneBrush;
+				}
+				case ERigElementType::Control:
+				{
+					return ControlBrush;
+				}
+				case ERigElementType::Null:
+				{
+					return NullBrush;
+				}
+				default:
+				{
+					break;
+				}
+			}
 			break;
 		}
 	}
@@ -750,20 +835,17 @@ const FSlateBrush* FControlRigSchematicModel::GetBrushForNode(const FSchematicGr
 			static const FSlateBrush* OutlineSingleBrush = FSchematicGraphStyle::Get().GetBrush( "Schematic.Outline.Single");
 			static const FSlateBrush* OutlineDoubleBrush = FSchematicGraphStyle::Get().GetBrush( "Schematic.Outline.Double");
 
-			if(Node->GetKey().Type == ERigElementType::Socket)
+			if (ControlRigBlueprint.IsValid())
 			{
-				if (ControlRigBlueprint.IsValid())
+				int32 Count = 0;
+				for(const TPair<FRigElementKey, FRigElementKey>& Pair : ControlRigBlueprint->ConnectionMap)
 				{
-					int32 Count = 0;
-					for(const TPair<FRigElementKey, FRigElementKey>& Pair : ControlRigBlueprint->ConnectionMap)
+					if(Pair.Value == Node->GetKey())
 					{
-						if(Pair.Value == Node->GetKey())
+						Count++;
+						if(Count == 2)
 						{
-							Count++;
-							if(Count == 2)
-							{
-								return OutlineDoubleBrush;
-							}
+							return OutlineDoubleBrush;
 						}
 					}
 				}
@@ -771,7 +853,7 @@ const FSlateBrush* FControlRigSchematicModel::GetBrushForNode(const FSchematicGr
 			return OutlineSingleBrush;
 		}
 
-		if(const FSlateBrush* Brush = GetBrushForKey(Node->GetKey(), true))
+		if(const FSlateBrush* Brush = GetBrushForKey(Node->GetKey(), Node))
 		{
 			return Brush;
 		}
@@ -862,6 +944,56 @@ ESchematicGraphPlacementConstraint::Type FControlRigSchematicModel::GetPlacement
 	return Super::GetPlacementForNode(InNode);
 }
 
+ESchematicGraphVisibility::Type FControlRigSchematicModel::GetVisibilityForNode(const FSchematicGraphNode* InNode) const
+{
+	if(ControlRigBlueprint.IsValid())
+	{
+		if(const FControlRigSchematicRigElementKeyNode* Node = Cast<FControlRigSchematicRigElementKeyNode>(InNode))
+		{
+			if(Node->GetKey().Type == ERigElementType::Connector)
+			{
+				// hide secondary connectors if the primary connector isn't resolved yet
+				if (const UModularRig* ControlRig = Cast<UModularRig>(ControlRigBlueprint->GetDebuggedControlRig()))
+				{
+					if(const URigHierarchy* Hierarchy = ControlRig->GetHierarchy())
+					{
+						if(const FRigConnectorElement* Connector = Hierarchy->Find<FRigConnectorElement>(Node->GetKey()))
+						{
+							if(Connector->Settings.Type != EConnectorType::Primary)
+							{
+								if(Connector->Settings.bOptional)
+								{
+									if(!IsConnectorResolved(Connector->GetKey()))
+									{
+										return ESchematicGraphVisibility::Hidden;
+									}
+								}
+								
+								// find the corresponding primary connector
+								const FName ModulePathName = Hierarchy->GetNameMetadata(Node->GetKey(), URigHierarchy::ModuleMetadataName, NAME_None);
+								if(!ModulePathName.IsNone())
+								{
+									if(const FRigModuleInstance* ModuleInstance = ControlRig->FindModule(ModulePathName.ToString()))
+									{
+										if(const FRigConnectorElement* PrimaryConnector = ModuleInstance->FindPrimaryConnector())
+										{
+											if(!IsConnectorResolved(PrimaryConnector->GetKey()))
+											{
+												return ESchematicGraphVisibility::Hidden;
+											}
+										}
+									}
+								}
+							}
+						}
+					}
+				}
+			}
+		}
+	}
+	return Super::GetVisibilityForNode(InNode);
+}
+
 const FSlateBrush* FControlRigSchematicModel::GetBrushForLink(const FSchematicGraphLink* InLink) const
 {
 	if(InLink->IsA<FControlRigSchematicRigElementKeyLink>())
@@ -905,33 +1037,12 @@ ESchematicGraphVisibility::Type FControlRigSchematicModel::GetVisibilityForTag(c
 	{
 		if(const FControlRigSchematicRigElementKeyNode* Node = Cast<FControlRigSchematicRigElementKeyNode>(InTag->GetNode()))
 		{
-			if(Node->GetKey().Type == ERigElementType::Socket)
+			if(Node->GetNumChildNodes() < 3)
 			{
-				if(Node->GetNumChildNodes() < 3)
-				{
-					return ESchematicGraphVisibility::Hidden;
-				}
+				return ESchematicGraphVisibility::Hidden;
 			}
 		}
 	}
-	if(InTag->IsA<FControlRigSchematicConnectorTag>())
-	{
-		if(const FControlRigSchematicRigElementKeyNode* Node = Cast<FControlRigSchematicRigElementKeyNode>(InTag->GetNode()))
-		{
-			if(Node->GetKey().Type == ERigElementType::Connector)
-			{
-				FRigElementKey ResolvedKey;
-				if(IsConnectorResolved(Node->GetKey(), &ResolvedKey))
-				{
-					if(ResolvedKey.Type != ERigElementType::Socket)
-					{
-						return ESchematicGraphVisibility::Visible;
-					}
-				}
-			}
-		}
-		return ESchematicGraphVisibility::Hidden;
-	}
 	if(InTag->IsA<FControlRigSchematicWarningTag>())
 	{
 		if(const FControlRigSchematicRigElementKeyNode* ConnectorNode = Cast<FControlRigSchematicRigElementKeyNode>(InTag->GetNode()))
@@ -967,18 +1078,6 @@ ESchematicGraphVisibility::Type FControlRigSchematicModel::GetVisibilityForTag(c
 	return Super::GetVisibilityForTag(InTag);
 }
 
-const FSlateBrush* FControlRigSchematicModel::GetForegroundBrushForTag(const FSchematicGraphTag* InTag) const
-{
-	if(const FControlRigSchematicConnectorTag* ConnectorTag = Cast<FControlRigSchematicConnectorTag>(InTag))
-	{
-		if(const FControlRigSchematicRigElementKeyNode* Node = Cast<FControlRigSchematicRigElementKeyNode>(ConnectorTag->GetNode()))
-		{
-			return GetBrushForKey(Node->GetKey(), false);
-		}
-	}
-	return Super::GetForegroundBrushForTag(InTag);
-}
-
 const FText FControlRigSchematicModel::GetToolTipForTag(const FSchematicGraphTag* InTag) const
 {
 	if(InTag->IsA<FControlRigSchematicWarningTag>())
diff --git a/Engine/Plugins/Animation/ControlRig/Source/ControlRigEditor/Private/Editor/ControlRigSchematicModel.h b/Engine/Plugins/Animation/ControlRig/Source/ControlRigEditor/Private/Editor/ControlRigSchematicModel.h
--- a/Engine/Plugins/Animation/ControlRig/Source/ControlRigEditor/Private/Editor/ControlRigSchematicModel.h
+++ b/Engine/Plugins/Animation/ControlRig/Source/ControlRigEditor/Private/Editor/ControlRigSchematicModel.h
@@ -53,17 +53,6 @@ protected:
 	friend class FControlRigSchematicModel;
 };
 
-/** Tag to mark up connectors */
-class FControlRigSchematicConnectorTag : public FSchematicGraphTag
-{
-public:
-
-	SCHEMATICGRAPHTAG_BODY(FControlRigSchematicConnectorTag, FSchematicGraphTag)
-
-	FControlRigSchematicConnectorTag();
-	virtual ~FControlRigSchematicConnectorTag() override {}
-};
-
 class FControlRigSchematicWarningTag : public FSchematicGraphTag
 {
 public:
@@ -87,6 +76,7 @@ public:
 	void SetEditor(const TSharedRef<FControlRigEditor>& InEditor);
 
 	virtual void Reset() override;
+	virtual void Tick(float InDeltaTime) override;
 	FControlRigSchematicRigElementKeyNode* AddElementKeyNode(const FRigElementKey& InKey, bool bNotify = true);
 	const FControlRigSchematicRigElementKeyNode* FindElementKeyNode(const FRigElementKey& InKey) const;
 	bool ContainsElementKeyNode(const FRigElementKey& InKey) const;
@@ -105,14 +95,14 @@ public:
 	virtual FVector2d GetPositionForNode(const FSchematicGraphNode* InNode) const override;
 	virtual bool GetPositionAnimationEnabledForNode(const FSchematicGraphNode* InNode) const override;
 	virtual int32 GetNumLayersForNode(const FSchematicGraphNode* InNode) const override;
-	const FSlateBrush* GetBrushForKey(const FRigElementKey& InKey, bool bResolveConnectors) const;
+	const FSlateBrush* GetBrushForKey(const FRigElementKey& InKey, const FSchematicGraphNode* InNode) const;
 	virtual const FSlateBrush* GetBrushForNode(const FSchematicGraphNode* InNode, int32 InLayerIndex) const override;
 	virtual FLinearColor GetColorForNode(const FSchematicGraphNode* InNode, int32 InLayerIndex) const override;
 	virtual ESchematicGraphPlacementConstraint::Type GetPlacementForNode(const FSchematicGraphNode* InNode) const override;
+	virtual ESchematicGraphVisibility::Type GetVisibilityForNode(const FSchematicGraphNode* InNode) const override;
 	virtual const FSlateBrush* GetBrushForLink(const FSchematicGraphLink* InLink) const override;
 	virtual FLinearColor GetColorForLink(const FSchematicGraphLink* InLink) const override;
 	virtual ESchematicGraphVisibility::Type GetVisibilityForTag(const FSchematicGraphTag* InTag) const override;
-	virtual const FSlateBrush* GetForegroundBrushForTag(const FSchematicGraphTag* InTag) const override;
 	virtual const FText GetToolTipForTag(const FSchematicGraphTag* InTag) const override;
 
 private:
@@ -131,6 +121,8 @@ private:
 
 	TArray<FGuid> TemporaryNodeGuids;
 	TMap<FRigElementKey, FGuid> RigElementKeyToGuid;
+	mutable TMap<FSoftObjectPath, FSlateBrush> ModuleIcons;
+	bool bUpdateConnectorTargetsOnTick = false;
 
 	friend class FControlRigEditor;
 };
\ No newline at end of file
diff --git a/Engine/Source/Editor/AnimationEditorWidgets/Private/SchematicGraphPanel/SSchematicGraphPanel.cpp b/Engine/Source/Editor/AnimationEditorWidgets/Private/SchematicGraphPanel/SSchematicGraphPanel.cpp
--- a/Engine/Source/Editor/AnimationEditorWidgets/Private/SchematicGraphPanel/SSchematicGraphPanel.cpp
+++ b/Engine/Source/Editor/AnimationEditorWidgets/Private/SchematicGraphPanel/SSchematicGraphPanel.cpp
@@ -77,7 +77,7 @@ void SSchematicGraphNode::Construct(const FArguments& InArgs)
 {
 	if(InArgs._NodeData)
 	{
-		NodeData = const_cast<FSchematicGraphNode*>(InArgs._NodeData);
+		NodeData = const_cast<FSchematicGraphNode*>(InArgs._NodeData)->AsShared();
 	}
 	OnClickedDelegate = InArgs._OnClicked;
 	OnBeginDragDelegate = InArgs._OnBeginDrag;
@@ -99,7 +99,7 @@ void SSchematicGraphNode::Construct(const FArguments& InArgs)
 
 	SetVisibility(TAttribute<EVisibility>::CreateSP(this, &SSchematicGraphNode::GetNodeVisibility));
 
-	if(const FSchematicGraphGroupNode* GroupNode = Cast<FSchematicGraphGroupNode>(NodeData))
+	if(const FSchematicGraphGroupNode* GroupNode = Cast<FSchematicGraphGroupNode>(GetNodeData()))
 	{
 		ExpansionCircleFactor = FFloatAttribute::Create(GroupNode->GetAnimationSettings(), 0.f);
 	}
@@ -121,7 +121,7 @@ int32 SSchematicGraphNode::OnPaint(const FPaintArgs& Args, const FGeometry& Allo
 	const bool bIsFadedOut = IsFadedOut();
 	const float FadedOutFactor = bIsFadedOut ? 0.5f : 1.f;
 
-	if(const FSchematicGraphGroupNode* GroupNode = Cast<FSchematicGraphGroupNode>(NodeData))
+	if(const FSchematicGraphGroupNode* GroupNode = Cast<FSchematicGraphGroupNode>(GetNodeData()))
 	{
 		if(GroupNode->GetExpansionState() > SMALL_NUMBER)
 		{
@@ -494,7 +494,7 @@ const bool SSchematicGraphNode::IsFadedOut() const
 	{
 		if(const FSchematicGraphModel* Graph = NodeData->GetGraph())
 		{
-			return Graph->GetVisibilityForNode(NodeData) == ESchematicGraphVisibility::FadedOut;
+			return Graph->GetVisibilityForNode(GetNodeData()) == ESchematicGraphVisibility::FadedOut;
 		}
 	}
 	return false;
@@ -535,11 +535,16 @@ void SSchematicGraphPanel::Construct(const FArguments& InArgs)
 	OnBeginDragDelegate = InArgs._OnBeginDrag;
 	OnEndDragDelegate = InArgs._OnEndDrag;
 	OnDropDelegate = InArgs._OnDrop;
+	BackgroundAlpha = InArgs._BackgroundAlpha;
 
-	TEasingAttributeInterpolator<float>::FSettings FloatInterpolationSettings(EEasingInterpolatorType::CubicEaseOut, 0.35f);
+	TEasingAttributeInterpolator<float>::FSettings FloatInterpolationSettings(EEasingInterpolatorType::CubicEaseOut, 0.75f);
 	FadeBackgroundAlpha = FFloatAttribute::CreateWithGetter(FloatInterpolationSettings, FFloatAttribute::FGetter::CreateLambda([this]() -> float
 	{
-		return bIsDragDropping ? 0.5f : 0.f;
+		if(!bIsOverlay || (GetVisibility() == EVisibility::Hidden))
+		{
+			return 0.f;
+		}
+		return BackgroundAlpha.Get();
 	}));
 	
 	SNodePanel::Construct();
@@ -578,7 +583,11 @@ void SSchematicGraphPanel::AddNode(const FSchematicGraphNode* InNodeToAdd)
 	{
 		return;
 	}
-
+	if(NodeByGuid.Contains(InNodeToAdd->GetGuid()))
+	{
+		return;
+	}
+	
 	static const TEasingAttributeInterpolator<FVector2d>::FSettings Vector2DInterpolationSettings(EEasingInterpolatorType::CubicEaseOut, 0.2f);
 	static const TEasingAttributeInterpolator<float>::FSettings FloatInterpolationSettings(EEasingInterpolatorType::CubicEaseOut, 0.2f);
 	static const TEasingAttributeInterpolator<FLinearColor>::FSettings ColorInterpolationSettings(EEasingInterpolatorType::CubicEaseOut, 0.2f);
@@ -607,11 +616,11 @@ void SSchematicGraphPanel::AddNode(const FSchematicGraphNode* InNodeToAdd)
 														.LayerColors(Colors)
 														.EnableAutoScale(this, &SSchematicGraphPanel::IsAutoScaleEnabledForNode, InNodeToAdd->GetGuid())
 														.BrushGetter(BrushGetter)
-														.ToolTipText_Raw(GraphData, &FSchematicGraphModel::GetToolTipForNode, InNodeToAdd)
-														.OnClicked_Raw(this, &SSchematicGraphPanel::OnNodeClicked)
-														.OnBeginDrag_Raw(this, &SSchematicGraphPanel::OnBeginDragEvent)
-														.OnEndDrag_Raw(this, &SSchematicGraphPanel::OnEndDragEvent)
-														.OnDrop_Raw(this, &SSchematicGraphPanel::OnDropEvent)
+														.ToolTipText(this, &SSchematicGraphPanel::GetToolTipForNode, InNodeToAdd->GetGuid())
+														.OnClicked(this, &SSchematicGraphPanel::OnNodeClicked)
+														.OnBeginDrag(this, &SSchematicGraphPanel::OnBeginDragEvent)
+														.OnEndDrag(this, &SSchematicGraphPanel::OnEndDragEvent)
+														.OnDrop(this, &SSchematicGraphPanel::OnDropEvent)
 														.NodeData(InNodeToAdd);
 	SNodePanel::AddGraphNode(NewNode);
 	NewNode->SchematicGraphPanel = this;
@@ -712,7 +721,9 @@ void SSchematicGraphPanel::OnArrangeChildren(const FGeometry& AllottedGeometry,
 
 int32 SSchematicGraphPanel::OnPaint(const FPaintArgs& Args, const FGeometry& AllottedGeometry, const FSlateRect& MyCullingRect, FSlateWindowElementList& OutDrawElements, int32 LayerId, const FWidgetStyle& InWidgetStyle, bool bParentEnabled) const
 {
-	const int32 LinkLayerId = LayerId + 1;
+	const int32 BackgroundLayer = LayerId + 1;
+	const int32 OutlineLayer = BackgroundLayer + 1;
+	const int32 LinkLayerId = OutlineLayer + 2;
 	const int32 NodeLayerId = LinkLayerId + 1;
     int32 MaxLayerId = NodeLayerId;
     	
@@ -722,22 +733,6 @@ int32 SSchematicGraphPanel::OnPaint(const FPaintArgs& Args, const FGeometry& All
 		PaintBackgroundAsLines(DefaultBackground, AllottedGeometry, MyCullingRect, OutDrawElements, LayerId);
 		MaxLayerId++;
 	}
-
-	const float BackgroundAlpha = FadeBackgroundAlpha->Get();
-	if (BackgroundAlpha > 0.f)
-	{
-		const FSlateBrush* Brush = FAppStyle::GetBrush(TEXT("Graph.Panel.SolidBackground"));
-		const FLinearColor TransparentGrey = FLinearColor(0.5, 0.5, 0.5, BackgroundAlpha);
-		FSlateDrawElement::MakeBox(
-				OutDrawElements,
-				MaxLayerId,
-				AllottedGeometry.ToPaintGeometry(),
-				Brush,
-				ESlateDrawEffect::None,
-				TransparentGrey
-				);
-		MaxLayerId++;
-	}
 	
 	if (!GraphData)
 	{
@@ -749,6 +744,8 @@ int32 SSchematicGraphPanel::OnPaint(const FPaintArgs& Args, const FGeometry& All
 
 	NodeCenterByGuid.Reset();
 	NodeCenterByGuid.Reserve(ArrangedChildren.Num());
+	NodeCenterByIndex.Reset();
+	NodeCenterByIndex.Reserve(ArrangedChildren.Num());
 	NodeVisibilityByIndex.Reset();
 	NodeVisibilityByIndex.Reserve(ArrangedChildren.Num());
 	NodeVisibilityByGuid.Reset();
@@ -759,7 +756,9 @@ int32 SSchematicGraphPanel::OnPaint(const FPaintArgs& Args, const FGeometry& All
 		FArrangedWidget& CurWidget = ArrangedChildren[ChildIndex];
 		const TSharedRef<SSchematicGraphNode> ChildNode = StaticCastSharedRef<SSchematicGraphNode>(CurWidget.Widget);
 
-		NodeCenterByGuid.Add(ChildNode->GetGuid(), CurWidget.Geometry.GetLocalPositionAtCoordinates({0.5, 0.5}));
+		const FVector2d NodeCenter = CurWidget.Geometry.GetLocalPositionAtCoordinates({0.5, 0.5});
+		NodeCenterByIndex.Add(NodeCenter);
+		NodeCenterByGuid.Add(ChildNode->GetGuid(), NodeCenter);
 
 		const FSchematicGraphNode* NodeData = ChildNode->GetNodeData();
 
@@ -809,6 +808,93 @@ int32 SSchematicGraphPanel::OnPaint(const FPaintArgs& Args, const FGeometry& All
 		}
 	}
 
+	FBox2d Bounds(EForceInit::ForceInit);
+	for (int32 ChildIndex = 0; ChildIndex < ArrangedChildren.Num(); ++ChildIndex)
+	{
+		if(NodeVisibilityByIndex[ChildIndex] == ESchematicGraphVisibility::Hidden)
+		{
+			continue;
+		}
+		if(NodeCenterByIndex[ChildIndex].IsNearlyZero())
+		{
+			continue;
+		}
+		Bounds += NodeCenterByIndex[ChildIndex] + FVector2d(120, 120);
+		Bounds += NodeCenterByIndex[ChildIndex] - FVector2d(120, 120);
+	}
+
+	const float Alpha = FadeBackgroundAlpha->Get();
+	if (Alpha > 0.f && !Bounds.GetSize().IsNearlyZero())
+	{
+		const FSlateBrush* Brush = FAppStyle::GetBrush(TEXT("Graph.Panel.SolidBackground"));
+		const FLinearColor TransparentGrey = FLinearColor(0.5, 0.5, 0.5, Alpha);
+		FSlateDrawElement::MakeBox(
+				OutDrawElements,
+				BackgroundLayer,
+				AllottedGeometry.ToPaintGeometry(Bounds.GetSize(), FSlateLayoutTransform(Bounds.Min)),
+				Brush,
+				ESlateDrawEffect::None,
+				TransparentGrey
+			);
+
+		TArray<FVector2d> LinePoints;
+		LinePoints.AddZeroed(2);
+
+		constexpr float Step = 10.f;
+		for(float X = Bounds.Min.X; X < Bounds.Max.X - Step; X += Step * 2)
+		{
+			LinePoints[0] = FVector2d(X, Bounds.Min.Y + 0.5f);
+			LinePoints[1] = LinePoints[0] + FVector2d(Step, 0.f);
+			
+			FSlateDrawElement::MakeLines(
+				OutDrawElements,
+				OutlineLayer,
+				AllottedGeometry.ToPaintGeometry(),
+				LinePoints,
+				ESlateDrawEffect::None,
+				FLinearColor(0.7, 0.7, 0.7, 0.95),
+				true);
+
+			LinePoints[0] = FVector2d(X, Bounds.Max.Y - 0.5f);
+			LinePoints[1] = LinePoints[0] + FVector2d(Step, 0.f);
+			
+			FSlateDrawElement::MakeLines(
+				OutDrawElements,
+				OutlineLayer,
+				AllottedGeometry.ToPaintGeometry(),
+				LinePoints,
+				ESlateDrawEffect::None,
+				FLinearColor(0.7, 0.7, 0.7, 0.95),
+				true);
+		}
+		for(float Y = Bounds.Min.Y; Y < Bounds.Max.Y - Step; Y += Step * 2)
+		{
+			LinePoints[0] = FVector2d(Bounds.Min.X + 0.5f, Y);
+			LinePoints[1] = LinePoints[0] + FVector2d(0.f, Step);
+			
+			FSlateDrawElement::MakeLines(
+				OutDrawElements,
+				OutlineLayer,
+				AllottedGeometry.ToPaintGeometry(),
+				LinePoints,
+				ESlateDrawEffect::None,
+				FLinearColor(0.7, 0.7, 0.7, 0.95),
+				true);
+
+			LinePoints[0] = FVector2d(Bounds.Max.X - 0.5f, Y);
+			LinePoints[1] = LinePoints[0] + FVector2d(0.f, Step);
+			
+			FSlateDrawElement::MakeLines(
+				OutDrawElements,
+				OutlineLayer,
+				AllottedGeometry.ToPaintGeometry(),
+				LinePoints,
+				ESlateDrawEffect::None,
+				FLinearColor(0.7, 0.7, 0.7, 0.95),
+				true);
+		}
+	}
+	
 	// draw all of the links
 	for(const TPair<FGuid, TSharedPtr<FSchematicLinkWidgetInfo>>& Pair : LinkByGuid)
 	{
@@ -1014,6 +1100,10 @@ void SSchematicGraphPanel::Tick(float DeltaTime)
 
 		// update the animation state of the node
 		Widget->EnablePositionAnimation(GraphData->GetPositionAnimationEnabledForNode(Widget->GetGuid()));
+		if(Widget->GetVisibility() != EVisibility::Visible)
+		{
+			continue;
+		}
 
 		// collect the nodes displayed in the corners
 		switch(GraphData->GetPlacementForNode(Widget->GetNodeData()))
@@ -1076,6 +1166,15 @@ void SSchematicGraphPanel::Tick(const FGeometry& AllottedGeometry, const double
 	Tick(InDeltaTime);
 }
 
+void SSchematicGraphPanel::ToggleVisibility()
+{
+	const EVisibility PreviousVisibility = GetVisibility();
+	SetVisibility(
+		PreviousVisibility == EVisibility::Hidden ?
+		EVisibility::SelfHitTestInvisible :
+		EVisibility::Hidden);
+}
+
 void SSchematicGraphPanel::OnNodeClicked(SSchematicGraphNode* Node, const FPointerEvent& MouseEvent)
 {
 	OnNodeClickedDelegate.ExecuteIfBound(this, Node, MouseEvent);
@@ -1187,6 +1286,15 @@ FLinearColor SSchematicGraphPanel::GetColorForNode(FGuid InNodeGuid, int32 InLay
 	return FLinearColor::White;
 }
 
+FText SSchematicGraphPanel::GetToolTipForNode(FGuid InNodeGuid) const
+{
+	if(GraphData)
+	{
+		return GraphData->GetToolTipForNode(InNodeGuid);
+	}
+	return FText();
+}
+
 float SSchematicGraphPanel::GetScaleForNode(FGuid InNodeGuid) const
 {
 	// check if the node may be auto scaled
diff --git a/Engine/Source/Editor/AnimationEditorWidgets/Public/SchematicGraphPanel/SSchematicGraphPanel.h b/Engine/Source/Editor/AnimationEditorWidgets/Public/SchematicGraphPanel/SSchematicGraphPanel.h
--- a/Engine/Source/Editor/AnimationEditorWidgets/Public/SchematicGraphPanel/SSchematicGraphPanel.h
+++ b/Engine/Source/Editor/AnimationEditorWidgets/Public/SchematicGraphPanel/SSchematicGraphPanel.h
@@ -97,8 +97,8 @@ public:
 
 	const FVector2d& GetOriginalSize() const { return OriginalSize; }
 
-	const FSchematicGraphNode* GetNodeData() const { return NodeData; }
-	FSchematicGraphNode* GetNodeData() { return NodeData; }
+	const FSchematicGraphNode* GetNodeData() const { return NodeData.Get(); }
+	FSchematicGraphNode* GetNodeData() { return NodeData.Get(); }
 	const FGuid GetGuid() const;
 	bool IsInteractive() const;
 
@@ -111,7 +111,7 @@ private:
 	static inline const FVector2d DefaultNodeSize = FVector2d(32.0,32.0);  
 	FVector2d OriginalSize = DefaultNodeSize;
 
-	FSchematicGraphNode* NodeData = nullptr;
+	TSharedPtr<FSchematicGraphNode> NodeData;
 	TSharedPtr<FVector2dAttribute> Position;
 	TOptional<FVector2d> PositionDuringDrag;
 	TOptional<FVector2d> OffsetDuringDrag;
@@ -147,7 +147,9 @@ public:
 	DECLARE_DELEGATE_ThreeParams(FOnEndDrag, SSchematicGraphPanel*, SSchematicGraphNode*, const FDragDropOperation&);
 	DECLARE_DELEGATE_ThreeParams(FOnDrop, SSchematicGraphPanel*, SSchematicGraphNode*, const FDragDropEvent&);
 	
-	SLATE_BEGIN_ARGS(SSchematicGraphPanel) {}
+	SLATE_BEGIN_ARGS(SSchematicGraphPanel)
+		: _BackgroundAlpha(0.5f)
+	{}
 	SLATE_ARGUMENT(bool, IsOverlay)
 	SLATE_ARGUMENT(FSchematicGraphModel*, GraphData)
 	SLATE_ARGUMENT(int32, PaddingLeft)
@@ -159,6 +161,7 @@ public:
 	SLATE_EVENT(FOnBeginDrag, OnBeginDrag)
 	SLATE_EVENT(FOnEndDrag, OnEndDrag)
 	SLATE_EVENT(FOnDrop, OnDrop)
+	SLATE_ATTRIBUTE(float, BackgroundAlpha)
 	SLATE_END_ARGS()
 
 	struct FSchematicLinkWidgetInfo
@@ -206,6 +209,7 @@ public:
 	virtual bool IsTickable() const override { return true; }
 	// End of FTickableEditorObject interface
 
+	void ToggleVisibility();
 	void OnNodeClicked(SSchematicGraphNode* Node, const FPointerEvent& MouseEvent);
 	void OnBeginDragEvent(SSchematicGraphNode* Node, const FDragDropOperation& InDragDropEvent);
 	void OnEndDragEvent(SSchematicGraphNode* Node, const FDragDropOperation& InDragDropEvent);
@@ -214,6 +218,7 @@ public:
 
 	virtual FVector2d GetPositionForNode(FGuid InNodeGuid) const;
 	virtual FLinearColor GetColorForNode(FGuid InNodeGuid, int32 InLayerIndex) const;
+	virtual FText GetToolTipForNode(FGuid InNodeGuid) const;
 	virtual float GetScaleForNode(FGuid InNodeGuid) const;
 
 	virtual bool IsAutoGroupingEnabled() const;
@@ -240,6 +245,7 @@ private:
 	int32 PaddingTop = 0;
 	int32 PaddingBottom = 0;
 	int32 PaddingInterNode = 0;
+	TAttribute<float> BackgroundAlpha;
 	TSharedPtr<FFloatAttribute> FadeBackgroundAlpha;
 	FSchematicGraphModel* GraphData = nullptr;
 	FOnNodeClicked OnNodeClickedDelegate;
@@ -275,6 +281,7 @@ private:
 	TMap<FGuid, TSharedPtr<FSchematicLinkWidgetInfo>> LinkByGuid;
 
 	mutable TMap<FGuid, FVector2d> NodeCenterByGuid;
+	mutable TArray<FVector2d> NodeCenterByIndex;
 	mutable TArray<ESchematicGraphVisibility::Type> NodeVisibilityByIndex;
 	mutable TMap<FGuid, ESchematicGraphVisibility::Type> NodeVisibilityByGuid;
 };
diff --git a/Engine/Source/Editor/AnimationEditorWidgets/Public/SchematicGraphPanel/SchematicGraphModel.h b/Engine/Source/Editor/AnimationEditorWidgets/Public/SchematicGraphPanel/SchematicGraphModel.h
--- a/Engine/Source/Editor/AnimationEditorWidgets/Public/SchematicGraphPanel/SchematicGraphModel.h
+++ b/Engine/Source/Editor/AnimationEditorWidgets/Public/SchematicGraphPanel/SchematicGraphModel.h
@@ -271,7 +271,7 @@ public:
 		SetLastExpandedNode(nullptr);
 	}
 
-	void Tick(float InDeltaTime);
+	virtual void Tick(float InDeltaTime);
 
 protected:
 
