Control Rig:  Fix broken space evaluation and handle fact that we may no longer ::Enter our edit mode after ::Exit so we need to make sure to keep our old toolkit and handle issues related to that.

diff --git a/Engine/Plugins/Animation/ControlRig/Source/ControlRig/Private/ControlRig.cpp b/Engine/Plugins/Animation/ControlRig/Source/ControlRig/Private/ControlRig.cpp
--- a/Engine/Plugins/Animation/ControlRig/Source/ControlRig/Private/ControlRig.cpp
+++ b/Engine/Plugins/Animation/ControlRig/Source/ControlRig/Private/ControlRig.cpp
@@ -2025,10 +2025,13 @@ void UControlRig::SwitchToParent(const FRigElementKey& InElementKey, const FRigE
 	{
 		return;
 	}
-	FRigBaseElement* Parent = DynamicHierarchy->Find<FRigBaseElement>(InNewParentKey);
-	if(Parent == nullptr)
+	if (InNewParentKey.Type != ERigElementType::Reference)
 	{
-		return;
+		FRigBaseElement* Parent = DynamicHierarchy->Find<FRigBaseElement>(InNewParentKey);
+		if (Parent == nullptr)
+		{
+			return;
+		}
 	}
 	if (bIsAdditive)
 	{
diff --git a/Engine/Plugins/Animation/ControlRig/Source/ControlRig/Private/Sequencer/MovieSceneControlRigParameterTemplate.cpp b/Engine/Plugins/Animation/ControlRig/Source/ControlRig/Private/Sequencer/MovieSceneControlRigParameterTemplate.cpp
--- a/Engine/Plugins/Animation/ControlRig/Source/ControlRig/Private/Sequencer/MovieSceneControlRigParameterTemplate.cpp
+++ b/Engine/Plugins/Animation/ControlRig/Source/ControlRig/Private/Sequencer/MovieSceneControlRigParameterTemplate.cpp
@@ -785,7 +785,6 @@ struct FControlRigParameterPreAnimatedTokenProducer : IMovieScenePreAnimatedToke
 							const bool bSetupUndo = false;
 							if (URigHierarchy* RigHierarchy = ControlRig->GetHierarchy())
 							{
-
 								FRigElementKey ControlKey;
 								ControlKey.Type = ERigElementType::Control;
 								for (FControlSpaceAndValue& SpaceNameAndValue : SpaceValues)
@@ -793,19 +792,19 @@ struct FControlRigParameterPreAnimatedTokenProducer : IMovieScenePreAnimatedToke
 									ControlKey.Name = SpaceNameAndValue.ControlName;
 									switch (SpaceNameAndValue.Value.SpaceType)
 									{
-									case EMovieSceneControlRigSpaceType::Parent:
-										RigHierarchy->SwitchToDefaultParent(ControlKey);
+										case EMovieSceneControlRigSpaceType::Parent:
+											ControlRig->SwitchToParent(ControlKey, RigHierarchy->GetDefaultParent(ControlKey), false, true);
 										break;
-									case EMovieSceneControlRigSpaceType::World:
-										RigHierarchy->SwitchToWorldSpace(ControlKey);
+										case EMovieSceneControlRigSpaceType::World:
+											ControlRig->SwitchToParent(ControlKey, RigHierarchy->GetWorldSpaceReferenceKey(), false, true);
 										break;
-									case EMovieSceneControlRigSpaceType::ControlRig:
+										case EMovieSceneControlRigSpaceType::ControlRig:
 										{
 #if WITH_EDITOR
 											ControlRig->SwitchToParent(ControlKey, SpaceNameAndValue.Value.ControlRigElement, false, true);
 #else
-											RigHierarchy->SwitchToParent(ControlKey, SpaceNameAndValue.Value.ControlRigElement);
-#endif
+											ControlRig->SwitchToParent(ControlKey, SpaceNameAndValue.Value.ControlRigElement, false, true);
+#endif	
 										}
 										break;
 									}
diff --git a/Engine/Plugins/Animation/ControlRig/Source/ControlRigEditor/Private/EditMode/ControlRigEditMode.cpp b/Engine/Plugins/Animation/ControlRig/Source/ControlRigEditor/Private/EditMode/ControlRigEditMode.cpp
--- a/Engine/Plugins/Animation/ControlRig/Source/ControlRigEditor/Private/EditMode/ControlRigEditMode.cpp
+++ b/Engine/Plugins/Animation/ControlRig/Source/ControlRigEditor/Private/EditMode/ControlRigEditMode.cpp
@@ -386,10 +386,9 @@ void FControlRigEditMode::Enter()
 		if (!Toolkit.IsValid())
 		{
 			Toolkit = MakeShareable(new FControlRigEditModeToolkit(*this));
+			Toolkit->Init(Owner->GetToolkitHost());
 		}
 
-		Toolkit->Init(Owner->GetToolkitHost());
-
 		FEditorModeTools* ModeManager = GetModeManager();
 
 		bIsChangingCoordSystem = false;
@@ -471,7 +470,6 @@ void FControlRigEditMode::Exit()
 	if (Toolkit.IsValid())
 	{
 		FToolkitManager::Get().CloseToolkit(Toolkit.ToSharedRef());
-		Toolkit.Reset();
 	}
 
 	DestroyShapesActors(nullptr);
diff --git a/Engine/Plugins/Animation/ControlRig/Source/ControlRigEditor/Private/EditMode/ControlRigEditModeToolkit.cpp b/Engine/Plugins/Animation/ControlRig/Source/ControlRigEditor/Private/EditMode/ControlRigEditModeToolkit.cpp
--- a/Engine/Plugins/Animation/ControlRig/Source/ControlRigEditor/Private/EditMode/ControlRigEditModeToolkit.cpp
+++ b/Engine/Plugins/Animation/ControlRig/Source/ControlRigEditor/Private/EditMode/ControlRigEditModeToolkit.cpp
@@ -48,7 +48,7 @@ const FName FControlRigEditModeToolkit::SpacePickerTabName = FName(TEXT("Control
 
 void FControlRigEditModeToolkit::Init(const TSharedPtr<IToolkitHost>& InitToolkitHost)
 {
-	SAssignNew(ModeTools, SControlRigEditModeTools, SharedThis(this), EditMode, EditMode.GetWorld());
+	SAssignNew(ModeTools, SControlRigEditModeTools, SharedThis(this), EditMode);
 
 	FPropertyEditorModule& PropertyEditorModule = FModuleManager::GetModuleChecked<FPropertyEditorModule>("PropertyEditor");
 
diff --git a/Engine/Plugins/Animation/ControlRig/Source/ControlRigEditor/Private/EditMode/SControlRigEditModeTools.cpp b/Engine/Plugins/Animation/ControlRig/Source/ControlRigEditor/Private/EditMode/SControlRigEditModeTools.cpp
--- a/Engine/Plugins/Animation/ControlRig/Source/ControlRigEditor/Private/EditMode/SControlRigEditModeTools.cpp
+++ b/Engine/Plugins/Animation/ControlRig/Source/ControlRigEditor/Private/EditMode/SControlRigEditModeTools.cpp
@@ -87,7 +87,7 @@ const URigHierarchy* SControlRigEditModeTools::GetHierarchy() const
 	return nullptr;
 }
 
-void SControlRigEditModeTools::Construct(const FArguments& InArgs, TSharedPtr<FControlRigEditModeToolkit> InOwningToolkit, FControlRigEditMode& InEditMode,UWorld* InWorld)
+void SControlRigEditModeTools::Construct(const FArguments& InArgs, TSharedPtr<FControlRigEditModeToolkit> InOwningToolkit, FControlRigEditMode& InEditMode)
 {
 	bIsChangingRigHierarchy = false;
 	OwningToolkit = InOwningToolkit;
@@ -420,13 +420,17 @@ void SControlRigEditModeTools::Construct(const FArguments& InArgs, TSharedPtr<FC
 						SNew(SButton)
 						.ContentPadding(0.0f)
 						.ButtonStyle(FAppStyle::Get(), "NoBorder")
-						.IsEnabled_Lambda([InWorld]()
+						.IsEnabled_Lambda([this]()
 						{
-							const ULevel* CurrentLevel = InWorld->GetCurrentLevel();
-							const TArray<AActor*> SelectedActors = ObjectPtrDecay(CurrentLevel->Actors).FilterByPredicate( [](const AActor* Actor)
+							TArray<AActor*> SelectedActors;
+							if (FControlRigEditMode* EditMode = static_cast<FControlRigEditMode*>(ModeTools->GetActiveMode(FControlRigEditMode::ModeName)))
 							{
-								return Actor && Actor->IsSelected();
-							});
+								const ULevel* CurrentLevel = EditMode->GetWorld()->GetCurrentLevel();
+								SelectedActors = ObjectPtrDecay(CurrentLevel->Actors).FilterByPredicate([](const AActor* Actor)
+								{
+									return Actor && Actor->IsSelected();
+								});
+							}
 							return !SelectedActors.IsEmpty();
 						})
 						.OnClicked(this, &SControlRigEditModeTools::HandleAddConstraintClicked)
diff --git a/Engine/Plugins/Animation/ControlRig/Source/ControlRigEditor/Private/EditMode/SControlRigEditModeTools.h b/Engine/Plugins/Animation/ControlRig/Source/ControlRigEditor/Private/EditMode/SControlRigEditModeTools.h
--- a/Engine/Plugins/Animation/ControlRig/Source/ControlRigEditor/Private/EditMode/SControlRigEditModeTools.h
+++ b/Engine/Plugins/Animation/ControlRig/Source/ControlRigEditor/Private/EditMode/SControlRigEditModeTools.h
@@ -31,7 +31,7 @@ public:
 	SLATE_BEGIN_ARGS(SControlRigEditModeTools) {}
 	SLATE_END_ARGS();
 
-	void Construct(const FArguments& InArgs, TSharedPtr<FControlRigEditModeToolkit> InOwningToolkit, FControlRigEditMode& InEditMode, UWorld* InWorld);
+	void Construct(const FArguments& InArgs, TSharedPtr<FControlRigEditModeToolkit> InOwningToolkit, FControlRigEditMode& InEditMode);
 	/** Set the objects to be displayed in the details panel */
 	void SetSettingsDetailsObject(const TWeakObjectPtr<>& InObject);
 #if USE_LOCAL_DETAILS
diff --git a/Engine/Plugins/Animation/ControlRig/Source/ControlRigEditor/Private/Sequencer/ControlRigParameterTrackEditor.cpp b/Engine/Plugins/Animation/ControlRig/Source/ControlRigEditor/Private/Sequencer/ControlRigParameterTrackEditor.cpp
--- a/Engine/Plugins/Animation/ControlRig/Source/ControlRigEditor/Private/Sequencer/ControlRigParameterTrackEditor.cpp
+++ b/Engine/Plugins/Animation/ControlRig/Source/ControlRigEditor/Private/Sequencer/ControlRigParameterTrackEditor.cpp
@@ -1672,17 +1672,19 @@ static void EvaluateThisControl(UMovieSceneControlRigParameterSection* Section,
 			TOptional<FMovieSceneControlRigSpaceBaseKey> SpaceKey = Section->EvaluateSpaceChannel(FrameTime, ControlName);
 			if (SpaceKey.IsSet())
 			{
+				const FRigElementKey ControlKey = ControlElement->GetKey();
 				switch (SpaceKey.GetValue().SpaceType)
 				{
 				case EMovieSceneControlRigSpaceType::Parent:
-					RigHierarchy->SwitchToDefaultParent(ControlElement->GetKey());
+					ControlRig->SwitchToParent(ControlKey, RigHierarchy->GetDefaultParent(ControlKey), false, true);
 					break;
 				case EMovieSceneControlRigSpaceType::World:
-					RigHierarchy->SwitchToWorldSpace(ControlElement->GetKey());
+					ControlRig->SwitchToParent(ControlKey, RigHierarchy->GetWorldSpaceReferenceKey(), false, true);
 					break;
 				case EMovieSceneControlRigSpaceType::ControlRig:
-					ControlRig->SwitchToParent(ControlElement->GetKey(), SpaceKey.GetValue().ControlRigElement, false, true);
-					break;
+					ControlRig->SwitchToParent(ControlKey, SpaceKey.GetValue().ControlRigElement, false, true);
+					break;	
+
 				}
 			}
 		}
@@ -5295,6 +5297,14 @@ FControlRigEditMode* FControlRigParameterTrackEditor::GetEditMode(bool bForceAct
 		if (bForceActivate && !EditorModetools->IsModeActive(FControlRigEditMode::ModeName))
 		{
 			EditorModetools->ActivateMode(FControlRigEditMode::ModeName);
+			if (bForceActivate)
+			{
+				FControlRigEditMode* EditMode = static_cast<FControlRigEditMode*>(EditorModetools->GetActiveMode(FControlRigEditMode::ModeName));
+				if (EditMode && EditMode->GetToolkit().IsValid() == false)
+				{
+					EditMode->Enter();
+				}
+			}
 		}
 
 		return static_cast<FControlRigEditMode*>(EditorModetools->GetActiveMode(FControlRigEditMode::ModeName));
diff --git a/Engine/Plugins/Animation/ControlRig/Source/ControlRigEditor/Private/Sequencer/ControlRigSpaceChannelEditors.cpp b/Engine/Plugins/Animation/ControlRig/Source/ControlRigEditor/Private/Sequencer/ControlRigSpaceChannelEditors.cpp
--- a/Engine/Plugins/Animation/ControlRig/Source/ControlRigEditor/Private/Sequencer/ControlRigSpaceChannelEditors.cpp
+++ b/Engine/Plugins/Animation/ControlRig/Source/ControlRigEditor/Private/Sequencer/ControlRigSpaceChannelEditors.cpp
@@ -712,10 +712,10 @@ void  FControlRigSpaceChannelHelpers::SequencerSpaceChannelKeyDeleted(UControlRi
 			switch (PreviousValue.SpaceType)
 			{
 			case EMovieSceneControlRigSpaceType::Parent:
-				RigHierarchy->SwitchToDefaultParent(ControlKey);
+				ControlRig->SwitchToParent(ControlKey, RigHierarchy->GetDefaultParent(ControlKey), false, true);
 				break;
 			case EMovieSceneControlRigSpaceType::World:
-				RigHierarchy->SwitchToWorldSpace(ControlKey);
+				ControlRig->SwitchToParent(ControlKey, RigHierarchy->GetWorldSpaceReferenceKey(), false, true);
 				break;
 			case EMovieSceneControlRigSpaceType::ControlRig:
 				ControlRig->SwitchToParent(ControlKey, PreviousValue.ControlRigElement, false, true);
@@ -1130,10 +1130,10 @@ void FControlRigSpaceChannelHelpers::SequencerBakeControlInSpace(UControlRig* Co
 				switch (EndFrameValue.SpaceType)
 				{
 				case EMovieSceneControlRigSpaceType::Parent:
-					RigHierarchy->SwitchToDefaultParent(ControlKey);
+					ControlRig->SwitchToParent(ControlKey, RigHierarchy->GetDefaultParent(ControlKey), false, true);
 					break;
 				case EMovieSceneControlRigSpaceType::World:
-					RigHierarchy->SwitchToWorldSpace(ControlKey);
+					ControlRig->SwitchToParent(ControlKey, RigHierarchy->GetWorldSpaceReferenceKey(), false, true);
 					break;
 				case EMovieSceneControlRigSpaceType::ControlRig:
 					ControlRig->SwitchToParent(ControlKey, EndFrameValue.ControlRigElement, false, true);
@@ -1344,16 +1344,17 @@ void FControlRigSpaceChannelHelpers::CompensateIfNeeded(
 								ControlRigParentWorldTransforms, ControlWorldTransforms);
 
 							//set space to previous space value that's different.
+							const FRigElementKey ControlKey = Control->GetKey();
 							switch (PreviousValue.SpaceType)
 							{
 								case EMovieSceneControlRigSpaceType::Parent:
-									RigHierarchy->SwitchToDefaultParent(Control->GetKey());
+									ControlRig->SwitchToParent(ControlKey, RigHierarchy->GetDefaultParent(ControlKey), false, true);
 									break;
 								case EMovieSceneControlRigSpaceType::World:
-									RigHierarchy->SwitchToWorldSpace(Control->GetKey());
+									ControlRig->SwitchToParent(ControlKey, RigHierarchy->GetWorldSpaceReferenceKey(), false, true);
 									break;
 								case EMovieSceneControlRigSpaceType::ControlRig:
-									ControlRig->SwitchToParent(Control->GetKey(), PreviousValue.ControlRigElement, false, true);
+									ControlRig->SwitchToParent(ControlKey, PreviousValue.ControlRigElement, false, true);
 									break;
 							}
 							
