Control Rig: ensure that the parent depth has been inserted
    - follow-up to 30563771

#rnx

diff --git a/Engine/Plugins/Animation/ControlRig/Source/ControlRigEditor/Private/Editor/SRigHierarchyTreeView.cpp b/Engine/Plugins/Animation/ControlRig/Source/ControlRigEditor/Private/Editor/SRigHierarchyTreeView.cpp
--- a/Engine/Plugins/Animation/ControlRig/Source/ControlRigEditor/Private/Editor/SRigHierarchyTreeView.cpp
+++ b/Engine/Plugins/Animation/ControlRig/Source/ControlRigEditor/Private/Editor/SRigHierarchyTreeView.cpp
@@ -792,7 +792,10 @@ void SRigHierarchyTreeView::RefreshTreeView(bool bRebuildContent)
 				int32& Depth = ElementDepth.Add(Element, 0);
 				if(const FRigBaseElement* ParentElement = Hierarchy->GetFirstParent(Element))
 				{
-					Depth = ElementDepth.FindChecked(ParentElement) + 1;
+					if (int32* ParentDepth = ElementDepth.Find(ParentElement))
+					{
+						Depth = *ParentDepth + 1;
+					}
 				}
 				
 				if(const FRigSocketElement* Socket = Cast<FRigSocketElement>(Element))
