Control Rig: Module connections handled by the ModularRigController
#rnx

diff --git a/Engine/Plugins/Animation/ControlRig/Source/ControlRig/Private/ModularRig.cpp b/Engine/Plugins/Animation/ControlRig/Source/ControlRig/Private/ModularRig.cpp
--- a/Engine/Plugins/Animation/ControlRig/Source/ControlRig/Private/ModularRig.cpp
+++ b/Engine/Plugins/Animation/ControlRig/Source/ControlRig/Private/ModularRig.cpp
@@ -212,7 +212,7 @@ void UModularRig::UpdateModuleHierarchyFromCDO()
 		// for each module (data only) reference in the model.
 		// Note: The CDO does not contain any UObject module instances itself.
 		const FModularRigModel& Model = GetModularRigModel();
-		Model.ForEachModule([this](const FRigModuleReference* InModuleReference) -> bool
+		Model.ForEachModule([this, Model](const FRigModuleReference* InModuleReference) -> bool
 		{
 			check(InModuleReference);
 			if (!InModuleReference->Class.IsValid())
@@ -225,7 +225,7 @@ void UModularRig::UpdateModuleHierarchyFromCDO()
 					InModuleReference->Name,
 					InModuleReference->Class.Get(),
 					FindModule(InModuleReference->ParentPath),
-					InModuleReference->Connections,
+					Model.Connections.GetModuleConnectionMap(InModuleReference->GetPath()),
 					InModuleReference->ConfigValues);
 			}
 
diff --git a/Engine/Plugins/Animation/ControlRig/Source/ControlRig/Private/ModularRigController.cpp b/Engine/Plugins/Animation/ControlRig/Source/ControlRig/Private/ModularRigController.cpp
--- a/Engine/Plugins/Animation/ControlRig/Source/ControlRig/Private/ModularRigController.cpp
+++ b/Engine/Plugins/Animation/ControlRig/Source/ControlRig/Private/ModularRigController.cpp
@@ -130,17 +130,17 @@ FRigModuleReference* UModularRigController::FindModule(const FString& InPath)
 
 bool UModularRigController::CanConnectConnectorToElement(const FRigElementKey& InConnectorKey, const FRigElementKey& InTargetKey, FText& OutErrorMessage)
 {
-	FString ConnectorParentPath, ConnectorName;
-	if (!InConnectorKey.Name.ToString().Split(UModularRig::NamespaceSeparator, &ConnectorParentPath, &ConnectorName, ESearchCase::CaseSensitive, ESearchDir::FromEnd))
+	FString ConnectorModulePath, ConnectorName;
+	if (!InConnectorKey.Name.ToString().Split(UModularRig::NamespaceSeparator, &ConnectorModulePath, &ConnectorName, ESearchCase::CaseSensitive, ESearchDir::FromEnd))
 	{
 		OutErrorMessage = FText::FromString(FString::Printf(TEXT("Connector %s does not contain a namespace"), *InConnectorKey.ToString()));
 		return false;
 	}
 	
-	FRigModuleReference* Module = FindModule(ConnectorParentPath);
+	FRigModuleReference* Module = FindModule(ConnectorModulePath);
 	if (!Module)
 	{
-		OutErrorMessage = FText::FromString(FString::Printf(TEXT("Could not find module %s"), *ConnectorParentPath));
+		OutErrorMessage = FText::FromString(FString::Printf(TEXT("Could not find module %s"), *ConnectorModulePath));
 		return false;
 	}
 
@@ -174,9 +174,8 @@ bool UModularRigController::CanConnectConnectorToElement(const FRigElementKey& I
 		return false;
 	}
 
-	const FRigElementKey ConnectorKey(*ConnectorName, ERigElementType::Connector);
-	FRigElementKey* CurrentTarget = Module->Connections.Find(ConnectorKey);
-	if (CurrentTarget && InTargetKey == *CurrentTarget)
+	FRigElementKey CurrentTarget = Model->Connections.FindTargetFromConnector(InConnectorKey);
+	if (CurrentTarget.IsValid() && InTargetKey == CurrentTarget)
 	{
 		return true; // Nothing to do
 	}
@@ -188,9 +187,11 @@ bool UModularRigController::CanConnectConnectorToElement(const FRigElementKey& I
 		{
 			return Connector.Settings.Type == EConnectorType::Primary;
 		});
-		const FRigElementKey PrimaryConnectorKey(*PrimaryMdouleConnector->Name, ERigElementType::Connector);
-		FRigElementKey* PrimaryTarget = Module->Connections.Find(PrimaryConnectorKey);
-		if (!PrimaryTarget || !PrimaryTarget->IsValid())
+
+		const FString PrimaryConnectorPath = FString::Printf(TEXT("%s:%s"), *ConnectorModulePath, *PrimaryMdouleConnector->Name);
+		const FRigElementKey PrimaryConnectorKey(*PrimaryConnectorPath, ERigElementType::Connector);
+		FRigElementKey PrimaryTarget = Model->Connections.FindTargetFromConnector(PrimaryConnectorKey);
+		if (!PrimaryTarget.IsValid())
 		{
 			OutErrorMessage = FText::FromString(FString::Printf(TEXT("Cannot resolve connector %s because primary connector is not resolved"), *InConnectorKey.ToString()));
 			return false;
@@ -256,8 +257,7 @@ bool UModularRigController::ConnectConnectorToElement(const FRigElementKey& InCo
 	InConnectorKey.Name.ToString().Split(UModularRig::NamespaceSeparator, &ConnectorParentPath, &ConnectorName, ESearchCase::CaseSensitive, ESearchDir::FromEnd);
 	FRigModuleReference* Module = FindModule(ConnectorParentPath);
 
-	const FRigElementKey ConnectorKey(*ConnectorName, ERigElementType::Connector);
-	FRigElementKey* CurrentTarget = Module->Connections.Find(ConnectorKey);
+	FRigElementKey CurrentTarget = Model->Connections.FindTargetFromConnector(InConnectorKey);
 
 	UBlueprint* Blueprint = Cast<UBlueprint>(GetOuter());
 
@@ -271,12 +271,12 @@ bool UModularRigController::ConnectConnectorToElement(const FRigElementKey& InCo
 #endif 
 
 	// First disconnect before connecting to anything else. This might disconnect other secondary/optional connectors.
-	if (CurrentTarget && CurrentTarget->IsValid())
+	if (CurrentTarget.IsValid())
 	{
 		DisconnectConnector(InConnectorKey, bSetupUndo);
 	}
 
-	Module->Connections.FindOrAdd(ConnectorKey) = InTargetKey;
+	Model->Connections.AddConnection(InConnectorKey, InTargetKey);
 	Notify(EModularRigNotification::ConnectionChanged, Module);
 
 #if WITH_EDITOR
@@ -293,21 +293,17 @@ bool UModularRigController::ConnectConnectorToElement(const FRigElementKey& InCo
 
 					for (const FRigModuleConnector& OtherConnector : RigCDO->GetRigModuleSettings().ExposedConnectors)
 					{
-						FRigElementKey OtherConnectorShortKey;
-						OtherConnectorShortKey.Name = *OtherConnector.Name;
-						OtherConnectorShortKey.Type = ERigElementType::Connector;
-						if (!Module->Connections.Contains(OtherConnectorShortKey))
+						FRigElementKey OtherConnectorKey;
+						OtherConnectorKey.Name = *FString::Printf(TEXT("%s%s"), *Module->GetNamespace(), *OtherConnector.Name);
+						OtherConnectorKey.Type = ERigElementType::Connector;
+						if (!Model->Connections.HasConnection(OtherConnectorKey))
 						{
-							FRigElementKey OtherConnectorKey;
-							OtherConnectorKey.Name = *FString::Printf(TEXT("%s%s"), *Module->GetNamespace(), *OtherConnector.Name);
-							OtherConnectorKey.Type = ERigElementType::Connector;
-							
 							const FRigConnectorElement* OtherConnectorElement = Cast<FRigConnectorElement>(Hierarchy->Find(OtherConnectorKey));
 							FModularRigResolveResult RuleResults = RuleManager->FindMatches(OtherConnectorElement, ModuleInstance, ModularRig->GetElementKeyRedirector());
 
 							if (RuleResults.GetMatches().Num() == 1)
 							{
-								Module->Connections.FindOrAdd(OtherConnectorShortKey) = RuleResults.GetMatches()[0].GetKey();
+								Model->Connections.AddConnection(OtherConnectorKey, RuleResults.GetMatches()[0].GetKey());
 								Notify(EModularRigNotification::ConnectionChanged, Module);
 							}
 							else
@@ -316,7 +312,7 @@ bool UModularRigController::ConnectConnectorToElement(const FRigElementKey& InCo
 								{
 									if (Result.GetState() == ERigElementResolveState::DefaultTarget)
 									{
-										Module->Connections.FindOrAdd(OtherConnectorShortKey) = Result.GetKey();
+										Model->Connections.AddConnection(OtherConnectorKey, Result.GetKey());
 										Notify(EModularRigNotification::ConnectionChanged, Module);
 										break;
 									}
@@ -339,17 +335,17 @@ bool UModularRigController::ConnectConnectorToElement(const FRigElementKey& InCo
 
 bool UModularRigController::DisconnectConnector(const FRigElementKey& InConnectorKey, bool bSetupUndo)
 {
-	FString ConnectorParentPath, ConnectorName;
-	if (!InConnectorKey.Name.ToString().Split(UModularRig::NamespaceSeparator, &ConnectorParentPath, &ConnectorName, ESearchCase::CaseSensitive, ESearchDir::FromEnd))
+	FString ConnectorModulePath, ConnectorName;
+	if (!InConnectorKey.Name.ToString().Split(UModularRig::NamespaceSeparator, &ConnectorModulePath, &ConnectorName, ESearchCase::CaseSensitive, ESearchDir::FromEnd))
 	{
 		UE_LOG(LogControlRig, Error, TEXT("Connector %s does not contain a namespace"), *InConnectorKey.ToString());
 		return false;
 	}
 	
-	FRigModuleReference* Module = FindModule(ConnectorParentPath);
+	FRigModuleReference* Module = FindModule(ConnectorModulePath);
 	if (!Module)
 	{
-		UE_LOG(LogControlRig, Error, TEXT("Could not find module %s"), *ConnectorParentPath);
+		UE_LOG(LogControlRig, Error, TEXT("Could not find module %s"), *ConnectorModulePath);
 		return false;
 	}
 
@@ -380,8 +376,7 @@ bool UModularRigController::DisconnectConnector(const FRigElementKey& InConnecto
 		return false;
 	}
 
-	const FRigElementKey ConnectorKey(*ConnectorName, ERigElementType::Connector);
-	if(!Module->Connections.Contains(ConnectorKey))
+	if(!Model->Connections.HasConnection(InConnectorKey))
 	{
 		return false;
 	}
@@ -395,17 +390,45 @@ bool UModularRigController::DisconnectConnector(const FRigElementKey& InConnecto
 	}
 #endif 
 
-	Module->Connections.Remove(ConnectorKey);
+	Model->Connections.RemoveConnection(InConnectorKey);
 
 	if (Connector->Settings.Type == EConnectorType::Primary)
 	{
-		Module->Connections.Reset();
+		// Remove connections from module and child modules
+		TArray<FRigElementKey> ConnectionsToRemove;
+		for (const FModularRigSingleConnection& Connection : Model->Connections.ConnectionList)
+		{
+			if (Connection.Connector.Name.ToString().StartsWith(ConnectorModulePath, ESearchCase::CaseSensitive))
+			{
+				ConnectionsToRemove.Add(Connection.Connector);
+			}
+		}
+		for (const FRigElementKey& ToRemove : ConnectionsToRemove)
+		{
+			Model->Connections.RemoveConnection(ToRemove);
+		}
 	}
-	else
+	else if (!Connector->Settings.bOptional)
 	{
-		// todo: Make sure all the rest of the connections are still valid
+		// Remove connections from child modules
+		TArray<FRigElementKey> ConnectionsToRemove;
+		for (const FModularRigSingleConnection& Connection : Model->Connections.ConnectionList)
+		{
+			FString OtherConnectorModulePath, OtherConnectorName;
+			Connection.Connector.Name.ToString().Split(UModularRig::NamespaceSeparator, &OtherConnectorModulePath, &OtherConnectorName, ESearchCase::CaseSensitive, ESearchDir::FromEnd);
+			if (OtherConnectorModulePath.StartsWith(ConnectorModulePath, ESearchCase::CaseSensitive) && OtherConnectorModulePath.Len() > ConnectorModulePath.Len())
+			{
+				ConnectionsToRemove.Add(Connection.Connector);
+			}
+		}
+		for (const FRigElementKey& ToRemove : ConnectionsToRemove)
+		{
+			Model->Connections.RemoveConnection(ToRemove);
+		}
 	}
 
+	// todo: Make sure all the rest of the connections are still valid
+
 	Notify(EModularRigNotification::ConnectionChanged, Module);
 
 #if WITH_EDITOR
@@ -522,7 +545,7 @@ TArray<FString> UModularRigController::GetPossibleBindings(const FString& InModu
 	Model->ForEachModule([this, &PossibleBindings, InModulePath, InVariableName, InvalidModulePrefix](const FRigModuleReference* InModule) -> bool
 	{
 		const FString CurModulePath = InModule->GetPath();
-		if (InModulePath != CurModulePath && !CurModulePath.StartsWith(InvalidModulePrefix))
+		if (InModulePath != CurModulePath && !CurModulePath.StartsWith(InvalidModulePrefix, ESearchCase::CaseSensitive))
 		{
 			TArray<FRigVMExternalVariable> Variables = InModule->Class->GetDefaultObject<UControlRig>()->GetExternalVariables();
 			for (const FRigVMExternalVariable& Variable : Variables)
@@ -582,7 +605,7 @@ bool UModularRigController::CanBindModuleVariable(const FString& InModulePath, c
 			return false;
 		}
 
-		if (SourceModulePath.StartsWith(InModulePath))
+		if (SourceModulePath.StartsWith(InModulePath, ESearchCase::CaseSensitive))
 		{
 			OutErrorMessage = FText::FromString(FString::Printf(TEXT("Cannot bind variable of module %s to a variable of module %s because the source module is a child of the target module"), *InModulePath, *SourceModulePath));
 			return false;
@@ -729,7 +752,7 @@ bool UModularRigController::DeleteModule(const FString& InModulePath, bool bSetu
 	TSharedPtr<FScopedTransaction> TransactionPtr;
 	if (bSetupUndo)
 	{
-		TransactionPtr = MakeShared<FScopedTransaction>(NSLOCTEXT("ModularRigController", "RenameModuleTransaction", "Rename Module"));
+		TransactionPtr = MakeShared<FScopedTransaction>(NSLOCTEXT("ModularRigController", "DeleteModuleTransaction", "Delete Module"));
 		if(UBlueprint* Blueprint = Cast<UBlueprint>(GetOuter()))
 		{
 			Blueprint->Modify();
@@ -749,6 +772,32 @@ bool UModularRigController::DeleteModule(const FString& InModulePath, bool bSetu
 	Model->UpdateCachedChildren();
 	UpdateShortNames();
 
+	// Fix connections
+	{
+		TArray<FRigElementKey> ToRemove;
+		for (FModularRigSingleConnection& Connection : Model->Connections.ConnectionList)
+		{
+			FString ConnectionModulePath, ConnectionName;
+			Connection.Connector.Name.ToString().Split(UModularRig::NamespaceSeparator, &ConnectionModulePath, &ConnectionName, ESearchCase::CaseSensitive, ESearchDir::FromEnd);
+
+			if (ConnectionModulePath == InModulePath)
+			{
+				ToRemove.Add(Connection.Connector);
+			}
+
+			FString TargetModulePath, TargetName;
+			Connection.Target.Name.ToString().Split(UModularRig::NamespaceSeparator, &TargetModulePath, &TargetName, ESearchCase::CaseSensitive, ESearchDir::FromEnd);
+			if (TargetModulePath == InModulePath)
+			{
+				ToRemove.Add(Connection.Connector);
+			}
+		}
+		for (FRigElementKey& KeyToRemove : ToRemove)
+		{
+			Model->Connections.RemoveConnection(KeyToRemove);
+		}
+	}
+
 	// Fix bindings
 	for (FRigModuleReference& Reference : Model->Modules)
 	{
@@ -783,19 +832,19 @@ FString UModularRigController::RenameModule(const FString& InModulePath, const F
 		return FString();
 	}
 
-	FText ErrorMessage;
-	if (!CanRenameModule(InModulePath, InNewName, ErrorMessage))
-	{
-		UE_LOG(LogControlRig, Error, TEXT("Could not rename module %s: %s"), *InModulePath, *ErrorMessage.ToString());
-		return FString();
-	}
-
 	const FString OldName = Module->Name.ToString();
 	const FString NewName = InNewName.ToString();
 	if (OldName.Equals(NewName))
 	{
 		return Module->GetPath();
 	}
+
+	FText ErrorMessage;
+	if (!CanRenameModule(InModulePath, InNewName, ErrorMessage))
+	{
+		UE_LOG(LogControlRig, Error, TEXT("Could not rename module %s: %s"), *InModulePath, *ErrorMessage.ToString());
+		return FString();
+	}
 	
 #if WITH_EDITOR
 	TSharedPtr<FScopedTransaction> TransactionPtr;
@@ -823,6 +872,35 @@ FString UModularRigController::RenameModule(const FString& InModulePath, const F
 		Children.Append(Child->CachedChildren);
 	}
 
+	// Fix connections
+	{
+		for (FModularRigSingleConnection& Connection : Model->Connections.ConnectionList)
+		{
+			if (Connection.Connector.Name.ToString().StartsWith(OldPath, ESearchCase::CaseSensitive))
+			{
+				Connection.Connector.Name = *FString::Printf(TEXT("%s%s"), *NewPath, *Connection.Connector.Name.ToString().RightChop(OldPath.Len()));
+			}
+			if (Connection.Target.Name.ToString().StartsWith(OldPath, ESearchCase::CaseSensitive))
+			{
+				Connection.Target.Name = *FString::Printf(TEXT("%s%s"), *NewPath, *Connection.Target.Name.ToString().RightChop(OldPath.Len()));
+			}
+		}
+		for (TPair<FRigElementKey, TArray<FRigElementKey>>& Pair : Model->Connections.ReverseConnectionMap)
+		{
+			if (Pair.Key.Name.ToString().StartsWith(OldPath, ESearchCase::CaseSensitive))
+			{
+				Pair.Key.Name = *FString::Printf(TEXT("%s%s"), *NewPath, *Pair.Key.Name.ToString().RightChop(OldPath.Len()));
+			}
+			for (FRigElementKey& TargetConnectorKey : Pair.Value)
+			{
+				if (TargetConnectorKey.Name.ToString().StartsWith(OldPath, ESearchCase::CaseSensitive))
+				{
+					TargetConnectorKey.Name = *FString::Printf(TEXT("%s%s"), *NewPath, *TargetConnectorKey.Name.ToString().RightChop(OldPath.Len()));
+				}
+			}
+		}
+	}
+
 	// Fix bindings
 	for (FRigModuleReference& Reference : Model->Modules)
 	{
@@ -918,6 +996,35 @@ FString UModularRigController::ReparentModule(const FString& InModulePath, const
 	Model->UpdateCachedChildren();
 	UpdateShortNames();
 
+	// Fix connections
+	{
+		for (FModularRigSingleConnection& Connection : Model->Connections.ConnectionList)
+		{
+			if (Connection.Connector.Name.ToString().StartsWith(OldPath, ESearchCase::CaseSensitive))
+			{
+				Connection.Connector.Name = *FString::Printf(TEXT("%s%s"), *NewPath, *Connection.Connector.Name.ToString().RightChop(OldPath.Len()));
+			}
+			if (Connection.Target.Name.ToString().StartsWith(OldPath, ESearchCase::CaseSensitive))
+			{
+				Connection.Target.Name = *FString::Printf(TEXT("%s%s"), *NewPath, *Connection.Target.Name.ToString().RightChop(OldPath.Len()));
+			}
+		}
+		for (TPair<FRigElementKey, TArray<FRigElementKey>>& Pair : Model->Connections.ReverseConnectionMap)
+		{
+			if (Pair.Key.Name.ToString().StartsWith(OldPath, ESearchCase::CaseSensitive))
+			{
+				Pair.Key.Name = *FString::Printf(TEXT("%s%s"), *NewPath, *Pair.Key.Name.ToString().RightChop(OldPath.Len()));
+			}
+			for (FRigElementKey& TargetConnectorKey : Pair.Value)
+			{
+				if (TargetConnectorKey.Name.ToString().StartsWith(OldPath, ESearchCase::CaseSensitive))
+				{
+					TargetConnectorKey.Name = *FString::Printf(TEXT("%s%s"), *NewPath, *TargetConnectorKey.Name.ToString().RightChop(OldPath.Len()));
+				}
+			}
+		}
+	}
+
 	// Fix bindings
 	for (FRigModuleReference& Reference : Model->Modules)
 	{
diff --git a/Engine/Plugins/Animation/ControlRig/Source/ControlRig/Private/ModularRigModel.cpp b/Engine/Plugins/Animation/ControlRig/Source/ControlRig/Private/ModularRigModel.cpp
--- a/Engine/Plugins/Animation/ControlRig/Source/ControlRig/Private/ModularRigModel.cpp
+++ b/Engine/Plugins/Animation/ControlRig/Source/ControlRig/Private/ModularRigModel.cpp
@@ -29,6 +29,44 @@ FString FRigModuleReference::GetNamespace() const
 	return GetPath() + UModularRig::NamespaceSeparator;
 }
 
+TMap<FRigElementKey, FRigElementKey> FModularRigConnections::GetModuleConnectionMap(const FString& InModulePath) const
+{
+	TMap<FRigElementKey, FRigElementKey> Result;
+	for (const FModularRigSingleConnection& Connection : ConnectionList)
+	{
+		FString Path, Name;
+		Connection.Connector.Name.ToString().Split(UModularRig::NamespaceSeparator, &Path, &Name, ESearchCase::CaseSensitive, ESearchDir::FromEnd);
+
+		// Exactly the same path (do not return connectors from child modules)
+		if (Path == InModulePath)
+		{
+			Result.Add(FRigElementKey(*Name, ERigElementType::Connector), Connection.Target);
+		}
+	}
+	return Result;
+}
+
+bool FModularRigModel::PatchModelsOnLoad()
+{
+	bool bPatched = false;
+	if (Connections.ConnectionList.IsEmpty())
+	{
+		ForEachModule([this, &bPatched](const FRigModuleReference* Module) -> bool
+		{
+			FString ModuleNamespace = Module->GetNamespace();
+			for (const TTuple<FRigElementKey, FRigElementKey>& Connection : Module->Connections_DEPRECATED)
+			{
+				const FString ConnectorPath = FString::Printf(TEXT("%s%s"), *ModuleNamespace, *Connection.Key.Name.ToString());
+				FRigElementKey ConnectorKey(*ConnectorPath, ERigElementType::Connector);
+				Connections.AddConnection(ConnectorKey, Connection.Value);
+			}
+			return true;
+		});
+		bPatched = !Connections.ConnectionList.IsEmpty();
+	}
+	return bPatched;
+}
+
 UModularRigController* FModularRigModel::GetController(bool bCreateIfNeeded)
 {
 	if (bCreateIfNeeded && Controller == nullptr)
diff --git a/Engine/Plugins/Animation/ControlRig/Source/ControlRig/Public/ModularRigModel.h b/Engine/Plugins/Animation/ControlRig/Source/ControlRig/Public/ModularRigModel.h
--- a/Engine/Plugins/Animation/ControlRig/Source/ControlRig/Public/ModularRigModel.h
+++ b/Engine/Plugins/Animation/ControlRig/Source/ControlRig/Public/ModularRigModel.h
@@ -64,8 +64,8 @@ struct CONTROLRIG_API FRigModuleReference
 	UPROPERTY()
 	TSoftClassPtr<UControlRig> Class;
 
-	UPROPERTY()
-	TMap<FRigElementKey, FRigElementKey> Connections; // Connectors to Connection element
+	UPROPERTY(meta = (DeprecatedProperty))
+	TMap<FRigElementKey, FRigElementKey> Connections_DEPRECATED; // Connectors to Connection element
 
 	UPROPERTY()
 	TMap<FName, FString> ConfigValues;
@@ -101,6 +101,97 @@ struct CONTROLRIG_API FRigModuleReference
 	friend class UModularRigController;
 };
 
+USTRUCT(BlueprintType)
+struct FModularRigSingleConnection
+{
+	GENERATED_BODY()
+
+	UPROPERTY()
+	FRigElementKey Connector;
+
+	UPROPERTY()
+	FRigElementKey Target;
+};
+
+USTRUCT(BlueprintType)
+struct CONTROLRIG_API FModularRigConnections
+{
+	GENERATED_BODY()
+
+	/** Connections sorted by creation order */
+	UPROPERTY()
+	TArray<FModularRigSingleConnection> ConnectionList;
+
+	/** Target key to connector array */
+	TMap<FRigElementKey, TArray<FRigElementKey>> ReverseConnectionMap;
+
+	void UpdateFromConnectionList()
+	{
+		ReverseConnectionMap.Reset();
+		for (const FModularRigSingleConnection& Connection : ConnectionList)
+		{
+			TArray<FRigElementKey>& Connectors = ReverseConnectionMap.FindOrAdd(Connection.Target);
+			Connectors.AddUnique(Connection.Connector);
+		}
+	}
+
+	void AddConnection(const FRigElementKey& Connector, const FRigElementKey& Target)
+	{
+		// Remove any existing connection
+		RemoveConnection(Connector);
+
+		ConnectionList.Add(FModularRigSingleConnection(Connector, Target));
+		ReverseConnectionMap.FindOrAdd(Target).AddUnique(Connector);
+	}
+
+	void RemoveConnection(const FRigElementKey& Connector)
+	{
+		int32 ExistingIndex = FindConnectionIndex(Connector);
+		if (ConnectionList.IsValidIndex(ExistingIndex))
+		{
+			if (TArray<FRigElementKey>* Connectors = ReverseConnectionMap.Find(ConnectionList[ExistingIndex].Target))
+			{
+				*Connectors = Connectors->FilterByPredicate([Connector](const FRigElementKey& TargetConnector)
+				{
+					return TargetConnector != Connector;
+				});
+				if (Connectors->IsEmpty())
+				{
+					ReverseConnectionMap.Remove(ConnectionList[ExistingIndex].Target);
+				}
+			}
+			ConnectionList.RemoveAt(ExistingIndex);
+		}
+	}
+
+	int32 FindConnectionIndex(const FRigElementKey& InConnectorKey) const
+	{
+		return ConnectionList.IndexOfByPredicate([InConnectorKey](const FModularRigSingleConnection& Connection)
+		{
+			return InConnectorKey == Connection.Connector;
+		});
+	}
+
+	FRigElementKey FindTargetFromConnector(const FRigElementKey& InConnectorKey) const
+	{
+		int32 Index = FindConnectionIndex(InConnectorKey);
+		if (ConnectionList.IsValidIndex(Index))
+		{
+			return ConnectionList[Index].Target;
+		}
+		static const FRigElementKey EmptyKey;
+		return EmptyKey;
+	}
+
+	bool HasConnection(const FRigElementKey& InConnectorKey) const
+	{
+		return ConnectionList.IsValidIndex(FindConnectionIndex(InConnectorKey));
+	}
+
+	/** Gets the connection map for a single module, where the connectors are identified without its namespace*/
+	TMap<FRigElementKey, FRigElementKey> GetModuleConnectionMap(const FString& InModulePath) const;
+};
+
 // A management struct containing all modules in the rig
 USTRUCT(BlueprintType)
 struct CONTROLRIG_API FModularRigModel
@@ -113,12 +204,14 @@ public:
 	FModularRigModel(const FModularRigModel& Other)
 	{
 		Modules = Other.Modules;
+		Connections = Other.Connections;
 		UpdateCachedChildren();
 	}
 	
 	FModularRigModel& operator=(const FModularRigModel& Other)
 	{
 		Modules = Other.Modules;
+		Connections = Other.Connections;
 		UpdateCachedChildren();
 		return *this;
 	}
@@ -130,9 +223,14 @@ public:
 	TArray<FRigModuleReference*> RootModules;
 	TArray<FRigModuleReference> DeletedModules;
 
+	UPROPERTY()
+	FModularRigConnections Connections;
+
 	UPROPERTY(transient)
 	TObjectPtr<UObject> Controller;
 
+	bool PatchModelsOnLoad();
+
 	UModularRigController* GetController(bool bCreateIfNeeded = true);
 
 	UObject* GetOuter() const { return OuterClientHost.IsValid() ? OuterClientHost.Get() : nullptr; }
diff --git a/Engine/Plugins/Animation/ControlRig/Source/ControlRigDeveloper/Private/ControlRigBlueprint.cpp b/Engine/Plugins/Animation/ControlRig/Source/ControlRigDeveloper/Private/ControlRigBlueprint.cpp
--- a/Engine/Plugins/Animation/ControlRig/Source/ControlRigDeveloper/Private/ControlRigBlueprint.cpp
+++ b/Engine/Plugins/Animation/ControlRig/Source/ControlRigDeveloper/Private/ControlRigBlueprint.cpp
@@ -988,15 +988,33 @@ bool UControlRigBlueprint::ResolveConnector(const FRigElementKey& DraggedKey, co
 		}
 		ExistingTargetKey = TargetKey;
 
-		// Add connection to the model
-		if (UModularRigController* Controller = GetModularRigController())
+		if (IsModularRig())
 		{
-			Controller->ConnectConnectorToElement(DraggedKey, TargetKey, bSetupUndoRedo, ModularRigSettings.bAutoResolve);
+			// Add connection to the model
+			if (UModularRigController* Controller = GetModularRigController())
+			{
+				Controller->ConnectConnectorToElement(DraggedKey, TargetKey, bSetupUndoRedo, ModularRigSettings.bAutoResolve);
+			}
+		}
+		else
+		{
+			ConnectionMap.FindOrAdd(DraggedKey) = TargetKey;
 		}
 	}
 	else
 	{
-		ConnectionMap.Remove(DraggedKey);
+		if (IsModularRig())
+		{
+			// Add connection to the model
+			if (UModularRigController* Controller = GetModularRigController())
+			{
+				Controller->DisconnectConnector(DraggedKey, bSetupUndoRedo);
+			}
+		}
+		else
+		{
+			ConnectionMap.Remove(DraggedKey);
+		}
 	}
 
 	RecompileModularRig();
@@ -1019,6 +1037,19 @@ bool UControlRigBlueprint::ResolveConnector(const FRigElementKey& DraggedKey, co
 	return true;
 }
 
+void UControlRigBlueprint::UpdateConnectionMapFromModel()
+{
+	if (IsModularRig())
+	{
+		ConnectionMap.Reset();
+
+		for (const FModularRigSingleConnection& Connection : ModularRigModel.Connections.ConnectionList)
+		{
+			ConnectionMap.Add(Connection.Connector, Connection.Target);
+		}
+	}
+}
+
 void UControlRigBlueprint::PostLoad()
 {
 	Super::PostLoad();
@@ -1134,6 +1165,7 @@ void UControlRigBlueprint::PostLoad()
 		}
 	}
 
+	ModularRigModel.PatchModelsOnLoad();
 	UpdateModularDependencyDelegates();
 }
 
@@ -2470,6 +2502,7 @@ void UControlRigBlueprint::HandleHierarchyModified(ERigHierarchyNotification InN
 
 void UControlRigBlueprint::HandleRigModulesModified(EModularRigNotification InNotification, const FRigModuleReference* InModule)
 {
+	bool bRecompile = true;
 	switch (InNotification)
 	{
 		case EModularRigNotification::ModuleAdded:
@@ -2552,12 +2585,7 @@ void UControlRigBlueprint::HandleRigModulesModified(EModularRigNotification InNo
 							}
 						}
 
-						// Resolve Connectors
-						for (TPair<FRigElementKey, ConnectionInfo>& Pair : RenamedConnectors)
-						{
-							ResolveConnector(FRigElementKey(*Pair.Value.NewPath, ERigElementType::Connector), Pair.Value.TargetConnection);
-						}
-
+						UpdateConnectionMapFromModel();
 						PropagateHierarchyFromBPToInstances();
 					}
 				}
@@ -2569,6 +2597,7 @@ void UControlRigBlueprint::HandleRigModulesModified(EModularRigNotification InNo
 			if (InModule)
 			{
 				RefreshModuleConnectors(InModule);
+				UpdateConnectionMapFromModel();
 				UpdateModularDependencyDelegates();
 			}
 			break;
@@ -2577,36 +2606,21 @@ void UControlRigBlueprint::HandleRigModulesModified(EModularRigNotification InNo
 		{
 			Hierarchy->Modify();
 			
-			const FString Namespace = InModule->GetNamespace();
-
-			TArray<FRigElementKey> ConnectionsToRemove;
-			for (TPair<FRigElementKey, FRigElementKey> ExistingConnection : ConnectionMap)
-			{
-				if(ExistingConnection.Key.Name.ToString().StartsWith(Namespace, ESearchCase::CaseSensitive))
-				{
-					if(!InModule->Connections.Contains(ExistingConnection.Key))
-					{
-						ConnectionsToRemove.Add(ExistingConnection.Key);
-					}
-				}
-			}
-				
-			for(const FRigElementKey& ConnectionToRemove : ConnectionsToRemove)
-			{
-				ConnectionMap.Remove(ConnectionToRemove);
-			}
-				
-			for (TPair<FRigElementKey, FRigElementKey> Connection : InModule->Connections)
-			{
-				const FString NamespacedName = FString::Printf(TEXT("%s%s"), *Namespace, *Connection.Key.Name.ToString());
-				ConnectionMap.FindOrAdd(FRigElementKey(*NamespacedName, ERigElementType::Connector)) = Connection.Value;
-			}
-				
+			UpdateConnectionMapFromModel();
 			HierarchyModifiedEvent.Broadcast(ERigHierarchyNotification::HierarchyReset, Hierarchy, nullptr);
 			break;
 		}
+		case EModularRigNotification::ModuleShortNameChanged:
+		{
+			bRecompile = false;
+			break;
+		}
+	}
+
+	if (bRecompile)
+	{
+		RecompileModularRig();
 	}
-	RecompileModularRig();
 }
 
 UControlRigBlueprint::FControlValueScope::FControlValueScope(UControlRigBlueprint* InBlueprint)
diff --git a/Engine/Plugins/Animation/ControlRig/Source/ControlRigDeveloper/Public/ControlRigBlueprint.h b/Engine/Plugins/Animation/ControlRig/Source/ControlRigDeveloper/Public/ControlRigBlueprint.h
--- a/Engine/Plugins/Animation/ControlRig/Source/ControlRigDeveloper/Public/ControlRigBlueprint.h
+++ b/Engine/Plugins/Animation/ControlRig/Source/ControlRigDeveloper/Public/ControlRigBlueprint.h
@@ -173,6 +173,8 @@ protected:
 
 	bool ResolveConnector(const FRigElementKey& DraggedKey, const FRigElementKey& TargetKey, bool bSetupUndoRedo = true);
 
+	void UpdateConnectionMapFromModel();
+
 	/** Asset searchable information about exposed public functions on this rig */
 	UPROPERTY(AssetRegistrySearchable)
 	TArray<FRigVMOldPublicFunctionData> PublicFunctions_DEPRECATED;
diff --git a/Engine/Plugins/Animation/ControlRig/Source/ControlRigEditor/Private/ControlRigModuleDetails.cpp b/Engine/Plugins/Animation/ControlRig/Source/ControlRigEditor/Private/ControlRigModuleDetails.cpp
--- a/Engine/Plugins/Animation/ControlRig/Source/ControlRigEditor/Private/ControlRigModuleDetails.cpp
+++ b/Engine/Plugins/Animation/ControlRig/Source/ControlRigEditor/Private/ControlRigModuleDetails.cpp
@@ -250,30 +250,31 @@ void FRigModuleInstanceDetails::CustomizeDetails(IDetailLayoutBuilder& DetailBui
 			FString ConnectorPath = FString::Printf(TEXT("%s:%s"), *PerModuleInfos[0].Path, *Connector.Name);
 			FRigElementKey ConnectorKey(*ConnectorPath, ERigElementType::Connector);
 
-			FRigTreeDelegates TreeDelegates;
-			TreeDelegates.OnGetHierarchy = FOnGetRigTreeHierarchy::CreateLambda([this]()
+			FModularRigResolveResult ConnectorMatches;
+			FPerModuleInfo Info = PerModuleInfos[0];
+			if (const FRigModuleInstance* Module = Info.GetModule())
 			{
-				return PerModuleInfos[0].GetModularRig()->GetHierarchy();
-			});
-			TreeDelegates.OnRigTreeIsItemVisible = FOnRigTreeIsItemVisible::CreateLambda([this, ConnectorKey](const FRigElementKey& InTarget)
-			{
-				FPerModuleInfo Info = PerModuleInfos[0];
-				if (const FRigModuleInstance* Module = Info.GetModule())
+				if (UModularRig* ModularRig = Info.GetModularRig())
 				{
-					if (UModularRig* ModularRig = Info.GetModularRig())
+					if (URigHierarchy* Hierarchy = ModularRig->GetHierarchy())
 					{
-						if (URigHierarchy* Hierarchy = ModularRig->GetHierarchy())
+						if (FRigConnectorElement* ConnectorElement = Cast<FRigConnectorElement>(Hierarchy->Find(ConnectorKey)))
 						{
-							if (FRigConnectorElement* ConnectorElement = Cast<FRigConnectorElement>(Hierarchy->Find(ConnectorKey)))
-							{
-								const UModularRigRuleManager* RuleManager = ModularRig->GetHierarchy()->GetRuleManager();
-								const FModularRigResolveResult Result = RuleManager->FindMatches(ConnectorElement, Module, ModularRig->GetElementKeyRedirector());
-								return Result.ContainsMatch(InTarget);
-							}
+							const UModularRigRuleManager* RuleManager = ModularRig->GetHierarchy()->GetRuleManager();
+							ConnectorMatches = RuleManager->FindMatches(ConnectorElement, Module, ModularRig->GetElementKeyRedirector());
 						}
 					}
 				}
-				return false;
+			}
+
+			FRigTreeDelegates TreeDelegates;
+			TreeDelegates.OnGetHierarchy = FOnGetRigTreeHierarchy::CreateLambda([this]()
+			{
+				return PerModuleInfos[0].GetModularRig()->GetHierarchy();
+			});
+			TreeDelegates.OnRigTreeIsItemVisible = FOnRigTreeIsItemVisible::CreateLambda([ConnectorMatches](const FRigElementKey& InTarget)
+			{
+				return ConnectorMatches.ContainsMatch(InTarget);
 			});
 			TreeDelegates.OnGetSelection.BindLambda([this, ConnectorKey]() -> TArray<FRigElementKey>
 			{
@@ -436,7 +437,7 @@ void FRigModuleInstanceDetails::CustomizeDetails(IDetailLayoutBuilder& DetailBui
 									{
 										if (const FRigElementKey* TargetKey = ModularRig->GetElementKeyRedirector().FindExternalKey(ConnectorKey))
 										{
-											ModularRig->GetHierarchy()->GetController()->SelectElement(*TargetKey);
+											ModularRig->GetHierarchy()->GetController()->SelectElement(*TargetKey, true, true);
 										}
 									}
 									return FReply::Handled();
diff --git a/Engine/Plugins/Animation/ControlRig/Source/ControlRigEditor/Private/Editor/ControlRigSchematicModel.cpp b/Engine/Plugins/Animation/ControlRig/Source/ControlRigEditor/Private/Editor/ControlRigSchematicModel.cpp
--- a/Engine/Plugins/Animation/ControlRig/Source/ControlRigEditor/Private/Editor/ControlRigSchematicModel.cpp
+++ b/Engine/Plugins/Animation/ControlRig/Source/ControlRigEditor/Private/Editor/ControlRigSchematicModel.cpp
@@ -578,12 +578,13 @@ void FControlRigSchematicModel::HandleModularRigModified(EModularRigNotification
 								if(!LocalName.Contains(UModularRig::NamespaceSeparator))
 								{
 									const FRigElementKey LocalKey(*LocalName, ConnectorKey.Type);
-									if(const FRigElementKey* ResolvedKey = InModule->Connections.Find(LocalKey))
+									const FRigElementKey ResolvedKey = ControlRigBlueprint->ModularRigModel.Connections.FindTargetFromConnector(ConnectorKey);
+									if(ResolvedKey.IsValid())
 									{
-										const FControlRigSchematicRigElementKeyNode* ResolvedNode = FindElementKeyNode(*ResolvedKey);
+										const FControlRigSchematicRigElementKeyNode* ResolvedNode = FindElementKeyNode(ResolvedKey);
 										if(ResolvedNode == nullptr)
 										{
-											ResolvedNode = AddElementKeyNode(*ResolvedKey);
+											ResolvedNode = AddElementKeyNode(ResolvedKey);
 										}
 										(void)SetParentNode(ConnectorNode, ResolvedNode);
 										const_cast<FControlRigSchematicRigElementKeyNode*>(ResolvedNode)->OnMouseLeave();
diff --git a/Engine/Plugins/Animation/ControlRig/Source/ControlRigEditor/Private/Editor/SModularRigModel.cpp b/Engine/Plugins/Animation/ControlRig/Source/ControlRigEditor/Private/Editor/SModularRigModel.cpp
--- a/Engine/Plugins/Animation/ControlRig/Source/ControlRigEditor/Private/Editor/SModularRigModel.cpp
+++ b/Engine/Plugins/Animation/ControlRig/Source/ControlRigEditor/Private/Editor/SModularRigModel.cpp
@@ -612,6 +612,15 @@ bool SModularRigModel::HandleVerifyNameChanged(const FString& InOldPath, const F
 	{
 		return false;
 	}
+
+	FString ParentPath;
+	FString OldName = InOldPath;
+	InOldPath.Split(UModularRig::NamespaceSeparator, &ParentPath, &OldName, ESearchCase::CaseSensitive, ESearchDir::FromEnd);
+
+	if (InNewName == OldName)
+	{
+		return true;
+	}
 	
 	if (ControlRigBlueprint.IsValid())
 	{
diff --git a/Engine/Plugins/Animation/ControlRig/Source/ControlRigEditor/Private/Editor/SModularRigTreeView.cpp b/Engine/Plugins/Animation/ControlRig/Source/ControlRigEditor/Private/Editor/SModularRigTreeView.cpp
--- a/Engine/Plugins/Animation/ControlRig/Source/ControlRigEditor/Private/Editor/SModularRigTreeView.cpp
+++ b/Engine/Plugins/Animation/ControlRig/Source/ControlRigEditor/Private/Editor/SModularRigTreeView.cpp
@@ -132,29 +132,30 @@ void SModularRigModelItem::Construct(const FArguments& InArgs, const TSharedRef<
 	TSharedPtr< SHorizontalBox > HorizontalBox;
 	TSharedPtr<SVerticalBox> ComboButtonBox;
 
-	FRigTreeDelegates TreeDelegates;
-	TreeDelegates.OnGetHierarchy = FOnGetRigTreeHierarchy::CreateLambda([this]()
-	{
-		return Delegates.GetModularRig()->GetHierarchy();
-	});
-	TreeDelegates.OnRigTreeIsItemVisible = FOnRigTreeIsItemVisible::CreateLambda([this, ConnectorKey, ModulePath](const FRigElementKey& InTarget)
+	FModularRigResolveResult ConnectorMatches;
+	if (const UModularRig* ModularRig = Delegates.GetModularRig())
 	{
-		if (const UModularRig* ModularRig = Delegates.GetModularRig())
+		if (const FRigModuleInstance* Module = ModularRig->FindModule(ModulePath))
 		{
-			if (const FRigModuleInstance* Module = ModularRig->FindModule(ModulePath))
+			if (URigHierarchy* Hierarchy = ModularRig->GetHierarchy())
 			{
-				if (URigHierarchy* Hierarchy = ModularRig->GetHierarchy())
+				if (FRigConnectorElement* ConnectorElement = Cast<FRigConnectorElement>(Hierarchy->Find(ConnectorKey)))
 				{
-					if (FRigConnectorElement* ConnectorElement = Cast<FRigConnectorElement>(Hierarchy->Find(ConnectorKey)))
-					{
-						const UModularRigRuleManager* RuleManager = ModularRig->GetHierarchy()->GetRuleManager();
-						const FModularRigResolveResult Result = RuleManager->FindMatches(ConnectorElement, Module, ModularRig->GetElementKeyRedirector());
-						return Result.ContainsMatch(InTarget);
-					}
+					const UModularRigRuleManager* RuleManager = ModularRig->GetHierarchy()->GetRuleManager();
+					ConnectorMatches = RuleManager->FindMatches(ConnectorElement, Module, ModularRig->GetElementKeyRedirector());
 				}
 			}
 		}
-		return false;
+	}
+
+	FRigTreeDelegates TreeDelegates;
+	TreeDelegates.OnGetHierarchy = FOnGetRigTreeHierarchy::CreateLambda([this]()
+	{
+		return Delegates.GetModularRig()->GetHierarchy();
+	});
+	TreeDelegates.OnRigTreeIsItemVisible = FOnRigTreeIsItemVisible::CreateLambda([ConnectorMatches](const FRigElementKey& InTarget)
+	{
+		return ConnectorMatches.ContainsMatch(InTarget);
 	});
 	TreeDelegates.OnGetSelection.BindLambda([this, ConnectorKey]() -> TArray<FRigElementKey>
 	{
@@ -220,6 +221,7 @@ void SModularRigModelItem::Construct(const FArguments& InArgs, const TSharedRef<
 					.OnTextCommitted(this, &SModularRigModelItem::OnNameCommitted)
 					.ToolTipText(this, &SModularRigModelItem::GetItemTooltip)
 					.MultiLine(false)
+					.Font(IDetailLayoutBuilder::GetDetailFont())
 					.ColorAndOpacity_Lambda([this]()
 					{
 						if(WeakRigTreeElement.IsValid())
@@ -362,9 +364,10 @@ void SModularRigModelItem::Construct(const FArguments& InArgs, const TSharedRef<
 						{
 							if (const UModularRig* ModularRig = Delegates.GetModularRig())
 							{
-								if (const FRigElementKey* TargetKey = ModularRig->GetElementKeyRedirector().FindExternalKey(ConnectorKey))
+								const FRigElementKeyRedirector& Redirector = ModularRig->GetElementKeyRedirector();
+								if (const FRigElementKey* TargetKey = Redirector.FindExternalKey(ConnectorKey))
 								{
-									ModularRig->GetHierarchy()->GetController()->SelectElement(*TargetKey);
+									ModularRig->GetHierarchy()->GetController()->SelectElement(*TargetKey, true, true);
 								}
 							}
 							return FReply::Handled();
@@ -387,12 +390,14 @@ void SModularRigModelItem::Construct(const FArguments& InArgs, const TSharedRef<
 		], OwnerTable);
 
 	InRigTreeElement->OnRenameRequested.BindSP(InlineWidget.Get(), &SInlineEditableTextBlock::EnterEditingMode);
+	const UModularRig* ModularRig = Delegates.GetModularRig();
+	const FRigElementKeyRedirector& Redirector = ModularRig->GetElementKeyRedirector();
 	FRigElementKey CurrentTargetKey;
-	if (const FRigElementKey* Key = Delegates.GetModularRig()->GetElementKeyRedirector().FindExternalKey(ConnectorKey))
+	if (const FRigElementKey* Key = Redirector.FindExternalKey(ConnectorKey))
 	{
 		CurrentTargetKey = *Key;
 	}
-	TPair<const FSlateBrush*, FSlateColor> IconAndColor = SRigHierarchyItem::GetBrushForElementType(Delegates.GetModularRig()->GetHierarchy(), CurrentTargetKey);
+	TPair<const FSlateBrush*, FSlateColor> IconAndColor = SRigHierarchyItem::GetBrushForElementType(ModularRig->GetHierarchy(), CurrentTargetKey);
 	PopulateConnectorCurrentTarget(ComboButtonBox, ConnectorKey, CurrentTargetKey, IconAndColor.Key, IconAndColor.Value, FText::FromName(CurrentTargetKey.Name));
 }
 
@@ -962,20 +967,13 @@ TPair<const FSlateBrush*, FSlateColor> FModularRigTreeElement::GetBrushAndColor(
 
 	if (const FRigModuleInstance* ConnectorModule = InModularRig->FindModule(ModulePath))
 	{
-		bool bIsConnected = false;
-		if (const FRigModuleReference* Reference = ConnectorModule->GetModuleReference())
+		const FModularRigModel& Model = InModularRig->GetModularRigModel();
+		const FString ConnectorPath = FString::Printf(TEXT("%s:%s"), *ModulePath, *ConnectorName);
+		bool bIsConnected = Model.Connections.HasConnection(FRigElementKey(*ConnectorPath, ERigElementType::Connector));
+		
+		if (!bIsConnected)
 		{
-			FRigElementKey ConnectorKey(*ConnectorName, ERigElementType::Connector);
-			const FRigElementKey* Target = Reference->Connections.Find(ConnectorKey);
-
-			if (!Target || !Target->IsValid())
-			{
-				Color = FStyleColors::AccentYellow;
-			}
-			else
-			{
-				bIsConnected = true;
-			}
+			Color = FStyleColors::AccentYellow;
 		}
 		
 		if (const UControlRig* ModuleRig = ConnectorModule->GetRig())
