Control Rig: Fix compilation errror
#rnx

diff --git a/Engine/Plugins/Animation/ControlRig/Source/ControlRig/Public/ModularRigModel.h b/Engine/Plugins/Animation/ControlRig/Source/ControlRig/Public/ModularRigModel.h
--- a/Engine/Plugins/Animation/ControlRig/Source/ControlRig/Public/ModularRigModel.h
+++ b/Engine/Plugins/Animation/ControlRig/Source/ControlRig/Public/ModularRigModel.h
@@ -106,6 +106,12 @@ struct FModularRigSingleConnection
 {
 	GENERATED_BODY()
 
+	FModularRigSingleConnection()
+		: Connector(FRigElementKey()), Target(FRigElementKey()) {}
+	
+	FModularRigSingleConnection(const FRigElementKey& InConnector, const FRigElementKey& InTarget)
+		: Connector(InConnector), Target(InTarget) {}
+
 	UPROPERTY()
 	FRigElementKey Connector;
 
