Control Rig: Fix module rigs not receiving important context info
#rnx

diff --git a/Engine/Plugins/Animation/ControlRig/Source/ControlRig/Private/ModularRig.cpp b/Engine/Plugins/Animation/ControlRig/Source/ControlRig/Private/ModularRig.cpp
--- a/Engine/Plugins/Animation/ControlRig/Source/ControlRig/Private/ModularRig.cpp
+++ b/Engine/Plugins/Animation/ControlRig/Source/ControlRig/Private/ModularRig.cpp
@@ -152,6 +152,21 @@ const FRigConnectorElement* FRigModuleInstance::FindPrimaryConnector() const
 	return nullptr;
 }
 
+void UModularRig::Serialize(FArchive& Ar)
+{
+	Super::Serialize(Ar);
+	if (Ar.IsLoading())
+	{
+		ModularRigModel.UpdateCachedChildren();
+	}
+}
+
+void UModularRig::PostLoad()
+{
+	Super::PostLoad();
+	ModularRigModel.UpdateCachedChildren();
+}
+
 void UModularRig::InitializeVMs(bool bRequestInit)
 {
 	URigVMHost::Initialize(bRequestInit);
@@ -355,10 +370,21 @@ void UModularRig::ExecuteQueue()
 				FRigUnitContext& RigUnitContext = RigPublicContext.UnitContext;
 				RigUnitContext = PublicContext.UnitContext;
 
-				// forward the draw interface to each module
+				// forward important context info to each module
 				RigPublicContext.SetDrawInterface(PublicContext.GetDrawInterface());
 				RigPublicContext.SetDrawContainer(PublicContext.GetDrawContainer());
 				RigPublicContext.RigModuleInstance = ExecutionElement.ModuleInstance;
+				RigPublicContext.SetAbsoluteTime(PublicContext.GetAbsoluteTime());
+				RigPublicContext.SetDeltaTime(PublicContext.GetDeltaTime());
+				RigPublicContext.SetWorld(PublicContext.GetWorld());
+				RigPublicContext.SetOwningActor(PublicContext.GetOwningActor());
+				RigPublicContext.SetOwningComponent(PublicContext.GetOwningComponent());
+#if WITH_EDITOR
+				RigPublicContext.SetLog(PublicContext.GetLog());
+#endif
+				RigPublicContext.SetFramesPerSecond(PublicContext.GetFramesPerSecond());
+				RigPublicContext.SetToWorldSpaceTransform(PublicContext.GetToWorldSpaceTransform());
+				RigPublicContext.RuntimeSettings = PublicContext.RuntimeSettings;
 
 				// re-initialize the module in case only the VM side got recompiled.
 				// this happens when the user relies on auto recompilation when editing the
diff --git a/Engine/Plugins/Animation/ControlRig/Source/ControlRig/Public/ModularRig.h b/Engine/Plugins/Animation/ControlRig/Source/ControlRig/Public/ModularRig.h
--- a/Engine/Plugins/Animation/ControlRig/Source/ControlRig/Public/ModularRig.h
+++ b/Engine/Plugins/Animation/ControlRig/Source/ControlRig/Public/ModularRig.h
@@ -127,6 +127,8 @@ class CONTROLRIG_API UModularRig : public UControlRig
 public:
 
 	// BEGIN ControlRig
+	virtual void Serialize(FArchive& Ar) override;
+	virtual void PostLoad() override;
 	virtual void InitializeVMs(bool bRequestInit = true) override;
 	virtual bool InitializeVMs(const FName& InEventName) override;
 	virtual void InitializeVMsFromCDO() override { URigVMHost::InitializeFromCDO(); }
