Control Rig: Fixing workflow bugs for modular rigging

#rnx

diff --git a/Engine/Plugins/Animation/ControlRig/Source/ControlRig/Private/ModularRigController.cpp b/Engine/Plugins/Animation/ControlRig/Source/ControlRig/Private/ModularRigController.cpp
--- a/Engine/Plugins/Animation/ControlRig/Source/ControlRig/Private/ModularRigController.cpp
+++ b/Engine/Plugins/Animation/ControlRig/Source/ControlRig/Private/ModularRigController.cpp
@@ -262,6 +262,18 @@ bool UModularRigController::ConnectConnectorToElement(const FRigElementKey& InCo
 	UBlueprint* Blueprint = Cast<UBlueprint>(GetOuter());
 
 #if WITH_EDITOR
+	FName TargetModulePathName = NAME_None;
+	if (UControlRig* RigCDO = Module->Class->GetDefaultObject<UControlRig>())
+	{
+		if (UModularRig* ModularRig = Cast<UModularRig>(Blueprint->GetObjectBeingDebugged()))
+		{
+			if (URigHierarchy* Hierarchy = ModularRig->GetHierarchy())
+			{
+				TargetModulePathName = Hierarchy->GetNameMetadata(InTargetKey, URigHierarchy::ModuleMetadataName, NAME_None);
+			}
+		}
+	}
+	
 	TSharedPtr<FScopedTransaction> TransactionPtr;
 	if (bSetupUndo)
 	{
@@ -275,18 +287,18 @@ bool UModularRigController::ConnectConnectorToElement(const FRigElementKey& InCo
 	{
 		DisconnectConnector(InConnectorKey, bSetupUndo);
 	}
-
+	
 	Model->Connections.AddConnection(InConnectorKey, InTargetKey);
 	Notify(EModularRigNotification::ConnectionChanged, Module);
 
 #if WITH_EDITOR
-	if (bAutoResolveOtherConnectors)
+	if (UControlRig* RigCDO = Module->Class->GetDefaultObject<UControlRig>())
 	{
-		if (UControlRig* RigCDO = Module->Class->GetDefaultObject<UControlRig>())
+		if (UModularRig* ModularRig = Cast<UModularRig>(Blueprint->GetObjectBeingDebugged()))
 		{
-			if (UModularRig* ModularRig = Cast<UModularRig>(Blueprint->GetObjectBeingDebugged()))
+			if (URigHierarchy* Hierarchy = ModularRig->GetHierarchy())
 			{
-				if (URigHierarchy* Hierarchy = ModularRig->GetHierarchy())
+				if (bAutoResolveOtherConnectors)
 				{
 					UModularRigRuleManager* RuleManager = Hierarchy->GetRuleManager();
 					const FRigModuleInstance* ModuleInstance = ModularRig->FindModule(Module->GetPath());
@@ -321,6 +333,18 @@ bool UModularRigController::ConnectConnectorToElement(const FRigElementKey& InCo
 						}
 					}
 				}
+
+				// automatically re-parent the module in the module tree as well
+				if(const FRigConnectorElement* Connector = Hierarchy->Find<FRigConnectorElement>(InConnectorKey))
+				{
+					if(Connector->Settings.Type == EConnectorType::Primary)
+					{
+						if(!TargetModulePathName.IsNone())
+						{
+							ReparentModule(Module->GetPath(), TargetModulePathName.ToString(), bSetupUndo);
+						}
+					}
+				}
 			}
 		}
 	}
@@ -964,6 +988,14 @@ FString UModularRigController::ReparentModule(const FString& InModulePath, const
 		return FString();
 	}
 
+	const FRigModuleReference* NewParentModule = FindModule(InNewParentModulePath);
+	const FString PreviousParentPath = Module->ParentPath;
+	const FString ParentPath = (NewParentModule) ? NewParentModule->GetPath() : FString();
+	if(PreviousParentPath.Equals(ParentPath, ESearchCase::CaseSensitive))
+	{
+		return Module->GetPath();
+	}
+
 #if WITH_EDITOR
 	TSharedPtr<FScopedTransaction> TransactionPtr;
 	if (bSetupUndo)
@@ -977,7 +1009,6 @@ FString UModularRigController::ReparentModule(const FString& InModulePath, const
 #endif
 
 	// Reparent or unparent children
-	FRigModuleReference* NewParentModule = FindModule(InNewParentModulePath);
 	const FString OldPath = Module->GetPath();
 	Module->PreviousParentPath = Module->ParentPath;
 	Module->ParentPath = (NewParentModule) ? NewParentModule->GetPath() : FString();
diff --git a/Engine/Plugins/Animation/ControlRig/Source/ControlRig/Private/Rigs/RigHierarchy.cpp b/Engine/Plugins/Animation/ControlRig/Source/ControlRig/Private/Rigs/RigHierarchy.cpp
--- a/Engine/Plugins/Animation/ControlRig/Source/ControlRig/Private/Rigs/RigHierarchy.cpp
+++ b/Engine/Plugins/Animation/ControlRig/Source/ControlRig/Private/Rigs/RigHierarchy.cpp
@@ -1527,22 +1527,21 @@ FText URigHierarchy::GetDisplayNameForUI(const FRigBaseElement* InElement, bool
 	check(InElement);
 
 	const FName& DisplayName = InElement->GetDisplayName();
+	FString DisplayNameString = DisplayName.ToString();
+	(void)DisplayNameString.Split(UModularRig::NamespaceSeparator, nullptr, &DisplayNameString, ESearchCase::CaseSensitive, ESearchDir::FromEnd);
 	
 	if(bIncludeNameSpace)
 	{
 		const FName ModuleShortName = GetNameMetadata(InElement->Key, ShortModuleNameMetadataName, NAME_None);
 		if(!ModuleShortName.IsNone())
 		{
-			FString DisplayNameString = DisplayName.ToString();
-			(void)DisplayNameString.Split(UModularRig::NamespaceSeparator, nullptr, &DisplayNameString, ESearchCase::CaseSensitive, ESearchDir::FromEnd);
-			
 			const FString ModuleShortNameString = ModuleShortName.ToString();
 			const FString ModuleDisplayName = JoinNameSpace(ModuleShortNameString, DisplayNameString);
 			return FText::FromString(ModuleDisplayName);
 		}
 	}
 
-	return FText::FromName(DisplayName);
+	return FText::FromString(*DisplayNameString);
 }
 
 FText URigHierarchy::GetDisplayNameForUI(const FRigElementKey& InKey, bool bIncludeNameSpace) const
diff --git a/Engine/Plugins/Animation/ControlRig/Source/ControlRig/Public/ModularRigModel.h b/Engine/Plugins/Animation/ControlRig/Source/ControlRig/Public/ModularRigModel.h
--- a/Engine/Plugins/Animation/ControlRig/Source/ControlRig/Public/ModularRigModel.h
+++ b/Engine/Plugins/Animation/ControlRig/Source/ControlRig/Public/ModularRigModel.h
@@ -189,6 +189,16 @@ struct CONTROLRIG_API FModularRigConnections
 		return EmptyKey;
 	}
 
+	const TArray<FRigElementKey>& FindConnectorsFromTarget(const FRigElementKey& InTargetKey) const
+	{
+		if(const TArray<FRigElementKey>* Connectors = ReverseConnectionMap.Find(InTargetKey))
+		{
+			return *Connectors;
+		}
+		static const TArray<FRigElementKey> EmptyList;
+		return EmptyList;
+	}
+
 	bool HasConnection(const FRigElementKey& InConnectorKey) const
 	{
 		return ConnectionList.IsValidIndex(FindConnectionIndex(InConnectorKey));
diff --git a/Engine/Plugins/Animation/ControlRig/Source/ControlRigEditor/Private/ControlRigBlueprintActions.cpp b/Engine/Plugins/Animation/ControlRig/Source/ControlRigEditor/Private/ControlRigBlueprintActions.cpp
--- a/Engine/Plugins/Animation/ControlRig/Source/ControlRigEditor/Private/ControlRigBlueprintActions.cpp
+++ b/Engine/Plugins/Animation/ControlRig/Source/ControlRigEditor/Private/ControlRigBlueprintActions.cpp
@@ -253,7 +253,6 @@ UControlRigBlueprint* FControlRigBlueprintActions::CreateControlRigFromSkeletalM
 	{
 		Controller->ImportBones(*RefSkeleton, NAME_None, false, false, false, false);
 		Controller->ImportCurves(Skeleton, NAME_None, false, false);
-		Controller->AddDefaultRootSocket();
 	}
 	NewControlRigBlueprint->SourceHierarchyImport = Skeleton;
 	NewControlRigBlueprint->SourceCurveImport = Skeleton;
diff --git a/Engine/Plugins/Animation/ControlRig/Source/ControlRigEditor/Private/ControlRigDragOps.cpp b/Engine/Plugins/Animation/ControlRig/Source/ControlRigEditor/Private/ControlRigDragOps.cpp
new file mode 100644
--- /dev/null
+++ b/Engine/Plugins/Animation/ControlRig/Source/ControlRigEditor/Private/ControlRigDragOps.cpp
@@ -0,0 +1,114 @@
+// Copyright Epic Games, Inc. All Rights Reserved.
+
+#include "ControlRigDragOps.h"
+#include "Editor/SRigHierarchyTagWidget.h"
+
+//////////////////////////////////////////////////////////////
+/// FRigElementHierarchyDragDropOp
+///////////////////////////////////////////////////////////
+
+TSharedRef<FRigElementHierarchyDragDropOp> FRigElementHierarchyDragDropOp::New(const TArray<FRigElementKey>& InElements)
+{
+	TSharedRef<FRigElementHierarchyDragDropOp> Operation = MakeShared<FRigElementHierarchyDragDropOp>();
+	Operation->Elements = InElements;
+	Operation->Construct();
+	return Operation;
+}
+
+TSharedPtr<SWidget> FRigElementHierarchyDragDropOp::GetDefaultDecorator() const
+{
+	return SNew(SBorder)
+		.Visibility(EVisibility::Visible)
+		.BorderImage(FAppStyle::GetBrush("Menu.Background"))
+		[
+			SNew(STextBlock)
+			.Text(FText::FromString(GetJoinedElementNames()))
+			//.Font(FAppStyle::Get().GetFontStyle("FontAwesome.10"))
+		];
+}
+
+FString FRigElementHierarchyDragDropOp::GetJoinedElementNames() const
+{
+	TArray<FString> ElementNameStrings;
+	for (const FRigElementKey& Element: Elements)
+	{
+		ElementNameStrings.Add(Element.Name.ToString());
+	}
+	return FString::Join(ElementNameStrings, TEXT(","));
+}
+
+bool FRigElementHierarchyDragDropOp::IsDraggingSingleConnector() const
+{
+	if(Elements.Num() == 1)
+	{
+		return Elements[0].Type == ERigElementType::Connector;
+	}
+	return false;
+}
+
+bool FRigElementHierarchyDragDropOp::IsDraggingSingleSocket() const
+{
+	if(Elements.Num() == 1)
+	{
+		return Elements[0].Type == ERigElementType::Socket;
+	}
+	return false;
+}
+
+//////////////////////////////////////////////////////////////
+/// FRigHierarchyTagDragDropOp
+///////////////////////////////////////////////////////////
+
+TSharedRef<FRigHierarchyTagDragDropOp> FRigHierarchyTagDragDropOp::New(TSharedPtr<SRigHierarchyTagWidget> InTagWidget)
+{
+	TSharedRef<FRigHierarchyTagDragDropOp> Operation = MakeShared<FRigHierarchyTagDragDropOp>();
+	Operation->Text = InTagWidget->Text.Get();
+	Operation->Identifier = InTagWidget->Identifier.Get();
+	Operation->Construct();
+	return Operation;
+}
+
+TSharedPtr<SWidget> FRigHierarchyTagDragDropOp::GetDefaultDecorator() const
+{
+	return SNew(SBorder)
+		.Visibility(EVisibility::Visible)
+		.BorderImage(FAppStyle::GetBrush("Menu.Background"))
+		[
+			SNew(STextBlock)
+			.Text(Text)
+		];
+}
+
+//////////////////////////////////////////////////////////////
+/// FModularRigModuleDragDropOp
+///////////////////////////////////////////////////////////
+
+TSharedRef<FModularRigModuleDragDropOp> FModularRigModuleDragDropOp::New(const TArray<FString>& InElements)
+{
+	TSharedRef<FModularRigModuleDragDropOp> Operation = MakeShared<FModularRigModuleDragDropOp>();
+	Operation->Elements = InElements;
+	Operation->Construct();
+	return Operation;
+}
+
+TSharedPtr<SWidget> FModularRigModuleDragDropOp::GetDefaultDecorator() const
+{
+	return SNew(SBorder)
+		.Visibility(EVisibility::Visible)
+		.BorderImage(FAppStyle::GetBrush("Menu.Background"))
+		[
+			SNew(STextBlock)
+			.Text(FText::FromString(GetJoinedElementNames()))
+			//.Font(FAppStyle::Get().GetFontStyle("FontAwesome.10"))
+		];
+}
+
+FString FModularRigModuleDragDropOp::GetJoinedElementNames() const
+{
+	TArray<FString> ElementNameStrings;
+	for (const FString& Element: Elements)
+	{
+		ElementNameStrings.Add(Element);
+	}
+	return FString::Join(ElementNameStrings, TEXT(","));
+}
diff --git a/Engine/Plugins/Animation/ControlRig/Source/ControlRigEditor/Private/ControlRigDragOps.h b/Engine/Plugins/Animation/ControlRig/Source/ControlRigEditor/Private/ControlRigDragOps.h
new file mode 100644
--- /dev/null
+++ b/Engine/Plugins/Animation/ControlRig/Source/ControlRigEditor/Private/ControlRigDragOps.h
@@ -0,0 +1,91 @@
+// Copyright Epic Games, Inc. All Rights Reserved.
+
+#pragma once
+
+#include "DragAndDrop/GraphNodeDragDropOp.h"
+#include "Rigs/RigHierarchyDefines.h"
+
+class SRigHierarchyTagWidget;
+
+class FRigElementHierarchyDragDropOp : public FGraphNodeDragDropOp
+{
+public:
+	DRAG_DROP_OPERATOR_TYPE(FRigElementHierarchyDragDropOp, FGraphNodeDragDropOp)
+
+	static TSharedRef<FRigElementHierarchyDragDropOp> New(const TArray<FRigElementKey>& InElements);
+
+	virtual TSharedPtr<SWidget> GetDefaultDecorator() const override;
+
+	/** @return true if this drag operation contains property paths */
+	bool HasElements() const
+	{
+		return Elements.Num() > 0;
+	}
+
+	/** @return The property paths from this drag operation */
+	const TArray<FRigElementKey>& GetElements() const
+	{
+		return Elements;
+	}
+
+	FString GetJoinedElementNames() const;
+
+	bool IsDraggingSingleConnector() const;
+	bool IsDraggingSingleSocket() const;
+
+private:
+
+	/** Data for the property paths this item represents */
+	TArray<FRigElementKey> Elements;
+};
+
+class FRigHierarchyTagDragDropOp : public FDecoratedDragDropOp
+{
+public:
+	DRAG_DROP_OPERATOR_TYPE(FRigHierarchyTagDragDropOp, FDecoratedDragDropOp)
+
+	static TSharedRef<FRigHierarchyTagDragDropOp> New(TSharedPtr<SRigHierarchyTagWidget> InTagWidget);
+
+	virtual TSharedPtr<SWidget> GetDefaultDecorator() const override;
+
+	/** @return The identifier being dragged */
+	const FString& GetIdentifier() const
+	{
+		return Identifier;
+	}
+
+private:
+
+	FText Text;
+	FString Identifier;
+};
+
+class FModularRigModuleDragDropOp : public FDragDropOperation
+{
+public:
+	DRAG_DROP_OPERATOR_TYPE(FModularRigModuleDragDropOp, FDragDropOperation)
+
+	static TSharedRef<FModularRigModuleDragDropOp> New(const TArray<FString>& InElements);
+
+	virtual TSharedPtr<SWidget> GetDefaultDecorator() const override;
+
+	/** @return true if this drag operation contains property paths */
+	bool HasElements() const
+	{
+		return Elements.Num() > 0;
+	}
+
+	/** @return The property paths from this drag operation */
+	const TArray<FString>& GetElements() const
+	{
+		return Elements;
+	}
+
+	FString GetJoinedElementNames() const;
+
+private:
+
+	/** Data for the property paths this item represents */
+	TArray<FString> Elements;
+};
+
diff --git a/Engine/Plugins/Animation/ControlRig/Source/ControlRigEditor/Private/Editor/ControlRigEditor.cpp b/Engine/Plugins/Animation/ControlRig/Source/ControlRigEditor/Private/Editor/ControlRigEditor.cpp
--- a/Engine/Plugins/Animation/ControlRig/Source/ControlRigEditor/Private/Editor/ControlRigEditor.cpp
+++ b/Engine/Plugins/Animation/ControlRig/Source/ControlRigEditor/Private/Editor/ControlRigEditor.cpp
@@ -2550,6 +2550,14 @@ void FControlRigEditor::UpdateRigVMHost()
 			ControlRig->OnPreForwardsSolve_AnyThread().AddSP(this, &FControlRigEditor::OnPreForwardsSolve_AnyThread);
 			ControlRig->ControlModified().AddSP(this, &FControlRigEditor::HandleOnControlModified);
 		}
+
+		if(IsModularRig() && ControlRig)
+		{
+			if(SchematicModel.ControlRigBlueprint.IsValid())
+			{
+				SchematicModel.OnSetObjectBeingDebugged(ControlRig);
+			}
+		}
 	}
 }
 
@@ -5046,9 +5054,6 @@ void FControlRigEditor::OnPreConstruction_AnyThread(UControlRig* InRig, const FN
 							}
 						}
 
-						// create a default socket under the root bone
-						Controller->AddDefaultRootSocket();
-
 						// create a null to store controls under
 						static const FRigElementKey ControlParentKey(TEXT("Controls"), ERigElementType::Null);
 						if(!Hierarchy->Contains(ControlParentKey))
diff --git a/Engine/Plugins/Animation/ControlRig/Source/ControlRigEditor/Private/Editor/ControlRigSchematicModel.cpp b/Engine/Plugins/Animation/ControlRig/Source/ControlRigEditor/Private/Editor/ControlRigSchematicModel.cpp
--- a/Engine/Plugins/Animation/ControlRig/Source/ControlRigEditor/Private/Editor/ControlRigSchematicModel.cpp
+++ b/Engine/Plugins/Animation/ControlRig/Source/ControlRigEditor/Private/Editor/ControlRigSchematicModel.cpp
@@ -14,6 +14,8 @@
 #include <SchematicGraphPanel/SchematicGraphStyle.h>
 #include "Editor/SModularRigModel.h"
 #include "Blueprint/WidgetBlueprintLibrary.h"
+#include "ControlRigDragOps.h"
+#include "SchematicGraphPanel/SSchematicGraphPanel.h"
 
 #define LOCTEXT_NAMESPACE "ControlRigSchematicModel"
 
@@ -46,6 +48,7 @@ FControlRigSchematicModel::~FControlRigSchematicModel()
 			if(!ControlRigBeingDebugged->HasAnyFlags(RF_BeginDestroyed))
 			{
 				ControlRigBeingDebugged->GetHierarchy()->OnModified().RemoveAll(this);
+				ControlRigBeingDebugged->OnPostConstruction_AnyThread().RemoveAll(this);
 			}
 		}
 	}
@@ -59,7 +62,7 @@ void FControlRigSchematicModel::SetEditor(const TSharedRef<FControlRigEditor>& I
 {
 	ControlRigEditor = InEditor;
 	ControlRigBlueprint = ControlRigEditor.Pin()->GetControlRigBlueprint();
-	ControlRigBeingDebuggedPtr = ControlRigBlueprint->GetDebuggedControlRig();
+	OnSetObjectBeingDebugged(ControlRigBlueprint->GetDebuggedControlRig());
 }
 
 void FControlRigSchematicModel::Reset()
@@ -76,9 +79,10 @@ void FControlRigSchematicModel::Tick(float InDeltaTime)
 	{
 		if(ControlRigBlueprint.IsValid())
 		{
-			TArray<FRigElementKey> KeysWithConnectors;
-			ControlRigBlueprint->ConnectionMap.GenerateValueArray(KeysWithConnectors);
-			
+			const URigHierarchy* Hierarchy = ControlRigBlueprint->GetDebuggedControlRig()->GetHierarchy();
+			check(Hierarchy);
+			const FModularRigConnections& Connections = ControlRigBlueprint->ModularRigModel.Connections;
+
 			TArray<FRigElementKey> KeysToRemove;
 			for(const TPair<FRigElementKey,FGuid>& Pair : RigElementKeyToGuid)
 			{
@@ -87,10 +91,23 @@ void FControlRigSchematicModel::Tick(float InDeltaTime)
 				{
 					continue;
 				}
-				if(KeysWithConnectors.Contains(Pair.Key))
+
+				const TArray<FRigElementKey>& Connectors = Connections.FindConnectorsFromTarget(Pair.Key);
+				if(Connectors.Num() > 1)
 				{
 					continue;
 				}
+				if(Connectors.Num() == 1)
+				{
+					if(const FRigConnectorElement* Connector = Hierarchy->Find<FRigConnectorElement>(Pair.Key))
+					{
+						if(Connector->Settings.Type != EConnectorType::Primary)
+						{
+							continue;
+						}
+					}
+				}
+				
 				KeysToRemove.Add(Pair.Key);
 			}
 			for(const FRigElementKey& KeyToRemove : KeysToRemove)
@@ -98,7 +115,16 @@ void FControlRigSchematicModel::Tick(float InDeltaTime)
 				RemoveElementKeyNode(KeyToRemove);
 			}
 		}
-		bUpdateConnectorTargetsOnTick = true;
+		bUpdateConnectorTargetsOnTick = false;
+	}
+
+	if(FSlateApplication::Get().IsDragDropping())
+	{
+		
+	}
+	else
+	{
+		
 	}
 }
 
@@ -278,7 +304,10 @@ void FControlRigSchematicModel::UpdateElementKeyLinks()
 					const FRigElementKey ParentB = Hierarchy->GetFirstParent(ChildNode->GetKey());
 					if(!ParentA || !ParentB.IsValid() || ParentA != ParentB)
 					{
-						OutExpectedLinks.Emplace(Node->GetKey(), ChildNode->GetKey());
+						if(ChildNode->GetKey().Type != ERigElementType::Socket)
+						{
+							OutExpectedLinks.Emplace(Node->GetKey(), ChildNode->GetKey());
+						}
 					}
 				}
 			}
@@ -415,6 +444,7 @@ void FControlRigSchematicModel::OnSetObjectBeingDebugged(UObject* InObject)
 			if(!ControlRigBeingDebugged->HasAnyFlags(RF_BeginDestroyed))
 			{
 				ControlRigBeingDebugged->GetHierarchy()->OnModified().RemoveAll(this);
+				ControlRigBeingDebugged->OnPostConstruction_AnyThread().RemoveAll(this);
 			}
 		}
 	}
@@ -428,6 +458,8 @@ void FControlRigSchematicModel::OnSetObjectBeingDebugged(UObject* InObject)
 		{
 			Hierarchy->OnModified().RemoveAll(this);
 			Hierarchy->OnModified().AddRaw(this, &FControlRigSchematicModel::OnHierarchyModified);
+			ControlRig->OnPostConstruction_AnyThread().AddRaw(this, &FControlRigSchematicModel::HandlePostConstruction);
+
 			UpdateControlRigContent();
 		}
 	}
@@ -556,8 +588,6 @@ void FControlRigSchematicModel::OnHierarchyModified(ERigHierarchyNotification In
 
 void FControlRigSchematicModel::HandleModularRigModified(EModularRigNotification InNotification, const FRigModuleReference* InModule)
 {
-	bool bUnresolvedConnectors = false;
-	
 	switch (InNotification)
 	{
 		case EModularRigNotification::ConnectionChanged:
@@ -595,7 +625,6 @@ void FControlRigSchematicModel::HandleModularRigModified(EModularRigNotification
 									{
 										(void)RemoveFromParentNode(ConnectorNode);
 										const_cast<FSchematicGraphNode*>(ConnectorNode)->OnMouseLeave();
-										bUnresolvedConnectors = true;
 									}
 								}
 							}
@@ -611,11 +640,7 @@ void FControlRigSchematicModel::HandleModularRigModified(EModularRigNotification
 		}
 	}
 
-	// if we have unresolved connectors we may want to remove nodes which no longer have any connectors resolved to them
-	if(bUnresolvedConnectors && ControlRigBlueprint.IsValid())
-	{
-		bUpdateConnectorTargetsOnTick = true;
-	}
+	bUpdateConnectorTargetsOnTick = true;
 }
 
 FSchematicGraphGroupNode* FControlRigSchematicModel::AddAutoGroupNode()
@@ -1115,6 +1140,104 @@ const FText FControlRigSchematicModel::GetToolTipForTag(const FSchematicGraphTag
 	return Super::GetToolTipForTag(InTag);
 }
 
+bool FControlRigSchematicModel::GetForwardedNodeForDrag(FGuid& InOutGuid) const
+{
+	if(const FControlRigSchematicRigElementKeyNode* Node = FindNode<FControlRigSchematicRigElementKeyNode>(InOutGuid))
+	{
+		if(ControlRigBlueprint.IsValid())
+		{
+			const URigHierarchy* Hierarchy = ControlRigBlueprint->GetDebuggedControlRig()->GetHierarchy();
+			check(Hierarchy);
+			const FModularRigConnections& Connections = ControlRigBlueprint->ModularRigModel.Connections;
+			const TArray<FRigElementKey>& Connectors = Connections.FindConnectorsFromTarget(Node->GetKey());
+			if(Connectors.Num() == 1)
+			{
+				if(const FControlRigSchematicRigElementKeyNode* ConnectorNode = FindElementKeyNode(Connectors[0]))
+				{
+					InOutGuid = ConnectorNode->GetGuid();
+					return true;
+				}
+			}
+		}
+	}
+	return Super::GetForwardedNodeForDrag(InOutGuid);
+}
+
+TArray<FRigElementKey> FControlRigSchematicModel::GetElementKeysFromDragDropEvent(const FDragDropOperation& InDragDropOperation, const UControlRig* InControlRig)
+{
+	TArray<FRigElementKey> DraggedKeys;
+
+	if(InDragDropOperation.IsOfType<FRigElementHierarchyDragDropOp>())
+	{
+		const FRigElementHierarchyDragDropOp* RigDragDropOp = StaticCast<const FRigElementHierarchyDragDropOp*>(&InDragDropOperation);
+		return RigDragDropOp->GetElements();
+	}
+
+	if(InDragDropOperation.IsOfType<FRigHierarchyTagDragDropOp>())
+	{
+		const FRigHierarchyTagDragDropOp* TagDragDropOp = StaticCast<const FRigHierarchyTagDragDropOp*>(&InDragDropOperation);
+		FRigElementKey DraggedKey;
+		FRigElementKey::StaticStruct()->ImportText(*TagDragDropOp->GetIdentifier(), &DraggedKey, nullptr, EPropertyPortFlags::PPF_None, nullptr, FRigElementKey::StaticStruct()->GetName(), true);
+		return {DraggedKey};
+	}
+	
+	if(InDragDropOperation.IsOfType<FModularRigModuleDragDropOp>())
+	{
+		const FModularRigModuleDragDropOp* ModuleDropOp = StaticCast<const FModularRigModuleDragDropOp*>(&InDragDropOperation);
+		const TArray<FString> Sources = ModuleDropOp->GetElements();
+
+		URigHierarchy* Hierarchy = InControlRig->GetHierarchy();
+		if (!Hierarchy)
+		{
+			return DraggedKeys;
+		}
+
+		const TArray<FRigConnectorElement*> Connectors = Hierarchy->GetConnectors();
+		for (const FString& ModulePathOrConnectorName : Sources)
+		{
+			if(const FRigConnectorElement* Connector = Hierarchy->Find<FRigConnectorElement>(FRigElementKey(*ModulePathOrConnectorName, ERigElementType::Connector)))
+			{
+				DraggedKeys.Add(Connector->GetKey());
+			}
+
+			if(const UModularRig* ModularRig = Cast<UModularRig>(InControlRig))
+			{
+				if(const FRigModuleInstance* ModuleInstance = ModularRig->FindModule(ModulePathOrConnectorName))
+				{
+					const FString ModuleNameSpace = ModuleInstance->GetNamespace();
+					for(const FRigConnectorElement* Connector : Connectors)
+					{
+						const FString ConnectorNameSpace = Hierarchy->GetNameMetadata(Connector->GetKey(), URigHierarchy::NameSpaceMetadataName, NAME_None).ToString();
+						if(ConnectorNameSpace.Equals(ModuleNameSpace))
+						{
+							if(Connector->Settings.Type == EConnectorType::Primary)
+							{
+								DraggedKeys.Add(Connector->GetKey());
+							}
+						}
+					}
+				}
+			}
+		}
+	}
+
+	if(InDragDropOperation.IsOfType<FSchematicGraphNodeDragDropOp>())
+	{
+		const FSchematicGraphNodeDragDropOp* SchematicGraphNodeDragDropOp = StaticCast<const FSchematicGraphNodeDragDropOp*>(&InDragDropOperation);
+		const TArray<const FSchematicGraphNode*> Nodes = SchematicGraphNodeDragDropOp->GetNodes();
+
+		for(const FSchematicGraphNode* Node : Nodes)
+		{
+			if(const FControlRigSchematicRigElementKeyNode* ElementKeyNode = Cast<FControlRigSchematicRigElementKeyNode>(Node))
+			{
+				DraggedKeys.Add(ElementKeyNode->GetKey());
+			}
+		}
+	}
+
+	return DraggedKeys;
+}
+
 void FControlRigSchematicModel::HandleSchematicNodeClicked(SSchematicGraphPanel* InPanel, SSchematicGraphNode* InNode, const FPointerEvent& InMouseEvent)
 {
 	for (const TSharedPtr<FSchematicGraphNode>& Node : Nodes)
@@ -1159,9 +1282,22 @@ void FControlRigSchematicModel::HandleSchematicBeginDrag(SSchematicGraphPanel* I
 		return;
 	}
 
-	const FControlRigSchematicRigElementKeyNode* ElementKeyNode =
-		Cast<FControlRigSchematicRigElementKeyNode>(InNode->GetNodeData());
-	if(ElementKeyNode == nullptr)
+	FRigElementKey DraggedKey;
+	if(const FControlRigSchematicRigElementKeyNode* ElementKeyNode =
+		Cast<FControlRigSchematicRigElementKeyNode>(InNode->GetNodeData()))
+	{
+		DraggedKey = ElementKeyNode->GetKey();
+	}
+	else
+	{
+		const TArray<FRigElementKey> DraggedKeys = GetElementKeysFromDragDropEvent(InDragDropOperation, ControlRigBlueprint->GetDebuggedControlRig());
+		if(!DraggedKeys.IsEmpty())
+		{
+			DraggedKey = DraggedKeys[0];
+		}
+	}
+
+	if(!DraggedKey.IsValid())
 	{
 		return;
 	}
@@ -1172,12 +1308,6 @@ void FControlRigSchematicModel::HandleSchematicBeginDrag(SSchematicGraphPanel* I
 		return;
 	}
 
-	const FRigElementKey DraggedKey = ElementKeyNode->GetKey();
-	if (!DraggedKey.IsValid())
-	{
-		return;
-	}
-
 	FRigBaseElement* Element = Hierarchy->Find(DraggedKey);
 	if (!Element)
 	{
@@ -1430,6 +1560,11 @@ void FControlRigSchematicModel::HandleSchematicDrop(SSchematicGraphPanel* InPane
 	}
 }
 
+void FControlRigSchematicModel::HandlePostConstruction(UControlRig* Subject, const FName& InEventName)
+{
+	UpdateControlRigContent();
+}
+
 bool FControlRigSchematicModel::IsConnectorResolved(const FRigElementKey& InConnectorKey, FRigElementKey* OutKey) const
 {
 	if(ControlRigBlueprint.IsValid() && InConnectorKey.Type == ERigElementType::Connector)
diff --git a/Engine/Plugins/Animation/ControlRig/Source/ControlRigEditor/Private/Editor/ControlRigSchematicModel.h b/Engine/Plugins/Animation/ControlRig/Source/ControlRigEditor/Private/Editor/ControlRigSchematicModel.h
--- a/Engine/Plugins/Animation/ControlRig/Source/ControlRigEditor/Private/Editor/ControlRigSchematicModel.h
+++ b/Engine/Plugins/Animation/ControlRig/Source/ControlRigEditor/Private/Editor/ControlRigSchematicModel.h
@@ -6,6 +6,7 @@
 #include "SchematicGraphPanel/SSchematicGraphPanel.h"
 #include "Rigs/RigHierarchyDefines.h"
 #include "ControlRigBlueprint.h"
+#include "Input/DragAndDrop.h"
 
 class FControlRigEditor;
 class UControlRig;
@@ -104,6 +105,9 @@ public:
 	virtual FLinearColor GetColorForLink(const FSchematicGraphLink* InLink) const override;
 	virtual ESchematicGraphVisibility::Type GetVisibilityForTag(const FSchematicGraphTag* InTag) const override;
 	virtual const FText GetToolTipForTag(const FSchematicGraphTag* InTag) const override;
+	virtual bool GetForwardedNodeForDrag(FGuid& InOutGuid) const override;
+
+	static TArray<FRigElementKey> GetElementKeysFromDragDropEvent(const FDragDropOperation& InDragDropOperation, const UControlRig* InControlRig);
 
 private:
 
@@ -113,8 +117,10 @@ private:
 	void HandleSchematicBeginDrag(SSchematicGraphPanel* InPanel, SSchematicGraphNode* InNode, const FDragDropOperation& InDragDropOperation);
 	void HandleSchematicEndDrag(SSchematicGraphPanel* InPanel, SSchematicGraphNode* InNode, const FDragDropOperation& InDragDropOperation);
 	void HandleSchematicDrop(SSchematicGraphPanel* InPanel, SSchematicGraphNode* InNode, const FDragDropEvent& InDragDropEvent);
+	void HandlePostConstruction(UControlRig* Subject, const FName& InEventName);
 	bool IsConnectorResolved(const FRigElementKey& InConnectorKey, FRigElementKey* OutKey = nullptr) const;
 
+
 	TWeakPtr<FControlRigEditor> ControlRigEditor;
 	TWeakObjectPtr<UControlRigBlueprint> ControlRigBlueprint;
 	TWeakObjectPtr<UControlRig> ControlRigBeingDebuggedPtr;
diff --git a/Engine/Plugins/Animation/ControlRig/Source/ControlRigEditor/Private/Editor/SModularRigModel.cpp b/Engine/Plugins/Animation/ControlRig/Source/ControlRigEditor/Private/Editor/SModularRigModel.cpp
--- a/Engine/Plugins/Animation/ControlRig/Source/ControlRigEditor/Private/Editor/SModularRigModel.cpp
+++ b/Engine/Plugins/Animation/ControlRig/Source/ControlRigEditor/Private/Editor/SModularRigModel.cpp
@@ -57,39 +57,6 @@
 
 #define LOCTEXT_NAMESPACE "SModularRigModel"
 
-//////////////////////////////////////////////////////////////
-/// FModularRigModuleDragDropOp
-///////////////////////////////////////////////////////////
-TSharedRef<FModularRigModuleDragDropOp> FModularRigModuleDragDropOp::New(const TArray<FString>& InElements)
-{
-	TSharedRef<FModularRigModuleDragDropOp> Operation = MakeShared<FModularRigModuleDragDropOp>();
-	Operation->Elements = InElements;
-	Operation->Construct();
-	return Operation;
-}
-
-TSharedPtr<SWidget> FModularRigModuleDragDropOp::GetDefaultDecorator() const
-{
-	return SNew(SBorder)
-		.Visibility(EVisibility::Visible)
-		.BorderImage(FAppStyle::GetBrush("Menu.Background"))
-		[
-			SNew(STextBlock)
-			.Text(FText::FromString(GetJoinedElementNames()))
-			//.Font(FAppStyle::Get().GetFontStyle("FontAwesome.10"))
-		];
-}
-
-FString FModularRigModuleDragDropOp::GetJoinedElementNames() const
-{
-	TArray<FString> ElementNameStrings;
-	for (const FString& Element: Elements)
-	{
-		ElementNameStrings.Add(Element);
-	}
-	return FString::Join(ElementNameStrings, TEXT(","));
-}
-
 ///////////////////////////////////////////////////////////
 
 const FName SModularRigModel::ContextMenuName = TEXT("ControlRigEditor.ModularRigModel.ContextMenu");
diff --git a/Engine/Plugins/Animation/ControlRig/Source/ControlRigEditor/Private/Editor/SModularRigModel.h b/Engine/Plugins/Animation/ControlRig/Source/ControlRigEditor/Private/Editor/SModularRigModel.h
--- a/Engine/Plugins/Animation/ControlRig/Source/ControlRigEditor/Private/Editor/SModularRigModel.h
+++ b/Engine/Plugins/Animation/ControlRig/Source/ControlRigEditor/Private/Editor/SModularRigModel.h
@@ -8,6 +8,7 @@
 #include "ControlRigBlueprint.h"
 #include "DragAndDrop/GraphNodeDragDropOp.h"
 #include "Editor/RigVMEditor.h"
+#include "ControlRigDragOps.h"
 
 class SModularRigModel;
 class FControlRigEditor;
@@ -20,35 +21,6 @@ class FMenuBuilder;
 class UToolMenu;
 struct FToolMenuContext;
 
-class FModularRigModuleDragDropOp : public FDragDropOperation
-{
-public:
-	DRAG_DROP_OPERATOR_TYPE(FModularRigModuleDragDropOp, FDragDropOperation)
-
-	static TSharedRef<FModularRigModuleDragDropOp> New(const TArray<FString>& InElements);
-
-	virtual TSharedPtr<SWidget> GetDefaultDecorator() const override;
-
-	/** @return true if this drag operation contains property paths */
-	bool HasElements() const
-	{
-		return Elements.Num() > 0;
-	}
-
-	/** @return The property paths from this drag operation */
-	const TArray<FString>& GetElements() const
-	{
-		return Elements;
-	}
-
-	FString GetJoinedElementNames() const;
-
-private:
-
-	/** Data for the property paths this item represents */
-	TArray<FString> Elements;
-};
-
 /** Widget allowing editing of a control rig's structure */
 class SModularRigModel : public SCompoundWidget, public FEditorUndoClient
 {
diff --git a/Engine/Plugins/Animation/ControlRig/Source/ControlRigEditor/Private/Editor/SRigHierarchy.cpp b/Engine/Plugins/Animation/ControlRig/Source/ControlRigEditor/Private/Editor/SRigHierarchy.cpp
--- a/Engine/Plugins/Animation/ControlRig/Source/ControlRigEditor/Private/Editor/SRigHierarchy.cpp
+++ b/Engine/Plugins/Animation/ControlRig/Source/ControlRigEditor/Private/Editor/SRigHierarchy.cpp
@@ -53,62 +53,12 @@
 #include "DragAndDrop/AssetDragDropOp.h"
 #include "RigVMFunctions/Math/RigVMMathLibrary.h"
 #include "Preferences/PersonaOptions.h"
+#include "Editor/SModularRigModel.h"
 
 #include UE_INLINE_GENERATED_CPP_BY_NAME(SRigHierarchy)
 
 #define LOCTEXT_NAMESPACE "SRigHierarchy"
 
-//////////////////////////////////////////////////////////////
-/// FRigElementHierarchyDragDropOp
-///////////////////////////////////////////////////////////
-TSharedRef<FRigElementHierarchyDragDropOp> FRigElementHierarchyDragDropOp::New(const TArray<FRigElementKey>& InElements)
-{
-	TSharedRef<FRigElementHierarchyDragDropOp> Operation = MakeShared<FRigElementHierarchyDragDropOp>();
-	Operation->Elements = InElements;
-	Operation->Construct();
-	return Operation;
-}
-
-TSharedPtr<SWidget> FRigElementHierarchyDragDropOp::GetDefaultDecorator() const
-{
-	return SNew(SBorder)
-		.Visibility(EVisibility::Visible)
-		.BorderImage(FAppStyle::GetBrush("Menu.Background"))
-		[
-			SNew(STextBlock)
-			.Text(FText::FromString(GetJoinedElementNames()))
-			//.Font(FAppStyle::Get().GetFontStyle("FontAwesome.10"))
-		];
-}
-
-FString FRigElementHierarchyDragDropOp::GetJoinedElementNames() const
-{
-	TArray<FString> ElementNameStrings;
-	for (const FRigElementKey& Element: Elements)
-	{
-		ElementNameStrings.Add(Element.Name.ToString());
-	}
-	return FString::Join(ElementNameStrings, TEXT(","));
-}
-
-bool FRigElementHierarchyDragDropOp::IsDraggingSingleConnector() const
-{
-	if(Elements.Num() == 1)
-	{
-		return Elements[0].Type == ERigElementType::Connector;
-	}
-	return false;
-}
-
-bool FRigElementHierarchyDragDropOp::IsDraggingSingleSocket() const
-{
-	if(Elements.Num() == 1)
-	{
-		return Elements[0].Type == ERigElementType::Socket;
-	}
-	return false;
-}
-
 ///////////////////////////////////////////////////////////
 
 const FName SRigHierarchy::ContextMenuName = TEXT("ControlRigEditor.RigHierarchy.ContextMenu");
@@ -1743,10 +1693,6 @@ void SRigHierarchy::ImportHierarchy(const FAssetData& InAssetData)
 
 		const TArray<FRigElementKey> ImportedBones = Controller->ImportBones(Mesh->GetSkeleton(), NAME_None, false, false, bSelectBones, true, true);
 		Controller->ImportCurves(Mesh->GetSkeleton(), NAME_None, false, true);
-		if(bIsModularRig && Hierarchy->GetSockets().Num() == 0)
-		{
-			(void)Controller->AddDefaultRootSocket();
-		}
 
 		ControlRigBlueprint->SourceHierarchyImport = Mesh->GetSkeleton();
 		ControlRigBlueprint->SourceCurveImport = Mesh->GetSkeleton();
@@ -2738,6 +2684,44 @@ TOptional<EItemDropZone> SRigHierarchy::OnCanAcceptDrop(const FDragDropEvent& Dr
 			}
 		}
 	}
+	
+	const TSharedPtr<FModularRigModuleDragDropOp> ModuleDropOp = DragDropEvent.GetOperationAs<FModularRigModuleDragDropOp>();
+	if (ModuleDropOp.IsValid() && TargetItem.IsValid())
+	{
+		if(DropZone != EItemDropZone::OntoItem)
+		{
+			return InvalidDropZone;
+		}
+
+		const UModularRig* ControlRig = Cast<UModularRig>(ControlRigBlueprint->GetDebuggedControlRig());
+		if (!ControlRig)
+		{
+			return InvalidDropZone;
+		}
+
+		const FRigElementKey TargetKey = TargetItem->Key;
+		const TArray<FRigElementKey> DraggedKeys = FControlRigSchematicModel::GetElementKeysFromDragDropEvent(*ModuleDropOp.Get(), ControlRig);
+		for(const FRigElementKey& DraggedKey : DraggedKeys)
+		{
+			if(DraggedKey.Type != ERigElementType::Connector)
+			{
+				continue;
+			}
+			
+			if(!DragRigResolveResults.Contains(DraggedKey))
+			{
+				UpdateConnectorMatchesOnDrag({DraggedKey});
+			}
+			
+			const FModularRigResolveResult& ResolveResult = DragRigResolveResults.FindChecked(DraggedKey);
+			if(ResolveResult.ContainsMatch(TargetKey))
+			{
+				return DropZone;
+			}
+		}
+
+		return InvalidDropZone;
+	}
 
 	TSharedPtr<FAssetDragDropOp> AssetDragDropOp = DragDropEvent.GetOperationAs<FAssetDragDropOp>();
 	if (AssetDragDropOp.IsValid())
@@ -2872,18 +2856,39 @@ FReply SRigHierarchy::OnAcceptDrop(const FDragDropEvent& DragDropEvent, EItemDro
 	const TSharedPtr<FRigHierarchyTagDragDropOp> TagDragDropOp = DragDropEvent.GetOperationAs<FRigHierarchyTagDragDropOp>();
 	if(TagDragDropOp.IsValid())
 	{
-		if (const URigHierarchy* Hierarchy = GetHierarchy())
+		FRigElementKey DraggedKey;
+		FRigElementKey::StaticStruct()->ImportText(*TagDragDropOp->GetIdentifier(), &DraggedKey, nullptr, EPropertyPortFlags::PPF_None, nullptr, FRigElementKey::StaticStruct()->GetName(), true);
+		if(TargetItem.IsValid())
 		{
-			FRigElementKey DraggedKey;
-			FRigElementKey::StaticStruct()->ImportText(*TagDragDropOp->GetIdentifier(), &DraggedKey, nullptr, EPropertyPortFlags::PPF_None, nullptr, FRigElementKey::StaticStruct()->GetName(), true);
-			if(TargetItem.IsValid())
+			return ResolveConnector(DraggedKey, TargetItem->Key);
+		}
+		return ResolveConnector(DraggedKey, FRigElementKey());
+	}
+
+	const TSharedPtr<FModularRigModuleDragDropOp> ModuleDropOp = DragDropEvent.GetOperationAs<FModularRigModuleDragDropOp>();
+	if (ModuleDropOp.IsValid() && TargetItem.IsValid())
+	{
+		UModularRig* ControlRig = Cast<UModularRig>(ControlRigBlueprint->GetDebuggedControlRig());
+		if (!ControlRig)
+		{
+			return FReply::Handled();
+		}
+
+		const FRigElementKey TargetKey = TargetItem->Key;
+		const TArray<FRigElementKey> DraggedKeys = FControlRigSchematicModel::GetElementKeysFromDragDropEvent(*ModuleDropOp.Get(), ControlRig);
+
+		bool bSuccess = false;
+		for(const FRigElementKey& DraggedKey : DraggedKeys)
+		{
+			const FReply Reply = ResolveConnector(DraggedKey, TargetKey);
+			if(Reply.IsEventHandled())
 			{
-				return ResolveConnector(DraggedKey, TargetItem->Key);
+				bSuccess = true;
 			}
-			return ResolveConnector(DraggedKey, FRigElementKey());
 		}
+		return bSuccess ? FReply::Handled() : FReply::Unhandled();
 	}
-
+	
 	TSharedPtr<FAssetDragDropOp> AssetDragDropOp = DragDropEvent.GetOperationAs<FAssetDragDropOp>();
 	if (AssetDragDropOp.IsValid())
 	{
diff --git a/Engine/Plugins/Animation/ControlRig/Source/ControlRigEditor/Private/Editor/SRigHierarchy.h b/Engine/Plugins/Animation/ControlRig/Source/ControlRigEditor/Private/Editor/SRigHierarchy.h
--- a/Engine/Plugins/Animation/ControlRig/Source/ControlRigEditor/Private/Editor/SRigHierarchy.h
+++ b/Engine/Plugins/Animation/ControlRig/Source/ControlRigEditor/Private/Editor/SRigHierarchy.h
@@ -4,12 +4,13 @@
 
 #include "CoreMinimal.h"
 #include "EditorUndoClient.h"
-#include "DragAndDrop/GraphNodeDragDropOp.h"
 #include "Engine/SkeletalMesh.h"
 #include "Editor/SRigHierarchyTreeView.h"
 #include "Units/RigUnitContext.h"
 #include "ControlRigBlueprint.h"
 #include "Editor/RigVMEditor.h"
+#include "ControlRigSchematicModel.h"
+#include "ControlRigDragOps.h"
 #include "SRigHierarchy.generated.h"
 
 class SRigHierarchy;
@@ -23,38 +24,6 @@ class FMenuBuilder;
 class UToolMenu;
 struct FToolMenuContext;
 
-class FRigElementHierarchyDragDropOp : public FGraphNodeDragDropOp
-{
-public:
-	DRAG_DROP_OPERATOR_TYPE(FRigElementHierarchyDragDropOp, FGraphNodeDragDropOp)
-
-	static TSharedRef<FRigElementHierarchyDragDropOp> New(const TArray<FRigElementKey>& InElements);
-
-	virtual TSharedPtr<SWidget> GetDefaultDecorator() const override;
-
-	/** @return true if this drag operation contains property paths */
-	bool HasElements() const
-	{
-		return Elements.Num() > 0;
-	}
-
-	/** @return The property paths from this drag operation */
-	const TArray<FRigElementKey>& GetElements() const
-	{
-		return Elements;
-	}
-
-	FString GetJoinedElementNames() const;
-
-	bool IsDraggingSingleConnector() const;
-	bool IsDraggingSingleSocket() const;
-
-private:
-
-	/** Data for the property paths this item represents */
-	TArray<FRigElementKey> Elements;
-};
-
 USTRUCT()
 struct FRigHierarchyImportSettings
 {
diff --git a/Engine/Plugins/Animation/ControlRig/Source/ControlRigEditor/Private/Editor/SRigHierarchyTagWidget.cpp b/Engine/Plugins/Animation/ControlRig/Source/ControlRigEditor/Private/Editor/SRigHierarchyTagWidget.cpp
--- a/Engine/Plugins/Animation/ControlRig/Source/ControlRigEditor/Private/Editor/SRigHierarchyTagWidget.cpp
+++ b/Engine/Plugins/Animation/ControlRig/Source/ControlRigEditor/Private/Editor/SRigHierarchyTagWidget.cpp
@@ -6,6 +6,7 @@
 #include "Widgets/Text/STextBlock.h"
 #include "Brushes/SlateRoundedBoxBrush.h"
 #include "DetailLayoutBuilder.h"
+#include "ControlRigDragOps.h"
 
 #define LOCTEXT_NAMESPACE "SRigHierarchyTagWidget"
 
@@ -13,26 +14,6 @@
 /// SRigHierarchyTagWidget
 ///////////////////////////////////////////////////////////
 
-TSharedRef<FRigHierarchyTagDragDropOp> FRigHierarchyTagDragDropOp::New(TSharedPtr<SRigHierarchyTagWidget> InTagWidget)
-{
-	TSharedRef<FRigHierarchyTagDragDropOp> Operation = MakeShared<FRigHierarchyTagDragDropOp>();
-	Operation->Text = InTagWidget->Text.Get();
-	Operation->Identifier = InTagWidget->Identifier.Get();
-	Operation->Construct();
-	return Operation;
-}
-
-TSharedPtr<SWidget> FRigHierarchyTagDragDropOp::GetDefaultDecorator() const
-{
-	return SNew(SBorder)
-		.Visibility(EVisibility::Visible)
-		.BorderImage(FAppStyle::GetBrush("Menu.Background"))
-		[
-			SNew(STextBlock)
-			.Text(Text)
-		];
-}
-
 void SRigHierarchyTagWidget::Construct(const FArguments& InArgs)
 {
 	Text = InArgs._Text;
diff --git a/Engine/Plugins/Animation/ControlRig/Source/ControlRigEditor/Private/Editor/SRigHierarchyTagWidget.h b/Engine/Plugins/Animation/ControlRig/Source/ControlRigEditor/Private/Editor/SRigHierarchyTagWidget.h
--- a/Engine/Plugins/Animation/ControlRig/Source/ControlRigEditor/Private/Editor/SRigHierarchyTagWidget.h
+++ b/Engine/Plugins/Animation/ControlRig/Source/ControlRigEditor/Private/Editor/SRigHierarchyTagWidget.h
@@ -11,27 +11,6 @@ class SRigHierarchyTagWidget;
 
 DECLARE_DELEGATE_OneParam(FOnRigTreeElementKeyTagDragDetected, const FRigElementKey&);
 
-class FRigHierarchyTagDragDropOp : public FDecoratedDragDropOp
-{
-public:
-	DRAG_DROP_OPERATOR_TYPE(FRigHierarchyTagDragDropOp, FDecoratedDragDropOp)
-
-	static TSharedRef<FRigHierarchyTagDragDropOp> New(TSharedPtr<SRigHierarchyTagWidget> InTagWidget);
-
-	virtual TSharedPtr<SWidget> GetDefaultDecorator() const override;
-
-	/** @return The identifier being dragged */
-	const FString& GetIdentifier() const
-	{
-		return Identifier;
-	}
-
-private:
-
-	FText Text;
-	FString Identifier;
-};
-
 class SRigHierarchyTagWidget : public SCompoundWidget
 {
 public:
diff --git a/Engine/Plugins/Animation/ControlRig/Source/ControlRigEditor/Private/Editor/SRigHierarchyTreeView.cpp b/Engine/Plugins/Animation/ControlRig/Source/ControlRigEditor/Private/Editor/SRigHierarchyTreeView.cpp
--- a/Engine/Plugins/Animation/ControlRig/Source/ControlRigEditor/Private/Editor/SRigHierarchyTreeView.cpp
+++ b/Engine/Plugins/Animation/ControlRig/Source/ControlRigEditor/Private/Editor/SRigHierarchyTreeView.cpp
@@ -48,7 +48,7 @@ FRigTreeElement::FRigTreeElement(const FRigElementKey& InKey, TWeakPtr<SRigHiera
 	{
 		if(const URigHierarchy* Hierarchy = InTreeView.Pin()->GetRigTreeDelegates().GetHierarchy())
 		{
-			ShortName = *Hierarchy->GetDisplayNameForUI(InKey).ToString();
+			ShortName = *Hierarchy->GetDisplayNameForUI(InKey, false).ToString();
 			
 			const FRigTreeDisplaySettings& Settings = InTreeView.Pin()->GetRigTreeDelegates().GetDisplaySettings();
 			RefreshDisplaySettings(Hierarchy, Settings);
@@ -199,7 +199,7 @@ void SRigHierarchyItem::Construct(const FArguments& InArgs, const TSharedRef<STa
 			.VAlign(VAlign_Center)
 			[
 				SAssignNew(InlineWidget, SInlineEditableTextBlock)
-				.Text(this, &SRigHierarchyItem::GetName, DisplaySettings.bUseShortName)
+				.Text(this, &SRigHierarchyItem::GetNameForUI)
 				.ToolTipText(this, &SRigHierarchyItem::GetItemTooltip)
 				.OnVerifyTextChanged(this, &SRigHierarchyItem::OnVerifyNameChanged)
 				.OnTextCommitted(this, &SRigHierarchyItem::OnNameCommitted)
@@ -230,6 +230,11 @@ void SRigHierarchyItem::Construct(const FArguments& InArgs, const TSharedRef<STa
 	InRigTreeElement->OnRenameRequested.BindSP(InlineWidget.Get(), &SInlineEditableTextBlock::EnterEditingMode);
 }
 
+FText SRigHierarchyItem::GetNameForUI() const
+{
+	return GetName(Delegates.GetDisplaySettings().bUseShortName);
+}
+
 FText SRigHierarchyItem::GetName(bool bUseShortName) const
 {
 	if(WeakRigTreeElement.Pin()->bIsTransient)
diff --git a/Engine/Plugins/Animation/ControlRig/Source/ControlRigEditor/Private/Editor/SRigHierarchyTreeView.h b/Engine/Plugins/Animation/ControlRig/Source/ControlRigEditor/Private/Editor/SRigHierarchyTreeView.h
--- a/Engine/Plugins/Animation/ControlRig/Source/ControlRigEditor/Private/Editor/SRigHierarchyTreeView.h
+++ b/Engine/Plugins/Animation/ControlRig/Source/ControlRigEditor/Private/Editor/SRigHierarchyTreeView.h
@@ -271,6 +271,7 @@ private:
 	TWeakPtr<FRigTreeElement> WeakRigTreeElement;
  	FRigTreeDelegates Delegates;
 
+	FText GetNameForUI() const;
 	FText GetName(bool bUseShortName) const;
 	FText GetItemTooltip() const;
 
diff --git a/Engine/Source/Editor/AnimationEditorWidgets/Private/SchematicGraphPanel/SSchematicGraphPanel.cpp b/Engine/Source/Editor/AnimationEditorWidgets/Private/SchematicGraphPanel/SSchematicGraphPanel.cpp
--- a/Engine/Source/Editor/AnimationEditorWidgets/Private/SchematicGraphPanel/SSchematicGraphPanel.cpp
+++ b/Engine/Source/Editor/AnimationEditorWidgets/Private/SchematicGraphPanel/SSchematicGraphPanel.cpp
@@ -45,6 +45,16 @@ TSharedPtr<SWidget> FSchematicGraphNodeDragDropOp::GetDefaultDecorator() const
 		];
 }
 
+const TArray<const FSchematicGraphNode*> FSchematicGraphNodeDragDropOp::GetNodes() const
+{
+	TArray<const FSchematicGraphNode*> Result;
+	for(SSchematicGraphNode* Node : SchematicGraphNodes)
+	{
+		Result.Add(Node->GetNodeData());
+	}
+	return Result;
+}
+
 FString FSchematicGraphNodeDragDropOp::GetJoinedDecoratorLabels() const
 {
 	TArray<FString> DecoratorLabels;
diff --git a/Engine/Source/Editor/AnimationEditorWidgets/Public/SchematicGraphPanel/SSchematicGraphPanel.h b/Engine/Source/Editor/AnimationEditorWidgets/Public/SchematicGraphPanel/SSchematicGraphPanel.h
--- a/Engine/Source/Editor/AnimationEditorWidgets/Public/SchematicGraphPanel/SSchematicGraphPanel.h
+++ b/Engine/Source/Editor/AnimationEditorWidgets/Public/SchematicGraphPanel/SSchematicGraphPanel.h
@@ -11,7 +11,7 @@
 class SSchematicGraphPanel;
 class SSchematicGraphNode;
 
-class FSchematicGraphNodeDragDropOp : public FDragDropOperation
+class ANIMATIONEDITORWIDGETS_API FSchematicGraphNodeDragDropOp : public FDragDropOperation
 {
 public:
 	DRAG_DROP_OPERATOR_TYPE(FSchematicGraphNodeDragDropOp, FDragDropOperation)
@@ -36,6 +36,9 @@ public:
 		return Elements;
 	}
 
+	/** @return The nodes being dragged */
+	const TArray<const FSchematicGraphNode*> GetNodes() const;
+
 	FString GetJoinedDecoratorLabels() const;
 
 private:
