Control Rig: Fix drag and drop operation issues in module tree

#pf skip
#rnx

diff --git a/Engine/Plugins/Animation/ControlRig/Source/ControlRigEditor/Private/Editor/SModularRigModel.cpp b/Engine/Plugins/Animation/ControlRig/Source/ControlRigEditor/Private/Editor/SModularRigModel.cpp
--- a/Engine/Plugins/Animation/ControlRig/Source/ControlRigEditor/Private/Editor/SModularRigModel.cpp
+++ b/Engine/Plugins/Animation/ControlRig/Source/ControlRigEditor/Private/Editor/SModularRigModel.cpp
@@ -846,6 +846,16 @@ TOptional<EItemDropZone> SModularRigModel::OnCanAcceptDrop(const FDragDropEvent&
 	const TOptional<EItemDropZone> InvalidDropZone;
 	TOptional<EItemDropZone> ReturnDropZone = DropZone;
 
+	if(DropZone == EItemDropZone::BelowItem && !TargetItem.IsValid())
+	{
+		DropZone = EItemDropZone::OntoItem;
+	}
+	
+	if(DropZone != EItemDropZone::OntoItem)
+	{
+		return InvalidDropZone;
+	}
+
 	TSharedPtr<FAssetDragDropOp> AssetDragDropOperation = DragDropEvent.GetOperationAs<FAssetDragDropOp>();
 	TSharedPtr<FModularRigModuleDragDropOp> ModuleDragDropOperation = DragDropEvent.GetOperationAs<FModularRigModuleDragDropOp>();
 	if (AssetDragDropOperation)
@@ -870,7 +880,14 @@ TOptional<EItemDropZone> SModularRigModel::OnCanAcceptDrop(const FDragDropEvent&
 	}
 	else if(ModuleDragDropOperation)
 	{
-		// Accept this drop
+		if(TargetItem.IsValid())
+		{
+			// we cannot drag a module onto itself
+			if(ModuleDragDropOperation->GetElements().Contains(TargetItem->ModulePath))
+			{
+				return InvalidDropZone;
+			}
+		}
 	}
 	else
 	{
@@ -882,11 +899,10 @@ TOptional<EItemDropZone> SModularRigModel::OnCanAcceptDrop(const FDragDropEvent&
 
 FReply SModularRigModel::OnAcceptDrop(const FDragDropEvent& DragDropEvent, EItemDropZone DropZone, TSharedPtr<FModularRigTreeElement> TargetItem)
 {
-	const TSharedPtr<FModularRigTreeElement>* ItemAtMouse = TreeView->FindItemAtPosition(DragDropEvent.GetScreenSpacePosition());
 	FString ParentPath;
-	if (ItemAtMouse && ItemAtMouse->IsValid())
+	if (TargetItem.IsValid())
 	{
-		ParentPath = ItemAtMouse->Get()->ModulePath;
+		ParentPath = TargetItem->ModulePath;
 	}
 
 	TSharedPtr<FAssetDragDropOp> AssetDragDropOperation = DragDropEvent.GetOperationAs<FAssetDragDropOp>();
