Control Rig: Multiple selection of modules in details panel
#rnx

diff --git a/Engine/Plugins/Animation/ControlRig/Source/ControlRig/Private/ModularRigRuleManager.cpp b/Engine/Plugins/Animation/ControlRig/Source/ControlRig/Private/ModularRigRuleManager.cpp
--- a/Engine/Plugins/Animation/ControlRig/Source/ControlRig/Private/ModularRigRuleManager.cpp
+++ b/Engine/Plugins/Animation/ControlRig/Source/ControlRig/Private/ModularRigRuleManager.cpp
@@ -2,6 +2,8 @@
 
 #include "ModularRigRuleManager.h"
 
+#include "Units/Execution/RigUnit_PrepareForExecution.h"
+
 #include UE_INLINE_GENERATED_CPP_BY_NAME(ModularRigRuleManager)
 
 #define LOCTEXT_NAMESPACE "ModularRigRuleManager"
@@ -22,6 +24,14 @@ FModularRigResolveResult UModularRigRuleManager::FindMatches(
 		return Result;
 	}
 
+	if (UControlRig* ControlRig = Hierarchy->GetTypedOuter<UControlRig>())
+	{
+		if (ControlRig->IsConstructionRequired())
+		{
+			ControlRig->Execute(FRigUnit_PrepareForExecution::EventName);
+		}
+	}
+
 	// start with a full set of possible targets
 	TArray<bool> VisitedElement;
 	VisitedElement.AddZeroed(Hierarchy->Num());
diff --git a/Engine/Plugins/Animation/ControlRig/Source/ControlRig/Public/ModularRig.h b/Engine/Plugins/Animation/ControlRig/Source/ControlRig/Public/ModularRig.h
--- a/Engine/Plugins/Animation/ControlRig/Source/ControlRig/Public/ModularRig.h
+++ b/Engine/Plugins/Animation/ControlRig/Source/ControlRig/Public/ModularRig.h
@@ -24,7 +24,11 @@ public:
 	
 	const UModularRig* GetModularRig() const
 	{
-		return ModularRig.Get();
+		if (ModularRig.IsValid())
+		{
+			return ModularRig.Get();
+		}
+		return nullptr;
 	}
 	const FString& GetPath() const { return Path; }
 
diff --git a/Engine/Plugins/Animation/ControlRig/Source/ControlRig/Public/Rigs/RigModuleDefines.h b/Engine/Plugins/Animation/ControlRig/Source/ControlRig/Public/Rigs/RigModuleDefines.h
--- a/Engine/Plugins/Animation/ControlRig/Source/ControlRig/Public/Rigs/RigModuleDefines.h
+++ b/Engine/Plugins/Animation/ControlRig/Source/ControlRig/Public/Rigs/RigModuleDefines.h
@@ -45,6 +45,11 @@ struct CONTROLRIG_API FRigModuleConnector
 	FRigModuleConnector()
 	{}
 
+	bool operator==(const FRigModuleConnector& Other) const
+	{
+		return Name == Other.Name && GetTypeHash(Settings) == GetTypeHash(Other.Settings);
+	}
+
 	UPROPERTY(BlueprintReadWrite, EditAnywhere, Category = Connector)
 	FString Name;
 
diff --git a/Engine/Plugins/Animation/ControlRig/Source/ControlRigEditor/Private/ControlRigModuleDetails.cpp b/Engine/Plugins/Animation/ControlRig/Source/ControlRigEditor/Private/ControlRigModuleDetails.cpp
--- a/Engine/Plugins/Animation/ControlRig/Source/ControlRigEditor/Private/ControlRigModuleDetails.cpp
+++ b/Engine/Plugins/Animation/ControlRig/Source/ControlRigEditor/Private/ControlRigModuleDetails.cpp
@@ -169,7 +169,6 @@ void FRigModuleInstanceDetails::CustomizeDetails(IDetailLayoutBuilder& DetailBui
 			.Text(FText::FromString(TEXT("Name")))
 			.Font(IDetailLayoutBuilder::GetDetailFont())
 			.ToolTipText(NameTooltip)
-			.IsEnabled(PerModuleInfos.Num() == 1)
 		]
 		.ValueContent()
 		[
@@ -178,7 +177,6 @@ void FRigModuleInstanceDetails::CustomizeDetails(IDetailLayoutBuilder& DetailBui
 			.Text(this, &FRigModuleInstanceDetails::GetName)
 			.OnTextCommitted(this, &FRigModuleInstanceDetails::SetName, DetailBuilder.GetPropertyUtilities())
 			.ToolTipText(NameTooltip)
-			.IsEnabled(PerModuleInfos.Num() == 1)
 			.OnVerifyTextChanged(this, &FRigModuleInstanceDetails::OnVerifyNameChanged)
 		];
 
@@ -242,231 +240,314 @@ void FRigModuleInstanceDetails::CustomizeDetails(IDetailLayoutBuilder& DetailBui
 
 	IDetailCategoryBuilder& ConnectionsCategory = DetailBuilder.EditCategory(TEXT("Connections"), LOCTEXT("Connections", "Connections"));
 	{
-		TArray<FRigModuleConnector> Connectors = GetConnectors();
-		for(FRigModuleConnector& Connector : Connectors)
+		bool bDisplayConnectors = PerModuleInfos.Num() >= 1;
+		if (PerModuleInfos.Num() > 1)
 		{
-			const FText Label = FText::FromString(Connector.Name);
-			TSharedPtr<SVerticalBox> ButtonBox;
-			FString ConnectorPath = FString::Printf(TEXT("%s:%s"), *PerModuleInfos[0].Path, *Connector.Name);
-			FRigElementKey ConnectorKey(*ConnectorPath, ERigElementType::Connector);
-
-			FModularRigResolveResult ConnectorMatches;
-			FPerModuleInfo Info = PerModuleInfos[0];
-			if (const FRigModuleInstance* Module = Info.GetModule())
+			UModularRig* ModularRig = PerModuleInfos[0].GetModularRig();
+			for (FPerModuleInfo& Info : PerModuleInfos)
 			{
-				if (UModularRig* ModularRig = Info.GetModularRig())
+				if (Info.GetModularRig() != ModularRig)
 				{
-					if (URigHierarchy* Hierarchy = ModularRig->GetHierarchy())
+					bDisplayConnectors = false;
+					break;
+				}
+			}
+		}
+		if (bDisplayConnectors)
+		{
+			TArray<FRigModuleConnector> Connectors = GetConnectors();
+			for(const FRigModuleConnector& Connector : Connectors)
+			{
+				const FText Label = FText::FromString(Connector.Name);
+				TSharedPtr<SVerticalBox> ButtonBox;
+
+				TArray<FRigElementResolveResult> Matches;
+				for (int32 ModuleIndex=0; ModuleIndex<PerModuleInfos.Num(); ++ModuleIndex)
+				{
+					const FPerModuleInfo& Info = PerModuleInfos[ModuleIndex];
+					if (const FRigModuleInstance* Module = Info.GetModule())
 					{
-						if (FRigConnectorElement* ConnectorElement = Cast<FRigConnectorElement>(Hierarchy->Find(ConnectorKey)))
+						if (UModularRig* ModularRig = Info.GetModularRig())
 						{
-							const UModularRigRuleManager* RuleManager = ModularRig->GetHierarchy()->GetRuleManager();
-							ConnectorMatches = RuleManager->FindMatches(ConnectorElement, Module, ModularRig->GetElementKeyRedirector());
+							if (URigHierarchy* Hierarchy = ModularRig->GetHierarchy())
+							{
+								FString ConnectorPath = FString::Printf(TEXT("%s:%s"), *Info.Path, *Connector.Name);
+								FRigElementKey ConnectorKey(*ConnectorPath, ERigElementType::Connector);
+								if (FRigConnectorElement* ConnectorElement = Cast<FRigConnectorElement>(Hierarchy->Find(ConnectorKey)))
+								{
+									const UModularRigRuleManager* RuleManager = ModularRig->GetHierarchy()->GetRuleManager();
+									if (ModuleIndex == 0)
+									{
+										Matches = RuleManager->FindMatches(ConnectorElement, Module, ModularRig->GetElementKeyRedirector()).GetMatches();
+									}
+									else
+									{
+										const FModularRigResolveResult& ConnectorMatches = RuleManager->FindMatches(ConnectorElement, Module, ModularRig->GetElementKeyRedirector());
+										Matches.FilterByPredicate([ConnectorMatches](const FRigElementResolveResult& Match)
+										{
+											return ConnectorMatches.ContainsMatch(Match.GetKey());
+										});
+									}
+								}
+							}
 						}
 					}
 				}
-			}
 
-			FRigTreeDelegates TreeDelegates;
-			TreeDelegates.OnGetHierarchy = FOnGetRigTreeHierarchy::CreateLambda([this]()
-			{
-				return PerModuleInfos[0].GetModularRig()->GetHierarchy();
-			});
-			TreeDelegates.OnRigTreeIsItemVisible = FOnRigTreeIsItemVisible::CreateLambda([ConnectorMatches](const FRigElementKey& InTarget)
-			{
-				return ConnectorMatches.ContainsMatch(InTarget);
-			});
-			TreeDelegates.OnGetSelection.BindLambda([this, ConnectorKey]() -> TArray<FRigElementKey>
-			{
-				FPerModuleInfo Info = PerModuleInfos[0];
-				FRigElementKeyRedirector Redirector = Info.GetModularRig()->GetElementKeyRedirector();
-				if (const FRigElementKey* Key = PerModuleInfos[0].GetModularRig()->GetElementKeyRedirector().FindExternalKey(ConnectorKey))
+				FRigTreeDelegates TreeDelegates;
+				TreeDelegates.OnGetHierarchy = FOnGetRigTreeHierarchy::CreateLambda([this]()
 				{
-					return {*Key};
-				}
-				return {};
-			});
-			TreeDelegates.OnSelectionChanged.BindSP(this, &FRigModuleInstanceDetails::OnConnectorTargetChanged, ConnectorKey);
+					return PerModuleInfos[0].GetModularRig()->GetHierarchy();
+				});
+				TreeDelegates.OnRigTreeIsItemVisible = FOnRigTreeIsItemVisible::CreateLambda([Matches](const FRigElementKey& InTarget)
+				{
+					return Matches.ContainsByPredicate([InTarget](const FRigElementResolveResult& Match)
+					{
+						return Match.GetKey() == InTarget;
+					});
+				});
+				TreeDelegates.OnGetSelection.BindLambda([this, Connector]() -> TArray<FRigElementKey>
+				{
+					FRigElementKey Target;
+					FRigElementKeyRedirector Redirector = PerModuleInfos[0].GetModularRig()->GetElementKeyRedirector();
+					for (int32 ModuleIndex=0; ModuleIndex<PerModuleInfos.Num(); ++ModuleIndex)
+					{
+						FString ConnectorPath = FString::Printf(TEXT("%s:%s"), *PerModuleInfos[ModuleIndex].Path, *Connector.Name);
+						FRigElementKey ConnectorKey(*ConnectorPath, ERigElementType::Connector);
+						if (const FRigElementKey* Key = Redirector.FindExternalKey(ConnectorKey))
+						{
+							if (ModuleIndex == 0)
+							{
+								Target = *Key;
+							}
+							else if (Target != *Key)
+							{
+								Target.Name = *ControlRigModuleDetailsMultipleValues.ToString();
+								return {Target};
+							}
+						}
+						else
+						{
+							Target.Name = *ControlRigModuleDetailsMultipleValues.ToString();
+							return {Target};
+						}
+					}
+					return {};
+				});
+				TreeDelegates.OnSelectionChanged.BindSP(this, &FRigModuleInstanceDetails::OnConnectorTargetChanged, Connector);
 			
 
-			ConnectionsCategory.AddCustomRow(Label)
-				.NameContent()
-				[
-					SNew(STextBlock)
-					.Text(Label)
-					.Font(IDetailLayoutBuilder::GetDetailFont())
-					.IsEnabled(true)
-				]
-				.ValueContent()
-				[
-					SNew(SHorizontalBox)
-
-					// Combo button
-					+SHorizontalBox::Slot()
-					.AutoWidth()
-					.Padding(0.f, 0.f, 0.f, 0.f)
+				ConnectionsCategory.AddCustomRow(Label)
+					.NameContent()
 					[
-						SNew( SComboButton )
-						.ContentPadding(3)
-						.MenuPlacement(MenuPlacement_BelowAnchor)
-						.OnComboBoxOpened(this, &FRigModuleInstanceDetails::PopulateConnectorTargetList, ConnectorKey)
-						.ButtonContent()
+						SNew(STextBlock)
+						.Text(Label)
+						.Font(IDetailLayoutBuilder::GetDetailFont())
+						.IsEnabled(true)
+					]
+					.ValueContent()
+					[
+						SNew(SHorizontalBox)
+
+						// Combo button
+						+SHorizontalBox::Slot()
+						.AutoWidth()
+						.Padding(0.f, 0.f, 0.f, 0.f)
 						[
-							// Wrap in configurable box to restrain height/width of menu
-							SNew(SBox)
-							.MinDesiredWidth(150.0f)
+							SNew( SComboButton )
+							.ContentPadding(3)
+							.MenuPlacement(MenuPlacement_BelowAnchor)
+							.OnComboBoxOpened(this, &FRigModuleInstanceDetails::PopulateConnectorTargetList, Connector.Name)
+							.ButtonContent()
 							[
-								SAssignNew(ButtonBox, SVerticalBox)
+								// Wrap in configurable box to restrain height/width of menu
+								SNew(SBox)
+								.MinDesiredWidth(150.0f)
+								[
+									SAssignNew(ButtonBox, SVerticalBox)
+								]
 							]
-						]
-						.MenuContent()
-						[
-							SNew(SBorder)
-							.Visibility(EVisibility::Visible)
-							.BorderImage(FAppStyle::GetBrush("Menu.Background"))
+							.MenuContent()
 							[
-								SAssignNew(ConnectionListBox.FindOrAdd(ConnectorKey), SSearchableRigHierarchyTreeView)
-									.RigTreeDelegates(TreeDelegates)
+								SNew(SBorder)
+								.Visibility(EVisibility::Visible)
+								.BorderImage(FAppStyle::GetBrush("Menu.Background"))
+								[
+									SAssignNew(ConnectionListBox.FindOrAdd(Connector.Name), SSearchableRigHierarchyTreeView)
+										.RigTreeDelegates(TreeDelegates)
+								]
 							]
 						]
-					]
 
-					+SHorizontalBox::Slot()
-					.AutoWidth()
-					.VAlign(VAlign_Center)
-					.Padding(0.f, 0.f, 0.f, 0.f)
-					[
-						SNew(SVerticalBox)
-						+SVerticalBox::Slot()
+						+SHorizontalBox::Slot()
+						.AutoWidth()
 						.VAlign(VAlign_Center)
-						.AutoHeight()
+						.Padding(0.f, 0.f, 0.f, 0.f)
 						[
-							// Reset button
-							SNew(SHorizontalBox)
-							+SHorizontalBox::Slot()
-							.AutoWidth()
-							.Padding(0.f, 0.f, 0.f, 0.f)
+							SNew(SVerticalBox)
+							+SVerticalBox::Slot()
+							.VAlign(VAlign_Center)
+							.AutoHeight()
 							[
-								SAssignNew(ResetConnectorButton.FindOrAdd(ConnectorKey), SButton)
-								.ButtonStyle( FAppStyle::Get(), "NoBorder" )
-								.ButtonColorAndOpacity_Lambda([this, ConnectorKey]()
-								{
-									const TSharedPtr<SButton>& Button = ResetConnectorButton.FindRef(ConnectorKey);
-									return Button.IsValid() && Button->IsHovered()
-										? FSlateColor(FLinearColor(1.f, 1.f, 1.f, 0.8))
-										: FSlateColor(FLinearColor(1.f, 1.f, 1.f, 0.4));
-								})
-								.OnClicked_Lambda([this, ConnectorKey]()
-								{
-									PerModuleInfos[0].GetBlueprint()->GetModularRigController()->DisconnectConnector(ConnectorKey);
-									return FReply::Handled();
-								})
-								.ContentPadding(1.f)
-								.ToolTipText(NSLOCTEXT("ControlRigModuleDetails", "Reset_Connector", "Reset Connector"))
+								// Reset button
+								SNew(SHorizontalBox)
+								+SHorizontalBox::Slot()
+								.AutoWidth()
+								.Padding(0.f, 0.f, 0.f, 0.f)
 								[
-									SNew(SImage)
-									.ColorAndOpacity_Lambda( [this, ConnectorKey]()
+									SAssignNew(ResetConnectorButton.FindOrAdd(Connector.Name), SButton)
+									.ButtonStyle( FAppStyle::Get(), "NoBorder" )
+									.ButtonColorAndOpacity_Lambda([this, Connector]()
 									{
-										const TSharedPtr<SButton>& Button = ResetConnectorButton.FindRef(ConnectorKey);
+										const TSharedPtr<SButton>& Button = ResetConnectorButton.FindRef(Connector.Name);
 										return Button.IsValid() && Button->IsHovered()
-										? FSlateColor(FLinearColor(1.f, 1.f, 1.f, 0.8))
-										: FSlateColor(FLinearColor(1.f, 1.f, 1.f, 0.4));
+											? FSlateColor(FLinearColor(1.f, 1.f, 1.f, 0.8))
+											: FSlateColor(FLinearColor(1.f, 1.f, 1.f, 0.4));
 									})
-									.Image(FSlateIcon(FAppStyle::Get().GetStyleSetName(), "PropertyWindow.DiffersFromDefault").GetIcon())
-								]
-							]
-
-							// Use button
-							+SHorizontalBox::Slot()
-							.AutoWidth()
-							.Padding(0.f, 0.f, 0.f, 0.f)
-							[
-								SAssignNew(UseSelectedButton.FindOrAdd(ConnectorKey), SButton)
-								.ButtonStyle( FAppStyle::Get(), "NoBorder" )
-								.ButtonColorAndOpacity_Lambda([this, ConnectorKey]()
-								{
-									const TSharedPtr<SButton>& Button = UseSelectedButton.FindRef(ConnectorKey);
-									return Button.IsValid() && Button->IsHovered()
-										? FSlateColor(FLinearColor(1.f, 1.f, 1.f, 0.8))
-										: FSlateColor(FLinearColor(1.f, 1.f, 1.f, 0.4));
-								})
-								.OnClicked_Lambda([this, ConnectorKey]()
-								{
-									if (UModularRig* ModularRig = PerModuleInfos[0].GetModularRig())
+									.OnClicked_Lambda([this, Connector]()
 									{
-										const TArray<FRigElementKey>& Selected = ModularRig->GetHierarchy()->GetSelectedKeys();
-										if (Selected.Num() > 0)
+										for (FPerModuleInfo& Info : PerModuleInfos)
 										{
-											PerModuleInfos[0].GetBlueprint()->GetModularRigController()->ConnectConnectorToElement(ConnectorKey, Selected[0], true, ModularRig->GetModularRigSettings().bAutoResolve);
+											FString ConnectorPath = FString::Printf(TEXT("%s:%s"), *Info.Path, *Connector.Name);
+											FRigElementKey ConnectorKey(*ConnectorPath, ERigElementType::Connector);
+											Info.GetBlueprint()->GetModularRigController()->DisconnectConnector(ConnectorKey);
 										}
-									}
-									return FReply::Handled();
-								})
-								.ContentPadding(1.f)
-								.ToolTipText(NSLOCTEXT("ControlRigModuleDetails", "Use_Selected", "Use Selected"))
+										return FReply::Handled();
+									})
+									.ContentPadding(1.f)
+									.ToolTipText(NSLOCTEXT("ControlRigModuleDetails", "Reset_Connector", "Reset Connector"))
+									[
+										SNew(SImage)
+										.ColorAndOpacity_Lambda( [this, Connector]()
+										{
+											const TSharedPtr<SButton>& Button = ResetConnectorButton.FindRef(Connector.Name);
+											return Button.IsValid() && Button->IsHovered()
+											? FSlateColor(FLinearColor(1.f, 1.f, 1.f, 0.8))
+											: FSlateColor(FLinearColor(1.f, 1.f, 1.f, 0.4));
+										})
+										.Image(FSlateIcon(FAppStyle::Get().GetStyleSetName(), "PropertyWindow.DiffersFromDefault").GetIcon())
+									]
+								]
+
+								// Use button
+								+SHorizontalBox::Slot()
+								.AutoWidth()
+								.Padding(0.f, 0.f, 0.f, 0.f)
 								[
-									SNew(SImage)
-									.ColorAndOpacity_Lambda( [this, ConnectorKey]()
+									SAssignNew(UseSelectedButton.FindOrAdd(Connector.Name), SButton)
+									.ButtonStyle( FAppStyle::Get(), "NoBorder" )
+									.ButtonColorAndOpacity_Lambda([this, Connector]()
 									{
-										const TSharedPtr<SButton>& Button = UseSelectedButton.FindRef(ConnectorKey);
+										const TSharedPtr<SButton>& Button = UseSelectedButton.FindRef(Connector.Name);
 										return Button.IsValid() && Button->IsHovered()
-										? FSlateColor(FLinearColor(1.f, 1.f, 1.f, 0.8))
-										: FSlateColor(FLinearColor(1.f, 1.f, 1.f, 0.4));
+											? FSlateColor(FLinearColor(1.f, 1.f, 1.f, 0.8))
+											: FSlateColor(FLinearColor(1.f, 1.f, 1.f, 0.4));
 									})
-									.Image(FAppStyle::GetBrush("Icons.CircleArrowLeft"))
-								]
-							]
-
-							// Select in hierarchy button
-							+SHorizontalBox::Slot()
-							.AutoWidth()
-							.Padding(0.f, 0.f, 0.f, 0.f)
-							[
-								SAssignNew(SelectElementButton.FindOrAdd(ConnectorKey), SButton)
-								.ButtonStyle( FAppStyle::Get(), "NoBorder" )
-								.ButtonColorAndOpacity_Lambda([this, ConnectorKey]()
-								{
-									const TSharedPtr<SButton>& Button = SelectElementButton.FindRef(ConnectorKey);
-									return Button.IsValid() && Button->IsHovered()
-										? FSlateColor(FLinearColor(1.f, 1.f, 1.f, 0.8))
-										: FSlateColor(FLinearColor(1.f, 1.f, 1.f, 0.4));
-								})
-								.OnClicked_Lambda([this, ConnectorKey]()
-								{
-									if (UModularRig* ModularRig = PerModuleInfos[0].GetModularRig())
+									.OnClicked_Lambda([this, Connector]()
 									{
-										if (const FRigElementKey* TargetKey = ModularRig->GetElementKeyRedirector().FindExternalKey(ConnectorKey))
+										if (UModularRig* ModularRig = PerModuleInfos[0].GetModularRig())
 										{
-											ModularRig->GetHierarchy()->GetController()->SelectElement(*TargetKey, true, true);
+											const TArray<FRigElementKey>& Selected = ModularRig->GetHierarchy()->GetSelectedKeys();
+											if (Selected.Num() > 0)
+											{
+												for (FPerModuleInfo& Info : PerModuleInfos)
+												{
+													FString ConnectorPath = FString::Printf(TEXT("%s:%s"), *Info.Path, *Connector.Name);
+													FRigElementKey ConnectorKey(*ConnectorPath, ERigElementType::Connector);
+													Info.GetBlueprint()->GetModularRigController()->ConnectConnectorToElement(ConnectorKey, Selected[0], true, ModularRig->GetModularRigSettings().bAutoResolve);
+												}
+											}
 										}
-									}
-									return FReply::Handled();
-								})
-								.ContentPadding(1.f)
-								.ToolTipText(NSLOCTEXT("ControlRigModuleDetails", "Select_Element", "Select Element"))
+										return FReply::Handled();
+									})
+									.ContentPadding(1.f)
+									.ToolTipText(NSLOCTEXT("ControlRigModuleDetails", "Use_Selected", "Use Selected"))
+									[
+										SNew(SImage)
+										.ColorAndOpacity_Lambda( [this, Connector]()
+										{
+											const TSharedPtr<SButton>& Button = UseSelectedButton.FindRef(Connector.Name);
+											return Button.IsValid() && Button->IsHovered()
+											? FSlateColor(FLinearColor(1.f, 1.f, 1.f, 0.8))
+											: FSlateColor(FLinearColor(1.f, 1.f, 1.f, 0.4));
+										})
+										.Image(FAppStyle::GetBrush("Icons.CircleArrowLeft"))
+									]
+								]
+
+								// Select in hierarchy button
+								+SHorizontalBox::Slot()
+								.AutoWidth()
+								.Padding(0.f, 0.f, 0.f, 0.f)
 								[
-									SNew(SImage)
-									.ColorAndOpacity_Lambda( [this, ConnectorKey]()
+									SAssignNew(SelectElementButton.FindOrAdd(Connector.Name), SButton)
+									.ButtonStyle( FAppStyle::Get(), "NoBorder" )
+									.ButtonColorAndOpacity_Lambda([this, Connector]()
 									{
-										const TSharedPtr<SButton>& Button = SelectElementButton.FindRef(ConnectorKey);
+										const TSharedPtr<SButton>& Button = SelectElementButton.FindRef(Connector.Name);
 										return Button.IsValid() && Button->IsHovered()
-										? FSlateColor(FLinearColor(1.f, 1.f, 1.f, 0.8))
-										: FSlateColor(FLinearColor(1.f, 1.f, 1.f, 0.4));
+											? FSlateColor(FLinearColor(1.f, 1.f, 1.f, 0.8))
+											: FSlateColor(FLinearColor(1.f, 1.f, 1.f, 0.4));
 									})
-									.Image(FAppStyle::GetBrush("Icons.Search"))
+									.OnClicked_Lambda([this, Connector]()
+									{
+										if (UModularRig* ModularRig = PerModuleInfos[0].GetModularRig())
+										{
+											FString ConnectorPath = FString::Printf(TEXT("%s:%s"), *PerModuleInfos[0].Path, *Connector.Name);
+											FRigElementKey ConnectorKey(*ConnectorPath, ERigElementType::Connector);
+											if (const FRigElementKey* TargetKey = ModularRig->GetElementKeyRedirector().FindExternalKey(ConnectorKey))
+											{
+												ModularRig->GetHierarchy()->GetController()->SelectElement(*TargetKey, true, true);
+											}
+										}
+										return FReply::Handled();
+									})
+									.ContentPadding(1.f)
+									.ToolTipText(NSLOCTEXT("ControlRigModuleDetails", "Select_Element", "Select Element"))
+									[
+										SNew(SImage)
+										.ColorAndOpacity_Lambda( [this, Connector]()
+										{
+											const TSharedPtr<SButton>& Button = SelectElementButton.FindRef(Connector.Name);
+											return Button.IsValid() && Button->IsHovered()
+											? FSlateColor(FLinearColor(1.f, 1.f, 1.f, 0.8))
+											: FSlateColor(FLinearColor(1.f, 1.f, 1.f, 0.4));
+										})
+										.Image(FAppStyle::GetBrush("Icons.Search"))
+									]
 								]
 							]
 						]
-					]
-				];
+					];
 
-			FRigElementKey CurrentTargetKey;
-			if (const FRigElementKey* Key = PerModuleInfos[0].GetModularRig()->GetElementKeyRedirector().FindExternalKey(ConnectorKey))
-			{
-				CurrentTargetKey = *Key;
+				FRigElementKey CurrentTargetKey;
+				if (PerModuleInfos.Num() >= 1)
+				{
+					for (int32 ModuleIndex=0; ModuleIndex<PerModuleInfos.Num(); ++ModuleIndex)
+					{
+						FRigElementKey TargetKey;
+						FString ConnectorPath = FString::Printf(TEXT("%s:%s"), *PerModuleInfos[ModuleIndex].Path, *Connector.Name);
+						FRigElementKey ConnectorKey(*ConnectorPath, ERigElementType::Connector);
+						if (const FRigElementKey* Key = PerModuleInfos[ModuleIndex].GetModularRig()->GetElementKeyRedirector().FindExternalKey(ConnectorKey))
+						{
+							TargetKey = *Key;
+						}
+						if (ModuleIndex == 0)
+						{
+							CurrentTargetKey = TargetKey;
+						}
+						else
+						{
+							if (CurrentTargetKey != TargetKey)
+							{
+								CurrentTargetKey.Name = *ControlRigModuleDetailsMultipleValues.ToString();
+							}
+						}
+					}
+				}
+				TPair<const FSlateBrush*, FSlateColor> IconAndColor = SRigHierarchyItem::GetBrushForElementType(PerModuleInfos[0].GetModularRig()->GetHierarchy(), CurrentTargetKey);
+				PopulateConnectorCurrentTarget(ButtonBox, IconAndColor.Key, IconAndColor.Value, FText::FromName(CurrentTargetKey.Name));
 			}
-			TPair<const FSlateBrush*, FSlateColor> IconAndColor = SRigHierarchyItem::GetBrushForElementType(PerModuleInfos[0].GetModularRig()->GetHierarchy(), CurrentTargetKey);
-			PopulateConnectorCurrentTarget(ButtonBox, ConnectorKey, CurrentTargetKey, IconAndColor.Key, IconAndColor.Value, FText::FromName(CurrentTargetKey.Name));
 		}
 	}
 
@@ -587,7 +668,7 @@ FText FRigModuleInstanceDetails::GetName() const
 		{
 			if(const FRigModuleInstance* Module = PerModuleInfos[i].GetModule())
 			{
-				if (Module->Name.IsEqual(FirstValue, ENameCase::CaseSensitive))
+				if (!Module->Name.IsEqual(FirstValue, ENameCase::CaseSensitive))
 				{
 					bSame = false;
 					break;
@@ -609,23 +690,22 @@ void FRigModuleInstanceDetails::SetName(const FText& InValue, ETextCommit::Type
 		return;
 	}
 	
-	check(PerModuleInfos.Num() == 1);
-	FPerModuleInfo& Info = PerModuleInfos[0];
-	if (const FRigModuleInstance* ModuleInstance = Info.GetModule())
+	for (FPerModuleInfo& Info : PerModuleInfos)
 	{
-		if (UControlRigBlueprint* Blueprint = Info.GetBlueprint())
+		if (const FRigModuleInstance* ModuleInstance = Info.GetModule())
 		{
-			UModularRigController* Controller = Blueprint->GetModularRigController();
-			const FString OldPath = ModuleInstance->GetPath();
-			(void)Controller->RenameModule(OldPath, *InValue.ToString(), true);
+			if (UControlRigBlueprint* Blueprint = Info.GetBlueprint())
+			{
+				UModularRigController* Controller = Blueprint->GetModularRigController();
+				const FString OldPath = ModuleInstance->GetPath();
+				(void)Controller->RenameModule(OldPath, *InValue.ToString(), true);
+			}
 		}
 	}
 }
 
 bool FRigModuleInstanceDetails::OnVerifyNameChanged(const FText& InText, FText& OutErrorMessage)
 {
-	check(PerModuleInfos.Num() == 1);
-
 	if(InText.IsEmpty())
 	{
 		static const FText EmptyNameIsNotAllowed = LOCTEXT("EmptyNameIsNotAllowed", "Empty name is not allowed.");
@@ -633,13 +713,18 @@ bool FRigModuleInstanceDetails::OnVerifyNameChanged(const FText& InText, FText&
 		return false;
 	}
 
-	const FPerModuleInfo& Info = PerModuleInfos[0];
-	if (const FRigModuleInstance* ModuleInstance = Info.GetModule())
+	for (FPerModuleInfo& Info : PerModuleInfos)
 	{
-		if (UControlRigBlueprint* Blueprint = Info.GetBlueprint())
+		if (const FRigModuleInstance* ModuleInstance = Info.GetModule())
 		{
-			UModularRigController* Controller = Blueprint->GetModularRigController();
-			return Controller->CanRenameModule(ModuleInstance->GetPath(), *InText.ToString(), OutErrorMessage);
+			if (UControlRigBlueprint* Blueprint = Info.GetBlueprint())
+			{
+				UModularRigController* Controller = Blueprint->GetModularRigController();
+				if (!Controller->CanRenameModule(ModuleInstance->GetPath(), *InText.ToString(), OutErrorMessage))
+				{
+					return false;
+				}
+			}
 		}
 	}
 
@@ -662,7 +747,7 @@ FText FRigModuleInstanceDetails::GetShortName() const
 		{
 			if(const FRigModuleInstance* Module = PerModuleInfos[i].GetModule())
 			{
-				if (Module->GetShortName().Equals(FirstValue, ESearchCase::CaseSensitive))
+				if (!Module->GetShortName().Equals(FirstValue, ESearchCase::CaseSensitive))
 				{
 					bSame = false;
 					break;
@@ -734,7 +819,7 @@ FText FRigModuleInstanceDetails::GetLongName() const
 		{
 			if(const FRigModuleInstance* Module = PerModuleInfos[i].GetModule())
 			{
-				if (Module->GetLongName().Equals(FirstValue, ESearchCase::CaseSensitive))
+				if (!Module->GetLongName().Equals(FirstValue, ESearchCase::CaseSensitive))
 				{
 					bSame = false;
 					break;
@@ -759,7 +844,7 @@ FText FRigModuleInstanceDetails::GetRigClassPath() const
 			{
 				if(const FRigModuleInstance* Module = PerModuleInfos[i].GetModule())
 				{
-					if (Module->GetRig()->GetClass() !=  FirstModule->GetRig()->GetClass())
+					if (Module->GetRig()->GetClass() != FirstModule->GetRig()->GetClass())
 					{
 						bSame = false;
 						break;
@@ -788,7 +873,29 @@ TArray<FRigModuleConnector> FRigModuleInstanceDetails::GetConnectors() const
 {
 	if(PerModuleInfos.Num() > 1)
 	{
-		return TArray<FRigModuleConnector>();
+		TArray<FRigModuleConnector> CommonConnectors;
+		if (const FRigModuleInstance* Module = PerModuleInfos[0].GetModule())
+		{
+			if (const UControlRig* ModuleRig = Module->GetRig())
+			{
+				CommonConnectors = ModuleRig->GetRigModuleSettings().ExposedConnectors;
+			}
+		}
+		for (int32 ModuleIndex=1; ModuleIndex<PerModuleInfos.Num(); ++ModuleIndex)
+		{
+			if (const FRigModuleInstance* Module = PerModuleInfos[ModuleIndex].GetModule())
+			{
+				if (const UControlRig* ModuleRig = Module->GetRig())
+				{
+					const TArray<FRigModuleConnector>& ModuleConnectors = ModuleRig->GetRigModuleSettings().ExposedConnectors;
+					CommonConnectors = CommonConnectors.FilterByPredicate([ModuleConnectors](const FRigModuleConnector& Connector)
+					{
+						return ModuleConnectors.Contains(Connector);
+					});
+				}
+			}
+		}
+		return CommonConnectors;
 	}
 
 	if (const FRigModuleInstance* Module = PerModuleInfos[0].GetModule())
@@ -820,35 +927,15 @@ FRigElementKeyRedirector FRigModuleInstanceDetails::GetConnections() const
 	return FRigElementKeyRedirector();
 }
 
-void FRigModuleInstanceDetails::PopulateConnectorTargetList(const FRigElementKey InConnectorKey)
+void FRigModuleInstanceDetails::PopulateConnectorTargetList(const FString InConnectorName)
 {
-	if (PerModuleInfos.IsEmpty() || PerModuleInfos.Num() > 1)
-	{
-		return;
-	}
-	FPerModuleInfo& Info = PerModuleInfos[0];
-	if (const FRigModuleInstance* Module = Info.GetModule())
-	{
-		if (UModularRig* ModularRig = Info.GetModularRig())
-		{
-			if (URigHierarchy* Hierarchy = ModularRig->GetHierarchy())
-			{
-				if (FRigConnectorElement* ConnectorElement = Cast<FRigConnectorElement>(Hierarchy->Find(InConnectorKey)))
-				{
-					ConnectionListBox.FindRef(InConnectorKey)->GetTreeView()->RefreshTreeView(true);
-				}
-			}
-		}
-	}
+	ConnectionListBox.FindRef(InConnectorName)->GetTreeView()->RefreshTreeView(true);
 }
 
-void FRigModuleInstanceDetails::PopulateConnectorCurrentTarget(TSharedPtr<SVerticalBox> InListBox, const FRigElementKey& InConnectorKey, const FRigElementKey& InTargetKey,
-	const FSlateBrush* InBrush, const FSlateColor& InColor, const FText& InTitle)
+void FRigModuleInstanceDetails::PopulateConnectorCurrentTarget(TSharedPtr<SVerticalBox> InListBox, const FSlateBrush* InBrush, const FSlateColor& InColor, const FText& InTitle)
 {
 	static const FSlateBrush* RoundedBoxBrush = FControlRigEditorStyle::Get().GetBrush(TEXT("ControlRig.SpacePicker.RoundedRect"));
-	FPerModuleInfo Info = PerModuleInfos[0];
-	UModularRigController* Controller = Info.GetBlueprint()->GetModularRigController();
-
+	
 	TSharedPtr<SHorizontalBox> RowBox, ButtonBox;
 	InListBox->AddSlot()
 	.AutoHeight()
@@ -859,11 +946,6 @@ void FRigModuleInstanceDetails::PopulateConnectorCurrentTarget(TSharedPtr<SVerti
 		SNew( SButton )
 		.ButtonStyle(FAppStyle::Get(), "SimpleButton")
 		.ContentPadding(FMargin(0.0))
-		.OnClicked_Lambda([Info, Controller, InConnectorKey, InTargetKey]()
-		{
-			Controller->ConnectConnectorToElement(InConnectorKey, InTargetKey, true, Info.GetModularRig()->GetModularRigSettings().bAutoResolve);
-			return FReply::Handled();
-		})
 		[
 			SAssignNew(RowBox, SHorizontalBox)
 			+ SHorizontalBox::Slot()
@@ -914,6 +996,7 @@ void FRigModuleInstanceDetails::PopulateConnectorCurrentTarget(TSharedPtr<SVerti
 
 void FRigModuleInstanceDetails::OnConfigValueChanged(const FName InVariableName)
 {
+	TMap<FString, FString> ModuleValues;
 	for(const FPerModuleInfo& Info : PerModuleInfos)
 	{
 		if (const FRigModuleInstance* ModuleInstance = Info.GetModule())
@@ -921,17 +1004,22 @@ void FRigModuleInstanceDetails::OnConfigValueChanged(const FName InVariableName)
 			if (const UControlRig* ModuleRig = ModuleInstance->GetRig())
 			{
 				FString ValueStr = ModuleRig->GetVariableAsString(InVariableName);
-				if (UControlRigBlueprint* Blueprint = Info.GetBlueprint())
-				{
-					UModularRigController* Controller = Blueprint->GetModularRigController();
-					Controller->SetConfigValueInModule(ModuleInstance->GetPath(), InVariableName, ValueStr);
-				}
+				ModuleValues.Add(ModuleInstance->GetPath(), ValueStr);
 			}
 		}
 	}
+	
+	for (const TPair<FString, FString>& Value : ModuleValues)
+	{
+		if (UControlRigBlueprint* Blueprint = PerModuleInfos[0].GetBlueprint())
+		{
+			UModularRigController* Controller = Blueprint->GetModularRigController();
+			Controller->SetConfigValueInModule(Value.Key, InVariableName, Value.Value);
+		}
+	}
 }
 
-void FRigModuleInstanceDetails::OnConnectorTargetChanged(TSharedPtr<FRigTreeElement> Selection, ESelectInfo::Type SelectInfo, const FRigElementKey InConnectorKey)
+void FRigModuleInstanceDetails::OnConnectorTargetChanged(TSharedPtr<FRigTreeElement> Selection, ESelectInfo::Type SelectInfo, const FRigModuleConnector InConnector)
 {
 	if (SelectInfo == ESelectInfo::OnNavigation)
 	{
@@ -939,18 +1027,21 @@ void FRigModuleInstanceDetails::OnConnectorTargetChanged(TSharedPtr<FRigTreeElem
 	}
 	
 	FScopedTransaction Transaction(LOCTEXT("ModuleHierarchyResolveConnector", "Resolve Connector"));
-	FPerModuleInfo Info = PerModuleInfos[0];
-	TSharedPtr<FRigTreeElement> NewSelection = Selection;
-	if (UModularRigController* Controller = Info.GetBlueprint()->GetModularRigController())
+	for (FPerModuleInfo& Info : PerModuleInfos)
 	{
-		if (NewSelection.IsValid())
+		if (UModularRigController* Controller = Info.GetBlueprint()->GetModularRigController())
 		{
-			const FModularRigSettings& Settings = Info.GetModularRig()->GetModularRigSettings();
-			Controller->ConnectConnectorToElement(InConnectorKey, NewSelection->Key, true, Settings.bAutoResolve);
-		}
-		else
-		{
-			Controller->DisconnectConnector(InConnectorKey);
+			FString ConnectorPath = FString::Printf(TEXT("%s:%s"), *Info.Path, *InConnector.Name);
+			FRigElementKey ConnectorKey(*ConnectorPath, ERigElementType::Connector);
+			if (Selection.IsValid())
+			{
+				const FModularRigSettings& Settings = Info.GetModularRig()->GetModularRigSettings();
+				Controller->ConnectConnectorToElement(ConnectorKey, Selection->Key, true, Settings.bAutoResolve);
+			}
+			else
+			{
+				Controller->DisconnectConnector(ConnectorKey);
+			}
 		}
 	}
 }
@@ -995,14 +1086,14 @@ FText FRigModuleInstanceDetails::GetBindingText(const FProperty* InProperty) con
 {
 	const FName VariableName = InProperty->GetFName();
 	FText FirstValue;
-	for(const FPerModuleInfo& Info : PerModuleInfos)
+	for (int32 ModuleIndex=0; ModuleIndex<PerModuleInfos.Num(); ++ModuleIndex)
 	{
-		if (const FRigModuleReference* ModuleReference = Info.GetReference())
+		if (const FRigModuleReference* ModuleReference = PerModuleInfos[ModuleIndex].GetReference())
 		{
 			if(ModuleReference->Bindings.Contains(VariableName))
 			{
 				const FText BindingText = FText::FromString(ModuleReference->Bindings.FindChecked(VariableName));
-				if(FirstValue.IsEmpty())
+				if(ModuleIndex == 0)
 				{
 					FirstValue = BindingText;
 				}
diff --git a/Engine/Plugins/Animation/ControlRig/Source/ControlRigEditor/Private/ControlRigModuleDetails.h b/Engine/Plugins/Animation/ControlRig/Source/ControlRigEditor/Private/ControlRigModuleDetails.h
--- a/Engine/Plugins/Animation/ControlRig/Source/ControlRigEditor/Private/ControlRigModuleDetails.h
+++ b/Engine/Plugins/Animation/ControlRig/Source/ControlRigEditor/Private/ControlRigModuleDetails.h
@@ -50,17 +50,15 @@ public:
 	FText GetRigClassPath() const;
 	TArray<FRigModuleConnector> GetConnectors() const;
 	FRigElementKeyRedirector GetConnections() const;
-	void PopulateConnectorTargetList(const FRigElementKey InConnectorKey);
+	void PopulateConnectorTargetList(const FString InConnectorKey);
 	void PopulateConnectorCurrentTarget(
 		TSharedPtr<SVerticalBox> InListBox,
-		const FRigElementKey& InConnectorKey,
-		const FRigElementKey& InTargetKey,
 		const FSlateBrush* InBrush,
 		const FSlateColor& InColor,
 		const FText& InTitle);
 
 	void OnConfigValueChanged(const FName InVariableName);
-	void OnConnectorTargetChanged(TSharedPtr<FRigTreeElement> Selection, ESelectInfo::Type SelectInfo, const FRigElementKey InConnectorKey);
+	void OnConnectorTargetChanged(TSharedPtr<FRigTreeElement> Selection, ESelectInfo::Type SelectInfo, const FRigModuleConnector InConnector);
 	
 	struct FPerModuleInfo
 	{
@@ -146,10 +144,10 @@ protected:
 	void HandleChangeBinding(const FProperty* InProperty, const FString& InNewVariablePath) const;
 
 	TArray<FPerModuleInfo> PerModuleInfos;
-	TMap<FRigElementKey, TSharedPtr<SSearchableRigHierarchyTreeView>> ConnectionListBox;
+	TMap<FString, TSharedPtr<SSearchableRigHierarchyTreeView>> ConnectionListBox;
 
 	/** Helper buttons. */
-	TMap<FRigElementKey, TSharedPtr<SButton>> UseSelectedButton;
-	TMap<FRigElementKey, TSharedPtr<SButton>> SelectElementButton;
-	TMap<FRigElementKey, TSharedPtr<SButton>> ResetConnectorButton;
+	TMap<FString, TSharedPtr<SButton>> UseSelectedButton;
+	TMap<FString, TSharedPtr<SButton>> SelectElementButton;
+	TMap<FString, TSharedPtr<SButton>> ResetConnectorButton;
 };
diff --git a/Engine/Plugins/Animation/ControlRig/Source/ControlRigEditor/Private/Editor/SModularRigModel.cpp b/Engine/Plugins/Animation/ControlRig/Source/ControlRigEditor/Private/Editor/SModularRigModel.cpp
--- a/Engine/Plugins/Animation/ControlRig/Source/ControlRigEditor/Private/Editor/SModularRigModel.cpp
+++ b/Engine/Plugins/Animation/ControlRig/Source/ControlRigEditor/Private/Editor/SModularRigModel.cpp
@@ -311,7 +311,7 @@ void SModularRigModel::OnItemClicked(TSharedPtr<FModularRigTreeElement> InItem)
 
 	if (ControlRigEditor.IsValid() && InItem.IsValid())
 	{
-		ControlRigEditor.Pin()->SetDetailViewForRigModules({InItem->ModulePath});
+		ControlRigEditor.Pin()->SetDetailViewForRigModules(TreeView->GetSelectedKeys());
 	}
 }
 
diff --git a/Engine/Plugins/Animation/ControlRig/Source/ControlRigEditor/Private/Editor/SModularRigTreeView.cpp b/Engine/Plugins/Animation/ControlRig/Source/ControlRigEditor/Private/Editor/SModularRigTreeView.cpp
--- a/Engine/Plugins/Animation/ControlRig/Source/ControlRigEditor/Private/Editor/SModularRigTreeView.cpp
+++ b/Engine/Plugins/Animation/ControlRig/Source/ControlRigEditor/Private/Editor/SModularRigTreeView.cpp
@@ -922,7 +922,7 @@ TArray<FString> SModularRigTreeView::GetSelectedKeys() const
 	TArray<TSharedPtr<FModularRigTreeElement>> SelectedElements = GetSelectedItems();
 	for(const TSharedPtr<FModularRigTreeElement>& SelectedElement : SelectedElements)
 	{
-		Keys.Add(SelectedElement->Key);
+		Keys.AddUnique(SelectedElement->ModulePath);
 	}
 	return Keys;
 }
