Control Rig: Fix compilation issue
#rnx

diff --git a/Engine/Plugins/Animation/ControlRig/Source/ControlRig/Private/Rigs/RigModuleDefines.cpp b/Engine/Plugins/Animation/ControlRig/Source/ControlRig/Private/Rigs/RigModuleDefines.cpp
--- a/Engine/Plugins/Animation/ControlRig/Source/ControlRig/Private/Rigs/RigModuleDefines.cpp
+++ b/Engine/Plugins/Animation/ControlRig/Source/ControlRig/Private/Rigs/RigModuleDefines.cpp
@@ -4,3 +4,7 @@
 
 #include UE_INLINE_GENERATED_CPP_BY_NAME(RigModuleDefines)
 
+bool FRigModuleConnector::operator==(const FRigModuleConnector& Other) const
+{
+	return Name == Other.Name && GetTypeHash(Settings) == GetTypeHash(Other.Settings);
+}
diff --git a/Engine/Plugins/Animation/ControlRig/Source/ControlRig/Public/Rigs/RigModuleDefines.h b/Engine/Plugins/Animation/ControlRig/Source/ControlRig/Public/Rigs/RigModuleDefines.h
--- a/Engine/Plugins/Animation/ControlRig/Source/ControlRig/Public/Rigs/RigModuleDefines.h
+++ b/Engine/Plugins/Animation/ControlRig/Source/ControlRig/Public/Rigs/RigModuleDefines.h
@@ -45,10 +45,7 @@ struct CONTROLRIG_API FRigModuleConnector
 	FRigModuleConnector()
 	{}
 
-	bool operator==(const FRigModuleConnector& Other) const
-	{
-		return Name == Other.Name && GetTypeHash(Settings) == GetTypeHash(Other.Settings);
-	}
+	bool operator==(const FRigModuleConnector& Other) const;
 
 	UPROPERTY(BlueprintReadWrite, EditAnywhere, Category = Connector)
 	FString Name;
