Control Rig: Adding default root module

#pf https://horde.devtools.epicgames.com/job/65ae6643c86e0008406a9e8f
#rnx

#virtualized

diff --git a/Engine/Plugins/Animation/ControlRig/Source/ControlRig/Private/ModularRig.cpp b/Engine/Plugins/Animation/ControlRig/Source/ControlRig/Private/ModularRig.cpp
--- a/Engine/Plugins/Animation/ControlRig/Source/ControlRig/Private/ModularRig.cpp
+++ b/Engine/Plugins/Animation/ControlRig/Source/ControlRig/Private/ModularRig.cpp
@@ -74,7 +74,7 @@ FString FRigModuleInstance::GetShortName() const
 
 FString FRigModuleInstance::GetPath() const
 {
-	if (!ParentPath.IsEmpty())
+	if (!IsRootModule())
 	{
 		return URigHierarchy::JoinNameSpace(ParentPath, Name.ToString()); 
 	}
@@ -139,19 +139,59 @@ const FRigConnectorElement* FRigModuleInstance::FindPrimaryConnector() const
 {
 	if(const UControlRig* Rig = GetRig())
 	{
-		if(const FRigModuleConnector* ExposedConnector = Rig->GetRigModuleSettings().FindPrimaryConnector())
+		if(const URigHierarchy* Hierarchy = Rig->GetHierarchy())
 		{
-			if(const URigHierarchy* Hierarchy = Rig->GetHierarchy())
+			const FString MyModulePath = GetPath();
+			const TArray<FRigConnectorElement*> AllConnectors = Hierarchy->GetConnectors();
+			for(const FRigConnectorElement* Connector : AllConnectors)
 			{
-				const FString ConnectorName = URigHierarchy::JoinNameSpace(GetNamespace(), ExposedConnector->Name); 
-				const FRigElementKey ConnectorKey(*ConnectorName, ERigElementType::Connector);
-				return Hierarchy->Find<FRigConnectorElement>(ConnectorKey);
+				if(Connector->IsPrimary())
+				{
+					const FString ModulePath = Hierarchy->GetModulePath(Connector->GetKey());
+					if(!ModulePath.IsEmpty())
+					{
+						if(ModulePath.Equals(MyModulePath, ESearchCase::CaseSensitive))
+						{
+							return Connector;
+						}
+					}
+				}
 			}
 		}
 	}
 	return nullptr;
 }
 
+TArray<const FRigConnectorElement*> FRigModuleInstance::FindConnectors() const
+{
+	TArray<const FRigConnectorElement*> Connectors;
+	if(const UControlRig* Rig = GetRig())
+	{
+		if(const URigHierarchy* Hierarchy = Rig->GetHierarchy())
+		{
+			const FString MyModulePath = GetPath();
+			const TArray<FRigConnectorElement*> AllConnectors = Hierarchy->GetConnectors();
+			for(const FRigConnectorElement* Connector : AllConnectors)
+			{
+				const FString ModulePath = Hierarchy->GetModulePath(Connector->GetKey());
+				if(!ModulePath.IsEmpty())
+				{
+					if(ModulePath.Equals(MyModulePath, ESearchCase::CaseSensitive))
+					{
+						Connectors.Add(Connector);
+					}
+				}
+			}
+		}
+	}
+	return Connectors;
+}
+
+bool FRigModuleInstance::IsRootModule() const
+{
+	return ParentPath.IsEmpty();
+}
+
 void UModularRig::Serialize(FArchive& Ar)
 {
 	Super::Serialize(Ar);
@@ -558,7 +598,7 @@ void UModularRig::UpdateCachedChildren()
 	RootModules.Reset();
 	for (FRigModuleInstance& Module : Modules)
 	{
-		if (Module.ParentPath.IsEmpty())
+		if (Module.IsRootModule())
 		{
 			RootModules.Add(&Module);
 		}
@@ -761,6 +801,28 @@ const FRigModuleInstance* UModularRig::FindModule(const UControlRig* InModuleIns
 	return FoundModule;
 }
 
+const FRigModuleInstance* UModularRig::FindModule(const FRigBaseElement* InElement) const
+{
+	if(InElement)
+	{
+		return FindModule(InElement->GetKey());
+	}
+	return nullptr;
+}
+
+const FRigModuleInstance* UModularRig::FindModule(const FRigElementKey& InElementKey) const
+{
+	if(const URigHierarchy* Hierarchy = GetHierarchy())
+	{
+		const FString ModulePath = Hierarchy->GetModulePath(InElementKey);
+		if(!ModulePath.IsEmpty())
+		{
+			return FindModule(ModulePath);
+		}
+	}
+	return nullptr;
+}
+
 FString UModularRig::GetParentPath(const FString& InPath) const
 {
 	if (const FRigModuleInstance* Element = FindModule(InPath))
diff --git a/Engine/Plugins/Animation/ControlRig/Source/ControlRig/Private/ModularRigController.cpp b/Engine/Plugins/Animation/ControlRig/Source/ControlRig/Private/ModularRigController.cpp
--- a/Engine/Plugins/Animation/ControlRig/Source/ControlRig/Private/ModularRigController.cpp
+++ b/Engine/Plugins/Animation/ControlRig/Source/ControlRig/Private/ModularRigController.cpp
@@ -180,12 +180,12 @@ bool UModularRigController::CanConnectConnectorToElement(const FRigElementKey& I
 		return true; // Nothing to do
 	}
 
-	if (ModuleConnector->Settings.Type != EConnectorType::Primary)
+	if (!ModuleConnector->IsPrimary())
 	{
 		const FRigModuleConnector* PrimaryMdouleConnector = RigCDO->GetRigModuleSettings().ExposedConnectors.FindByPredicate(
 		[ConnectorName](FRigModuleConnector& Connector)
 		{
-			return Connector.Settings.Type == EConnectorType::Primary;
+			return Connector.IsPrimary();
 		});
 
 		const FString PrimaryConnectorPath = FString::Printf(TEXT("%s:%s"), *ConnectorModulePath, *PrimaryMdouleConnector->Name);
@@ -263,14 +263,11 @@ bool UModularRigController::ConnectConnectorToElement(const FRigElementKey& InCo
 
 #if WITH_EDITOR
 	FName TargetModulePathName = NAME_None;
-	if (UControlRig* RigCDO = Module->Class->GetDefaultObject<UControlRig>())
+	if (UModularRig* ModularRig = Cast<UModularRig>(Blueprint->GetObjectBeingDebugged()))
 	{
-		if (UModularRig* ModularRig = Cast<UModularRig>(Blueprint->GetObjectBeingDebugged()))
+		if (URigHierarchy* Hierarchy = ModularRig->GetHierarchy())
 		{
-			if (URigHierarchy* Hierarchy = ModularRig->GetHierarchy())
-			{
-				TargetModulePathName = Hierarchy->GetNameMetadata(InTargetKey, URigHierarchy::ModuleMetadataName, NAME_None);
-			}
+			TargetModulePathName = Hierarchy->GetModulePathFName(InTargetKey);
 		}
 	}
 	
@@ -337,7 +334,7 @@ bool UModularRigController::ConnectConnectorToElement(const FRigElementKey& InCo
 				// automatically re-parent the module in the module tree as well
 				if(const FRigConnectorElement* Connector = Hierarchy->Find<FRigConnectorElement>(InConnectorKey))
 				{
-					if(Connector->Settings.Type == EConnectorType::Primary)
+					if(Connector->IsPrimary())
 					{
 						if(!TargetModulePathName.IsNone())
 						{
@@ -416,7 +413,7 @@ bool UModularRigController::DisconnectConnector(const FRigElementKey& InConnecto
 
 	Model->Connections.RemoveConnection(InConnectorKey);
 
-	if (Connector->Settings.Type == EConnectorType::Primary)
+	if (Connector->IsPrimary())
 	{
 		// Remove connections from module and child modules
 		TArray<FRigElementKey> ConnectionsToRemove;
diff --git a/Engine/Plugins/Animation/ControlRig/Source/ControlRig/Private/ModularRigModel.cpp b/Engine/Plugins/Animation/ControlRig/Source/ControlRig/Private/ModularRigModel.cpp
--- a/Engine/Plugins/Animation/ControlRig/Source/ControlRig/Private/ModularRigModel.cpp
+++ b/Engine/Plugins/Animation/ControlRig/Source/ControlRig/Private/ModularRigModel.cpp
@@ -29,6 +29,52 @@ FString FRigModuleReference::GetNamespace() const
 	return GetPath() + UModularRig::NamespaceSeparator;
 }
 
+const FRigConnectorElement* FRigModuleReference::FindPrimaryConnector(const URigHierarchy* InHierarchy) const
+{
+	if(InHierarchy)
+	{
+		const FString MyModulePath = GetPath();
+		const TArray<FRigConnectorElement*> AllConnectors = InHierarchy->GetConnectors();
+		for(const FRigConnectorElement* Connector : AllConnectors)
+		{
+			if(Connector->IsPrimary())
+			{
+				const FString ModulePath = InHierarchy->GetModulePath(Connector->GetKey());
+				if(!ModulePath.IsEmpty())
+				{
+					if(ModulePath.Equals(MyModulePath, ESearchCase::CaseSensitive))
+					{
+						return Connector;
+					}
+				}
+			}
+		}
+	}
+	return nullptr;
+}
+
+TArray<const FRigConnectorElement*> FRigModuleReference::FindConnectors(const URigHierarchy* InHierarchy) const
+{
+	TArray<const FRigConnectorElement*> Connectors;
+	if(InHierarchy)
+	{
+		const FString MyModulePath = GetPath();
+		const TArray<FRigConnectorElement*> AllConnectors = InHierarchy->GetConnectors();
+		for(const FRigConnectorElement* Connector : AllConnectors)
+		{
+			const FString ModulePath = InHierarchy->GetModulePath(Connector->GetKey());
+			if(!ModulePath.IsEmpty())
+			{
+				if(ModulePath.Equals(MyModulePath, ESearchCase::CaseSensitive))
+				{
+					Connectors.Add(Connector);
+				}
+			}
+		}
+	}
+	return Connectors;
+}
+
 TMap<FRigElementKey, FRigElementKey> FModularRigConnections::GetModuleConnectionMap(const FString& InModulePath) const
 {
 	TMap<FRigElementKey, FRigElementKey> Result;
diff --git a/Engine/Plugins/Animation/ControlRig/Source/ControlRig/Private/Rigs/RigConnectionRules.cpp b/Engine/Plugins/Animation/ControlRig/Source/ControlRig/Private/Rigs/RigConnectionRules.cpp
--- a/Engine/Plugins/Animation/ControlRig/Source/ControlRig/Private/Rigs/RigConnectionRules.cpp
+++ b/Engine/Plugins/Animation/ControlRig/Source/ControlRig/Private/Rigs/RigConnectionRules.cpp
@@ -113,7 +113,7 @@ const FRigConnectorElement* FRigConnectionRuleInput::FindPrimaryConnector(FText*
 	Hierarchy->ForEach<FRigConnectorElement>(
 		[this, ModuleNameSpace, &PrimaryConnector](FRigConnectorElement* Connector) -> bool
 		{
-			if(Connector->Settings.Type == EConnectorType::Primary)
+			if(Connector->IsPrimary())
 			{
 				const FName ConnectorNameSpace = Hierarchy->GetNameMetadata(Connector->GetKey(), URigHierarchy::NameSpaceMetadataName, NAME_None);
 				if(!ConnectorNameSpace.IsNone() && ConnectorNameSpace.IsEqual(ModuleNameSpace, ENameCase::CaseSensitive))
@@ -149,7 +149,7 @@ TArray<const FRigConnectorElement*> FRigConnectionRuleInput::FindSecondaryConnec
 	Hierarchy->ForEach<FRigConnectorElement>(
 		[this, bOptional, ModuleNameSpace, &SecondaryConnectors](FRigConnectorElement* Connector) -> bool
 		{
-			if(Connector->Settings.Type == EConnectorType::Secondary && Connector->Settings.bOptional == bOptional)
+			if(Connector->IsOptional())
 			{
 				const FName ConnectorNameSpace = Hierarchy->GetNameMetadata(Connector->GetKey(), URigHierarchy::NameSpaceMetadataName, NAME_None);
 				if(!ConnectorNameSpace.IsNone() && ConnectorNameSpace.IsEqual(ModuleNameSpace, ENameCase::CaseSensitive))
diff --git a/Engine/Plugins/Animation/ControlRig/Source/ControlRig/Private/Rigs/RigHierarchy.cpp b/Engine/Plugins/Animation/ControlRig/Source/ControlRig/Private/Rigs/RigHierarchy.cpp
--- a/Engine/Plugins/Animation/ControlRig/Source/ControlRig/Private/Rigs/RigHierarchy.cpp
+++ b/Engine/Plugins/Animation/ControlRig/Source/ControlRig/Private/Rigs/RigHierarchy.cpp
@@ -1179,6 +1179,22 @@ bool URigHierarchy::RemoveAllMetadata(FRigElementKey InItem)
 	return RemoveAllMetadataForElement(Find(InItem));
 }
 
+FName URigHierarchy::GetModulePathFName(FRigElementKey InItem) const
+{
+	return GetNameMetadata(InItem, ModuleMetadataName, NAME_None);
+	
+}
+
+FString URigHierarchy::GetModulePath(FRigElementKey InItem) const
+{
+	const FName ModulePathName = GetModulePathFName(InItem);
+	if(!ModulePathName.IsNone())
+	{
+		return ModulePathName.ToString();
+	}
+	return FString();
+}
+
 TArray<const FRigBaseElement*> URigHierarchy::GetSelectedElements(ERigElementType InTypeFilter) const
 {
 	LLM_SCOPE_BYNAME(TEXT("Animation/ControlRig"));
diff --git a/Engine/Plugins/Animation/ControlRig/Source/ControlRig/Private/Rigs/RigHierarchyController.cpp b/Engine/Plugins/Animation/ControlRig/Source/ControlRig/Private/Rigs/RigHierarchyController.cpp
--- a/Engine/Plugins/Animation/ControlRig/Source/ControlRig/Private/Rigs/RigHierarchyController.cpp
+++ b/Engine/Plugins/Animation/ControlRig/Source/ControlRig/Private/Rigs/RigHierarchyController.cpp
@@ -604,7 +604,7 @@ FRigElementKey URigHierarchyController::AddConnector(FName InName, FRigConnector
 		const TArray<FRigConnectorElement*>& Connectors = Hierarchy->GetConnectors();
 		for(const FRigConnectorElement* Connector : Connectors)
 		{
-			if(Connector->Settings.Type == EConnectorType::Primary)
+			if(Connector->IsPrimary())
 			{
 				if(Hierarchy->HasExecuteContext())
 				{
diff --git a/Engine/Plugins/Animation/ControlRig/Source/ControlRig/Private/Settings/ControlRigSettings.cpp b/Engine/Plugins/Animation/ControlRig/Source/ControlRig/Private/Settings/ControlRigSettings.cpp
--- a/Engine/Plugins/Animation/ControlRig/Source/ControlRig/Private/Settings/ControlRigSettings.cpp
+++ b/Engine/Plugins/Animation/ControlRig/Source/ControlRig/Private/Settings/ControlRigSettings.cpp
@@ -1,6 +1,7 @@
 // Copyright Epic Games, Inc. All Rights Reserved.
 
 #include "Settings/ControlRigSettings.h"
+#include "ControlRig.h"
 
 #include UE_INLINE_GENERATED_CPP_BY_NAME(ControlRigSettings)
 
@@ -9,6 +10,7 @@ UControlRigSettings::UControlRigSettings(const FObjectInitializer& ObjectInitial
 {
 #if WITH_EDITORONLY_DATA
 	DefaultShapeLibrary = LoadObject<UControlRigShapeLibrary>(nullptr, TEXT("/ControlRig/Controls/DefaultGizmoLibraryNormalized.DefaultGizmoLibraryNormalized"));
+	DefaultRootModule = TEXT("/ControlRig/Modules/CR_Root.CR_Root");
 #endif
 }
 
diff --git a/Engine/Plugins/Animation/ControlRig/Source/ControlRig/Public/ModularRig.h b/Engine/Plugins/Animation/ControlRig/Source/ControlRig/Public/ModularRig.h
--- a/Engine/Plugins/Animation/ControlRig/Source/ControlRig/Public/ModularRig.h
+++ b/Engine/Plugins/Animation/ControlRig/Source/ControlRig/Public/ModularRig.h
@@ -84,6 +84,8 @@ public:
 	bool ContainsRig(const UControlRig* InRig) const;
 	const FRigModuleReference* GetModuleReference() const;
 	const FRigConnectorElement* FindPrimaryConnector() const;
+	TArray<const FRigConnectorElement*> FindConnectors() const;
+	bool IsRootModule() const;
 };
 
 USTRUCT(BlueprintType)
@@ -176,6 +178,8 @@ public:
 
 	const FRigModuleInstance* FindModule(const FString& InPath) const;
 	const FRigModuleInstance* FindModule(const UControlRig* InModuleInstance) const;
+	const FRigModuleInstance* FindModule(const FRigBaseElement* InElement) const;
+	const FRigModuleInstance* FindModule(const FRigElementKey& InElementKey) const;
 	FString GetParentPath(const FString& InPath) const;
 
 	void ForEachModule(TFunctionRef<bool(FRigModuleInstance*)> PerModuleFunction);
diff --git a/Engine/Plugins/Animation/ControlRig/Source/ControlRig/Public/ModularRigModel.h b/Engine/Plugins/Animation/ControlRig/Source/ControlRig/Public/ModularRigModel.h
--- a/Engine/Plugins/Animation/ControlRig/Source/ControlRig/Public/ModularRigModel.h
+++ b/Engine/Plugins/Animation/ControlRig/Source/ControlRig/Public/ModularRigModel.h
@@ -97,7 +97,10 @@ struct CONTROLRIG_API FRigModuleReference
 		return A.ParentPath == B.ParentPath &&
 			A.Name == B.Name;
 	}
-	
+
+	const FRigConnectorElement* FindPrimaryConnector(const URigHierarchy* InHierarchy) const;
+	TArray<const FRigConnectorElement*> FindConnectors(const URigHierarchy* InHierarchy) const;
+
 	friend class UModularRigController;
 };
 
diff --git a/Engine/Plugins/Animation/ControlRig/Source/ControlRig/Public/Rigs/RigHierarchy.h b/Engine/Plugins/Animation/ControlRig/Source/ControlRig/Public/Rigs/RigHierarchy.h
--- a/Engine/Plugins/Animation/ControlRig/Source/ControlRig/Public/Rigs/RigHierarchy.h
+++ b/Engine/Plugins/Animation/ControlRig/Source/ControlRig/Public/Rigs/RigHierarchy.h
@@ -1565,6 +1565,20 @@ public:
 		return SetArrayMetadata<FRigElementKey>(InItem, ERigMetadataType::RigElementKeyArray, InMetadataName, InValue);
 	}
 
+	/**
+	 * Returns the path of the module an element belong to (or NAME_None in case the element doesn't belong to a module)
+	 * @return The path the element belongs to (or NAME_None)
+	 */
+	UFUNCTION(BlueprintPure, Category = URigHierarchy)
+	FName GetModulePathFName(FRigElementKey InItem) const;
+	
+	/**
+	 * Returns the path of the module an element belong to (or an empty string in case the element doesn't belong to a module)
+	 * @return The path the element belongs to (or empty string)
+	 */
+	UFUNCTION(BlueprintPure, Category = URigHierarchy)
+	FString GetModulePath(FRigElementKey InItem) const;
+
 	/*
 	 * Returns the tags for a given item
 	 * @param InItem The item to return the tags for
diff --git a/Engine/Plugins/Animation/ControlRig/Source/ControlRig/Public/Rigs/RigHierarchyElements.h b/Engine/Plugins/Animation/ControlRig/Source/ControlRig/Public/Rigs/RigHierarchyElements.h
--- a/Engine/Plugins/Animation/ControlRig/Source/ControlRig/Public/Rigs/RigHierarchyElements.h
+++ b/Engine/Plugins/Animation/ControlRig/Source/ControlRig/Public/Rigs/RigHierarchyElements.h
@@ -1727,6 +1727,10 @@ struct CONTROLRIG_API FRigConnectorElement final : public FRigBaseElement
 
 	FRigConnectorState GetConnectorState(const URigHierarchy* InHierarchy) const;
 
+	bool IsPrimary() const { return Settings.Type == EConnectorType::Primary; }
+	bool IsSecondary() const { return Settings.Type == EConnectorType::Secondary; }
+	bool IsOptional() const { return IsSecondary() && Settings.bOptional; }
+
 private:
 	explicit FRigConnectorElement(URigHierarchy* InOwner)
 		: FRigBaseElement(InOwner, ERigElementType::Connector)
diff --git a/Engine/Plugins/Animation/ControlRig/Source/ControlRig/Public/Rigs/RigModuleDefines.h b/Engine/Plugins/Animation/ControlRig/Source/ControlRig/Public/Rigs/RigModuleDefines.h
--- a/Engine/Plugins/Animation/ControlRig/Source/ControlRig/Public/Rigs/RigModuleDefines.h
+++ b/Engine/Plugins/Animation/ControlRig/Source/ControlRig/Public/Rigs/RigModuleDefines.h
@@ -52,6 +52,10 @@ struct CONTROLRIG_API FRigModuleConnector
 
 	UPROPERTY(BlueprintReadWrite, EditAnywhere, Category = Connector)
 	FRigConnectorSettings Settings;
+	
+	bool IsPrimary() const { return Settings.Type == EConnectorType::Primary; }
+	bool IsSecondary() const { return Settings.Type == EConnectorType::Secondary; }
+	bool IsOptional() const { return IsSecondary() && Settings.bOptional; }
 };
 
 USTRUCT(BlueprintType)
@@ -73,7 +77,7 @@ struct CONTROLRIG_API FRigModuleSettings
 	{
 		return ExposedConnectors.FindByPredicate([](const FRigModuleConnector& Connector)
 		{
-			return Connector.Settings.Type == EConnectorType::Primary;
+			return Connector.IsPrimary();
 		});
 	}
 
diff --git a/Engine/Plugins/Animation/ControlRig/Source/ControlRig/Public/Settings/ControlRigSettings.h b/Engine/Plugins/Animation/ControlRig/Source/ControlRig/Public/Settings/ControlRigSettings.h
--- a/Engine/Plugins/Animation/ControlRig/Source/ControlRig/Public/Settings/ControlRigSettings.h
+++ b/Engine/Plugins/Animation/ControlRig/Source/ControlRig/Public/Settings/ControlRigSettings.h
@@ -17,6 +17,7 @@
 #include "ControlRigSettings.generated.h"
 
 class UStaticMesh;
+class UControlRig;
 
 USTRUCT()
 struct FControlRigSettingsPerPinBool
@@ -40,6 +41,9 @@ public:
 #if WITH_EDITORONLY_DATA
 	UPROPERTY(EditAnywhere, config, Category = Shapes)
 	TSoftObjectPtr<UControlRigShapeLibrary> DefaultShapeLibrary;
+
+	UPROPERTY(EditAnywhere, config, Category = ModularRigging, meta=(AllowedClasses="/Script/ControlRigDeveloper.ControlRigBlueprint"))
+	FSoftObjectPath DefaultRootModule;
 #endif
 	
 	static UControlRigSettings * Get() { return GetMutableDefault<UControlRigSettings>(); }
diff --git a/Engine/Plugins/Animation/ControlRig/Source/ControlRigEditor/Private/ControlRigBlueprintFactory.cpp b/Engine/Plugins/Animation/ControlRig/Source/ControlRigEditor/Private/ControlRigBlueprintFactory.cpp
--- a/Engine/Plugins/Animation/ControlRig/Source/ControlRigEditor/Private/ControlRigBlueprintFactory.cpp
+++ b/Engine/Plugins/Animation/ControlRig/Source/ControlRigEditor/Private/ControlRigBlueprintFactory.cpp
@@ -21,10 +21,13 @@
 #include "ControlRigBlueprint.h"
 #include "ControlRigBlueprintActions.h"
 #include "ControlRigEditorModule.h"
+#include "ModularRig.h"
 #include "RigVMBlueprintGeneratedClass.h"
 #include "Graph/ControlRigGraphSchema.h"
 #include "Graph/ControlRigGraph.h"
 #include "Kismet2/BlueprintEditorUtils.h"
+#include "ModularRigController.h"
+#include "Settings/ControlRigSettings.h"
 
 #include UE_INLINE_GENERATED_CPP_BY_NAME(ControlRigBlueprintFactory)
 
@@ -295,6 +298,23 @@ UObject* UControlRigBlueprintFactory::FactoryCreateNew(UClass* Class, UObject* I
 	{
 		UControlRigBlueprint* ControlRigBlueprint = CastChecked<UControlRigBlueprint>(FKismetEditorUtilities::CreateBlueprint(ParentClass, InParent, Name, BPTYPE_Normal, UControlRigBlueprint::StaticClass(), URigVMBlueprintGeneratedClass::StaticClass(), CallingContext));
 		FControlRigEditorModule::Get().CreateRootGraphIfRequired(ControlRigBlueprint);
+
+		// add the default module
+		if(ParentClass->IsChildOf(UModularRig::StaticClass()))
+		{
+			const FSoftObjectPath DefaultRootModulePath = UControlRigSettings::Get()->DefaultRootModule;
+			if(const UControlRigBlueprint* DefaultRootModule = Cast<UControlRigBlueprint>(DefaultRootModulePath.TryLoad()))
+			{
+				if(UClass* DefaultRootModuleClass = Cast<UClass>(DefaultRootModule->GetControlRigClass()))
+				{
+					if(DefaultRootModuleClass->IsChildOf(UControlRig::StaticClass()))
+					{
+						static const FName RootName = TEXT("Root");
+						ControlRigBlueprint->ModularRigModel.GetController()->AddModule(RootName, DefaultRootModuleClass, FString(), false);
+					}
+				}
+			}
+		}
 		return ControlRigBlueprint;
 	}
 }
diff --git a/Engine/Plugins/Animation/ControlRig/Source/ControlRigEditor/Private/ControlRigElementDetails.cpp b/Engine/Plugins/Animation/ControlRig/Source/ControlRigEditor/Private/ControlRigElementDetails.cpp
--- a/Engine/Plugins/Animation/ControlRig/Source/ControlRigEditor/Private/ControlRigElementDetails.cpp
+++ b/Engine/Plugins/Animation/ControlRig/Source/ControlRigEditor/Private/ControlRigElementDetails.cpp
@@ -1140,7 +1140,7 @@ bool FRigBaseElementDetails::IsAnyConnectorPrimary() const
 	{
 		if(const FRigConnectorElement* Connector = Info.Element.Get<FRigConnectorElement>())
 		{
-			return Connector->Settings.Type == EConnectorType::Primary;
+			return Connector->IsPrimary();
 		}
 		return false;
 	});
diff --git a/Engine/Plugins/Animation/ControlRig/Source/ControlRigEditor/Private/Editor/ControlRigSchematicModel.cpp b/Engine/Plugins/Animation/ControlRig/Source/ControlRigEditor/Private/Editor/ControlRigSchematicModel.cpp
--- a/Engine/Plugins/Animation/ControlRig/Source/ControlRigEditor/Private/Editor/ControlRigSchematicModel.cpp
+++ b/Engine/Plugins/Animation/ControlRig/Source/ControlRigEditor/Private/Editor/ControlRigSchematicModel.cpp
@@ -75,56 +75,57 @@ void FControlRigSchematicModel::Tick(float InDeltaTime)
 {
 	Super::Tick(InDeltaTime);
 
-	if(bUpdateConnectorTargetsOnTick)
+	if(ControlRigBlueprint.IsValid())
 	{
-		if(ControlRigBlueprint.IsValid())
-		{
-			const URigHierarchy* Hierarchy = ControlRigBlueprint->GetDebuggedControlRig()->GetHierarchy();
-			check(Hierarchy);
-			const FModularRigConnections& Connections = ControlRigBlueprint->ModularRigModel.Connections;
+		const URigHierarchy* Hierarchy = ControlRigBlueprint->GetDebuggedControlRig()->GetHierarchy();
+		check(Hierarchy);
+		const FModularRigConnections& Connections = ControlRigBlueprint->ModularRigModel.Connections;
 
-			TArray<FRigElementKey> KeysToRemove;
-			for(const TPair<FRigElementKey,FGuid>& Pair : RigElementKeyToGuid)
+		TArray<FRigElementKey> KeysToRemove;
+		for(const TPair<FRigElementKey,FGuid>& Pair : RigElementKeyToGuid)
+		{
+			if((Pair.Key.Type == ERigElementType::Socket) ||
+				(Pair.Key.Type == ERigElementType::Connector))
 			{
-				if((Pair.Key.Type == ERigElementType::Socket) ||
-					(Pair.Key.Type == ERigElementType::Connector))
-				{
-					continue;
-				}
+				continue;
+			}
 
-				const TArray<FRigElementKey>& Connectors = Connections.FindConnectorsFromTarget(Pair.Key);
-				if(Connectors.Num() > 1)
-				{
-					continue;
-				}
-				if(Connectors.Num() == 1)
+			const TArray<FRigElementKey>& Connectors = Connections.FindConnectorsFromTarget(Pair.Key);
+			if(Connectors.Num() > 1)
+			{
+				continue;
+			}
+			if(Connectors.Num() == 1)
+			{
+				if(const FRigConnectorElement* Connector = Hierarchy->Find<FRigConnectorElement>(Pair.Key))
 				{
-					if(const FRigConnectorElement* Connector = Hierarchy->Find<FRigConnectorElement>(Pair.Key))
+					if(!Connector->IsPrimary())
 					{
-						if(Connector->Settings.Type != EConnectorType::Primary)
-						{
-							continue;
-						}
+						continue;
 					}
 				}
-				
-				KeysToRemove.Add(Pair.Key);
 			}
-			for(const FRigElementKey& KeyToRemove : KeysToRemove)
+			
+			KeysToRemove.Add(Pair.Key);
+		}
+		for(const FRigElementKey& KeyToRemove : KeysToRemove)
+		{
+			RemoveElementKeyNode(KeyToRemove);
+		}
+
+		TArray<FRigElementKey> ConnectorKeys;
+		for(const TPair<FRigElementKey,FGuid>& Pair : RigElementKeyToGuid)
+		{
+			if(Pair.Key.Type == ERigElementType::Connector)
 			{
-				RemoveElementKeyNode(KeyToRemove);
+				ConnectorKeys.Add(Pair.Key);
 			}
 		}
-		bUpdateConnectorTargetsOnTick = false;
-	}
 
-	if(FSlateApplication::Get().IsDragDropping())
-	{
-		
-	}
-	else
-	{
-		
+		for(const FRigElementKey& ConnectorKey : ConnectorKeys)
+		{
+			UpdateConnector(ConnectorKey);
+		}
 	}
 }
 
@@ -135,17 +136,7 @@ FControlRigSchematicRigElementKeyNode* FControlRigSchematicModel::AddElementKeyN
 	{
 		ConfigureElementKeyNode(Node, InKey);
 		RigElementKeyToGuid.Add(Node->GetKey(), Node->GetGuid());
-
-		FRigElementKey ResolvedKey;
-		if(IsConnectorResolved(InKey, &ResolvedKey))
-		{
-			const FControlRigSchematicRigElementKeyNode* ResolvedNode = FindElementKeyNode(ResolvedKey);
-			if(ResolvedNode == nullptr)
-			{
-				ResolvedNode = AddElementKeyNode(ResolvedKey);
-			}
-			SetParentNode(Node->GetGuid(), ResolvedNode->GetGuid());
-		}
+		UpdateConnector(InKey);
 
 		if(InKey.Type == ERigElementType::Connector)
 		{
@@ -430,6 +421,58 @@ void FControlRigSchematicModel::UpdateControlRigContent()
 	UpdateElementKeyLinks();
 }
 
+void FControlRigSchematicModel::UpdateConnector(const FRigElementKey& InElementKey)
+{
+	if(InElementKey.Type == ERigElementType::Connector)
+	{
+		if(const FControlRigSchematicRigElementKeyNode* ConnectorNode = FindElementKeyNode(InElementKey))
+		{
+			FRigElementKey ResolvedKey;
+			if(IsConnectorResolved(InElementKey, &ResolvedKey))
+			{
+				if(ControlRigBlueprint.IsValid() && ControlRigBeingDebuggedPtr.IsValid() && ResolvedKey.Type != ERigElementType::Socket)
+				{
+					if(const URigHierarchy* Hierarchy = ControlRigBeingDebuggedPtr->GetHierarchy())
+					{
+						if(const FRigConnectorElement* Connector = Hierarchy->Find<FRigConnectorElement>(InElementKey))
+						{
+							if(Connector->IsPrimary())
+							{
+								// if the resolved target has only one primary connector on it, and it is not a socket
+								// let's not draw the node in the schematic for now
+								if(ControlRigBlueprint->ModularRigModel.Connections.FindConnectorsFromTarget(ResolvedKey).Num() == 1)
+								{
+									ResolvedKey.Reset();
+								}
+							}
+						}
+					}
+				}
+			}
+
+			if(ResolvedKey.IsValid())
+			{
+				const FControlRigSchematicRigElementKeyNode* ResolvedNode = FindElementKeyNode(ResolvedKey);
+				if(ResolvedNode == nullptr)
+				{
+					ResolvedNode = AddElementKeyNode(ResolvedKey);
+				}
+
+				if(ConnectorNode->GetParentNode() != ResolvedNode)
+				{
+					SetParentNode(ConnectorNode->GetGuid(), ResolvedNode->GetGuid());
+					const_cast<FControlRigSchematicRigElementKeyNode*>(ConnectorNode)->OnMouseLeave();
+				}
+			}
+			else if(ConnectorNode->HasParentNode())
+			{
+				RemoveFromParentNode(ConnectorNode->GetGuid());
+				const_cast<FControlRigSchematicRigElementKeyNode*>(ConnectorNode)->OnMouseLeave();
+			}
+		}
+	}
+}
+
 void FControlRigSchematicModel::OnSetObjectBeingDebugged(UObject* InObject)
 {
 	if(ControlRigBeingDebuggedPtr.Get() == InObject)
@@ -559,6 +602,7 @@ void FControlRigSchematicModel::OnHierarchyModified(ERigHierarchyNotification In
 				{
 					AddElementKeyNode(Connector->GetKey());
 				}
+				UpdateConnector(Connector->GetKey());
 			}
 
 			// remove obsolete nodes
@@ -607,25 +651,7 @@ void FControlRigSchematicModel::HandleModularRigModified(EModularRigNotification
 								const FString LocalName = NameString.Mid(InModule->GetNamespace().Len());
 								if(!LocalName.Contains(UModularRig::NamespaceSeparator))
 								{
-									const FRigElementKey LocalKey(*LocalName, ConnectorKey.Type);
-									const FRigElementKey ResolvedKey = ControlRigBlueprint->ModularRigModel.Connections.FindTargetFromConnector(ConnectorKey);
-									if(ResolvedKey.IsValid())
-									{
-										const FControlRigSchematicRigElementKeyNode* ResolvedNode = FindElementKeyNode(ResolvedKey);
-										if(ResolvedNode == nullptr)
-										{
-											ResolvedNode = AddElementKeyNode(ResolvedKey);
-										}
-										(void)SetParentNode(ConnectorNode, ResolvedNode);
-										const_cast<FControlRigSchematicRigElementKeyNode*>(ResolvedNode)->OnMouseLeave();
-										const_cast<FSchematicGraphNode*>(ConnectorNode)->OnMouseLeave();
-										continue;
-									}
-									if(ConnectorNode->GetParentNode())
-									{
-										(void)RemoveFromParentNode(ConnectorNode);
-										const_cast<FSchematicGraphNode*>(ConnectorNode)->OnMouseLeave();
-									}
+									UpdateConnector(ConnectorKey);
 								}
 							}
 						}
@@ -639,8 +665,6 @@ void FControlRigSchematicModel::HandleModularRigModified(EModularRigNotification
 			break;
 		}
 	}
-
-	bUpdateConnectorTargetsOnTick = true;
 }
 
 FSchematicGraphGroupNode* FControlRigSchematicModel::AddAutoGroupNode()
@@ -759,33 +783,29 @@ const FSlateBrush* FControlRigSchematicModel::GetBrushForKey(const FRigElementKe
 					{
 						if(const FRigConnectorElement* Connector = Hierarchy->Find<FRigConnectorElement>(InKey))
 						{
-							if((Connector->Settings.Type == EConnectorType::Primary) && !IsConnectorResolved(InKey))
+							if((Connector->IsPrimary()) && !IsConnectorResolved(InKey))
 							{
-								const FName ModulePathName = Hierarchy->GetNameMetadata(Connector->GetKey(), URigHierarchy::ModuleMetadataName, NAME_None);
-								if(!ModulePathName.IsNone())
+								if(const FRigModuleInstance* ModuleInstance = ControlRig->FindModule(Connector))
 								{
-									if(const FRigModuleInstance* ModuleInstance = ControlRig->FindModule(ModulePathName.ToString()))
+									const FSoftObjectPath& IconPath = ModuleInstance->GetRig()->GetRigModuleSettings().Icon;
+									const FSlateBrush* IconBrush = ModuleIcons.Find(IconPath);
+									if (!IconBrush)
 									{
-										const FSoftObjectPath& IconPath = ModuleInstance->GetRig()->GetRigModuleSettings().Icon;
-										const FSlateBrush* IconBrush = ModuleIcons.Find(IconPath);
-										if (!IconBrush)
-										{
-											if(UTexture2D* Icon = Cast<UTexture2D>(IconPath.TryLoad()))
-											{
-												IconBrush = &ModuleIcons.Add(IconPath, UWidgetBlueprintLibrary::MakeBrushFromTexture(Icon, 16.0f, 16.0f));
-											}
-										}
-										if(IconBrush)
+										if(UTexture2D* Icon = Cast<UTexture2D>(IconPath.TryLoad()))
 										{
-											return IconBrush;
+											IconBrush = &ModuleIcons.Add(IconPath, UWidgetBlueprintLibrary::MakeBrushFromTexture(Icon, 16.0f, 16.0f));
 										}
 									}
+									if(IconBrush)
+									{
+										return IconBrush;
+									}
 								}
 							}
 
-							if(Connector->Settings.Type == EConnectorType::Secondary)
+							if(Connector->IsSecondary())
 							{
-								return Connector->Settings.bOptional ? 	ConnectorOptionalBrush : ConnectorSecondaryBrush;
+								return Connector->IsOptional() ? 	ConnectorOptionalBrush : ConnectorSecondaryBrush;
 							}
 						}
 					}
@@ -985,7 +1005,21 @@ ESchematicGraphVisibility::Type FControlRigSchematicModel::GetVisibilityForNode(
 					{
 						if(const FRigConnectorElement* Connector = Hierarchy->Find<FRigConnectorElement>(Node->GetKey()))
 						{
-							if(Connector->Settings.Type != EConnectorType::Primary)
+							if(Connector->IsPrimary())
+							{
+								// if the module is a root module
+								if(const FRigModuleInstance* ModuleInstance = ControlRig->FindModule(Connector))
+								{
+									if(ModuleInstance->IsRootModule())
+									{
+										if(IsConnectorResolved(Connector->GetKey()))
+										{
+											return ESchematicGraphVisibility::Hidden;
+										}
+									}
+								}
+							}
+							else
 							{
 								if(Connector->Settings.bOptional)
 								{
@@ -995,18 +1029,13 @@ ESchematicGraphVisibility::Type FControlRigSchematicModel::GetVisibilityForNode(
 									}
 								}
 								
-								// find the corresponding primary connector
-								const FName ModulePathName = Hierarchy->GetNameMetadata(Node->GetKey(), URigHierarchy::ModuleMetadataName, NAME_None);
-								if(!ModulePathName.IsNone())
+								if(const FRigModuleInstance* ModuleInstance = ControlRig->FindModule(Connector))
 								{
-									if(const FRigModuleInstance* ModuleInstance = ControlRig->FindModule(ModulePathName.ToString()))
+									if(const FRigConnectorElement* PrimaryConnector = ModuleInstance->FindPrimaryConnector())
 									{
-										if(const FRigConnectorElement* PrimaryConnector = ModuleInstance->FindPrimaryConnector())
+										if(!IsConnectorResolved(PrimaryConnector->GetKey()))
 										{
-											if(!IsConnectorResolved(PrimaryConnector->GetKey()))
-											{
-												return ESchematicGraphVisibility::Hidden;
-											}
+											return ESchematicGraphVisibility::Hidden;
 										}
 									}
 								}
@@ -1210,7 +1239,7 @@ TArray<FRigElementKey> FControlRigSchematicModel::GetElementKeysFromDragDropEven
 						const FString ConnectorNameSpace = Hierarchy->GetNameMetadata(Connector->GetKey(), URigHierarchy::NameSpaceMetadataName, NAME_None).ToString();
 						if(ConnectorNameSpace.Equals(ModuleNameSpace))
 						{
-							if(Connector->Settings.Type == EConnectorType::Primary)
+							if(Connector->IsPrimary())
 							{
 								DraggedKeys.Add(Connector->GetKey());
 							}
@@ -1454,7 +1483,7 @@ void FControlRigSchematicModel::HandleSchematicDrop(SSchematicGraphPanel* InPane
 							TArray<FRigConnectorElement*> Connectors = Hierarchy->GetElementsOfType<FRigConnectorElement>();
 							for (FRigConnectorElement* Connector : Connectors)
 							{
-								if (Connector->Settings.Type == EConnectorType::Primary)
+								if (Connector->IsPrimary())
 								{
 									FString Path, Name;
 									Connector->GetName().Split(UModularRig::NamespaceSeparator, &Path, &Name, ESearchCase::CaseSensitive, ESearchDir::FromEnd);
@@ -1546,7 +1575,7 @@ void FControlRigSchematicModel::HandleSchematicDrop(SSchematicGraphPanel* InPane
 							const FString ConnectorNameSpace = Hierarchy->GetNameMetadata(Connector->GetKey(), URigHierarchy::NameSpaceMetadataName, NAME_None).ToString();
 							if(ConnectorNameSpace.Equals(ModuleNameSpace))
 							{
-								if(Connector->Settings.Type == EConnectorType::Primary)
+								if(Connector->IsPrimary())
 								{
 									Controller->ConnectConnectorToElement(Connector->GetKey(), TargetKey, true, ControlRig->GetModularRigSettings().bAutoResolve);
 									return;
diff --git a/Engine/Plugins/Animation/ControlRig/Source/ControlRigEditor/Private/Editor/ControlRigSchematicModel.h b/Engine/Plugins/Animation/ControlRig/Source/ControlRigEditor/Private/Editor/ControlRigSchematicModel.h
--- a/Engine/Plugins/Animation/ControlRig/Source/ControlRigEditor/Private/Editor/ControlRigSchematicModel.h
+++ b/Engine/Plugins/Animation/ControlRig/Source/ControlRigEditor/Private/Editor/ControlRigSchematicModel.h
@@ -87,6 +87,7 @@ public:
 	void UpdateElementKeyNodes();
 	void UpdateElementKeyLinks();
 	void UpdateControlRigContent();
+	void UpdateConnector(const FRigElementKey& InElementKey);
 
 	void OnSetObjectBeingDebugged(UObject* InObject);
 	void OnHierarchyModified(ERigHierarchyNotification InNotif, URigHierarchy* InHierarchy, const FRigBaseElement* InElement);
@@ -128,7 +129,6 @@ private:
 	TArray<FGuid> TemporaryNodeGuids;
 	TMap<FRigElementKey, FGuid> RigElementKeyToGuid;
 	mutable TMap<FSoftObjectPath, FSlateBrush> ModuleIcons;
-	bool bUpdateConnectorTargetsOnTick = false;
 
 	friend class FControlRigEditor;
 };
\ No newline at end of file
diff --git a/Engine/Plugins/Animation/ControlRig/Source/ControlRigEditor/Private/Editor/SModularRigModel.cpp b/Engine/Plugins/Animation/ControlRig/Source/ControlRigEditor/Private/Editor/SModularRigModel.cpp
--- a/Engine/Plugins/Animation/ControlRig/Source/ControlRigEditor/Private/Editor/SModularRigModel.cpp
+++ b/Engine/Plugins/Animation/ControlRig/Source/ControlRigEditor/Private/Editor/SModularRigModel.cpp
@@ -726,7 +726,7 @@ void SModularRigModel::OnHierarchyModified(ERigHierarchyNotification InNotif, UR
 			}
 
 			FString ModulePathOrConnectorName;
-			if(Connector->Settings.Type == EConnectorType::Primary)
+			if(Connector->IsPrimary())
 			{
 				ModulePathOrConnectorName = InHierarchy->GetNameMetadata(Connector->GetKey(), URigHierarchy::NameSpaceMetadataName, NAME_None).ToString();
 				ModulePathOrConnectorName = ModulePathOrConnectorName.LeftChop(1);
diff --git a/Engine/Plugins/Animation/ControlRig/Source/ControlRigEditor/Private/Editor/SModularRigTreeView.cpp b/Engine/Plugins/Animation/ControlRig/Source/ControlRigEditor/Private/Editor/SModularRigTreeView.cpp
--- a/Engine/Plugins/Animation/ControlRig/Source/ControlRigEditor/Private/Editor/SModularRigTreeView.cpp
+++ b/Engine/Plugins/Animation/ControlRig/Source/ControlRigEditor/Private/Editor/SModularRigTreeView.cpp
@@ -712,7 +712,7 @@ bool SModularRigTreeView::AddElement(FString InKey, FString InParentKey)
 					const TArray<FRigModuleConnector>& Connectors = CDO->GetRigModuleSettings().ExposedConnectors;
 					for (const FRigModuleConnector& Connector : Connectors)
 					{
-						if (Connector.Settings.Type == EConnectorType::Primary)
+						if (Connector.IsPrimary())
 						{
 							continue;
 						}
@@ -984,7 +984,7 @@ TPair<const FSlateBrush*, FSlateColor> FModularRigTreeElement::GetBrushAndColor(
 			});
 			if (Connector)
 			{
-				if (Connector->Settings.Type == EConnectorType::Primary)
+				if (Connector->IsPrimary())
 				{
 					if (bIsConnected)
 					{
diff --git a/Engine/Plugins/Animation/ControlRig/Source/ControlRigEditor/Private/Editor/SRigHierarchy.cpp b/Engine/Plugins/Animation/ControlRig/Source/ControlRigEditor/Private/Editor/SRigHierarchy.cpp
--- a/Engine/Plugins/Animation/ControlRig/Source/ControlRigEditor/Private/Editor/SRigHierarchy.cpp
+++ b/Engine/Plugins/Animation/ControlRig/Source/ControlRigEditor/Private/Editor/SRigHierarchy.cpp
@@ -2748,7 +2748,7 @@ TOptional<EItemDropZone> SRigHierarchy::OnCanAcceptDrop(const FDragDropEvent& Dr
 					FRigModuleConnector* PrimaryConnector = nullptr;
 					for (FRigModuleConnector& Connector : AssetBlueprint->RigModuleSettings.ExposedConnectors)
 					{
-						if (Connector.Settings.Type == EConnectorType::Primary)
+						if (Connector.IsPrimary())
 						{
 							PrimaryConnector = &Connector;
 							break;
@@ -2913,7 +2913,7 @@ FReply SRigHierarchy::OnAcceptDrop(const FDragDropEvent& DragDropEvent, EItemDro
 						TArray<FRigConnectorElement*> Connectors = GetHierarchy()->GetElementsOfType<FRigConnectorElement>();
 						for (FRigConnectorElement* Connector : Connectors)
 						{
-							if (Connector->Settings.Type == EConnectorType::Primary)
+							if (Connector->IsPrimary())
 							{
 								FString Path, Name;
 								Connector->GetName().Split(UModularRig::NamespaceSeparator, &Path, &Name, ESearchCase::CaseSensitive, ESearchDir::FromEnd);
diff --git a/Engine/Source/Editor/AnimationEditorWidgets/Private/SchematicGraphPanel/SSchematicGraphPanel.cpp b/Engine/Source/Editor/AnimationEditorWidgets/Private/SchematicGraphPanel/SSchematicGraphPanel.cpp
--- a/Engine/Source/Editor/AnimationEditorWidgets/Private/SchematicGraphPanel/SSchematicGraphPanel.cpp
+++ b/Engine/Source/Editor/AnimationEditorWidgets/Private/SchematicGraphPanel/SSchematicGraphPanel.cpp
@@ -1114,6 +1114,11 @@ void SSchematicGraphPanel::Tick(float DeltaTime)
 	{
 		TSharedRef<SSchematicGraphNode> Widget = GetChild(i);
 
+		if(!bIsDragDropping && Widget->IsBeingDragged())
+		{
+			Widget->bIsBeingDragged = false;
+		}
+
 		// update the animation state of the node
 		Widget->EnablePositionAnimation(GraphData->GetPositionAnimationEnabledForNode(Widget->GetGuid()));
 		if(Widget->GetVisibility() != EVisibility::Visible)
@@ -1403,6 +1408,7 @@ void SSchematicGraphPanel::UpdatePerNodeCaches(bool bRemoveNodesFromAutoGroups)
 			}
 			Cache.bHasParent = Node->HasParentNode();
 			Cache.Visibility = GraphData->GetVisibilityForNode(Node);
+			Cache.Placement = GraphData->GetPlacementForNode(Node);
 			Cache.bIsAutoScaling = !Cache.bHasParent && !Widget->bIsBeingDragged && Widget->EnableAutoScale.Get();
 			Cache.Position = GetPositionForNode(Node->GetGuid());
 			const FVector2d NodeSize = GraphData->GetSizeForNode(Node->GetGuid());
@@ -1449,6 +1455,10 @@ void SSchematicGraphPanel::UpdateAutoGroupingForNodes()
 		{
 			continue;
 		}
+		if(PerNodeCaches[i].Placement != ESchematicGraphPlacementConstraint::Free)
+		{
+			continue;
+		}
 
 		TSharedRef<SSchematicGraphNode> Widget = GetChild(i);
 		const FSchematicGraphNode* Node = Widget->GetNodeData();
diff --git a/Engine/Source/Editor/AnimationEditorWidgets/Public/SchematicGraphPanel/SSchematicGraphPanel.h b/Engine/Source/Editor/AnimationEditorWidgets/Public/SchematicGraphPanel/SSchematicGraphPanel.h
--- a/Engine/Source/Editor/AnimationEditorWidgets/Public/SchematicGraphPanel/SSchematicGraphPanel.h
+++ b/Engine/Source/Editor/AnimationEditorWidgets/Public/SchematicGraphPanel/SSchematicGraphPanel.h
@@ -266,6 +266,7 @@ private:
 		FPerNodeCache()
 			: bHasParent(false)
 			, Visibility(ESchematicGraphVisibility::Visible)
+			, Placement(ESchematicGraphPlacementConstraint::Free)
 			, bIsAutoScaling(false)
 			, Position(FVector2d::ZeroVector)
 			, Radius(0.0)
@@ -273,6 +274,7 @@ private:
 
 		bool bHasParent;
 		ESchematicGraphVisibility::Type Visibility;
+		ESchematicGraphPlacementConstraint::Type Placement;
 		bool bIsAutoScaling;
 		FVector2d Position;
 		double Radius;
