Control Rig: Fix spawn elements created in wrong order when converting a rig into a module
#rnx

diff --git a/Engine/Plugins/Animation/ControlRig/Source/ControlRig/Private/Rigs/RigHierarchy.cpp b/Engine/Plugins/Animation/ControlRig/Source/ControlRig/Private/Rigs/RigHierarchy.cpp
--- a/Engine/Plugins/Animation/ControlRig/Source/ControlRig/Private/Rigs/RigHierarchy.cpp
+++ b/Engine/Plugins/Animation/ControlRig/Source/ControlRig/Private/Rigs/RigHierarchy.cpp
@@ -2537,10 +2537,11 @@ TArray<FRigElementKey> URigHierarchy::GetKeysByPredicate(
 			// TBitArray reserves 4, we'll do 16 so we can remember at least 512 elements before
 			// we need to hit the heap.
 			TBitArray<TInlineAllocator<16>> ElementVisited(false, Elements.Num());
-			
-			for (int32 ElementIndex = 0; ElementIndex < Elements.Num(); ElementIndex++)
+
+			const TArray<FRigBaseElement*> RootElements = GetRootElements();
+			for (FRigBaseElement* Element : RootElements)
 			{
-				FRigBaseElement* Element = Elements[ElementIndex];
+				const int32 ElementIndex = Element->GetIndex();
 				Traverse(Element, true, [&ElementVisited, InProcessFunc, InPredicateFunc](FRigBaseElement* InElement, bool& bContinue)
 				{
 					bContinue = !ElementVisited[InElement->GetIndex()];
