Control Rig: Cleanup of namespace splitting code and module path access code

#rnx

diff --git a/Engine/Plugins/Animation/ControlRig/Source/ControlRig/Private/ModularRig.cpp b/Engine/Plugins/Animation/ControlRig/Source/ControlRig/Private/ModularRig.cpp
--- a/Engine/Plugins/Animation/ControlRig/Source/ControlRig/Private/ModularRig.cpp
+++ b/Engine/Plugins/Animation/ControlRig/Source/ControlRig/Private/ModularRig.cpp
@@ -341,7 +341,7 @@ bool UModularRig::Execute_Internal(const FName& InEventName)
 					const bool bIsInteracting = UnitContext.ElementsBeingInteracted.ContainsByPredicate(
 						[ModuleNamespace, Hierarchy](const FRigElementKey& InteractionElement)
 						{
-							return ModuleNamespace == Hierarchy->GetNameMetadata(InteractionElement, URigHierarchy::NameSpaceMetadataName, NAME_None);
+							return ModuleNamespace == Hierarchy->GetNameSpace(InteractionElement);
 						});
 					if (!bIsInteracting)
 					{
@@ -454,7 +454,7 @@ void UModularRig::ExecuteQueue()
 				RigUnitContext.ElementsBeingInteracted = RigUnitContext.ElementsBeingInteracted.FilterByPredicate(
 					[ModuleNamespace, Hierarchy](const FRigElementKey& Key)
 				{
-					return ModuleNamespace == Hierarchy->GetNameMetadata(Key, URigHierarchy::NameSpaceMetadataName, NAME_None);
+					return ModuleNamespace == Hierarchy->GetNameSpace(Key);
 				});
 				RigUnitContext.InteractionType = RigUnitContext.ElementsBeingInteracted.IsEmpty() ?
 					(uint8) EControlRigInteractionType::None
@@ -732,7 +732,7 @@ bool UModularRig::SetModuleVariableBindings(const FString& InModulePath, const T
 		for (const TPair<FName, FString>& Pair : InVariableBindings)
 		{
 			FString SourceModulePath, SourceVariableName = Pair.Value;
-			Pair.Value.Split(NamespaceSeparator, &SourceModulePath, &SourceVariableName, ESearchCase::CaseSensitive, ESearchDir::FromEnd);
+			(void)URigHierarchy::SplitNameSpace(Pair.Value, &SourceModulePath, &SourceVariableName);
 			FRigVMExternalVariable SourceVariable;
 			if (SourceModulePath.IsEmpty())
 			{
diff --git a/Engine/Plugins/Animation/ControlRig/Source/ControlRig/Private/ModularRigController.cpp b/Engine/Plugins/Animation/ControlRig/Source/ControlRig/Private/ModularRigController.cpp
--- a/Engine/Plugins/Animation/ControlRig/Source/ControlRig/Private/ModularRigController.cpp
+++ b/Engine/Plugins/Animation/ControlRig/Source/ControlRig/Private/ModularRigController.cpp
@@ -98,40 +98,18 @@ FString UModularRigController::AddModule(const FName& InModuleName, TSubclassOf<
 
 FRigModuleReference* UModularRigController::FindModule(const FString& InPath)
 {
-	FString Path = InPath;
-	Path.RemoveFromEnd(UModularRig::NamespaceSeparator);
-	
-	TArray<FRigModuleReference*>* Children = &Model->RootModules;
-	FString Left = Path, Right;
-	while (Left.Split(UModularRig::NamespaceSeparator, &Left, &Right))
-	{
-		FRigModuleReference** Cur = Children->FindByPredicate([Left](FRigModuleReference* Module)
-		{
-			return Module->Name == Left;
-		});
-		if (!Cur)
-		{
-			return nullptr;
-		}
-		Children = &(*Cur)->CachedChildren;
-		Left = Right;
-	}
+	return Model->FindModule(InPath);
+}
 
-	FRigModuleReference** Cur = Children->FindByPredicate([Left](FRigModuleReference* Module)
-		{
-			return Module->Name == Left;
-		});
-	if (!Cur)
-	{
-		return nullptr;
-	}
-	return *Cur;
+const FRigModuleReference* UModularRigController::FindModule(const FString& InPath) const
+{
+	return const_cast<UModularRigController*>(this)->FindModule(InPath);
 }
 
 bool UModularRigController::CanConnectConnectorToElement(const FRigElementKey& InConnectorKey, const FRigElementKey& InTargetKey, FText& OutErrorMessage)
 {
 	FString ConnectorModulePath, ConnectorName;
-	if (!InConnectorKey.Name.ToString().Split(UModularRig::NamespaceSeparator, &ConnectorModulePath, &ConnectorName, ESearchCase::CaseSensitive, ESearchDir::FromEnd))
+	if (!URigHierarchy::SplitNameSpace(InConnectorKey.Name.ToString(), &ConnectorModulePath, &ConnectorName))
 	{
 		OutErrorMessage = FText::FromString(FString::Printf(TEXT("Connector %s does not contain a namespace"), *InConnectorKey.ToString()));
 		return false;
@@ -254,7 +232,7 @@ bool UModularRigController::ConnectConnectorToElement(const FRigElementKey& InCo
 	}
 	
 	FString ConnectorParentPath, ConnectorName;
-	InConnectorKey.Name.ToString().Split(UModularRig::NamespaceSeparator, &ConnectorParentPath, &ConnectorName, ESearchCase::CaseSensitive, ESearchDir::FromEnd);
+	(void)URigHierarchy::SplitNameSpace(InConnectorKey.Name.ToString(), &ConnectorParentPath, &ConnectorName);
 	FRigModuleReference* Module = FindModule(ConnectorParentPath);
 
 	FRigElementKey CurrentTarget = Model->Connections.FindTargetFromConnector(InConnectorKey);
@@ -347,6 +325,8 @@ bool UModularRigController::ConnectConnectorToElement(const FRigElementKey& InCo
 	}
 #endif
 
+	(void)DisconnectCyclicConnectors();
+
 #if WITH_EDITOR
 	TransactionPtr.Reset();
 #endif
@@ -357,7 +337,7 @@ bool UModularRigController::ConnectConnectorToElement(const FRigElementKey& InCo
 bool UModularRigController::DisconnectConnector(const FRigElementKey& InConnectorKey, bool bSetupUndo)
 {
 	FString ConnectorModulePath, ConnectorName;
-	if (!InConnectorKey.Name.ToString().Split(UModularRig::NamespaceSeparator, &ConnectorModulePath, &ConnectorName, ESearchCase::CaseSensitive, ESearchDir::FromEnd))
+	if (!URigHierarchy::SplitNameSpace(InConnectorKey.Name.ToString(), &ConnectorModulePath, &ConnectorName))
 	{
 		UE_LOG(LogControlRig, Error, TEXT("Connector %s does not contain a namespace"), *InConnectorKey.ToString());
 		return false;
@@ -389,13 +369,6 @@ bool UModularRigController::DisconnectConnector(const FRigElementKey& InConnecto
 	}
 
 	UBlueprint* Blueprint = Cast<UBlueprint>(GetOuter());
-	const IRigHierarchyProvider* HierarchyProvider = CastChecked<IRigHierarchyProvider>(Blueprint);
-	const FRigConnectorElement* Connector = Cast<FRigConnectorElement>(HierarchyProvider->GetHierarchy()->Find(InConnectorKey));
-	if (!Connector)
-	{
-		UE_LOG(LogControlRig, Error, TEXT("Could not find connector %s"), *InConnectorKey.ToString());
-		return false;
-	}
 
 	if(!Model->Connections.HasConnection(InConnectorKey))
 	{
@@ -413,7 +386,7 @@ bool UModularRigController::DisconnectConnector(const FRigElementKey& InConnecto
 
 	Model->Connections.RemoveConnection(InConnectorKey);
 
-	if (Connector->IsPrimary())
+	if (ModuleConnector->IsPrimary())
 	{
 		// Remove connections from module and child modules
 		TArray<FRigElementKey> ConnectionsToRemove;
@@ -429,14 +402,14 @@ bool UModularRigController::DisconnectConnector(const FRigElementKey& InConnecto
 			Model->Connections.RemoveConnection(ToRemove);
 		}
 	}
-	else if (!Connector->Settings.bOptional)
+	else if (!ModuleConnector->IsOptional())
 	{
 		// Remove connections from child modules
 		TArray<FRigElementKey> ConnectionsToRemove;
 		for (const FModularRigSingleConnection& Connection : Model->Connections.ConnectionList)
 		{
 			FString OtherConnectorModulePath, OtherConnectorName;
-			Connection.Connector.Name.ToString().Split(UModularRig::NamespaceSeparator, &OtherConnectorModulePath, &OtherConnectorName, ESearchCase::CaseSensitive, ESearchDir::FromEnd);
+			(void)URigHierarchy::SplitNameSpace(Connection.Connector.Name.ToString(), &OtherConnectorModulePath, &OtherConnectorName);
 			if (OtherConnectorModulePath.StartsWith(ConnectorModulePath, ESearchCase::CaseSensitive) && OtherConnectorModulePath.Len() > ConnectorModulePath.Len())
 			{
 				ConnectionsToRemove.Add(Connection.Connector);
@@ -450,6 +423,12 @@ bool UModularRigController::DisconnectConnector(const FRigElementKey& InConnecto
 
 	// todo: Make sure all the rest of the connections are still valid
 
+	// un-parent the module if we've disconnected the primary
+	if(ModuleConnector->IsPrimary() && !Module->IsRootModule())
+	{
+		(void)ReparentModule(Module->GetPath(), FString(), bSetupUndo);
+	}
+
 	Notify(EModularRigNotification::ConnectionChanged, Module);
 
 #if WITH_EDITOR
@@ -459,6 +438,63 @@ bool UModularRigController::DisconnectConnector(const FRigElementKey& InConnecto
 	return true;
 }
 
+TArray<FRigElementKey> UModularRigController::DisconnectCyclicConnectors(bool bSetupUndo)
+{
+	TArray<FRigElementKey> DisconnectedConnectors;
+
+#if WITH_EDITOR
+	const UBlueprint* Blueprint = Cast<UBlueprint>(GetOuter());
+	check(Blueprint);
+
+	const UModularRig* ModularRig = Cast<UModularRig>(Blueprint->GetObjectBeingDebugged());
+	if (!ModularRig)
+	{
+		return DisconnectedConnectors;
+	}
+	
+	const URigHierarchy* Hierarchy = ModularRig->GetHierarchy();
+	if (!Hierarchy)
+	{
+		return DisconnectedConnectors;
+	}
+
+	TArray<FRigElementKey> ConnectorsToDisconnect;
+	for (const FModularRigSingleConnection& Connection : Model->Connections.ConnectionList)
+	{
+		const FString ConnectorModulePath = Hierarchy->GetModulePath(Connection.Connector);
+		const FString TargetModulePath = Hierarchy->GetModulePath(Connection.Target);
+
+		// targets in the base hierarchy are always allowed
+		if(TargetModulePath.IsEmpty())
+		{
+			continue;
+		}
+
+		const FRigModuleReference* ConnectorModule = Model->FindModule(ConnectorModulePath);
+		const FRigModuleReference* TargetModule = Model->FindModule(TargetModulePath);
+		if(ConnectorModule == nullptr || TargetModule == nullptr)
+		{
+			continue;
+		}
+
+		if(!Model->IsModuleParentedTo(ConnectorModule, TargetModule))
+		{
+			ConnectorsToDisconnect.Add(Connection.Connector);
+		}
+	}
+
+	for(const FRigElementKey& ConnectorToDisconnect : ConnectorsToDisconnect)
+	{
+		if(DisconnectConnector(ConnectorToDisconnect, bSetupUndo))
+		{
+			DisconnectedConnectors.Add(ConnectorToDisconnect);
+		}
+	}
+#endif
+
+	return DisconnectedConnectors;
+}
+
 bool UModularRigController::SetConfigValueInModule(const FString& InModulePath, const FName& InVariableName, const FString& InValue, bool bSetupUndo)
 {
 	FRigModuleReference* Module = FindModule(InModulePath);
@@ -614,7 +650,7 @@ bool UModularRigController::CanBindModuleVariable(const FString& InModulePath, c
 	}
 
 	FString SourceModulePath, SourceVariableName = InSourcePath;
-	InSourcePath.Split(UModularRig::NamespaceSeparator, &SourceModulePath, &SourceVariableName, ESearchCase::CaseSensitive, ESearchDir::FromEnd);
+	(void)URigHierarchy::SplitNameSpace(InSourcePath, &SourceModulePath, &SourceVariableName);
 
 	FRigModuleReference* SourceModule = nullptr;
 	if (!SourceModulePath.IsEmpty())
@@ -675,7 +711,7 @@ bool UModularRigController::BindModuleVariable(const FString& InModulePath, cons
 	const FProperty* TargetProperty = Module->Class->FindPropertyByName(InVariableName);
 
 	FString SourceModulePath, SourceVariableName = InSourcePath;
-	InSourcePath.Split(UModularRig::NamespaceSeparator, &SourceModulePath, &SourceVariableName, ESearchCase::CaseSensitive, ESearchDir::FromEnd);
+	(void)URigHierarchy::SplitNameSpace(InSourcePath, &SourceModulePath, &SourceVariableName);
 
 	FRigModuleReference* SourceModule = nullptr;
 	if (!SourceModulePath.IsEmpty())
@@ -799,7 +835,7 @@ bool UModularRigController::DeleteModule(const FString& InModulePath, bool bSetu
 		for (FModularRigSingleConnection& Connection : Model->Connections.ConnectionList)
 		{
 			FString ConnectionModulePath, ConnectionName;
-			Connection.Connector.Name.ToString().Split(UModularRig::NamespaceSeparator, &ConnectionModulePath, &ConnectionName, ESearchCase::CaseSensitive, ESearchDir::FromEnd);
+			(void)URigHierarchy::SplitNameSpace(Connection.Connector.Name.ToString(), &ConnectionModulePath, &ConnectionName);
 
 			if (ConnectionModulePath == InModulePath)
 			{
@@ -807,7 +843,7 @@ bool UModularRigController::DeleteModule(const FString& InModulePath, bool bSetu
 			}
 
 			FString TargetModulePath, TargetName;
-			Connection.Target.Name.ToString().Split(UModularRig::NamespaceSeparator, &TargetModulePath, &TargetName, ESearchCase::CaseSensitive, ESearchDir::FromEnd);
+			(void)URigHierarchy::SplitNameSpace(Connection.Target.Name.ToString(), &TargetModulePath, &TargetName);
 			if (TargetModulePath == InModulePath)
 			{
 				ToRemove.Add(Connection.Connector);
@@ -825,7 +861,7 @@ bool UModularRigController::DeleteModule(const FString& InModulePath, bool bSetu
 		Reference.Bindings = Reference.Bindings.FilterByPredicate([InModulePath](const TPair<FName, FString>& Binding)
 		{
 			FString ModulePath, VariableName = Binding.Value;
-			Binding.Value.Split(UModularRig::NamespaceSeparator, &ModulePath, &VariableName, ESearchCase::CaseSensitive, ESearchDir::FromEnd);
+			(void)URigHierarchy::SplitNameSpace(Binding.Value, &ModulePath, &VariableName);
 			if (ModulePath == InModulePath)
 			{
 				return false;
@@ -928,7 +964,7 @@ FString UModularRigController::RenameModule(const FString& InModulePath, const F
 		for (TPair<FName, FString>& Binding : Reference.Bindings)
 		{
 			FString ModulePath, VariableName = Binding.Value;
-			Binding.Value.Split(UModularRig::NamespaceSeparator, &ModulePath, &VariableName, ESearchCase::CaseSensitive, ESearchDir::FromEnd);
+			(void)URigHierarchy::SplitNameSpace(Binding.Value, &ModulePath, &VariableName);
 			if (ModulePath == OldPath)
 			{
 				Binding.Value = URigHierarchy::JoinNameSpace(NewPath, VariableName);
@@ -1060,7 +1096,7 @@ FString UModularRigController::ReparentModule(const FString& InModulePath, const
 		for (TPair<FName, FString>& Binding : Reference.Bindings)
 		{
 			FString ModulePath, VariableName = Binding.Value;
-			Binding.Value.Split(UModularRig::NamespaceSeparator, &ModulePath, &VariableName, ESearchCase::CaseSensitive, ESearchDir::FromEnd);
+			(void)URigHierarchy::SplitNameSpace(Binding.Value, &ModulePath, &VariableName);
 			if (ModulePath == OldPath)
 			{
 				Binding.Value = URigHierarchy::JoinNameSpace(NewPath, VariableName);
@@ -1081,6 +1117,12 @@ FString UModularRigController::ReparentModule(const FString& InModulePath, const
 		});
 	}
 
+	// Fix connectors in the hierarchies
+	
+
+	// since we've reparented the module now we should clear out all connectors which are cyclic
+	(void)DisconnectCyclicConnectors(bSetupUndo);
+
 	Notify(EModularRigNotification::ModuleReparented, Module);
 	
 	return NewPath;
@@ -1316,7 +1358,7 @@ void UModularRigController::UpdateShortNames()
 		{
 			FString RemainingPath = Module.GetPath();
 			TokenToCount.FindOrAdd(RemainingPath, 0)++;
-			while(RemainingPath.Split(UModularRig::NamespaceSeparator, nullptr, &RemainingPath, ESearchCase::IgnoreCase, ESearchDir::FromStart))
+			while(URigHierarchy::SplitNameSpace(RemainingPath, nullptr, &RemainingPath, false))
 			{
 				TokenToCount.FindOrAdd(RemainingPath, 0)++;
 			}
@@ -1337,7 +1379,7 @@ void UModularRigController::UpdateShortNames()
 				FString Left, Right, RemainingPath = Module.GetPath();
 				ShortPath.Reset();
 
-				while (RemainingPath.Split(UModularRig::NamespaceSeparator, &Left, &Right, ESearchCase::IgnoreCase, ESearchDir::FromEnd))
+				while(URigHierarchy::SplitNameSpace(RemainingPath, &Left, &Right))
 				{
 					ShortPath = ShortPath.IsEmpty() ? Right : URigHierarchy::JoinNameSpace(Right, ShortPath);
 
diff --git a/Engine/Plugins/Animation/ControlRig/Source/ControlRig/Private/ModularRigModel.cpp b/Engine/Plugins/Animation/ControlRig/Source/ControlRig/Private/ModularRigModel.cpp
--- a/Engine/Plugins/Animation/ControlRig/Source/ControlRig/Private/ModularRigModel.cpp
+++ b/Engine/Plugins/Animation/ControlRig/Source/ControlRig/Private/ModularRigModel.cpp
@@ -21,7 +21,7 @@ FString FRigModuleReference::GetPath() const
 	{
 		return Name.ToString();
 	}
-	return ParentPath + UModularRig::NamespaceSeparator + Name.ToString();
+	return URigHierarchy::JoinNameSpace(ParentPath, Name.ToString());
 }
 
 FString FRigModuleReference::GetNamespace() const
@@ -81,7 +81,7 @@ TMap<FRigElementKey, FRigElementKey> FModularRigConnections::GetModuleConnection
 	for (const FModularRigSingleConnection& Connection : ConnectionList)
 	{
 		FString Path, Name;
-		Connection.Connector.Name.ToString().Split(UModularRig::NamespaceSeparator, &Path, &Name, ESearchCase::CaseSensitive, ESearchDir::FromEnd);
+		URigHierarchy::SplitNameSpace(Connection.Connector.Name.ToString(), &Path, &Name);
 
 		// Exactly the same path (do not return connectors from child modules)
 		if (Path == InModulePath)
@@ -156,12 +156,12 @@ void FModularRigModel::UpdateCachedChildren()
 	}
 }
 
-FRigModuleReference* FModularRigModel::FindModule(const FString InPath) const
+FRigModuleReference* FModularRigModel::FindModule(const FString& InPath)
 {
 	const TArray<FRigModuleReference*>* Children = &RootModules;
 
 	FString Left = InPath, Right;
-	while (Left.Split(UModularRig::NamespaceSeparator, &Left, &Right))
+	while (URigHierarchy::SplitNameSpace(Left, &Left, &Right, false))
 	{
 		FRigModuleReference* const * Child = Children->FindByPredicate([Left](FRigModuleReference* Module)
 		{
@@ -189,9 +189,80 @@ FRigModuleReference* FModularRigModel::FindModule(const FString InPath) const
 	return *Child;
 }
 
-FString FModularRigModel::FindParentPath(const FString InPath) const
+const FRigModuleReference* FModularRigModel::FindModule(const FString& InPath) const
 {
-	if (FRigModuleReference* Element = FindModule(InPath))
+	FRigModuleReference* Module = const_cast<FModularRigModel*>(this)->FindModule(InPath);
+	return const_cast<FRigModuleReference*>(Module);
+}
+
+FRigModuleReference* FModularRigModel::GetParentModule(const FString& InPath)
+{
+	const FString ParentPath = GetParentPath(InPath);
+	if(!ParentPath.IsEmpty())
+	{
+		return FindModule(ParentPath);
+	}
+	return nullptr;
+}
+
+const FRigModuleReference* FModularRigModel::GetParentModule(const FString& InPath) const
+{
+	FRigModuleReference* Module = const_cast<FModularRigModel*>(this)->GetParentModule(InPath);
+	return const_cast<FRigModuleReference*>(Module);
+}
+
+FRigModuleReference* FModularRigModel::GetParentModule(const FRigModuleReference* InChildModule)
+{
+	if(InChildModule)
+	{
+		return FindModule(InChildModule->ParentPath);
+	}
+	return nullptr;
+}
+
+const FRigModuleReference* FModularRigModel::GetParentModule(const FRigModuleReference* InChildModule) const
+{
+	FRigModuleReference* Module = const_cast<FModularRigModel*>(this)->GetParentModule(InChildModule);
+	return const_cast<FRigModuleReference*>(Module);
+}
+
+bool FModularRigModel::IsModuleParentedTo(const FString& InChildModulePath, const FString& InParentModulePath) const
+{
+	const FRigModuleReference* ChildModule = FindModule(InChildModulePath);
+	const FRigModuleReference* ParentModule = FindModule(InParentModulePath);
+	if(ChildModule == nullptr || ParentModule == nullptr)
+	{
+		return false;
+	}
+	return IsModuleParentedTo(ChildModule, ParentModule);
+}
+
+bool FModularRigModel::IsModuleParentedTo(const FRigModuleReference* InChildModule, const FRigModuleReference* InParentModule) const
+{
+	if(InChildModule == nullptr || InParentModule == nullptr)
+	{
+		return false;
+	}
+	if(InChildModule == InParentModule)
+	{
+		return false;
+	}
+	
+	while(InChildModule)
+	{
+		if(InChildModule == InParentModule)
+		{
+			return true;
+		}
+		InChildModule = GetParentModule(InChildModule);
+	}
+
+	return false;
+}
+
+FString FModularRigModel::GetParentPath(const FString& InPath) const
+{
+	if (const FRigModuleReference* Element = FindModule(InPath))
 	{
 		return Element->ParentPath;
 	}
diff --git a/Engine/Plugins/Animation/ControlRig/Source/ControlRig/Private/ModularRigRuleManager.cpp b/Engine/Plugins/Animation/ControlRig/Source/ControlRig/Private/ModularRigRuleManager.cpp
--- a/Engine/Plugins/Animation/ControlRig/Source/ControlRig/Private/ModularRigRuleManager.cpp
+++ b/Engine/Plugins/Animation/ControlRig/Source/ControlRig/Private/ModularRigRuleManager.cpp
@@ -175,7 +175,7 @@ void UModularRigRuleManager::FilterIncompatibleTypes(FWorkData& InOutWorkData)
 
 void UModularRigRuleManager::FilterInvalidNameSpaces(FWorkData& InOutWorkData)
 {
-	const FName NameSpace = InOutWorkData.Hierarchy->GetNameMetadata(InOutWorkData.Connector->GetKey(), URigHierarchy::NameSpaceMetadataName, NAME_None);
+	const FName NameSpace = InOutWorkData.Hierarchy->GetNameSpaceFName(InOutWorkData.Connector->GetKey());
 	if(NameSpace.IsNone())
 	{
 		return;
@@ -184,7 +184,7 @@ void UModularRigRuleManager::FilterInvalidNameSpaces(FWorkData& InOutWorkData)
 	const FString NameSpaceString = NameSpace.ToString(); 
 	InOutWorkData.Filter([NameSpaceString, InOutWorkData](FRigElementResolveResult& Result)
 	{
-		const FName MatchNameSpace = InOutWorkData.Hierarchy->GetNameMetadata(Result.GetKey(), URigHierarchy::NameSpaceMetadataName, NAME_None);
+		const FName MatchNameSpace = InOutWorkData.Hierarchy->GetNameSpaceFName(Result.GetKey());
 		if(!MatchNameSpace.IsNone())
 		{
 			const FString MatchNameSpaceString = MatchNameSpace.ToString();
diff --git a/Engine/Plugins/Animation/ControlRig/Source/ControlRig/Private/Rigs/RigConnectionRules.cpp b/Engine/Plugins/Animation/ControlRig/Source/ControlRig/Private/Rigs/RigConnectionRules.cpp
--- a/Engine/Plugins/Animation/ControlRig/Source/ControlRig/Private/Rigs/RigConnectionRules.cpp
+++ b/Engine/Plugins/Animation/ControlRig/Source/ControlRig/Private/Rigs/RigConnectionRules.cpp
@@ -115,7 +115,7 @@ const FRigConnectorElement* FRigConnectionRuleInput::FindPrimaryConnector(FText*
 		{
 			if(Connector->IsPrimary())
 			{
-				const FName ConnectorNameSpace = Hierarchy->GetNameMetadata(Connector->GetKey(), URigHierarchy::NameSpaceMetadataName, NAME_None);
+				const FName ConnectorNameSpace = Hierarchy->GetNameSpaceFName(Connector->GetKey());
 				if(!ConnectorNameSpace.IsNone() && ConnectorNameSpace.IsEqual(ModuleNameSpace, ENameCase::CaseSensitive))
 				{
 					PrimaryConnector = Connector;
@@ -147,11 +147,11 @@ TArray<const FRigConnectorElement*> FRigConnectionRuleInput::FindSecondaryConnec
 
 	TArray<const FRigConnectorElement*> SecondaryConnectors;
 	Hierarchy->ForEach<FRigConnectorElement>(
-		[this, bOptional, ModuleNameSpace, &SecondaryConnectors](FRigConnectorElement* Connector) -> bool
+		[this, bOptional, ModuleNameSpace, &SecondaryConnectors](const FRigConnectorElement* Connector) -> bool
 		{
-			if(Connector->IsOptional())
+			if(Connector->IsSecondary() && Connector->IsOptional() == bOptional)
 			{
-				const FName ConnectorNameSpace = Hierarchy->GetNameMetadata(Connector->GetKey(), URigHierarchy::NameSpaceMetadataName, NAME_None);
+				const FName ConnectorNameSpace = Hierarchy->GetNameSpaceFName(Connector->GetKey());
 				if(!ConnectorNameSpace.IsNone() && ConnectorNameSpace.IsEqual(ModuleNameSpace, ENameCase::CaseSensitive))
 				{
 					SecondaryConnectors.Add(Connector);
diff --git a/Engine/Plugins/Animation/ControlRig/Source/ControlRig/Private/Rigs/RigHierarchy.cpp b/Engine/Plugins/Animation/ControlRig/Source/ControlRig/Private/Rigs/RigHierarchy.cpp
--- a/Engine/Plugins/Animation/ControlRig/Source/ControlRig/Private/Rigs/RigHierarchy.cpp
+++ b/Engine/Plugins/Animation/ControlRig/Source/ControlRig/Private/Rigs/RigHierarchy.cpp
@@ -1181,8 +1181,29 @@ bool URigHierarchy::RemoveAllMetadata(FRigElementKey InItem)
 
 FName URigHierarchy::GetModulePathFName(FRigElementKey InItem) const
 {
-	return GetNameMetadata(InItem, ModuleMetadataName, NAME_None);
+	if(!InItem.IsValid())
+	{
+		return NAME_None;
+	}
 	
+	const FName Result = GetNameMetadata(InItem, ModuleMetadataName, NAME_None);
+	if(!Result.IsNone())
+	{
+		return Result;
+	}
+
+	// fall back on the name of the item
+	const FString NameString = InItem.Name.ToString();
+	FString ModulePathFromName;
+	if(SplitNameSpace(NameString, &ModulePathFromName, nullptr))
+	{
+		if(!ModulePathFromName.IsEmpty())
+		{
+			return *ModulePathFromName;
+		}
+	}
+
+	return NAME_None;
 }
 
 FString URigHierarchy::GetModulePath(FRigElementKey InItem) const
@@ -1195,6 +1216,43 @@ FString URigHierarchy::GetModulePath(FRigElementKey InItem) const
 	return FString();
 }
 
+FName URigHierarchy::GetNameSpaceFName(FRigElementKey InItem) const
+{
+	if(!InItem.IsValid())
+	{
+		return NAME_None;
+	}
+	
+	const FName Result = GetNameMetadata(InItem, NameSpaceMetadataName, NAME_None);
+	if(!Result.IsNone())
+	{
+		return Result;
+	}
+
+	// fall back on the name of the item
+	const FString NameString = InItem.Name.ToString();
+	FString NameSpaceFromName;
+	if(SplitNameSpace(NameString, &NameSpaceFromName, nullptr))
+	{
+		if(!NameSpaceFromName.IsEmpty())
+		{
+			return *(NameSpaceFromName + UModularRig::NamespaceSeparator);
+		}
+	}
+	
+	return NAME_None;
+}
+
+FString URigHierarchy::GetNameSpace(FRigElementKey InItem) const
+{
+	const FName NameSpaceName = GetNameSpaceFName(InItem);
+	if(!NameSpaceName.IsNone())
+	{
+		return NameSpaceName.ToString();
+	}
+	return FString();
+}
+
 TArray<const FRigBaseElement*> URigHierarchy::GetSelectedElements(ERigElementType InTypeFilter) const
 {
 	LLM_SCOPE_BYNAME(TEXT("Animation/ControlRig"));
@@ -1262,6 +1320,42 @@ FString URigHierarchy::JoinNameSpace(const FString& InLeft, const FString& InRig
 	return InLeft + UModularRig::NamespaceSeparator + InRight;
 }
 
+TPair<FString, FString> URigHierarchy::SplitNameSpace(const FString& InNameSpacedPath, bool bFromEnd)
+{
+	TPair<FString, FString> Result;
+	(void)SplitNameSpace(InNameSpacedPath, &Result.Key, &Result.Value, bFromEnd);
+	return Result;
+}
+
+TPair<FRigName, FRigName> URigHierarchy::SplitNameSpace(const FRigName& InNameSpacedPath, bool bFromEnd)
+{
+	const TPair<FString, FString> Result = SplitNameSpace(InNameSpacedPath.GetName(), bFromEnd);
+	return {FRigName(Result.Key), FRigName(Result.Value)};
+}
+
+bool URigHierarchy::SplitNameSpace(const FString& InNameSpacedPath, FString* OutNameSpace, FString* OutName, bool bFromEnd)
+{
+	return InNameSpacedPath.Split(UModularRig::NamespaceSeparator, OutNameSpace, OutName, ESearchCase::CaseSensitive, bFromEnd ? ESearchDir::FromEnd : ESearchDir::FromStart);
+}
+
+bool URigHierarchy::SplitNameSpace(const FRigName& InNameSpacedPath, FRigName* OutNameSpace, FRigName* OutName, bool bFromEnd)
+{
+	FString NameSpace, Name;
+	if(SplitNameSpace(InNameSpacedPath.GetName(), &NameSpace, &Name, bFromEnd))
+	{
+		if(OutNameSpace)
+		{
+			OutNameSpace->SetName(NameSpace);
+		}
+		if(OutName)
+		{
+			OutName->SetName(Name);
+		}
+		return true;
+	}
+	return false;
+}
+
 void URigHierarchy::SanitizeName(FRigName& InOutName, bool bAllowNameSpaces)
 {
 	// Sanitize the name
@@ -1544,7 +1638,7 @@ FText URigHierarchy::GetDisplayNameForUI(const FRigBaseElement* InElement, bool
 
 	const FName& DisplayName = InElement->GetDisplayName();
 	FString DisplayNameString = DisplayName.ToString();
-	(void)DisplayNameString.Split(UModularRig::NamespaceSeparator, nullptr, &DisplayNameString, ESearchCase::CaseSensitive, ESearchDir::FromEnd);
+	(void)SplitNameSpace(DisplayNameString, nullptr, &DisplayNameString);
 	
 	if(bIncludeNameSpace)
 	{
diff --git a/Engine/Plugins/Animation/ControlRig/Source/ControlRig/Private/Rigs/RigHierarchyController.cpp b/Engine/Plugins/Animation/ControlRig/Source/ControlRig/Private/Rigs/RigHierarchyController.cpp
--- a/Engine/Plugins/Animation/ControlRig/Source/ControlRig/Private/Rigs/RigHierarchyController.cpp
+++ b/Engine/Plugins/Animation/ControlRig/Source/ControlRig/Private/Rigs/RigHierarchyController.cpp
@@ -612,8 +612,8 @@ FRigElementKey URigHierarchyController::AddConnector(FName InName, FRigConnector
 					const FString NameSpace = CRContext.GetRigModuleNameSpace();
 					if(!NameSpace.IsEmpty())
 					{
-						const FName& ConnectorNameSpace = Hierarchy->GetNameMetadata(Connector->GetKey(), URigHierarchy::NameSpaceMetadataName, NAME_None);
-						if(!ConnectorNameSpace.IsNone() && ConnectorNameSpace.ToString().Equals(NameSpace, ESearchCase::CaseSensitive))
+						const FString ConnectorNameSpace = Hierarchy->GetNameSpace(Connector->GetKey());
+						if(!ConnectorNameSpace.IsEmpty() && ConnectorNameSpace.Equals(NameSpace, ESearchCase::CaseSensitive))
 						{
 							static constexpr TCHAR Format[] = TEXT("Cannot add connector '%s' - there already is a primary connector.");
 							ReportAndNotifyErrorf(Format, *InName.ToString());
diff --git a/Engine/Plugins/Animation/ControlRig/Source/ControlRig/Private/Units/Execution/RigUnit_RigModules.cpp b/Engine/Plugins/Animation/ControlRig/Source/ControlRig/Private/Units/Execution/RigUnit_RigModules.cpp
--- a/Engine/Plugins/Animation/ControlRig/Source/ControlRig/Private/Units/Execution/RigUnit_RigModules.cpp
+++ b/Engine/Plugins/Animation/ControlRig/Source/ControlRig/Private/Units/Execution/RigUnit_RigModules.cpp
@@ -53,10 +53,10 @@ FRigUnit_GetItemNameSpace_Execute()
 
 	if(const URigHierarchy* Hierarchy = ExecuteContext.Hierarchy)
 	{
-		const FName NameSpaceName = Hierarchy->GetNameMetadata(Item, URigHierarchy::NameSpaceMetadataName, NAME_None);
-		if(!NameSpaceName.IsNone())
+		const FString NameSpaceForItem = Hierarchy->GetNameSpace(Item);
+		if(!NameSpaceForItem.IsEmpty())
 		{
-			NameSpace = NameSpaceName.ToString();
+			NameSpace = NameSpaceForItem;
 			HasNameSpace = true;
 		}
 	}
@@ -109,7 +109,7 @@ FRigUnit_GetItemsInNameSpace_Execute()
 				const FRigElementKey Key = InElement->GetKey();
 				if(((uint8)TypeToSearch & (uint8)Key.Type) == (uint8)Key.Type)
 				{
-					const FName ItemNameSpace = Hierarchy->GetNameMetadata(Key, URigHierarchy::NameSpaceMetadataName, NAME_None);
+					const FName ItemNameSpace = Hierarchy->GetNameSpaceFName(Key);
 					if(!ItemNameSpace.IsNone())
 					{
 						if(ItemNameSpace.IsEqual(NameSpaceName, ENameCase::CaseSensitive))
diff --git a/Engine/Plugins/Animation/ControlRig/Source/ControlRig/Public/ModularRigController.h b/Engine/Plugins/Animation/ControlRig/Source/ControlRig/Public/ModularRigController.h
--- a/Engine/Plugins/Animation/ControlRig/Source/ControlRig/Public/ModularRigController.h
+++ b/Engine/Plugins/Animation/ControlRig/Source/ControlRig/Public/ModularRigController.h
@@ -35,6 +35,9 @@ class CONTROLRIG_API UModularRigController : public UObject
 
 	UFUNCTION(BlueprintCallable, Category = "Control Rig | Modules")
 	bool DisconnectConnector(const FRigElementKey& InConnectorKey, bool bSetupUndo = true);
+	
+	UFUNCTION(BlueprintCallable, Category = "Control Rig | Modules")
+	TArray<FRigElementKey> DisconnectCyclicConnectors(bool bSetupUndo = true);
 
 	UFUNCTION(BlueprintCallable, Category = "Control Rig | Modules")
 	bool SetConfigValueInModule(const FString& InModulePath, const FName& InVariableName, const FString& InValue, bool bSetupUndo = true);
@@ -69,6 +72,7 @@ class CONTROLRIG_API UModularRigController : public UObject
 	FRigName GetSafeNewName(const FString& InParentModulePath, const FRigName& InDesiredName) const;
 	FRigName GetSafeNewShortName(const FRigName& InDesiredShortName) const;
 	FRigModuleReference* FindModule(const FString& InPath);
+	const FRigModuleReference* FindModule(const FString& InPath) const;
 	FModularRigModifiedEvent& OnModified() { return ModifiedEvent; }
 
 private:
diff --git a/Engine/Plugins/Animation/ControlRig/Source/ControlRig/Public/ModularRigModel.h b/Engine/Plugins/Animation/ControlRig/Source/ControlRig/Public/ModularRigModel.h
--- a/Engine/Plugins/Animation/ControlRig/Source/ControlRig/Public/ModularRigModel.h
+++ b/Engine/Plugins/Animation/ControlRig/Source/ControlRig/Public/ModularRigModel.h
@@ -91,6 +91,10 @@ struct CONTROLRIG_API FRigModuleReference
 	FString GetPath() const;
 
 	FString GetNamespace() const;
+
+	bool HasParentModule() const { return !ParentPath.IsEmpty(); }
+
+	bool IsRootModule() const { return !HasParentModule(); }
 	
 	friend bool operator==(const FRigModuleReference& A, const FRigModuleReference& B)
 	{
@@ -258,14 +262,25 @@ public:
 
 	void UpdateCachedChildren();
 
-	FRigModuleReference* FindModule(const FString InPath) const;
+	FRigModuleReference* FindModule(const FString& InPath);
+	const FRigModuleReference* FindModule(const FString& InPath) const;
+
+	FRigModuleReference* GetParentModule(const FString& InPath);
+	const FRigModuleReference* GetParentModule(const FString& InPath) const;
 
-	FString FindParentPath(const FString InPath) const;
+	FRigModuleReference* GetParentModule(const FRigModuleReference* InChildModule);
+	const FRigModuleReference* GetParentModule(const FRigModuleReference* InChildModule) const;
+
+	FString GetParentPath(const FString& InPath) const;
 
 	void ForEachModule(TFunction<bool(const FRigModuleReference*)> PerModule) const;
 
 	TArray<FString> SortPaths(const TArray<FString>& InPaths) const;
 
+	bool IsModuleParentedTo(const FString& InChildModulePath, const FString& InParentModulePath) const;
+	
+	bool IsModuleParentedTo(const FRigModuleReference* InChildModule, const FRigModuleReference* InParentModule) const;
+
 private:
 	TWeakObjectPtr<UObject> OuterClientHost;
 
diff --git a/Engine/Plugins/Animation/ControlRig/Source/ControlRig/Public/Rigs/RigHierarchy.h b/Engine/Plugins/Animation/ControlRig/Source/ControlRig/Public/Rigs/RigHierarchy.h
--- a/Engine/Plugins/Animation/ControlRig/Source/ControlRig/Public/Rigs/RigHierarchy.h
+++ b/Engine/Plugins/Animation/ControlRig/Source/ControlRig/Public/Rigs/RigHierarchy.h
@@ -1579,6 +1579,20 @@ public:
 	UFUNCTION(BlueprintPure, Category = URigHierarchy)
 	FString GetModulePath(FRigElementKey InItem) const;
 
+	/**
+	 * Returns the namespace of an element belong to (or NAME_None in case the element doesn't belong to a module / namespace)
+	 * @return The namespace the element belongs to (or NAME_None)
+	 */
+	UFUNCTION(BlueprintPure, Category = URigHierarchy)
+	FName GetNameSpaceFName(FRigElementKey InItem) const;
+	
+	/**
+	 * Returns the namespace of an element belong to (or an empty string in case the element doesn't belong to a module / namespace)
+	 * @return The namespace the element belongs to (or empty string)
+	 */
+	UFUNCTION(BlueprintPure, Category = URigHierarchy)
+	FString GetNameSpace(FRigElementKey InItem) const;
+
 	/*
 	 * Returns the tags for a given item
 	 * @param InItem The item to return the tags for
@@ -1681,6 +1695,14 @@ public:
 	static FString JoinNameSpace(const FString& InLeft, const FString& InRight);
 	static FRigName JoinNameSpace(const FRigName& InLeft, const FRigName& InRight);
 
+	/**
+	 * Returns the two name sections with the right namespace separator
+	 */
+	static TPair<FString, FString> SplitNameSpace(const FString& InNameSpacedPath, bool bFromEnd = true);
+	static TPair<FRigName, FRigName> SplitNameSpace(const FRigName& InNameSpacedPath, bool bFromEnd = true);
+	static bool SplitNameSpace(const FString& InNameSpacedPath, FString* OutNameSpace, FString* OutName, bool bFromEnd = true);
+	static bool SplitNameSpace(const FRigName& InNameSpacedPath, FRigName* OutNameSpace, FRigName* OutName, bool bFromEnd = true);
+
 	/**
 	 * Returns the max allowed length for a name within the hierarchy.
 	 * @return Returns the max allowed length for a name within the hierarchy.
diff --git a/Engine/Plugins/Animation/ControlRig/Source/ControlRigDeveloper/Private/ControlRigBlueprint.cpp b/Engine/Plugins/Animation/ControlRig/Source/ControlRigDeveloper/Private/ControlRigBlueprint.cpp
--- a/Engine/Plugins/Animation/ControlRig/Source/ControlRigDeveloper/Private/ControlRigBlueprint.cpp
+++ b/Engine/Plugins/Animation/ControlRig/Source/ControlRigDeveloper/Private/ControlRigBlueprint.cpp
@@ -2350,8 +2350,8 @@ void UControlRigBlueprint::RefreshModuleConnectors(const FRigModuleReference* In
 			const TArray<FRigElementKey> ExistingConnectors = AllConnectors.FilterByPredicate([Namespace](const FRigElementKey& ConnectorKey) -> bool
 			{
 				const FString ConnectorName = ConnectorKey.Name.ToString();
-				FString ConnectorNamespace, ShortName = ConnectorName;
-				ConnectorName.Split(UModularRig::NamespaceSeparator, &ConnectorNamespace, &ShortName, ESearchCase::CaseSensitive, ESearchDir::FromEnd);
+				FString ConnectorNamespace;
+				(void)URigHierarchy::SplitNameSpace(ConnectorName, &ConnectorNamespace, nullptr);
 				ConnectorNamespace.Append(UModularRig::NamespaceSeparator);
 				return ConnectorNamespace.Equals(Namespace, ESearchCase::CaseSensitive);
 			});
@@ -2368,7 +2368,7 @@ void UControlRigBlueprint::RefreshModuleConnectors(const FRigModuleReference* In
 			{
 				const FString ConnectorNameString = Connector.Name.ToString();
 				FString ConnectorNamespace, ShortName = ConnectorNameString;
-				ConnectorNameString.Split(UModularRig::NamespaceSeparator, &ConnectorNamespace, &ShortName, ESearchCase::CaseSensitive, ESearchDir::FromEnd);
+				(void)URigHierarchy::SplitNameSpace(ConnectorNameString, &ConnectorNamespace, &ShortName);
 				const bool bConnectorExpected = ExpectedConnectors.ContainsByPredicate(
 					[ShortName](const FRigModuleConnector& ExpectedConnector) -> bool
 					{
@@ -2550,7 +2550,7 @@ void UControlRigBlueprint::HandleRigModulesModified(EModularRigNotification InNo
 						{
 							FString OldConnectorName = Connector.Name.ToString();
 							FString OldConnectorNamespace, ShortName;
-							OldConnectorName.Split(UModularRig::NamespaceSeparator, &OldConnectorNamespace, &ShortName, ESearchCase::CaseSensitive, ESearchDir::FromEnd);
+							(void)URigHierarchy::SplitNameSpace(OldConnectorName, &OldConnectorNamespace, &ShortName);
 							OldConnectorNamespace.Append(UModularRig::NamespaceSeparator);
 							if (OldConnectorNamespace.StartsWith(OldNamespace))
 							{
@@ -2578,7 +2578,7 @@ void UControlRigBlueprint::HandleRigModulesModified(EModularRigNotification InNo
 							for (TPair<FRigElementKey, ConnectionInfo>& Pair : RenamedConnectors)
 							{
 								FString Namespace, ConnectorName;
-								Pair.Value.NewPath.Split(UModularRig::NamespaceSeparator, &Namespace, &ConnectorName, ESearchCase::CaseSensitive, ESearchDir::FromEnd);
+								(void)URigHierarchy::SplitNameSpace(Pair.Value.NewPath, &Namespace, &ConnectorName);
 								Namespace.Append(UModularRig::NamespaceSeparator);
 								FControlRigExecuteContextRigModuleGuard RigModuleGuard(PublicContext, Namespace);
 								Controller->AddConnector(*ConnectorName, Pair.Value.Settings);
diff --git a/Engine/Plugins/Animation/ControlRig/Source/ControlRigEditor/Private/ControlRigModuleDetails.cpp b/Engine/Plugins/Animation/ControlRig/Source/ControlRigEditor/Private/ControlRigModuleDetails.cpp
--- a/Engine/Plugins/Animation/ControlRig/Source/ControlRigEditor/Private/ControlRigModuleDetails.cpp
+++ b/Engine/Plugins/Animation/ControlRig/Source/ControlRigEditor/Private/ControlRigModuleDetails.cpp
@@ -1207,7 +1207,7 @@ void FRigModuleInstanceDetails::FillBindingMenu(FMenuBuilder& MenuBuilder, const
 			for(const FString& VariablePath : Data.Variables)
 			{
 				FString VariableName = VariablePath;
-				(void)VariablePath.Split(UModularRig::NamespaceSeparator, nullptr, &VariableName, ESearchCase::CaseSensitive, ESearchDir::FromEnd);
+				(void)URigHierarchy::SplitNameSpace(VariablePath, nullptr, &VariableName);
 				
 				InMenuBuilder.AddMenuEntry(
 					FUIAction(FExecuteAction::CreateLambda([ThisDetails, InProperty, VariablePath]()
@@ -1259,11 +1259,11 @@ void FRigModuleInstanceDetails::FillBindingMenu(FMenuBuilder& MenuBuilder, const
 	for(const FString& BindingPath : CombinedBindings)
 	{
 		FString MenuPath;
-		(void)BindingPath.Split(UModularRig::NamespaceSeparator, &MenuPath, nullptr, ESearchCase::CaseSensitive, ESearchDir::FromEnd);
+		(void)URigHierarchy::SplitNameSpace(BindingPath, &MenuPath, nullptr);
 
 		FString PreviousMenuPath = MenuPath;
 		FString ParentMenuPath = MenuPath, RemainingPath;
-		while(ParentMenuPath.Split(UModularRig::NamespaceSeparator, &ParentMenuPath, &RemainingPath, ESearchCase::CaseSensitive, ESearchDir::FromEnd))
+		while(URigHierarchy::SplitNameSpace(ParentMenuPath, &ParentMenuPath, &RemainingPath))
 		{
 			// scope since the map may change at the end of this block
 			{
diff --git a/Engine/Plugins/Animation/ControlRig/Source/ControlRigEditor/Private/Editor/ControlRigSchematicModel.cpp b/Engine/Plugins/Animation/ControlRig/Source/ControlRigEditor/Private/Editor/ControlRigSchematicModel.cpp
--- a/Engine/Plugins/Animation/ControlRig/Source/ControlRigEditor/Private/Editor/ControlRigSchematicModel.cpp
+++ b/Engine/Plugins/Animation/ControlRig/Source/ControlRigEditor/Private/Editor/ControlRigSchematicModel.cpp
@@ -89,6 +89,10 @@ void FControlRigSchematicModel::Tick(float InDeltaTime)
 			{
 				continue;
 			}
+			if(TemporaryNodeGuids.Contains(Pair.Value))
+			{
+				continue;
+			}
 
 			const TArray<FRigElementKey>& Connectors = Connections.FindConnectorsFromTarget(Pair.Key);
 			if(Connectors.Num() > 1)
@@ -1236,7 +1240,7 @@ TArray<FRigElementKey> FControlRigSchematicModel::GetElementKeysFromDragDropEven
 					const FString ModuleNameSpace = ModuleInstance->GetNamespace();
 					for(const FRigConnectorElement* Connector : Connectors)
 					{
-						const FString ConnectorNameSpace = Hierarchy->GetNameMetadata(Connector->GetKey(), URigHierarchy::NameSpaceMetadataName, NAME_None).ToString();
+						const FString ConnectorNameSpace = Hierarchy->GetNameSpace(Connector->GetKey());
 						if(ConnectorNameSpace.Equals(ModuleNameSpace))
 						{
 							if(Connector->IsPrimary())
@@ -1355,8 +1359,7 @@ void FControlRigSchematicModel::HandleSchematicBeginDrag(SSchematicGraphPanel* I
 		return;
 	}
 
-	FString ModulePath = Hierarchy->GetNameMetadata(DraggedKey, URigHierarchy::NameSpaceMetadataName, NAME_None).ToString();
-	ModulePath.RemoveFromEnd(UModularRig::NamespaceSeparator);
+	const FString ModulePath = Hierarchy->GetModulePath(DraggedKey);
 	const FRigModuleInstance* ModuleInstance = ModularRig->FindModule(ModulePath);
 	if (!ModuleInstance)
 	{
@@ -1388,6 +1391,7 @@ void FControlRigSchematicModel::HandleSchematicBeginDrag(SSchematicGraphPanel* I
 	{
 		if(FControlRigSchematicRigElementKeyNode* ExistingElementKeyNode = Cast<FControlRigSchematicRigElementKeyNode>(Node.Get()))
 		{
+			PreDragVisibilityPerNode.Add(Node->GetGuid(), ExistingElementKeyNode->Visibility);
 			ExistingElementKeyNode->SetVisibility(Matches.ContainsByPredicate([ExistingElementKeyNode](const FRigElementResolveResult& Match)
 			{
 				return ExistingElementKeyNode->GetKey() == Match.GetKey();
@@ -1398,16 +1402,26 @@ void FControlRigSchematicModel::HandleSchematicBeginDrag(SSchematicGraphPanel* I
 
 void FControlRigSchematicModel::HandleSchematicEndDrag(SSchematicGraphPanel* InPanel, SSchematicGraphNode* InNode, const FDragDropOperation& InDragDropOperation)
 {
+	if(TemporaryNodeGuids.IsEmpty() && PreDragVisibilityPerNode.IsEmpty())
+	{
+		return;
+	}
+	
 	for (FGuid& TempNodeGuid : TemporaryNodeGuids)
 	{
 		RemoveNode(TempNodeGuid);
 	}
 	TemporaryNodeGuids.Reset();
 
-	for (const TSharedPtr<FSchematicGraphNode>& Node : Nodes)
+	for(const TPair<FGuid, ESchematicGraphVisibility::Type>& Pair : PreDragVisibilityPerNode)
 	{
-		Node->SetVisibility(ESchematicGraphVisibility::Visible);
+		if(FControlRigSchematicRigElementKeyNode* Node = FindNode<FControlRigSchematicRigElementKeyNode>(Pair.Key))
+		{
+			Node->Visibility = Pair.Value;
+		}
 	}
+	PreDragVisibilityPerNode.Reset();
+
 	UpdateElementKeyLinks();
 }
 
@@ -1486,7 +1500,7 @@ void FControlRigSchematicModel::HandleSchematicDrop(SSchematicGraphPanel* InPane
 								if (Connector->IsPrimary())
 								{
 									FString Path, Name;
-									Connector->GetName().Split(UModularRig::NamespaceSeparator, &Path, &Name, ESearchCase::CaseSensitive, ESearchDir::FromEnd);
+									(void)URigHierarchy::SplitNameSpace(Connector->GetName(), &Path, &Name);
 									if (Path == ModulePath)
 									{
 										PrimaryConnectorKey = Connector->GetKey();
@@ -1495,7 +1509,7 @@ void FControlRigSchematicModel::HandleSchematicDrop(SSchematicGraphPanel* InPane
 								}
 							}
 
-							const FName TargetModulePath = Hierarchy->GetNameMetadata(TargetKey, URigHierarchy::NameSpaceMetadataName, NAME_None);
+							const FName TargetModulePath = Hierarchy->GetNameSpaceFName(TargetKey);
 							if(!TargetModulePath.IsNone())
 							{
 								(void)Controller->ReparentModule(ModulePath, TargetModulePath.ToString());
@@ -1572,7 +1586,7 @@ void FControlRigSchematicModel::HandleSchematicDrop(SSchematicGraphPanel* InPane
 						const FString ModuleNameSpace = Module->GetNamespace();
 						for(const FRigConnectorElement* Connector : Connectors)
 						{
-							const FString ConnectorNameSpace = Hierarchy->GetNameMetadata(Connector->GetKey(), URigHierarchy::NameSpaceMetadataName, NAME_None).ToString();
+							const FString ConnectorNameSpace = Hierarchy->GetNameSpace(Connector->GetKey());
 							if(ConnectorNameSpace.Equals(ModuleNameSpace))
 							{
 								if(Connector->IsPrimary())
diff --git a/Engine/Plugins/Animation/ControlRig/Source/ControlRigEditor/Private/Editor/ControlRigSchematicModel.h b/Engine/Plugins/Animation/ControlRig/Source/ControlRigEditor/Private/Editor/ControlRigSchematicModel.h
--- a/Engine/Plugins/Animation/ControlRig/Source/ControlRigEditor/Private/Editor/ControlRigSchematicModel.h
+++ b/Engine/Plugins/Animation/ControlRig/Source/ControlRigEditor/Private/Editor/ControlRigSchematicModel.h
@@ -127,6 +127,7 @@ private:
 	TWeakObjectPtr<UControlRig> ControlRigBeingDebuggedPtr;
 
 	TArray<FGuid> TemporaryNodeGuids;
+	TMap<FGuid, ESchematicGraphVisibility::Type> PreDragVisibilityPerNode;
 	TMap<FRigElementKey, FGuid> RigElementKeyToGuid;
 	mutable TMap<FSoftObjectPath, FSlateBrush> ModuleIcons;
 
diff --git a/Engine/Plugins/Animation/ControlRig/Source/ControlRigEditor/Private/Editor/SModularRigModel.cpp b/Engine/Plugins/Animation/ControlRig/Source/ControlRigEditor/Private/Editor/SModularRigModel.cpp
--- a/Engine/Plugins/Animation/ControlRig/Source/ControlRigEditor/Private/Editor/SModularRigModel.cpp
+++ b/Engine/Plugins/Animation/ControlRig/Source/ControlRigEditor/Private/Editor/SModularRigModel.cpp
@@ -582,7 +582,7 @@ bool SModularRigModel::HandleVerifyNameChanged(const FString& InOldPath, const F
 
 	FString ParentPath;
 	FString OldName = InOldPath;
-	InOldPath.Split(UModularRig::NamespaceSeparator, &ParentPath, &OldName, ESearchCase::CaseSensitive, ESearchDir::FromEnd);
+	(void)URigHierarchy::SplitNameSpace(InOldPath, &ParentPath, &OldName);
 
 	if (InNewName == OldName)
 	{
@@ -728,8 +728,7 @@ void SModularRigModel::OnHierarchyModified(ERigHierarchyNotification InNotif, UR
 			FString ModulePathOrConnectorName;
 			if(Connector->IsPrimary())
 			{
-				ModulePathOrConnectorName = InHierarchy->GetNameMetadata(Connector->GetKey(), URigHierarchy::NameSpaceMetadataName, NAME_None).ToString();
-				ModulePathOrConnectorName = ModulePathOrConnectorName.LeftChop(1);
+				ModulePathOrConnectorName = InHierarchy->GetModulePath(Connector->GetKey());
 			}
 			else
 			{
diff --git a/Engine/Plugins/Animation/ControlRig/Source/ControlRigEditor/Private/Editor/SModularRigTreeView.cpp b/Engine/Plugins/Animation/ControlRig/Source/ControlRigEditor/Private/Editor/SModularRigTreeView.cpp
--- a/Engine/Plugins/Animation/ControlRig/Source/ControlRigEditor/Private/Editor/SModularRigTreeView.cpp
+++ b/Engine/Plugins/Animation/ControlRig/Source/ControlRigEditor/Private/Editor/SModularRigTreeView.cpp
@@ -43,7 +43,7 @@ FModularRigTreeElement::FModularRigTreeElement(const FString& InKey, TWeakPtr<SM
 	bIsPrimary = bInIsPrimary;
 
 	FString ShortNameStr = Key;
-	ShortNameStr.Split(UModularRig::NamespaceSeparator, &ModulePath, &ShortNameStr, ESearchCase::Type::CaseSensitive, ESearchDir::FromEnd);
+	(void)URigHierarchy::SplitNameSpace(ShortNameStr, &ModulePath, &ShortNameStr);
 	if (bIsPrimary)
 	{
 		ModulePath = Key;
diff --git a/Engine/Plugins/Animation/ControlRig/Source/ControlRigEditor/Private/Editor/SRigHierarchy.cpp b/Engine/Plugins/Animation/ControlRig/Source/ControlRigEditor/Private/Editor/SRigHierarchy.cpp
--- a/Engine/Plugins/Animation/ControlRig/Source/ControlRigEditor/Private/Editor/SRigHierarchy.cpp
+++ b/Engine/Plugins/Animation/ControlRig/Source/ControlRigEditor/Private/Editor/SRigHierarchy.cpp
@@ -2916,7 +2916,7 @@ FReply SRigHierarchy::OnAcceptDrop(const FDragDropEvent& DragDropEvent, EItemDro
 							if (Connector->IsPrimary())
 							{
 								FString Path, Name;
-								Connector->GetName().Split(UModularRig::NamespaceSeparator, &Path, &Name, ESearchCase::CaseSensitive, ESearchDir::FromEnd);
+								(void)URigHierarchy::SplitNameSpace(Connector->GetName(), &Path, &Name);
 								if (Path == ModulePath)
 								{
 									PrimaryConnectorKey = Connector->GetKey();
@@ -2964,7 +2964,7 @@ void SRigHierarchy::UpdateConnectorMatchesOnDrag(const TArray<FRigElementKey>& I
 						{
 							if(const FRigConnectorElement* Connector = Hierarchy->Find<FRigConnectorElement>(DraggedElement))
 							{
-								const FName NameSpace = Hierarchy->GetNameMetadata(Connector->GetKey(), URigHierarchy::NameSpaceMetadataName, NAME_None);
+								const FName NameSpace = Hierarchy->GetNameSpaceFName(Connector->GetKey());
 								if(!NameSpace.IsNone())
 								{
 									if(const FRigModuleInstance* Module = ControlRig->FindModule(NameSpace.ToString()))
