Control Rig: Mirror modules
#rnx

diff --git a/Engine/Plugins/Animation/ControlRig/Source/ControlRig/Private/ModularRigController.cpp b/Engine/Plugins/Animation/ControlRig/Source/ControlRig/Private/ModularRigController.cpp
--- a/Engine/Plugins/Animation/ControlRig/Source/ControlRig/Private/ModularRigController.cpp
+++ b/Engine/Plugins/Animation/ControlRig/Source/ControlRig/Private/ModularRigController.cpp
@@ -8,6 +8,7 @@
 #include "ModularRigRuleManager.h"
 #include "Misc/DefaultValueHelper.h"
 #include "Rigs/RigHierarchyController.h"
+#include "RigVMFunctions/Math/RigVMMathLibrary.h"
 
 #include UE_INLINE_GENERATED_CPP_BY_NAME(ModularRigController)
 
@@ -1128,6 +1129,94 @@ FString UModularRigController::ReparentModule(const FString& InModulePath, const
 	return NewPath;
 }
 
+FString UModularRigController::MirrorModule(const FString& InModulePath, const FRigVMMirrorSettings& InSettings, bool bSetupUndo)
+{
+	FRigModuleReference* OriginalModule = FindModule(InModulePath);
+	if (!OriginalModule || !OriginalModule->Class.IsValid())
+	{
+		return FString();
+	}
+
+	FString NewModuleName = OriginalModule->Name.ToString();
+	if (!InSettings.SearchString.IsEmpty())
+	{
+		NewModuleName = NewModuleName.Replace(*InSettings.SearchString, *InSettings.ReplaceString, ESearchCase::CaseSensitive);
+	}
+
+	FString NewModulePath = AddModule(*NewModuleName, OriginalModule->Class.Get(), OriginalModule->ParentPath, bSetupUndo);
+	FRigModuleReference* NewModule = FindModule(NewModulePath);
+	if (!NewModule)
+	{
+		return FString();
+	}
+
+	const TMap<FRigElementKey, FRigElementKey> OriginalConnectionMap = Model->Connections.GetModuleConnectionMap(InModulePath);
+	for (const TPair<FRigElementKey, FRigElementKey>& Pair : OriginalConnectionMap)
+	{
+		FString OriginalTargetPath = Pair.Value.Name.ToString();
+		FString NewTargetPath = OriginalTargetPath.Replace(*InSettings.SearchString, *InSettings.ReplaceString, ESearchCase::CaseSensitive);
+		FRigElementKey NewTargetKey(*NewTargetPath, Pair.Value.Type);
+
+		FString NewConnectorPath = URigHierarchy::JoinNameSpace(NewModulePath, Pair.Key.Name.ToString());
+		FRigElementKey NewConnectorKey(*NewConnectorPath, ERigElementType::Connector);
+		ConnectConnectorToElement(NewConnectorKey, NewTargetKey, bSetupUndo);
+	}
+
+	for (const TPair<FName, FString>& Pair : OriginalModule->Bindings)
+	{
+		FString NewSourcePath = Pair.Value.Replace(*InSettings.SearchString, *InSettings.ReplaceString, ESearchCase::CaseSensitive);
+		BindModuleVariable(NewModulePath, Pair.Key, NewSourcePath, bSetupUndo);
+	}
+
+	for (const TPair<FName, FString>& Pair : OriginalModule->ConfigValues)
+	{
+#if WITH_EDITOR
+		const FProperty* Property = NewModule->Class->FindPropertyByName(Pair.Key);
+		if (!Property)
+		{
+			continue;
+		}
+
+		FString CPPType = Property->GetCPPType();
+		bool bIsVector;
+		if (CPPType == TEXT("FVector"))
+		{
+			bIsVector = true;
+		}
+		else if (CPPType == TEXT("FTransform"))
+		{
+			bIsVector = false;
+		}
+		else
+		{
+			SetConfigValueInModule(NewModulePath, Pair.Key, Pair.Value, bSetupUndo);
+			continue;
+		}
+
+		FString NewValue;
+		if (bIsVector)
+		{
+			FVector Value;
+			FBlueprintEditorUtils::PropertyValueFromString_Direct(Property, Pair.Value, (uint8*)&Value);
+			Value = InSettings.MirrorVector(Value);
+			FBlueprintEditorUtils::PropertyValueToString_Direct(Property, (uint8*)&Value, NewValue, nullptr);
+		}
+		else
+		{
+			FTransform Value;
+			FBlueprintEditorUtils::PropertyValueFromString_Direct(Property, Pair.Value, (uint8*)&Value);
+			Value = InSettings.MirrorTransform(Value);
+			FBlueprintEditorUtils::PropertyValueToString_Direct(Property, (uint8*)&Value, NewValue, nullptr);
+		}
+		SetConfigValueInModule(NewModulePath, Pair.Key, NewValue, bSetupUndo);
+#else
+		SetConfigValueInModule(NewModulePath, Pair.Key, Pair.Value, bSetupUndo);
+#endif
+	}
+	
+	return NewModulePath;
+}
+
 bool UModularRigController::SetModuleShortName(const FString& InModulePath, const FString& InNewShortName, bool bSetupUndo)
 {
 	FRigModuleReference* Module = FindModule(InModulePath);
diff --git a/Engine/Plugins/Animation/ControlRig/Source/ControlRig/Public/ModularRigController.h b/Engine/Plugins/Animation/ControlRig/Source/ControlRig/Public/ModularRigController.h
--- a/Engine/Plugins/Animation/ControlRig/Source/ControlRig/Public/ModularRigController.h
+++ b/Engine/Plugins/Animation/ControlRig/Source/ControlRig/Public/ModularRigController.h
@@ -60,6 +60,9 @@ class CONTROLRIG_API UModularRigController : public UObject
 	UFUNCTION(BlueprintCallable, Category = "ControlRig | Modules")
 	FString ReparentModule(const FString& InModulePath, const FString& InNewParentModulePath, bool bSetupUndo = true);
 
+	UFUNCTION(BlueprintCallable, Category = "ControlRig | Modules")
+	FString MirrorModule(const FString& InModulePath, const FRigVMMirrorSettings& InSettings, bool bSetupUndo = true);
+
 	UFUNCTION(BlueprintCallable, Category = "ControlRig | Modules")
 	bool SetModuleShortName(const FString& InModulePath, const FString& InNewShortName, bool bSetupUndo = true);
 	bool CanSetModuleShortName(const FString& InModulePath, const FString& InNewShortName, FText& OutErrorMessage) const;
diff --git a/Engine/Plugins/Animation/ControlRig/Source/ControlRigEditor/Private/ControlRigModularRigCommands.cpp b/Engine/Plugins/Animation/ControlRig/Source/ControlRigEditor/Private/ControlRigModularRigCommands.cpp
--- a/Engine/Plugins/Animation/ControlRig/Source/ControlRigEditor/Private/ControlRigModularRigCommands.cpp
+++ b/Engine/Plugins/Animation/ControlRig/Source/ControlRigEditor/Private/ControlRigModularRigCommands.cpp
@@ -9,6 +9,7 @@ void FControlRigModularRigCommands::RegisterCommands()
 	UI_COMMAND(AddModuleItem, "New Module", "Add new module to the rig.", EUserInterfaceActionType::Button, FInputChord());
 	UI_COMMAND(RenameModuleItem, "Rename", "Rename module to the rig.", EUserInterfaceActionType::Button, FInputChord(EKeys::F2));
 	UI_COMMAND(DeleteModuleItem, "Delete", "Delete module from the rig.", EUserInterfaceActionType::Button, FInputChord(EKeys::Delete));
+	UI_COMMAND(MirrorModuleItem, "Mirror", "Mirror module from the rig.", EUserInterfaceActionType::Button, FInputChord());
 }
 
 #undef LOCTEXT_NAMESPACE
diff --git a/Engine/Plugins/Animation/ControlRig/Source/ControlRigEditor/Private/ControlRigModularRigCommands.h b/Engine/Plugins/Animation/ControlRig/Source/ControlRigEditor/Private/ControlRigModularRigCommands.h
--- a/Engine/Plugins/Animation/ControlRig/Source/ControlRigEditor/Private/ControlRigModularRigCommands.h
+++ b/Engine/Plugins/Animation/ControlRig/Source/ControlRigEditor/Private/ControlRigModularRigCommands.h
@@ -27,6 +27,9 @@ public:
 	/** Delete Module */
 	TSharedPtr< FUICommandInfo > DeleteModuleItem;
 
+	/** Mirror Module */
+	TSharedPtr< FUICommandInfo > MirrorModuleItem;
+
 	/**
 	 * Initialize commands
 	 */
diff --git a/Engine/Plugins/Animation/ControlRig/Source/ControlRigEditor/Private/Editor/SModularRigModel.cpp b/Engine/Plugins/Animation/ControlRig/Source/ControlRigEditor/Private/Editor/SModularRigModel.cpp
--- a/Engine/Plugins/Animation/ControlRig/Source/ControlRigEditor/Private/Editor/SModularRigModel.cpp
+++ b/Engine/Plugins/Animation/ControlRig/Source/ControlRigEditor/Private/Editor/SModularRigModel.cpp
@@ -194,6 +194,10 @@ void SModularRigModel::BindCommands()
 	CommandList->MapAction(Commands.DeleteModuleItem,
 		FExecuteAction::CreateSP(this, &SModularRigModel::HandleDeleteModules),
 		FCanExecuteAction());
+
+	CommandList->MapAction(Commands.MirrorModuleItem,
+		FExecuteAction::CreateSP(this, &SModularRigModel::HandleMirrorModules),
+		FCanExecuteAction());
 }
 
 FReply SModularRigModel::OnKeyDown(const FGeometry& MyGeometry, const FKeyEvent& InKeyEvent)
@@ -358,6 +362,7 @@ void SModularRigModel::CreateContextMenu()
 					);
 					ElementsSection.AddMenuEntry(Commands.RenameModuleItem);
 					ElementsSection.AddMenuEntry(Commands.DeleteModuleItem);
+					ElementsSection.AddMenuEntry(Commands.MirrorModuleItem);
 				}
 			})
 		);
@@ -663,6 +668,71 @@ void SModularRigModel::HandleReparentModules(const TArray<FString>& InPaths, con
 	}
 }
 
+void SModularRigModel::HandleMirrorModules()
+{
+	if(!ControlRigEditor.IsValid())
+	{
+		return;
+	}
+
+	UModularRig* Rig = GetDefaultModularRig();
+	if (Rig)
+	{
+		TArray<TSharedPtr<FModularRigTreeElement>> SelectedItems = TreeView->GetSelectedItems();
+		TArray<FString> SelectedPaths;
+		Algo::Transform(SelectedItems, SelectedPaths, [](const TSharedPtr<FModularRigTreeElement>& Element)
+		{
+			if (Element.IsValid())
+			{
+				return Element->ModulePath;
+			}
+			return FString();
+		});
+		HandleMirrorModules(SelectedPaths);
+	}
+
+	return;
+}
+
+void SModularRigModel::HandleMirrorModules(const TArray<FString>& InPaths)
+{
+	if (ControlRigBlueprint.IsValid())
+	{
+		FRigVMMirrorSettings Settings;
+		TSharedPtr<FStructOnScope> StructToDisplay = MakeShareable(new FStructOnScope(FRigVMMirrorSettings::StaticStruct(), (uint8*)&Settings));
+		
+		TSharedRef<SKismetInspector> KismetInspector = SNew(SKismetInspector);
+		KismetInspector->ShowSingleStruct(StructToDisplay);
+		
+		TSharedRef<SCustomDialog> MirrorDialog = SNew(SCustomDialog)
+			.Title(FText(LOCTEXT("ControlModularModelMirror", "Mirror Selected Modules")))
+			.Content()
+			[
+				KismetInspector
+			]
+			.Buttons({
+				SCustomDialog::FButton(LOCTEXT("OK", "OK")),
+				SCustomDialog::FButton(LOCTEXT("Cancel", "Cancel"))
+		});
+
+		if (MirrorDialog->ShowModal() == 0)
+		{
+			FScopedTransaction Transaction(LOCTEXT("ModularRigModelMirror", "Mirror Modules"));
+
+			UModularRigController* Controller = ControlRigBlueprint->GetModularRigController();
+			check(Controller);
+
+			// Make sure we delete the modules from children to root
+			TArray<FString> SortedPaths = Controller->Model->SortPaths(InPaths);
+			Algo::Reverse(SortedPaths);
+			for (const FString& Path : SortedPaths)
+			{
+				Controller->MirrorModule(Path, Settings);
+			}
+		}
+	}
+}
+
 void SModularRigModel::HandleConnectorResolved(const FRigElementKey& InConnector, const FRigElementKey& InTarget)
 {
 	if (ControlRigBlueprint.IsValid())
diff --git a/Engine/Plugins/Animation/ControlRig/Source/ControlRigEditor/Private/Editor/SModularRigModel.h b/Engine/Plugins/Animation/ControlRig/Source/ControlRigEditor/Private/Editor/SModularRigModel.h
--- a/Engine/Plugins/Animation/ControlRig/Source/ControlRigEditor/Private/Editor/SModularRigModel.h
+++ b/Engine/Plugins/Animation/ControlRig/Source/ControlRigEditor/Private/Editor/SModularRigModel.h
@@ -74,6 +74,10 @@ private:
 	/** Reparent items */
 	void HandleReparentModules(const TArray<FString>& InPaths, const FString& InParentPath);
 
+	/** Mirror items */
+	void HandleMirrorModules();
+	void HandleMirrorModules(const TArray<FString>& InPaths);
+
 	/** Resolve connector */
 	void HandleConnectorResolved(const FRigElementKey& InConnector, const FRigElementKey& InTarget);
 
