Control Rig: Fixes for drag and drop within the schematic

#rnx

diff --git a/Engine/Plugins/Animation/ControlRig/Source/ControlRig/Private/ModularRigController.cpp b/Engine/Plugins/Animation/ControlRig/Source/ControlRig/Private/ModularRigController.cpp
--- a/Engine/Plugins/Animation/ControlRig/Source/ControlRig/Private/ModularRigController.cpp
+++ b/Engine/Plugins/Animation/ControlRig/Source/ControlRig/Private/ModularRigController.cpp
@@ -21,6 +21,7 @@ UModularRigController::UModularRigController(const FObjectInitializer& ObjectIni
 	: Super(ObjectInitializer)
 	, Model(nullptr)
 	, bSuspendNotifications(false)
+	, bAutomaticReparenting(true)
 {
 }
 
@@ -161,15 +162,15 @@ bool UModularRigController::CanConnectConnectorToElement(const FRigElementKey& I
 
 	if (!ModuleConnector->IsPrimary())
 	{
-		const FRigModuleConnector* PrimaryMdouleConnector = RigCDO->GetRigModuleSettings().ExposedConnectors.FindByPredicate(
-		[ConnectorName](FRigModuleConnector& Connector)
+		const FRigModuleConnector* PrimaryModuleConnector = RigCDO->GetRigModuleSettings().ExposedConnectors.FindByPredicate(
+		[](const FRigModuleConnector& Connector)
 		{
 			return Connector.IsPrimary();
 		});
 
-		const FString PrimaryConnectorPath = FString::Printf(TEXT("%s:%s"), *ConnectorModulePath, *PrimaryMdouleConnector->Name);
+		const FString PrimaryConnectorPath = FString::Printf(TEXT("%s:%s"), *ConnectorModulePath, *PrimaryModuleConnector->Name);
 		const FRigElementKey PrimaryConnectorKey(*PrimaryConnectorPath, ERigElementType::Connector);
-		FRigElementKey PrimaryTarget = Model->Connections.FindTargetFromConnector(PrimaryConnectorKey);
+		const FRigElementKey PrimaryTarget = Model->Connections.FindTargetFromConnector(PrimaryConnectorKey);
 		if (!PrimaryTarget.IsValid())
 		{
 			OutErrorMessage = FText::FromString(FString::Printf(TEXT("Cannot resolve connector %s because primary connector is not resolved"), *InConnectorKey.ToString()));
@@ -261,6 +262,7 @@ bool UModularRigController::ConnectConnectorToElement(const FRigElementKey& InCo
 	// First disconnect before connecting to anything else. This might disconnect other secondary/optional connectors.
 	if (CurrentTarget.IsValid())
 	{
+		const TGuardValue<bool> DisableAutomaticReparenting(bAutomaticReparenting, false);
 		DisconnectConnector(InConnectorKey, bSetupUndo);
 	}
 	
@@ -311,13 +313,16 @@ bool UModularRigController::ConnectConnectorToElement(const FRigElementKey& InCo
 				}
 
 				// automatically re-parent the module in the module tree as well
-				if(const FRigConnectorElement* Connector = Hierarchy->Find<FRigConnectorElement>(InConnectorKey))
+				if(bAutomaticReparenting)
 				{
-					if(Connector->IsPrimary())
+					if(const FRigConnectorElement* Connector = Hierarchy->Find<FRigConnectorElement>(InConnectorKey))
 					{
-						if(!TargetModulePathName.IsNone())
+						if(Connector->IsPrimary())
 						{
-							ReparentModule(Module->GetPath(), TargetModulePathName.ToString(), bSetupUndo);
+							if(!TargetModulePathName.IsNone())
+							{
+								ReparentModule(Module->GetPath(), TargetModulePathName.ToString(), bSetupUndo);
+							}
 						}
 					}
 				}
@@ -391,7 +396,7 @@ bool UModularRigController::DisconnectConnector(const FRigElementKey& InConnecto
 	{
 		// Remove connections from module and child modules
 		TArray<FRigElementKey> ConnectionsToRemove;
-		for (const FModularRigSingleConnection& Connection : Model->Connections.ConnectionList)
+		for (const FModularRigSingleConnection& Connection : Model->Connections)
 		{
 			if (Connection.Connector.Name.ToString().StartsWith(ConnectorModulePath, ESearchCase::CaseSensitive))
 			{
@@ -407,7 +412,7 @@ bool UModularRigController::DisconnectConnector(const FRigElementKey& InConnecto
 	{
 		// Remove connections from child modules
 		TArray<FRigElementKey> ConnectionsToRemove;
-		for (const FModularRigSingleConnection& Connection : Model->Connections.ConnectionList)
+		for (const FModularRigSingleConnection& Connection : Model->Connections)
 		{
 			FString OtherConnectorModulePath, OtherConnectorName;
 			(void)URigHierarchy::SplitNameSpace(Connection.Connector.Name.ToString(), &OtherConnectorModulePath, &OtherConnectorName);
@@ -425,9 +430,12 @@ bool UModularRigController::DisconnectConnector(const FRigElementKey& InConnecto
 	// todo: Make sure all the rest of the connections are still valid
 
 	// un-parent the module if we've disconnected the primary
-	if(ModuleConnector->IsPrimary() && !Module->IsRootModule())
+	if(bAutomaticReparenting)
 	{
-		(void)ReparentModule(Module->GetPath(), FString(), bSetupUndo);
+		if(ModuleConnector->IsPrimary() && !Module->IsRootModule())
+		{
+			(void)ReparentModule(Module->GetPath(), FString(), bSetupUndo);
+		}
 	}
 
 	Notify(EModularRigNotification::ConnectionChanged, Module);
@@ -460,7 +468,7 @@ TArray<FRigElementKey> UModularRigController::DisconnectCyclicConnectors(bool bS
 	}
 
 	TArray<FRigElementKey> ConnectorsToDisconnect;
-	for (const FModularRigSingleConnection& Connection : Model->Connections.ConnectionList)
+	for (const FModularRigSingleConnection& Connection : Model->Connections)
 	{
 		const FString ConnectorModulePath = Hierarchy->GetModulePath(Connection.Connector);
 		const FString TargetModulePath = Hierarchy->GetModulePath(Connection.Target);
@@ -833,7 +841,7 @@ bool UModularRigController::DeleteModule(const FString& InModulePath, bool bSetu
 	// Fix connections
 	{
 		TArray<FRigElementKey> ToRemove;
-		for (FModularRigSingleConnection& Connection : Model->Connections.ConnectionList)
+		for (FModularRigSingleConnection& Connection : Model->Connections)
 		{
 			FString ConnectionModulePath, ConnectionName;
 			(void)URigHierarchy::SplitNameSpace(Connection.Connector.Name.ToString(), &ConnectionModulePath, &ConnectionName);
@@ -932,7 +940,7 @@ FString UModularRigController::RenameModule(const FString& InModulePath, const F
 
 	// Fix connections
 	{
-		for (FModularRigSingleConnection& Connection : Model->Connections.ConnectionList)
+		for (FModularRigSingleConnection& Connection : Model->Connections)
 		{
 			if (Connection.Connector.Name.ToString().StartsWith(OldPath, ESearchCase::CaseSensitive))
 			{
@@ -943,20 +951,7 @@ FString UModularRigController::RenameModule(const FString& InModulePath, const F
 				Connection.Target.Name = *FString::Printf(TEXT("%s%s"), *NewPath, *Connection.Target.Name.ToString().RightChop(OldPath.Len()));
 			}
 		}
-		for (TPair<FRigElementKey, TArray<FRigElementKey>>& Pair : Model->Connections.ReverseConnectionMap)
-		{
-			if (Pair.Key.Name.ToString().StartsWith(OldPath, ESearchCase::CaseSensitive))
-			{
-				Pair.Key.Name = *FString::Printf(TEXT("%s%s"), *NewPath, *Pair.Key.Name.ToString().RightChop(OldPath.Len()));
-			}
-			for (FRigElementKey& TargetConnectorKey : Pair.Value)
-			{
-				if (TargetConnectorKey.Name.ToString().StartsWith(OldPath, ESearchCase::CaseSensitive))
-				{
-					TargetConnectorKey.Name = *FString::Printf(TEXT("%s%s"), *NewPath, *TargetConnectorKey.Name.ToString().RightChop(OldPath.Len()));
-				}
-			}
-		}
+		Model->Connections.UpdateFromConnectionList();
 	}
 
 	// Fix bindings
@@ -1063,7 +1058,7 @@ FString UModularRigController::ReparentModule(const FString& InModulePath, const
 
 	// Fix connections
 	{
-		for (FModularRigSingleConnection& Connection : Model->Connections.ConnectionList)
+		for (FModularRigSingleConnection& Connection : Model->Connections)
 		{
 			if (Connection.Connector.Name.ToString().StartsWith(OldPath, ESearchCase::CaseSensitive))
 			{
@@ -1074,20 +1069,7 @@ FString UModularRigController::ReparentModule(const FString& InModulePath, const
 				Connection.Target.Name = *FString::Printf(TEXT("%s%s"), *NewPath, *Connection.Target.Name.ToString().RightChop(OldPath.Len()));
 			}
 		}
-		for (TPair<FRigElementKey, TArray<FRigElementKey>>& Pair : Model->Connections.ReverseConnectionMap)
-		{
-			if (Pair.Key.Name.ToString().StartsWith(OldPath, ESearchCase::CaseSensitive))
-			{
-				Pair.Key.Name = *FString::Printf(TEXT("%s%s"), *NewPath, *Pair.Key.Name.ToString().RightChop(OldPath.Len()));
-			}
-			for (FRigElementKey& TargetConnectorKey : Pair.Value)
-			{
-				if (TargetConnectorKey.Name.ToString().StartsWith(OldPath, ESearchCase::CaseSensitive))
-				{
-					TargetConnectorKey.Name = *FString::Printf(TEXT("%s%s"), *NewPath, *TargetConnectorKey.Name.ToString().RightChop(OldPath.Len()));
-				}
-			}
-		}
+		Model->Connections.UpdateFromConnectionList();
 	}
 
 	// Fix bindings
diff --git a/Engine/Plugins/Animation/ControlRig/Source/ControlRig/Private/ModularRigModel.cpp b/Engine/Plugins/Animation/ControlRig/Source/ControlRig/Private/ModularRigModel.cpp
--- a/Engine/Plugins/Animation/ControlRig/Source/ControlRig/Private/ModularRigModel.cpp
+++ b/Engine/Plugins/Animation/ControlRig/Source/ControlRig/Private/ModularRigModel.cpp
@@ -92,12 +92,11 @@ TMap<FRigElementKey, FRigElementKey> FModularRigConnections::GetModuleConnection
 	return Result;
 }
 
-bool FModularRigModel::PatchModelsOnLoad()
+void FModularRigModel::PatchModelsOnLoad()
 {
-	bool bPatched = false;
-	if (Connections.ConnectionList.IsEmpty())
+	if (Connections.IsEmpty())
 	{
-		ForEachModule([this, &bPatched](const FRigModuleReference* Module) -> bool
+		ForEachModule([this](const FRigModuleReference* Module) -> bool
 		{
 			FString ModuleNamespace = Module->GetNamespace();
 			for (const TTuple<FRigElementKey, FRigElementKey>& Connection : Module->Connections_DEPRECATED)
@@ -108,9 +107,9 @@ bool FModularRigModel::PatchModelsOnLoad()
 			}
 			return true;
 		});
-		bPatched = !Connections.ConnectionList.IsEmpty();
 	}
-	return bPatched;
+
+	Connections.UpdateFromConnectionList();
 }
 
 UModularRigController* FModularRigModel::GetController(bool bCreateIfNeeded)
diff --git a/Engine/Plugins/Animation/ControlRig/Source/ControlRig/Private/ModularRigRuleManager.cpp b/Engine/Plugins/Animation/ControlRig/Source/ControlRig/Private/ModularRigRuleManager.cpp
--- a/Engine/Plugins/Animation/ControlRig/Source/ControlRig/Private/ModularRigRuleManager.cpp
+++ b/Engine/Plugins/Animation/ControlRig/Source/ControlRig/Private/ModularRigRuleManager.cpp
@@ -8,14 +8,11 @@
 
 #define LOCTEXT_NAMESPACE "ModularRigRuleManager"
 
-FModularRigResolveResult UModularRigRuleManager::FindMatches(
-	const FRigConnectorElement* InConnector,
-	const FRigModuleInstance* InModule,
-	const FRigElementKeyRedirector& InResolvedConnectors) const
+FModularRigResolveResult UModularRigRuleManager::FindMatches(FWorkData& InWorkData) const
 {
-	FModularRigResolveResult Result;
-	Result.Connector = InConnector->GetKey();
-	
+	check(InWorkData.Result);
+	FModularRigResolveResult& Result = *InWorkData.Result;
+
 	if(!Hierarchy.IsValid())
 	{
 		static const FText MissingHierarchyMessage = LOCTEXT("MissingHierarchyMessage", "The rule manager is missing the hierarchy.");
@@ -46,14 +43,40 @@ FModularRigResolveResult UModularRigRuleManager::FindMatches(
 		bContinue = true;
 	}, true);
 
+	InWorkData.Hierarchy = Hierarchy.Get();
+	ResolveConnector(InWorkData);
+	return Result;;
+}
+
+FModularRigResolveResult UModularRigRuleManager::FindMatches(
+	const FRigConnectorElement* InConnector,
+	const FRigModuleInstance* InModule,
+	const FRigElementKeyRedirector& InResolvedConnectors) const
+{
+	FModularRigResolveResult Result;
+	Result.Connector = InConnector->GetKey();
+
 	FWorkData WorkData;
 	WorkData.Hierarchy = Hierarchy.Get();
 	WorkData.Connector = InConnector;
 	WorkData.Module = InModule;
 	WorkData.ResolvedConnectors = &InResolvedConnectors;
 	WorkData.Result = &Result;
-	ResolveConnector(WorkData);
-	return Result;;
+
+	return FindMatches(WorkData);
+}
+
+FModularRigResolveResult UModularRigRuleManager::FindMatches(const FRigModuleConnector* InConnector) const
+{
+	FModularRigResolveResult Result;
+	Result.Connector = FRigElementKey(*InConnector->Name, ERigElementType::Connector);
+
+	FWorkData WorkData;
+	WorkData.Hierarchy = Hierarchy.Get();
+	WorkData.ModuleConnector = InConnector;
+	WorkData.Result = &Result;
+
+	return FindMatches(WorkData);
 }
 
 FModularRigResolveResult UModularRigRuleManager::FindMatchesForPrimaryConnector(const FRigModuleInstance* InModule) const
@@ -175,6 +198,11 @@ void UModularRigRuleManager::FilterIncompatibleTypes(FWorkData& InOutWorkData)
 
 void UModularRigRuleManager::FilterInvalidNameSpaces(FWorkData& InOutWorkData)
 {
+	if(InOutWorkData.Connector == nullptr)
+	{
+		return;
+	}
+	
 	const FName NameSpace = InOutWorkData.Hierarchy->GetNameSpaceFName(InOutWorkData.Connector->GetKey());
 	if(NameSpace.IsNone())
 	{
@@ -204,8 +232,12 @@ void UModularRigRuleManager::FilterInvalidNameSpaces(FWorkData& InOutWorkData)
 
 void UModularRigRuleManager::FilterByConnectorRules(FWorkData& InOutWorkData)
 {
-	const FRigConnectorElement* Connector = InOutWorkData.Connector;
-	for(const FRigConnectionRuleStash& Stash : Connector->Settings.Rules)
+	check(InOutWorkData.Connector != nullptr || InOutWorkData.ModuleConnector != nullptr);
+	
+	const TArray<FRigConnectionRuleStash>& Rules =
+		InOutWorkData.Connector ? InOutWorkData.Connector->Settings.Rules : InOutWorkData.ModuleConnector->Settings.Rules;
+	
+	for(const FRigConnectionRuleStash& Stash : Rules)
 	{
 		TSharedPtr<FStructOnScope> Storage;
 		const FRigConnectionRule* Rule = Stash.Get(Storage);
@@ -231,6 +263,11 @@ void UModularRigRuleManager::FilterByConnectorEvent(FWorkData& InOutWorkData)
 	{
 		return;
 	}
+
+	if(InOutWorkData.Connector == nullptr)
+	{
+		return;
+	}
 	
 	// see if we are nested below a modular rig
 	UModularRig* ModularRig = InOutWorkData.Hierarchy->GetTypedOuter<UModularRig>();
diff --git a/Engine/Plugins/Animation/ControlRig/Source/ControlRig/Public/ModularRigController.h b/Engine/Plugins/Animation/ControlRig/Source/ControlRig/Public/ModularRigController.h
--- a/Engine/Plugins/Animation/ControlRig/Source/ControlRig/Public/ModularRigController.h
+++ b/Engine/Plugins/Animation/ControlRig/Source/ControlRig/Public/ModularRigController.h
@@ -83,6 +83,8 @@ private:
 	void SetModel(FModularRigModel* InModel) { Model = InModel; }
 	void Notify(const EModularRigNotification& InNotification, const FRigModuleReference* InElement);
 	void UpdateShortNames();
+
+	bool bAutomaticReparenting;
 	
 	friend struct FModularRigModel;
 };
\ No newline at end of file
diff --git a/Engine/Plugins/Animation/ControlRig/Source/ControlRig/Public/ModularRigModel.h b/Engine/Plugins/Animation/ControlRig/Source/ControlRig/Public/ModularRigModel.h
--- a/Engine/Plugins/Animation/ControlRig/Source/ControlRig/Public/ModularRigModel.h
+++ b/Engine/Plugins/Animation/ControlRig/Source/ControlRig/Public/ModularRigModel.h
@@ -129,15 +129,32 @@ struct FModularRigSingleConnection
 USTRUCT(BlueprintType)
 struct CONTROLRIG_API FModularRigConnections
 {
+public:
+	
 	GENERATED_BODY()
-
 	/** Connections sorted by creation order */
+
+private:
+	
 	UPROPERTY()
 	TArray<FModularRigSingleConnection> ConnectionList;
 
 	/** Target key to connector array */
 	TMap<FRigElementKey, TArray<FRigElementKey>> ReverseConnectionMap;
 
+public:
+
+	const TArray<FModularRigSingleConnection>& GetConnectionList() const { return ConnectionList; }
+
+	bool IsEmpty() const { return ConnectionList.IsEmpty(); }
+	int32 Num() const { return ConnectionList.Num(); }
+	const FModularRigSingleConnection& operator[](int32 InIndex) const { return ConnectionList[InIndex]; }
+	FModularRigSingleConnection& operator[](int32 InIndex) { return ConnectionList[InIndex]; }
+	TArray<FModularRigSingleConnection>::RangedForIteratorType begin() { return ConnectionList.begin(); }
+	TArray<FModularRigSingleConnection>::RangedForIteratorType end() { return ConnectionList.end(); }
+	TArray<FModularRigSingleConnection>::RangedForConstIteratorType begin() const { return ConnectionList.begin(); }
+	TArray<FModularRigSingleConnection>::RangedForConstIteratorType end() const { return ConnectionList.end(); }
+
 	void UpdateFromConnectionList()
 	{
 		ReverseConnectionMap.Reset();
@@ -252,7 +269,7 @@ public:
 	UPROPERTY(transient)
 	TObjectPtr<UObject> Controller;
 
-	bool PatchModelsOnLoad();
+	void PatchModelsOnLoad();
 
 	UModularRigController* GetController(bool bCreateIfNeeded = true);
 
diff --git a/Engine/Plugins/Animation/ControlRig/Source/ControlRig/Public/ModularRigRuleManager.h b/Engine/Plugins/Animation/ControlRig/Source/ControlRig/Public/ModularRigRuleManager.h
--- a/Engine/Plugins/Animation/ControlRig/Source/ControlRig/Public/ModularRigRuleManager.h
+++ b/Engine/Plugins/Animation/ControlRig/Source/ControlRig/Public/ModularRigRuleManager.h
@@ -27,6 +27,16 @@ public:
 		const FRigElementKeyRedirector& InResolvedConnectors = FRigElementKeyRedirector()
 	) const;
 
+	/***
+	 * Returns the possible targets for the given external module connector in the current resolve stage
+	 * Note: This method is thread-safe. 
+	 * @param InConnector The connector to resolve
+	 * @return The resolve result including a list of matches
+	 */
+	FModularRigResolveResult FindMatches(
+		const FRigModuleConnector* InConnector
+	) const;
+
 	/***
 	 * Returns the possible targets for the primary connector in the current resolve stage
 	 * Note: This method is thread-safe. 
@@ -68,6 +78,7 @@ private:
 		FWorkData()
 		: Hierarchy(nullptr)
 		, Connector(nullptr)
+		, ModuleConnector(nullptr)
 		, Module(nullptr)
 		, ResolvedConnectors(nullptr)
 		, Result(nullptr)
@@ -78,11 +89,14 @@ private:
 
 		const URigHierarchy* Hierarchy;
 		const FRigConnectorElement* Connector;
+		const FRigModuleConnector* ModuleConnector;
 		const FRigModuleInstance* Module;
 		const FRigElementKeyRedirector* ResolvedConnectors;
 		FModularRigResolveResult* Result;
 	};
 
+	FModularRigResolveResult FindMatches(FWorkData& InWorkData) const;
+
 	void SetHierarchy(const URigHierarchy* InHierarchy);
 	static void ResolveConnector(FWorkData& InOutWorkData);
 	static void FilterIncompatibleTypes(FWorkData& InOutWorkData);
diff --git a/Engine/Plugins/Animation/ControlRig/Source/ControlRigDeveloper/Private/ControlRigBlueprint.cpp b/Engine/Plugins/Animation/ControlRig/Source/ControlRigDeveloper/Private/ControlRigBlueprint.cpp
--- a/Engine/Plugins/Animation/ControlRig/Source/ControlRigDeveloper/Private/ControlRigBlueprint.cpp
+++ b/Engine/Plugins/Animation/ControlRig/Source/ControlRigDeveloper/Private/ControlRigBlueprint.cpp
@@ -1043,7 +1043,7 @@ void UControlRigBlueprint::UpdateConnectionMapFromModel()
 	{
 		ConnectionMap.Reset();
 
-		for (const FModularRigSingleConnection& Connection : ModularRigModel.Connections.ConnectionList)
+		for (const FModularRigSingleConnection& Connection : ModularRigModel.Connections)
 		{
 			ConnectionMap.Add(Connection.Connector, Connection.Target);
 		}
@@ -2581,6 +2581,7 @@ void UControlRigBlueprint::HandleRigModulesModified(EModularRigNotification InNo
 								(void)URigHierarchy::SplitNameSpace(Pair.Value.NewPath, &Namespace, &ConnectorName);
 								Namespace.Append(UModularRig::NamespaceSeparator);
 								FControlRigExecuteContextRigModuleGuard RigModuleGuard(PublicContext, Namespace);
+								const TGuardValue<bool> DisableErrors(Controller->bReportWarningsAndErrors, false);
 								Controller->AddConnector(*ConnectorName, Pair.Value.Settings);
 							}
 						}
diff --git a/Engine/Plugins/Animation/ControlRig/Source/ControlRigEditor/Private/Editor/ControlRigEditor.cpp b/Engine/Plugins/Animation/ControlRig/Source/ControlRigEditor/Private/Editor/ControlRigEditor.cpp
--- a/Engine/Plugins/Animation/ControlRig/Source/ControlRigEditor/Private/Editor/ControlRigEditor.cpp
+++ b/Engine/Plugins/Animation/ControlRig/Source/ControlRigEditor/Private/Editor/ControlRigEditor.cpp
@@ -328,7 +328,6 @@ void FControlRigEditor::InitRigVMEditor(const EToolkitMode::Type Mode, const TSh
 		{
 			SchematicModel.SetEditor(SharedThis(this));
 			ControlRigBlueprint->OnSetObjectBeingDebugged().AddRaw(&SchematicModel, &FControlRigSchematicModel::OnSetObjectBeingDebugged);
-			ControlRigBlueprint->OnHierarchyModified().AddRaw(&SchematicModel, &FControlRigSchematicModel::OnHierarchyModified);
 			ControlRigBlueprint->GetModularRigController()->OnModified().AddRaw(&SchematicModel, &FControlRigSchematicModel::HandleModularRigModified);
 		}
 
@@ -2282,7 +2281,6 @@ void FControlRigEditor::HandleViewportCreated(const TSharedRef<class IPersonaVie
 																.PaddingInterNode(5)
 				;
 				InViewport->AddOverlayWidget(SchematicPanel);
-				HandleSchematicViewportCreated(SchematicPanel);
 			}
 		}
 	}
@@ -2599,15 +2597,6 @@ void FControlRigEditor::CacheNameLists()
 	}
 }
 
-void FControlRigEditor::HandleSchematicViewportCreated(const TSharedRef<SSchematicGraphPanel>& InViewport)
-{
-	SchematicViewport = InViewport.ToWeakPtr();
-	InViewport->OnNodeClicked().BindRaw(&SchematicModel, &FControlRigSchematicModel::HandleSchematicNodeClicked);
-	InViewport->OnBeginDrag().BindRaw(&SchematicModel, &FControlRigSchematicModel::HandleSchematicBeginDrag);
-	InViewport->OnEndDrag().BindRaw(&SchematicModel, &FControlRigSchematicModel::HandleSchematicEndDrag);
-	InViewport->OnAcceptDrop().BindRaw(&SchematicModel, &FControlRigSchematicModel::HandleSchematicDrop);
-}
-
 FVector2D FControlRigEditor::ComputePersonaProjectedScreenPos(const FVector& InWorldPos, bool bClampToScreenRectangle)
 {
 	if (PreviewViewport.IsValid())
@@ -2747,11 +2736,11 @@ void FControlRigEditor::FilterDraggedKeys(TArray<FRigElementKey>& Keys, bool bRe
 		FilteredKeys.Reserve(Keys.Num());
 		for (FRigElementKey Key : Keys)
 		{
-			for(const TPair<FRigElementKey,FRigElementKey>& Pair : ControlRigBlueprint->ConnectionMap)
+			for(const FModularRigSingleConnection& Connection : ControlRigBlueprint->ModularRigModel.Connections)
 			{
-				if(Pair.Value == Key)
+				if(Connection.Target == Key)
 				{
-					Key = Pair.Key;
+					Key = Connection.Connector;
 					break;
 				}
 			}
diff --git a/Engine/Plugins/Animation/ControlRig/Source/ControlRigEditor/Private/Editor/ControlRigEditor.h b/Engine/Plugins/Animation/ControlRig/Source/ControlRigEditor/Private/Editor/ControlRigEditor.h
--- a/Engine/Plugins/Animation/ControlRig/Source/ControlRigEditor/Private/Editor/ControlRigEditor.h
+++ b/Engine/Plugins/Animation/ControlRig/Source/ControlRigEditor/Private/Editor/ControlRigEditor.h
@@ -246,9 +246,6 @@ private:
 	void OnToolbarAxesScaleChanged(float InValue);
 	void HandleToggleSchematicViewport();
 
-	/** Handle schematic panel setup*/
-	void HandleSchematicViewportCreated(const TSharedRef<class SSchematicGraphPanel>& InViewport);
-
 		/** Handle switching skeletal meshes */
 	void HandlePreviewMeshChanged(USkeletalMesh* InOldSkeletalMesh, USkeletalMesh* InNewSkeletalMesh);
 
diff --git a/Engine/Plugins/Animation/ControlRig/Source/ControlRigEditor/Private/Editor/ControlRigSchematicModel.cpp b/Engine/Plugins/Animation/ControlRig/Source/ControlRigEditor/Private/Editor/ControlRigSchematicModel.cpp
--- a/Engine/Plugins/Animation/ControlRig/Source/ControlRigEditor/Private/Editor/ControlRigSchematicModel.cpp
+++ b/Engine/Plugins/Animation/ControlRig/Source/ControlRigEditor/Private/Editor/ControlRigSchematicModel.cpp
@@ -16,6 +16,8 @@
 #include "Blueprint/WidgetBlueprintLibrary.h"
 #include "ControlRigDragOps.h"
 #include "SchematicGraphPanel/SSchematicGraphPanel.h"
+#include "Framework/Application/SlateApplication.h"
+#include "ScopedTransaction.h"
 
 #define LOCTEXT_NAMESPACE "ControlRigSchematicModel"
 
@@ -65,6 +67,19 @@ void FControlRigSchematicModel::SetEditor(const TSharedRef<FControlRigEditor>& I
 	OnSetObjectBeingDebugged(ControlRigBlueprint->GetDebuggedControlRig());
 }
 
+void FControlRigSchematicModel::ApplyToPanel(SSchematicGraphPanel* InPanel)
+{
+	FSchematicGraphModel::ApplyToPanel(InPanel);
+	
+	InPanel->OnNodeClicked().BindRaw(this, &FControlRigSchematicModel::HandleSchematicNodeClicked);
+	InPanel->OnBeginDrag().BindRaw(this, &FControlRigSchematicModel::HandleSchematicBeginDrag);
+	InPanel->OnEndDrag().BindRaw(this, &FControlRigSchematicModel::HandleSchematicEndDrag);
+	InPanel->OnEnterDrag().BindRaw(this, &FControlRigSchematicModel::HandleSchematicEnterDrag);
+	InPanel->OnLeaveDrag().BindRaw(this, &FControlRigSchematicModel::HandleSchematicLeaveDrag);
+	InPanel->OnCancelDrag().BindRaw(this, &FControlRigSchematicModel::HandleSchematicCancelDrag);
+	InPanel->OnAcceptDrop().BindRaw(this, &FControlRigSchematicModel::HandleSchematicDrop);
+}
+
 void FControlRigSchematicModel::Reset()
 {
 	Super::Reset();
@@ -258,15 +273,54 @@ void FControlRigSchematicModel::UpdateElementKeyNodes()
 {
 	if (ControlRigBeingDebuggedPtr.IsValid())
 	{
-		if (URigHierarchy* Hierarchy = ControlRigBeingDebuggedPtr->GetHierarchy())
+		if (const URigHierarchy* Hierarchy = ControlRigBeingDebuggedPtr->GetHierarchy())
 		{
-			OnHierarchyModified(ERigHierarchyNotification::HierarchyReset, Hierarchy, nullptr);
+			const TArray<FRigSocketElement*> Sockets = Hierarchy->GetElementsOfType<FRigSocketElement>();
+			for (const FRigSocketElement* Socket : Sockets)
+			{
+				if(!ContainsElementKeyNode(Socket->GetKey()))
+				{
+					AddElementKeyNode(Socket->GetKey());
+				}
+			}
+
+			const TArray<FRigConnectorElement*> Connectors = Hierarchy->GetElementsOfType<FRigConnectorElement>();
+			for (const FRigConnectorElement* Connector : Connectors)
+			{
+				if(!ContainsElementKeyNode(Connector->GetKey()))
+				{
+					AddElementKeyNode(Connector->GetKey());
+				}
+				UpdateConnector(Connector->GetKey());
+			}
+
+			// remove obsolete nodes
+			TArray<FGuid> GuidsToRemove;
+			for(const TSharedPtr<FSchematicGraphNode>& Node : Nodes)
+			{
+				if(const FControlRigSchematicRigElementKeyNode* ElementKeyNode = Cast<FControlRigSchematicRigElementKeyNode>(Node.Get()))
+				{
+					if(!Hierarchy->Contains(ElementKeyNode->GetKey()))
+					{
+						GuidsToRemove.Add(ElementKeyNode->GetGuid());
+					}
+				}
+			}
+			for(const FGuid& Guid : GuidsToRemove)
+			{
+				RemoveNode(Guid);
+			}
 		}
 	}
 }
 
 void FControlRigSchematicModel::UpdateElementKeyLinks()
 {
+	if(!bUpdateElementKeyLinks)
+	{
+		return;
+	}
+	
 	const TSharedPtr<FControlRigEditor> Editor = ControlRigEditor.Pin();
 	if(!Editor.IsValid())
 	{
@@ -299,10 +353,7 @@ void FControlRigSchematicModel::UpdateElementKeyLinks()
 					const FRigElementKey ParentB = Hierarchy->GetFirstParent(ChildNode->GetKey());
 					if(!ParentA || !ParentB.IsValid() || ParentA != ParentB)
 					{
-						if(ChildNode->GetKey().Type != ERigElementType::Socket)
-						{
-							OutExpectedLinks.Emplace(Node->GetKey(), ChildNode->GetKey());
-						}
+						OutExpectedLinks.Emplace(Node->GetKey(), ChildNode->GetKey());
 					}
 				}
 			}
@@ -357,6 +408,10 @@ void FControlRigSchematicModel::UpdateElementKeyLinks()
 			const TArray<FRigBaseElement*> RootElements = Hierarchy->GetRootElements();
 			for(const FRigBaseElement* RootElement : RootElements)
 			{
+				if(RootElement->GetKey().Type == ERigElementType::Curve)
+				{
+					continue;
+				}
 				VisitElement(RootElement, nullptr, ExpectedLinks);
 			}
 			return ExpectedLinks;
@@ -421,7 +476,10 @@ void FControlRigSchematicModel::UpdateElementKeyLinks()
 
 void FControlRigSchematicModel::UpdateControlRigContent()
 {
-	UpdateElementKeyNodes();
+	{
+		const TGuardValue<bool> DisableUpdatingLinks(bUpdateElementKeyLinks, false);
+		UpdateElementKeyNodes();
+	}
 	UpdateElementKeyLinks();
 }
 
@@ -446,7 +504,7 @@ void FControlRigSchematicModel::UpdateConnector(const FRigElementKey& InElementK
 								// let's not draw the node in the schematic for now
 								if(ControlRigBlueprint->ModularRigModel.Connections.FindConnectorsFromTarget(ResolvedKey).Num() == 1)
 								{
-									ResolvedKey.Reset();
+									ResolvedKey = FRigElementKey();
 								}
 							}
 						}
@@ -503,8 +561,6 @@ void FControlRigSchematicModel::OnSetObjectBeingDebugged(UObject* InObject)
 		ControlRigBeingDebuggedPtr = ControlRig;
 		if(URigHierarchy* Hierarchy = ControlRig->GetHierarchy())
 		{
-			Hierarchy->OnModified().RemoveAll(this);
-			Hierarchy->OnModified().AddRaw(this, &FControlRigSchematicModel::OnHierarchyModified);
 			ControlRig->OnPostConstruction_AnyThread().AddRaw(this, &FControlRigSchematicModel::HandlePostConstruction);
 
 			UpdateControlRigContent();
@@ -512,128 +568,6 @@ void FControlRigSchematicModel::OnSetObjectBeingDebugged(UObject* InObject)
 	}
 }
 
-void FControlRigSchematicModel::OnHierarchyModified(ERigHierarchyNotification InNotif, URigHierarchy* InHierarchy, const FRigBaseElement* InElement)
-{
-	if(InHierarchy->IsCopyingHierarchy())
-	{
-		return;
-	}
-	
-	switch (InNotif)
-	{
-		case ERigHierarchyNotification::ElementAdded:
-		{
-			if (InElement)
-			{
-				if (InElement->GetType() == ERigElementType::Socket ||
-					InElement->GetType() == ERigElementType::Connector)
-				{
-					const FRigElementKey& NodeKey = InElement->GetKey();
-					if(!ContainsElementKeyNode(NodeKey))
-					{
-						AddElementKeyNode(NodeKey);
-					}
-				}
-			}
-			break;
-		}
-		case ERigHierarchyNotification::ElementRenamed:
-		{
-			if (InElement)
-			{
-				if (InElement->GetType() == ERigElementType::Socket ||
-					InElement->GetType() == ERigElementType::Connector)
-				{
-					const FString OldNameStr = InHierarchy->GetPreviousName(InElement->GetKey()).ToString();
-					FRigElementKey OldKey(*OldNameStr, InElement->GetType());
-
-					// just have to update the elementkey and lookup - since the guid doesn't need to change
-					FControlRigSchematicRigElementKeyNode* Node = const_cast<FControlRigSchematicRigElementKeyNode*>(FindElementKeyNode(OldKey));
-					if(Node)
-					{
-						ConfigureElementKeyNode(Node, InElement->GetKey());
-						RigElementKeyToGuid.Remove(OldKey);
-						RigElementKeyToGuid.Add(InElement->GetKey(), Node->GetGuid());
-					}
-
-					// update all of the links - update they keys
-					for(const TSharedPtr<FSchematicGraphLink>& Link : GetLinks())
-					{
-						if(FControlRigSchematicRigElementKeyLink* ElementKeyLink = Cast<FControlRigSchematicRigElementKeyLink>(Link.Get()))
-						{
-							if(ElementKeyLink->SourceKey == OldKey)
-							{
-								ElementKeyLink->SourceKey = InElement->GetKey();
-							}
-							if(ElementKeyLink->TargetKey == OldKey)
-							{
-								ElementKeyLink->TargetKey = InElement->GetKey();
-							}
-						}
-					}
-				}
-			}
-			break;
-		}
-		case ERigHierarchyNotification::ElementRemoved:
-		{
-			if (InElement)
-			{
-				if (InElement->GetType() == ERigElementType::Socket ||
-					InElement->GetType() == ERigElementType::Connector)
-				{
-					RemoveElementKeyNode(InElement->GetKey());
-				}
-			}
-			break;
-		}
-		case ERigHierarchyNotification::HierarchyReset:
-		case ERigHierarchyNotification::HierarchyCopied:
-		{
-			const TArray<FRigSocketElement*> Sockets = InHierarchy->GetElementsOfType<FRigSocketElement>();
-			for (const FRigSocketElement* Socket : Sockets)
-			{
-				if(!ContainsElementKeyNode(Socket->GetKey()))
-				{
-					AddElementKeyNode(Socket->GetKey());
-				}
-			}
-
-			const TArray<FRigConnectorElement*> Connectors = InHierarchy->GetElementsOfType<FRigConnectorElement>();
-			for (const FRigConnectorElement* Connector : Connectors)
-			{
-				if(!ContainsElementKeyNode(Connector->GetKey()))
-				{
-					AddElementKeyNode(Connector->GetKey());
-				}
-				UpdateConnector(Connector->GetKey());
-			}
-
-			// remove obsolete nodes
-			TArray<FGuid> GuidsToRemove;
-			for(const TSharedPtr<FSchematicGraphNode>& Node : Nodes)
-			{
-				if(const FControlRigSchematicRigElementKeyNode* ElementKeyNode = Cast<FControlRigSchematicRigElementKeyNode>(Node.Get()))
-				{
-					if(!InHierarchy->Contains(ElementKeyNode->GetKey()))
-					{
-						GuidsToRemove.Add(ElementKeyNode->GetGuid());
-					}
-				}
-			}
-			for(const FGuid& Guid : GuidsToRemove)
-			{
-				RemoveNode(Guid);
-			}
-			break;
-		}
-		default:
-		{
-			break;
-		}
-	}
-}
-
 void FControlRigSchematicModel::HandleModularRigModified(EModularRigNotification InNotification, const FRigModuleReference* InModule)
 {
 	switch (InNotification)
@@ -753,9 +687,9 @@ const FSlateBrush* FControlRigSchematicModel::GetBrushForKey(const FRigElementKe
 			if (ControlRigBlueprint.IsValid())
 			{
 				int32 Count = 0;
-				for(const TPair<FRigElementKey, FRigElementKey>& Pair : ControlRigBlueprint->ConnectionMap)
+				for(const FModularRigSingleConnection& Connection : ControlRigBlueprint->ModularRigModel.Connections)
 				{
-					if(Pair.Value == InKey)
+					if(Connection.Target == InKey)
 					{
 						Count++;
 						if(Count == 2)
@@ -827,9 +761,9 @@ const FSlateBrush* FControlRigSchematicModel::GetBrushForKey(const FRigElementKe
 					if (ControlRigBlueprint.IsValid())
 					{
 						int32 Count = 0;
-						for(const TPair<FRigElementKey, FRigElementKey>& Pair : ControlRigBlueprint->ConnectionMap)
+						for(const FModularRigSingleConnection& Connection : ControlRigBlueprint->ModularRigModel.Connections)
 						{
-							if(Pair.Value == InKey)
+							if(Connection.Target == InKey)
 							{
 								Count++;
 								if(Count == 2)
@@ -888,9 +822,9 @@ const FSlateBrush* FControlRigSchematicModel::GetBrushForNode(const FSchematicGr
 			if (ControlRigBlueprint.IsValid())
 			{
 				int32 Count = 0;
-				for(const TPair<FRigElementKey, FRigElementKey>& Pair : ControlRigBlueprint->ConnectionMap)
+				for(const FModularRigSingleConnection& Connection : ControlRigBlueprint->ModularRigModel.Connections)
 				{
-					if(Pair.Value == Node->GetKey())
+					if(Connection.Target == Node->GetKey())
 					{
 						Count++;
 						if(Count == 2)
@@ -1303,7 +1237,7 @@ void FControlRigSchematicModel::HandleSchematicNodeClicked(SSchematicGraphPanel*
 	}
 }
 
-void FControlRigSchematicModel::HandleSchematicBeginDrag(SSchematicGraphPanel* InPanel, SSchematicGraphNode* InNode, const FDragDropOperation& InDragDropOperation)
+void FControlRigSchematicModel::HandleSchematicBeginDrag(SSchematicGraphPanel* InPanel, SSchematicGraphNode* InNode, const TSharedPtr<FDragDropOperation>& InDragDropOperation)
 {
 	if (!ControlRigBlueprint.IsValid())
 	{
@@ -1321,108 +1255,111 @@ void FControlRigSchematicModel::HandleSchematicBeginDrag(SSchematicGraphPanel* I
 	{
 		DraggedKey = ElementKeyNode->GetKey();
 	}
-	else
+	else if(InDragDropOperation.IsValid())
 	{
-		const TArray<FRigElementKey> DraggedKeys = GetElementKeysFromDragDropEvent(InDragDropOperation, ControlRigBlueprint->GetDebuggedControlRig());
+		const TArray<FRigElementKey> DraggedKeys = GetElementKeysFromDragDropEvent(*InDragDropOperation.Get(), ControlRigBlueprint->GetDebuggedControlRig());
 		if(!DraggedKeys.IsEmpty())
 		{
 			DraggedKey = DraggedKeys[0];
 		}
 	}
 
-	if(!DraggedKey.IsValid())
-	{
-		return;
-	}
-	
-	URigHierarchy* Hierarchy = ControlRigBeingDebuggedPtr->GetHierarchy();
-	if (!Hierarchy)
-	{
-		return;
-	}
+	OnShowCandidatesForConnector(DraggedKey);
+}
 
-	FRigBaseElement* Element = Hierarchy->Find(DraggedKey);
-	if (!Element)
-	{
-		return;
-	}
-	
-	const FRigConnectorElement* Connector = Cast<FRigConnectorElement>(Element);
-	if (!Connector)
-	{
-		return;
-	}
+void FControlRigSchematicModel::HandleSchematicEndDrag(SSchematicGraphPanel* InPanel, SSchematicGraphNode* InNode, const TSharedPtr<FDragDropOperation>& InDragDropOperation)
+{
+	OnHideCandidatesForConnector();
+}
 
-	const UModularRig* ModularRig = Cast<UModularRig>(ControlRigBeingDebuggedPtr);
-	if (!ModularRig)
+void FControlRigSchematicModel::HandleSchematicEnterDrag(SSchematicGraphPanel* InPanel, const TSharedPtr<FDragDropOperation>& InDragDropOperation)
+{
+	if(!ControlRigBlueprint.IsValid())
 	{
 		return;
 	}
-
-	const FString ModulePath = Hierarchy->GetModulePath(DraggedKey);
-	const FRigModuleInstance* ModuleInstance = ModularRig->FindModule(ModulePath);
-	if (!ModuleInstance)
+	if(!InDragDropOperation.IsValid())
 	{
 		return;
 	}
-
-	if (UModularRigController* Controller = ControlRigBlueprint->GetModularRigController())
+	
+	const TArray<FRigElementKey> DraggedKeys = GetElementKeysFromDragDropEvent(*InDragDropOperation.Get(), ControlRigBlueprint->GetDebuggedControlRig());
+	if(!DraggedKeys.IsEmpty())
 	{
-		Controller->DisconnectConnector(Connector->GetKey(), true);
+		OnShowCandidatesForConnector(DraggedKeys[0]);
+		return;
 	}
 
-	const UModularRigRuleManager* RuleManager = Hierarchy->GetRuleManager();
-	const FModularRigResolveResult Result = RuleManager->FindMatches(Connector, ModuleInstance, ControlRigBeingDebuggedPtr->GetElementKeyRedirector());
-
-	const TArray<FRigElementResolveResult> Matches = Result.GetMatches();
-
-	for (const FRigElementResolveResult& Match : Matches)
+	if(InDragDropOperation->IsOfType<FAssetDragDropOp>())
 	{
-		// Create a temporary node that will be active only while this drag operation exists
-		if (!ContainsElementKeyNode(Match.GetKey()))
+		const TSharedPtr<FAssetDragDropOp> AssetDragDropOp = StaticCastSharedPtr<FAssetDragDropOp>(InDragDropOperation);
+		for (const FAssetData& AssetData : AssetDragDropOp->GetAssets())
 		{
-			const FSchematicGraphNode* NewNode = AddElementKeyNode(Match.GetKey());
-			TemporaryNodeGuids.Add(NewNode->GetGuid());
-		}
-	}
+			const UClass* AssetClass = AssetData.GetClass();
+			if (!AssetClass->IsChildOf(UControlRigBlueprint::StaticClass()))
+			{
+				continue;
+			}
 
-	// Fade all the unmatched nodes
-	for (const TSharedPtr<FSchematicGraphNode>& Node : Nodes)
-	{
-		if(FControlRigSchematicRigElementKeyNode* ExistingElementKeyNode = Cast<FControlRigSchematicRigElementKeyNode>(Node.Get()))
-		{
-			PreDragVisibilityPerNode.Add(Node->GetGuid(), ExistingElementKeyNode->Visibility);
-			ExistingElementKeyNode->SetVisibility(Matches.ContainsByPredicate([ExistingElementKeyNode](const FRigElementResolveResult& Match)
+			if(const UControlRigBlueprint* AssetBlueprint = Cast<UControlRigBlueprint>(AssetData.GetAsset()))
 			{
-				return ExistingElementKeyNode->GetKey() == Match.GetKey();
-			}) ? ESchematicGraphVisibility::Visible : ESchematicGraphVisibility::Hidden);
+				if(AssetBlueprint->IsControlRigModule())
+				{
+					if(const FRigModuleConnector* PrimaryConnector = AssetBlueprint->RigModuleSettings.FindPrimaryConnector())
+					{
+						OnShowCandidatesForConnector(PrimaryConnector);
+					}
+				}
+			}
 		}
-	};
+	}
 }
 
-void FControlRigSchematicModel::HandleSchematicEndDrag(SSchematicGraphPanel* InPanel, SSchematicGraphNode* InNode, const FDragDropOperation& InDragDropOperation)
+void FControlRigSchematicModel::HandleSchematicLeaveDrag(SSchematicGraphPanel* InPanel, const TSharedPtr<FDragDropOperation>& InDragDropOperation)
 {
-	if(TemporaryNodeGuids.IsEmpty() && PreDragVisibilityPerNode.IsEmpty())
+	OnHideCandidatesForConnector();
+}
+
+void FControlRigSchematicModel::HandleSchematicCancelDrag(SSchematicGraphPanel* InPanel, SSchematicGraphNode* InNode, const TSharedPtr<FDragDropOperation>& InDragDropOperation)
+{
+	if (!ControlRigBlueprint.IsValid())
 	{
 		return;
 	}
-	
-	for (FGuid& TempNodeGuid : TemporaryNodeGuids)
+
+	const FControlRigSchematicRigElementKeyNode* ElementKeyNode =
+		Cast<FControlRigSchematicRigElementKeyNode>(InNode->GetNodeData());
+	if(ElementKeyNode == nullptr)
 	{
-		RemoveNode(TempNodeGuid);
+		return;
 	}
-	TemporaryNodeGuids.Reset();
 
-	for(const TPair<FGuid, ESchematicGraphVisibility::Type>& Pair : PreDragVisibilityPerNode)
+	if (InDragDropOperation.IsValid() && InDragDropOperation->IsOfType<FSchematicGraphNodeDragDropOp>())
 	{
-		if(FControlRigSchematicRigElementKeyNode* Node = FindNode<FControlRigSchematicRigElementKeyNode>(Pair.Key))
+		const TSharedPtr<FSchematicGraphNodeDragDropOp> SchematicDragDropOp = StaticCastSharedPtr<FSchematicGraphNodeDragDropOp>(InDragDropOperation);
+		const TArray<FGuid> Sources = SchematicDragDropOp->GetElements();
+		TArray<FRigElementKey> Keys;
+		for(const FGuid& Source : Sources)
 		{
-			Node->Visibility = Pair.Value;
+			if(const FControlRigSchematicRigElementKeyNode* SourceNode = FindNode<FControlRigSchematicRigElementKeyNode>(Source))
+			{
+				Keys.Add(SourceNode->GetKey());
+			}
 		}
-	}
-	PreDragVisibilityPerNode.Reset();
 
-	UpdateElementKeyLinks();
+		FFunctionGraphTask::CreateAndDispatchWhenReady([this, Keys]()
+		{
+			if (UModularRigController* Controller = ControlRigBlueprint->GetModularRigController())
+			{
+				FScopedTransaction Transaction(LOCTEXT("DisconnectConnector", "Disconnect Connector"));
+				ControlRigBlueprint->Modify();
+				for(const FRigElementKey& Key : Keys)
+				{
+					Controller->DisconnectConnector(Key);
+				}
+			}
+		}, TStatId(), NULL, ENamedThreads::GameThread);
+	}
 }
 
 void FControlRigSchematicModel::HandleSchematicDrop(SSchematicGraphPanel* InPanel, SSchematicGraphNode* InNode, const FDragDropEvent& InDragDropEvent)
@@ -1508,12 +1445,6 @@ void FControlRigSchematicModel::HandleSchematicDrop(SSchematicGraphPanel* InPane
 									}
 								}
 							}
-
-							const FName TargetModulePath = Hierarchy->GetNameSpaceFName(TargetKey);
-							if(!TargetModulePath.IsNone())
-							{
-								(void)Controller->ReparentModule(ModulePath, TargetModulePath.ToString());
-							}
 							Controller->ConnectConnectorToElement(PrimaryConnectorKey, TargetKey, true, ControlRig->GetModularRigSettings().bAutoResolve);
 						}
 					}
@@ -1591,7 +1522,8 @@ void FControlRigSchematicModel::HandleSchematicDrop(SSchematicGraphPanel* InPane
 							{
 								if(Connector->IsPrimary())
 								{
-									Controller->ConnectConnectorToElement(Connector->GetKey(), TargetKey, true, ControlRig->GetModularRigSettings().bAutoResolve);
+									const FRigElementKey ConnectorKey = Connector->GetKey();
+									Controller->ConnectConnectorToElement(ConnectorKey, TargetKey, true, ControlRig->GetModularRigSettings().bAutoResolve);
 									return;
 								}
 							}
@@ -1612,14 +1544,15 @@ bool FControlRigSchematicModel::IsConnectorResolved(const FRigElementKey& InConn
 {
 	if(ControlRigBlueprint.IsValid() && InConnectorKey.Type == ERigElementType::Connector)
 	{
-		if(const FRigElementKey* TargetKey = ControlRigBlueprint->ConnectionMap.Find(InConnectorKey))
+		const FRigElementKey TargetKey = ControlRigBlueprint->ModularRigModel.Connections.FindTargetFromConnector(InConnectorKey);
+		if(TargetKey.IsValid())
 		{
 			// make sure the target exists
 			if(ControlRigBeingDebuggedPtr.IsValid())
 			{
 				if(const URigHierarchy* DebuggedHierarchy = ControlRigBeingDebuggedPtr->GetHierarchy())
 				{
-					if(!DebuggedHierarchy->Contains(*TargetKey))
+					if(!DebuggedHierarchy->Contains(TargetKey))
 					{
 						return false;
 					}
@@ -1628,7 +1561,7 @@ bool FControlRigSchematicModel::IsConnectorResolved(const FRigElementKey& InConn
 			
 			if(OutKey)
 			{
-				*OutKey = *TargetKey;
+				*OutKey = TargetKey;
 			}
 			return true;
 		}
@@ -1636,4 +1569,136 @@ bool FControlRigSchematicModel::IsConnectorResolved(const FRigElementKey& InConn
 	return false;
 }
 
+void FControlRigSchematicModel::OnShowCandidatesForConnector(const FRigElementKey& InConnectorKey)
+{
+	if(!InConnectorKey.IsValid())
+	{
+		return;
+	}
+	
+	URigHierarchy* Hierarchy = ControlRigBeingDebuggedPtr->GetHierarchy();
+	if (!Hierarchy)
+	{
+		return;
+	}
+
+	FRigBaseElement* Element = Hierarchy->Find(InConnectorKey);
+	if (!Element)
+	{
+		return;
+	}
+	
+	const FRigConnectorElement* Connector = Cast<FRigConnectorElement>(Element);
+	if (!Connector)
+	{
+		return;
+	}
+
+	const UModularRig* ModularRig = Cast<UModularRig>(ControlRigBeingDebuggedPtr);
+	if (!ModularRig)
+	{
+		return;
+	}
+
+	const FString ModulePath = Hierarchy->GetModulePath(InConnectorKey);
+	const FRigModuleInstance* ModuleInstance = ModularRig->FindModule(ModulePath);
+	if (!ModuleInstance)
+	{
+		return;
+	}
+
+	if(const FControlRigSchematicRigElementKeyNode* Node = FindElementKeyNode(InConnectorKey))
+	{
+		// close all expanded groups
+		FGuid ParentGuid = Node->GetParentNodeGuid();
+		while(ParentGuid.IsValid())
+		{
+			if(FSchematicGraphNode* ParentNode = FindNode(ParentGuid))
+			{
+				if(FSchematicGraphGroupNode* GroupNode = Cast<FSchematicGraphGroupNode>(ParentNode))
+				{
+					GroupNode->SetExpanded(false);
+				}
+				ParentGuid = ParentNode->GetParentNodeGuid();
+			}
+			else
+			{
+				break;
+			}
+		}
+	}
+	
+	const UModularRigRuleManager* RuleManager = Hierarchy->GetRuleManager();
+	const FModularRigResolveResult Result = RuleManager->FindMatches(Connector, ModuleInstance, ControlRigBeingDebuggedPtr->GetElementKeyRedirector());
+	OnShowCandidatesForMatches(Result);
+}
+
+void FControlRigSchematicModel::OnShowCandidatesForConnector(const FRigModuleConnector* InModuleConnector)
+{
+	check(InModuleConnector);
+	
+	URigHierarchy* Hierarchy = ControlRigBeingDebuggedPtr->GetHierarchy();
+	if (!Hierarchy)
+	{
+		return;
+	}
+	
+	const UModularRigRuleManager* RuleManager = Hierarchy->GetRuleManager();
+	const FModularRigResolveResult Result = RuleManager->FindMatches(InModuleConnector);
+	OnShowCandidatesForMatches(Result);
+}
+
+void FControlRigSchematicModel::OnShowCandidatesForMatches(const FModularRigResolveResult& InMatches)
+{
+	const TArray<FRigElementResolveResult> Matches = InMatches.GetMatches();
+
+	for (const FRigElementResolveResult& Match : Matches)
+	{
+		// Create a temporary node that will be active only while this drag operation exists
+		if (!ContainsElementKeyNode(Match.GetKey()))
+		{
+			const FSchematicGraphNode* NewNode = AddElementKeyNode(Match.GetKey());
+			TemporaryNodeGuids.Add(NewNode->GetGuid());
+		}
+	}
+
+	// Fade all the unmatched nodes
+	for (const TSharedPtr<FSchematicGraphNode>& Node : Nodes)
+	{
+		if(FControlRigSchematicRigElementKeyNode* ExistingElementKeyNode = Cast<FControlRigSchematicRigElementKeyNode>(Node.Get()))
+		{
+			PreDragVisibilityPerNode.Add(Node->GetGuid(), ExistingElementKeyNode->Visibility);
+			ExistingElementKeyNode->SetVisibility(Matches.ContainsByPredicate([ExistingElementKeyNode](const FRigElementResolveResult& Match)
+			{
+				return ExistingElementKeyNode->GetKey() == Match.GetKey();
+			}) ? ESchematicGraphVisibility::Visible : ESchematicGraphVisibility::Hidden);
+		}
+	};
+}
+
+void FControlRigSchematicModel::OnHideCandidatesForConnector()
+{
+	if(TemporaryNodeGuids.IsEmpty() && PreDragVisibilityPerNode.IsEmpty())
+	{
+		return;
+	}
+	
+	for (FGuid& TempNodeGuid : TemporaryNodeGuids)
+	{
+		RemoveNode(TempNodeGuid);
+	}
+	TemporaryNodeGuids.Reset();
+
+	for(const TPair<FGuid, ESchematicGraphVisibility::Type>& Pair : PreDragVisibilityPerNode)
+	{
+		if(FControlRigSchematicRigElementKeyNode* Node = FindNode<FControlRigSchematicRigElementKeyNode>(Pair.Key))
+		{
+			Node->Visibility = Pair.Value;
+		}
+	}
+	PreDragVisibilityPerNode.Reset();
+
+	UpdateElementKeyLinks();
+}
+
 #undef LOCTEXT_NAMESPACE
diff --git a/Engine/Plugins/Animation/ControlRig/Source/ControlRigEditor/Private/Editor/ControlRigSchematicModel.h b/Engine/Plugins/Animation/ControlRig/Source/ControlRigEditor/Private/Editor/ControlRigSchematicModel.h
--- a/Engine/Plugins/Animation/ControlRig/Source/ControlRigEditor/Private/Editor/ControlRigSchematicModel.h
+++ b/Engine/Plugins/Animation/ControlRig/Source/ControlRigEditor/Private/Editor/ControlRigSchematicModel.h
@@ -75,6 +75,7 @@ public:
 	virtual ~FControlRigSchematicModel() override;
 	
 	void SetEditor(const TSharedRef<FControlRigEditor>& InEditor);
+	virtual void ApplyToPanel(SSchematicGraphPanel* InPanel) override;
 
 	virtual void Reset() override;
 	virtual void Tick(float InDeltaTime) override;
@@ -90,7 +91,6 @@ public:
 	void UpdateConnector(const FRigElementKey& InElementKey);
 
 	void OnSetObjectBeingDebugged(UObject* InObject);
-	void OnHierarchyModified(ERigHierarchyNotification InNotif, URigHierarchy* InHierarchy, const FRigBaseElement* InElement);
 	void HandleModularRigModified(EModularRigNotification InNotification, const FRigModuleReference* InModule);
 
 	virtual FSchematicGraphGroupNode* AddAutoGroupNode() override;
@@ -115,12 +115,19 @@ private:
 	void ConfigureElementKeyNode(FControlRigSchematicRigElementKeyNode* InNode, const FRigElementKey& InKey);
 
 	void HandleSchematicNodeClicked(SSchematicGraphPanel* InPanel, SSchematicGraphNode* InNode, const FPointerEvent& InMouseEvent);
-	void HandleSchematicBeginDrag(SSchematicGraphPanel* InPanel, SSchematicGraphNode* InNode, const FDragDropOperation& InDragDropOperation);
-	void HandleSchematicEndDrag(SSchematicGraphPanel* InPanel, SSchematicGraphNode* InNode, const FDragDropOperation& InDragDropOperation);
+	void HandleSchematicBeginDrag(SSchematicGraphPanel* InPanel, SSchematicGraphNode* InNode, const TSharedPtr<FDragDropOperation>& InDragDropOperation);
+	void HandleSchematicEndDrag(SSchematicGraphPanel* InPanel, SSchematicGraphNode* InNode, const TSharedPtr<FDragDropOperation>& InDragDropOperation);
+	void HandleSchematicEnterDrag(SSchematicGraphPanel* InPanel, const TSharedPtr<FDragDropOperation>& InDragDropOperation);
+	void HandleSchematicLeaveDrag(SSchematicGraphPanel* InPanel, const TSharedPtr<FDragDropOperation>& InDragDropOperation);
+	void HandleSchematicCancelDrag(SSchematicGraphPanel* InPanel, SSchematicGraphNode* InNode, const TSharedPtr<FDragDropOperation>& InDragDropOperation);
 	void HandleSchematicDrop(SSchematicGraphPanel* InPanel, SSchematicGraphNode* InNode, const FDragDropEvent& InDragDropEvent);
 	void HandlePostConstruction(UControlRig* Subject, const FName& InEventName);
 	bool IsConnectorResolved(const FRigElementKey& InConnectorKey, FRigElementKey* OutKey = nullptr) const;
 
+	void OnShowCandidatesForConnector(const FRigElementKey& InConnectorKey);
+	void OnShowCandidatesForConnector(const FRigModuleConnector* InModuleConnector);
+	void OnShowCandidatesForMatches(const FModularRigResolveResult& InMatches);
+	void OnHideCandidatesForConnector();
 
 	TWeakPtr<FControlRigEditor> ControlRigEditor;
 	TWeakObjectPtr<UControlRigBlueprint> ControlRigBlueprint;
@@ -130,6 +137,7 @@ private:
 	TMap<FGuid, ESchematicGraphVisibility::Type> PreDragVisibilityPerNode;
 	TMap<FRigElementKey, FGuid> RigElementKeyToGuid;
 	mutable TMap<FSoftObjectPath, FSlateBrush> ModuleIcons;
+	bool bUpdateElementKeyLinks = true;
 
 	friend class FControlRigEditor;
 };
\ No newline at end of file
diff --git a/Engine/Plugins/Animation/ControlRig/Source/ControlRigEditor/Private/Editor/SModularRigModel.cpp b/Engine/Plugins/Animation/ControlRig/Source/ControlRigEditor/Private/Editor/SModularRigModel.cpp
--- a/Engine/Plugins/Animation/ControlRig/Source/ControlRigEditor/Private/Editor/SModularRigModel.cpp
+++ b/Engine/Plugins/Animation/ControlRig/Source/ControlRigEditor/Private/Editor/SModularRigModel.cpp
@@ -781,12 +781,12 @@ void SModularRigModel::OnHierarchyModified(ERigHierarchyNotification InNotif, UR
 			const FRigConnectorElement* Connector = Cast<FRigConnectorElement>(InElement);
 			if(Connector == nullptr)
 			{
-				for(const TPair<FRigElementKey, FRigElementKey>& Pair : ControlRigBlueprint->ConnectionMap)
+				for(const FModularRigSingleConnection& Connection : ControlRigBlueprint->ModularRigModel.Connections)
 				{
-					check(Pair.Key.Type == ERigElementType::Connector);
-					if(Pair.Value == InElement->GetKey())
+					check(Connection.Connector.Type == ERigElementType::Connector);
+					if(Connection.Target == InElement->GetKey())
 					{
-						if(const FRigConnectorElement* TargetConnector = InHierarchy->Find<FRigConnectorElement>(Pair.Key))
+						if(const FRigConnectorElement* TargetConnector = InHierarchy->Find<FRigConnectorElement>(Connection.Connector))
 						{
 							OnHierarchyModified(InNotif, InHierarchy, TargetConnector);
 						}
diff --git a/Engine/Plugins/Animation/ControlRig/Source/ControlRigEditor/Private/Editor/SModularRigTreeView.cpp b/Engine/Plugins/Animation/ControlRig/Source/ControlRigEditor/Private/Editor/SModularRigTreeView.cpp
--- a/Engine/Plugins/Animation/ControlRig/Source/ControlRigEditor/Private/Editor/SModularRigTreeView.cpp
+++ b/Engine/Plugins/Animation/ControlRig/Source/ControlRigEditor/Private/Editor/SModularRigTreeView.cpp
@@ -132,6 +132,7 @@ void SModularRigModelItem::Construct(const FArguments& InArgs, const TSharedRef<
 	TSharedPtr< SHorizontalBox > HorizontalBox;
 	TSharedPtr<SVerticalBox> ComboButtonBox;
 
+	// todo: this should be done lazily - the matches need to be pulled when the user clicks
 	FModularRigResolveResult ConnectorMatches;
 	if (const UModularRig* ModularRig = Delegates.GetModularRig())
 	{
diff --git a/Engine/Plugins/Animation/ControlRig/Source/ControlRigEditor/Private/Editor/SRigHierarchy.cpp b/Engine/Plugins/Animation/ControlRig/Source/ControlRigEditor/Private/Editor/SRigHierarchy.cpp
--- a/Engine/Plugins/Animation/ControlRig/Source/ControlRigEditor/Private/Editor/SRigHierarchy.cpp
+++ b/Engine/Plugins/Animation/ControlRig/Source/ControlRigEditor/Private/Editor/SRigHierarchy.cpp
@@ -1613,9 +1613,10 @@ FRigElementKey SRigHierarchy::OnGetResolvedKey(const FRigElementKey& InKey)
 {
 	if (const UControlRigBlueprint* Blueprint = ControlRigEditor.Pin()->GetControlRigBlueprint())
 	{
-		if(const FRigElementKey* ResolvedKey = Blueprint->ConnectionMap.Find(InKey))
+		const FRigElementKey ResolvedKey = Blueprint->ModularRigModel.Connections.FindTargetFromConnector(InKey);
+		if(ResolvedKey.IsValid())
 		{
-			return *ResolvedKey;
+			return ResolvedKey;
 		}
 	}
 	return InKey;
diff --git a/Engine/Source/Editor/AnimationEditorWidgets/Private/SchematicGraphPanel/SSchematicGraphPanel.cpp b/Engine/Source/Editor/AnimationEditorWidgets/Private/SchematicGraphPanel/SSchematicGraphPanel.cpp
--- a/Engine/Source/Editor/AnimationEditorWidgets/Private/SchematicGraphPanel/SSchematicGraphPanel.cpp
+++ b/Engine/Source/Editor/AnimationEditorWidgets/Private/SchematicGraphPanel/SSchematicGraphPanel.cpp
@@ -13,25 +13,17 @@
 
 #define LOCTEXT_NAMESPACE "SSchematicGraphPanel"
 
-TSharedRef<FSchematicGraphNodeDragDropOp> FSchematicGraphNodeDragDropOp::New(TArray<SSchematicGraphNode*> InSchematicGraphNodes, const TArray<FGuid>& InElements, FSchematicGraphNodeDragDropOp::FOnEndDrag InOnEndDragDelegate)
+TSharedRef<FSchematicGraphNodeDragDropOp> FSchematicGraphNodeDragDropOp::New(TArray<SSchematicGraphNode*> InSchematicGraphNodes, const TArray<FGuid>& InElements)
 {
 	TSharedRef<FSchematicGraphNodeDragDropOp> Operation = MakeShared<FSchematicGraphNodeDragDropOp>();
 	Operation->SchematicGraphNodes = InSchematicGraphNodes;
 	Operation->Elements = InElements;
-	Operation->OnEndDragDelegate = InOnEndDragDelegate;
 	Operation->Construct();
 	return Operation;
 }
 
 FSchematicGraphNodeDragDropOp::~FSchematicGraphNodeDragDropOp()
 {
-	if (OnEndDragDelegate.IsBound())
-	{
-		for (SSchematicGraphNode* Node : SchematicGraphNodes)
-		{
-			OnEndDragDelegate.Execute(Node, *this);
-		}
-	}
 }
 
 TSharedPtr<SWidget> FSchematicGraphNodeDragDropOp::GetDefaultDecorator() const
@@ -390,7 +382,7 @@ FReply SSchematicGraphNode::OnDrop(const FGeometry& MyGeometry, const FDragDropE
 	
 	SNode::OnDrop(MyGeometry, DragDropEvent);
 	OnDropDelegate.ExecuteIfBound(this, DragDropEvent);
-	OnEndDragDelegate.ExecuteIfBound(this, *DragDropEvent.GetOperation().Get());
+	OnEndDragDelegate.ExecuteIfBound(this, DragDropEvent.GetOperation());
 	return FReply::Unhandled();
 }
 
@@ -419,8 +411,8 @@ FReply SSchematicGraphNode::OnDragDetected(const FGeometry& MyGeometry, const FP
 				const FVector2d LocalMousePosition =  (AbsoluteMousePosition - MyGeometry.GetAbsolutePosition()) / MyGeometry.GetAccumulatedLayoutTransform().GetScale();
 				OffsetDuringDrag = -LocalMousePosition;
 			
-				const TSharedRef<FSchematicGraphNodeDragDropOp> DragDropOp = FSchematicGraphNodeDragDropOp::New({this}, MoveTemp(DraggedElements), OnEndDragDelegate);
-				OnBeginDragDelegate.ExecuteIfBound(this, DragDropOp.Get());
+				const TSharedRef<FSchematicGraphNodeDragDropOp> DragDropOp = FSchematicGraphNodeDragDropOp::New({this}, MoveTemp(DraggedElements));
+				OnBeginDragDelegate.ExecuteIfBound(this, DragDropOp.ToSharedPtr());
 				return FReply::Handled().BeginDragDrop(DragDropOp);
 			}
 		}
@@ -544,18 +536,9 @@ void SSchematicGraphPanel::Construct(const FArguments& InArgs)
 	OnNodeClickedDelegate = InArgs._OnNodeClicked;
 	OnBeginDragDelegate = InArgs._OnBeginDrag;
 	OnEndDragDelegate = InArgs._OnEndDrag;
+	OnEnterDragDelegate = InArgs._OnEnterDrag;
+	OnLeaveDragDelegate = InArgs._OnLeaveDrag;
 	OnDropDelegate = InArgs._OnDrop;
-	BackgroundAlpha = InArgs._BackgroundAlpha;
-
-	TEasingAttributeInterpolator<float>::FSettings FloatInterpolationSettings(EEasingInterpolatorType::CubicEaseOut, 0.75f);
-	FadeBackgroundAlpha = FFloatAttribute::CreateWithGetter(FloatInterpolationSettings, FFloatAttribute::FGetter::CreateLambda([this]() -> float
-	{
-		if(!bIsOverlay || (GetVisibility() == EVisibility::Hidden))
-		{
-			return 0.f;
-		}
-		return BackgroundAlpha.Get();
-	}));
 	
 	SNodePanel::Construct();
 
@@ -567,7 +550,8 @@ void SSchematicGraphPanel::Construct(const FArguments& InArgs)
 		GraphData->OnLinkAdded().AddSP(this, &SSchematicGraphPanel::AddLink);
 		GraphData->OnLinkRemoved().AddSP(this, &SSchematicGraphPanel::RemoveLink);
 		GraphData->OnGraphReset().AddSP(this, &SSchematicGraphPanel::RebuildPanel);
-	}
+		GraphData->ApplyToPanel(this);
+	};
 }
 
 void SSchematicGraphPanel::RebuildPanel()
@@ -818,99 +802,6 @@ int32 SSchematicGraphPanel::OnPaint(const FPaintArgs& Args, const FGeometry& All
 		}
 	}
 
-	/*
-	FBox2d Bounds(EForceInit::ForceInit);
-	for (int32 ChildIndex = 0; ChildIndex < ArrangedChildren.Num(); ++ChildIndex)
-	{
-		if(NodeVisibilityByIndex[ChildIndex] == ESchematicGraphVisibility::Hidden)
-		{
-			continue;
-		}
-		if(NodeCenterByIndex[ChildIndex].IsNearlyZero())
-		{
-			continue;
-		}
-		Bounds += NodeCenterByIndex[ChildIndex] + FVector2d(120, 120);
-		Bounds += NodeCenterByIndex[ChildIndex] - FVector2d(120, 120);
-	}
-	*/
-
-	const float Alpha = FadeBackgroundAlpha->Get();
-	if (Alpha > 0.f) // && !Bounds.GetSize().IsNearlyZero())
-	{
-		const FSlateBrush* Brush = FAppStyle::GetBrush(TEXT("Graph.Panel.SolidBackground"));
-		const FLinearColor TransparentGrey = FLinearColor(0.0, 0.0, 0.0, Alpha);
-		FSlateDrawElement::MakeBox(
-				OutDrawElements,
-				BackgroundLayer,
-				AllottedGeometry.ToPaintGeometry(),
-				//AllottedGeometry.ToPaintGeometry(Bounds.GetSize(), FSlateLayoutTransform(Bounds.Min)),
-				Brush,
-				ESlateDrawEffect::None,
-				TransparentGrey
-			);
-
-		/*
-		// Render an outline around the bounds
-		TArray<FVector2d> LinePoints;
-		LinePoints.AddZeroed(2);
-
-		constexpr float Step = 10.f;
-		for(float X = Bounds.Min.X; X < Bounds.Max.X - Step; X += Step * 2)
-		{
-			LinePoints[0] = FVector2d(X, Bounds.Min.Y + 0.5f);
-			LinePoints[1] = LinePoints[0] + FVector2d(Step, 0.f);
-			
-			FSlateDrawElement::MakeLines(
-				OutDrawElements,
-				OutlineLayer,
-				AllottedGeometry.ToPaintGeometry(),
-				LinePoints,
-				ESlateDrawEffect::None,
-				FLinearColor(0.7, 0.7, 0.7, 0.95),
-				true);
-
-			LinePoints[0] = FVector2d(X, Bounds.Max.Y - 0.5f);
-			LinePoints[1] = LinePoints[0] + FVector2d(Step, 0.f);
-			
-			FSlateDrawElement::MakeLines(
-				OutDrawElements,
-				OutlineLayer,
-				AllottedGeometry.ToPaintGeometry(),
-				LinePoints,
-				ESlateDrawEffect::None,
-				FLinearColor(0.7, 0.7, 0.7, 0.95),
-				true);
-		}
-		for(float Y = Bounds.Min.Y; Y < Bounds.Max.Y - Step; Y += Step * 2)
-		{
-			LinePoints[0] = FVector2d(Bounds.Min.X + 0.5f, Y);
-			LinePoints[1] = LinePoints[0] + FVector2d(0.f, Step);
-			
-			FSlateDrawElement::MakeLines(
-				OutDrawElements,
-				OutlineLayer,
-				AllottedGeometry.ToPaintGeometry(),
-				LinePoints,
-				ESlateDrawEffect::None,
-				FLinearColor(0.7, 0.7, 0.7, 0.95),
-				true);
-
-			LinePoints[0] = FVector2d(Bounds.Max.X - 0.5f, Y);
-			LinePoints[1] = LinePoints[0] + FVector2d(0.f, Step);
-			
-			FSlateDrawElement::MakeLines(
-				OutDrawElements,
-				OutlineLayer,
-				AllottedGeometry.ToPaintGeometry(),
-				LinePoints,
-				ESlateDrawEffect::None,
-				FLinearColor(0.7, 0.7, 0.7, 0.95),
-				true);
-		}
-		*/
-	}
-	
 	// draw all of the links
 	for(const TPair<FGuid, TSharedPtr<FSchematicLinkWidgetInfo>>& Pair : LinkByGuid)
 	{
@@ -1081,26 +972,38 @@ void SSchematicGraphPanel::Tick(float DeltaTime)
 		GraphData->Tick(DeltaTime);
 	}
 	
-	bIsDragDropping = false;
 	FSlateApplication& Application = FSlateApplication::Get();
 	if (Application.IsDragDropping())
 	{
 		TSharedPtr<FDragDropOperation> DragDropOp = Application.GetDragDroppingContent();
 		if (DragDropOp.IsValid())
 		{
+			const FVector2D MouseCursorLocation = FSlateApplication::Get().GetCursorPos();
+			if(GetPaintSpaceGeometry().GetRenderBoundingRect().ContainsPoint(MouseCursorLocation))
+			{
+				// when dragging and dropping we want the UI to update
+				// for that we are forcing an update on the tick of the animated attributes
+				FAnimatedAttributeManager::Get().Tick(DeltaTime);
+
+				// if we haven't seen this operation yet we need to let our model know
+				if(!DragDropOpFromOutside.IsValid())
+				{
+					DragDropOpFromOutside = DragDropOp;
+					OnEnterDragEvent(DragDropOpFromOutside);
+				}
+			}
+
 			bIsDragDropping = true;
 		}
 	}
-
-	if(bIsDragDropping)
+	else
 	{
-		const FVector2D MouseCursorLocation = FSlateApplication::Get().GetCursorPos();
-		if(GetPaintSpaceGeometry().GetRenderBoundingRect().ContainsPoint(MouseCursorLocation))
+		bIsDragDropping = false;
+		if(DragDropOpFromOutside)
 		{
-			// when dragging and dropping we want the UI to update
-			// for that we are forcing an update on the tick of the animated attributes
-			FAnimatedAttributeManager::Get().Tick(DeltaTime);
+			OnLeaveDragEvent(DragDropOpFromOutside);
 		}
+		DragDropOpFromOutside.Reset();
 	}
 
 	NodesTopLeft.Reset();
@@ -1201,18 +1104,41 @@ void SSchematicGraphPanel::OnNodeClicked(SSchematicGraphNode* Node, const FPoint
 	OnNodeClickedDelegate.ExecuteIfBound(this, Node, MouseEvent);
 }
 
-void SSchematicGraphPanel::OnBeginDragEvent(SSchematicGraphNode* Node, const FDragDropOperation& InDragDropOp)
+void SSchematicGraphPanel::OnBeginDragEvent(SSchematicGraphNode* Node, const TSharedPtr<FDragDropOperation>& InDragDropOp)
 {
+	DropTarget.Reset();
 	OnBeginDragDelegate.ExecuteIfBound(this, Node, InDragDropOp);
 }
 
-void SSchematicGraphPanel::OnEndDragEvent(SSchematicGraphNode* Node, const FDragDropOperation& InDragDropOp)
+void SSchematicGraphPanel::OnEndDragEvent(SSchematicGraphNode* Node, const TSharedPtr<FDragDropOperation>& InDragDropOp)
 {
+	if(!DropTarget.IsSet())
+	{
+		OnCancelDragEvent(Node, InDragDropOp);
+	}
 	OnEndDragDelegate.ExecuteIfBound(this, Node, InDragDropOp);
+	DropTarget.Reset();
+}
+
+void SSchematicGraphPanel::OnEnterDragEvent(const TSharedPtr<FDragDropOperation>& InDragDropEvent)
+{
+	OnEnterDragDelegate.ExecuteIfBound(this, InDragDropEvent);
+}
+
+void SSchematicGraphPanel::OnLeaveDragEvent(const TSharedPtr<FDragDropOperation>& InDragDropEvent)
+{
+	OnLeaveDragDelegate.ExecuteIfBound(this, InDragDropEvent);
+}
+
+void SSchematicGraphPanel::OnCancelDragEvent(SSchematicGraphNode* Node, const TSharedPtr<FDragDropOperation>& InDragDropEvent)
+{
+	OnCancelDragDelegate.ExecuteIfBound(this, Node, InDragDropEvent);
+	DropTarget.Reset();
 }
 
 void SSchematicGraphPanel::OnDropEvent(SSchematicGraphNode* Node, const FDragDropEvent& InDragDropEvent)
 {
+	DropTarget = Node ? Node->GetGuid() : FGuid(); 
 	OnDropDelegate.ExecuteIfBound(this, Node, InDragDropEvent);
 }
 
diff --git a/Engine/Source/Editor/AnimationEditorWidgets/Private/SchematicGraphPanel/SchematicGraphModel.cpp b/Engine/Source/Editor/AnimationEditorWidgets/Private/SchematicGraphPanel/SchematicGraphModel.cpp
--- a/Engine/Source/Editor/AnimationEditorWidgets/Private/SchematicGraphPanel/SchematicGraphModel.cpp
+++ b/Engine/Source/Editor/AnimationEditorWidgets/Private/SchematicGraphPanel/SchematicGraphModel.cpp
@@ -870,6 +870,28 @@ void FSchematicGraphModel::Tick(float InDeltaTime)
 			LastExpandedNode = FGuid();
 		}
 	}
+
+	// remove invalid parent / child relationships
+	for(const TSharedPtr<FSchematicGraphNode>& Node : Nodes)
+	{
+		if(Node->ParentNodeGuid.IsValid())
+		{
+			if(!NodeByGuid.Contains(Node->ParentNodeGuid))
+			{
+				Node->ParentNodeGuid = FGuid();
+			}
+		}
+
+		const FGuid Guid = Node->GetGuid();
+		Node->ChildNodeGuids.RemoveAll([this, Guid](const FGuid& InChildGuid) -> bool
+		{
+			if(const FSchematicGraphNode* ChildNode = FindNode(InChildGuid))
+			{
+				return ChildNode->ParentNodeGuid != Guid;
+			}
+			return true;
+		});
+	}
 }
 
 #undef LOCTEXT_NAMESPACE
diff --git a/Engine/Source/Editor/AnimationEditorWidgets/Public/SchematicGraphPanel/SSchematicGraphPanel.h b/Engine/Source/Editor/AnimationEditorWidgets/Public/SchematicGraphPanel/SSchematicGraphPanel.h
--- a/Engine/Source/Editor/AnimationEditorWidgets/Public/SchematicGraphPanel/SSchematicGraphPanel.h
+++ b/Engine/Source/Editor/AnimationEditorWidgets/Public/SchematicGraphPanel/SSchematicGraphPanel.h
@@ -16,9 +16,7 @@ class ANIMATIONEDITORWIDGETS_API FSchematicGraphNodeDragDropOp : public FDragDro
 public:
 	DRAG_DROP_OPERATOR_TYPE(FSchematicGraphNodeDragDropOp, FDragDropOperation)
 
-	DECLARE_DELEGATE_TwoParams(FOnEndDrag, SSchematicGraphNode*, const FDragDropOperation&);
-
-	static TSharedRef<FSchematicGraphNodeDragDropOp> New(TArray<SSchematicGraphNode*> InSchematicGraphNodes, const TArray<FGuid>& InElements, FOnEndDrag InOnEndDragDelegate);
+	static TSharedRef<FSchematicGraphNodeDragDropOp> New(TArray<SSchematicGraphNode*> InSchematicGraphNodes, const TArray<FGuid>& InElements);
 
 	virtual ~FSchematicGraphNodeDragDropOp() override;
 	
@@ -48,9 +46,6 @@ private:
 
 	/** Data for the property paths this item represents */
 	TArray<FGuid> Elements;
-
-	/** Delegate to call when this drag operation ends */
-	FOnEndDrag OnEndDragDelegate;
 };
 
 class ANIMATIONEDITORWIDGETS_API SSchematicGraphNode : public SNodePanel::SNode
@@ -62,8 +57,8 @@ public:
 	typedef TAnimatedAttribute<FLinearColor> FLinearColorAttribute;
 	
 	DECLARE_DELEGATE_TwoParams(FOnClicked, SSchematicGraphNode*, const FPointerEvent&);
-	DECLARE_DELEGATE_TwoParams(FOnBeginDrag, SSchematicGraphNode*, const FDragDropOperation&);
-	DECLARE_DELEGATE_TwoParams(FOnEndDrag, SSchematicGraphNode*, const FDragDropOperation&);
+	DECLARE_DELEGATE_TwoParams(FOnBeginDrag, SSchematicGraphNode*, const TSharedPtr<FDragDropOperation>&);
+	DECLARE_DELEGATE_TwoParams(FOnEndDrag, SSchematicGraphNode*, const TSharedPtr<FDragDropOperation>&);
 	DECLARE_DELEGATE_TwoParams(FOnDrop, SSchematicGraphNode*, const FDragDropEvent&);
 
 	SLATE_BEGIN_ARGS(SSchematicGraphNode);
@@ -146,12 +141,14 @@ public:
 	using FLinearColorAttribute = SSchematicGraphNode::FLinearColorAttribute;
 
 	DECLARE_DELEGATE_ThreeParams(FOnNodeClicked, SSchematicGraphPanel*, SSchematicGraphNode*, const FPointerEvent&);
-	DECLARE_DELEGATE_ThreeParams(FOnBeginDrag, SSchematicGraphPanel*, SSchematicGraphNode*, const FDragDropOperation&);
-	DECLARE_DELEGATE_ThreeParams(FOnEndDrag, SSchematicGraphPanel*, SSchematicGraphNode*, const FDragDropOperation&);
+	DECLARE_DELEGATE_ThreeParams(FOnBeginDrag, SSchematicGraphPanel*, SSchematicGraphNode*, const TSharedPtr<FDragDropOperation>&);
+	DECLARE_DELEGATE_ThreeParams(FOnEndDrag, SSchematicGraphPanel*, SSchematicGraphNode*, const TSharedPtr<FDragDropOperation>&);
+	DECLARE_DELEGATE_TwoParams(FOnEnterDrag, SSchematicGraphPanel*, const TSharedPtr<FDragDropOperation>&);
+	DECLARE_DELEGATE_TwoParams(FOnLeaveDrag, SSchematicGraphPanel*, const TSharedPtr<FDragDropOperation>&);
+	DECLARE_DELEGATE_ThreeParams(FOnCancelDrag, SSchematicGraphPanel*, SSchematicGraphNode*, const TSharedPtr<FDragDropOperation>&);
 	DECLARE_DELEGATE_ThreeParams(FOnDrop, SSchematicGraphPanel*, SSchematicGraphNode*, const FDragDropEvent&);
 	
 	SLATE_BEGIN_ARGS(SSchematicGraphPanel)
-		: _BackgroundAlpha(0.5f)
 	{}
 	SLATE_ARGUMENT(bool, IsOverlay)
 	SLATE_ARGUMENT(FSchematicGraphModel*, GraphData)
@@ -163,8 +160,10 @@ public:
 	SLATE_EVENT(FOnNodeClicked, OnNodeClicked)
 	SLATE_EVENT(FOnBeginDrag, OnBeginDrag)
 	SLATE_EVENT(FOnEndDrag, OnEndDrag)
+	SLATE_EVENT(FOnEnterDrag, OnEnterDrag)
+	SLATE_EVENT(FOnLeaveDrag, OnLeaveDrag)
+	SLATE_EVENT(FOnCancelDrag, OnCancelDrag)
 	SLATE_EVENT(FOnDrop, OnDrop)
-	SLATE_ATTRIBUTE(float, BackgroundAlpha)
 	SLATE_END_ARGS()
 
 	struct FSchematicLinkWidgetInfo
@@ -214,8 +213,11 @@ public:
 
 	void ToggleVisibility();
 	void OnNodeClicked(SSchematicGraphNode* Node, const FPointerEvent& MouseEvent);
-	void OnBeginDragEvent(SSchematicGraphNode* Node, const FDragDropOperation& InDragDropEvent);
-	void OnEndDragEvent(SSchematicGraphNode* Node, const FDragDropOperation& InDragDropEvent);
+	void OnBeginDragEvent(SSchematicGraphNode* Node, const TSharedPtr<FDragDropOperation>& InDragDropEvent);
+	void OnEndDragEvent(SSchematicGraphNode* Node, const TSharedPtr<FDragDropOperation>& InDragDropEvent);
+	void OnEnterDragEvent(const TSharedPtr<FDragDropOperation>& InDragDropEvent);
+	void OnLeaveDragEvent(const TSharedPtr<FDragDropOperation>& InDragDropEvent);
+	void OnCancelDragEvent(SSchematicGraphNode* Node, const TSharedPtr<FDragDropOperation>& InDragDropEvent);
 	void OnDropEvent(SSchematicGraphNode* Node, const FDragDropEvent& InDragDropEvent);
 	FReply HandleNodeDragDetected(FGuid Guid, const FGeometry& MyGeometry, const FPointerEvent& MouseEvent);
 
@@ -233,6 +235,9 @@ public:
 	FOnNodeClicked& OnNodeClicked() { return OnNodeClickedDelegate; }
 	FOnBeginDrag& OnBeginDrag() { return OnBeginDragDelegate; }
 	FOnEndDrag& OnEndDrag() { return OnEndDragDelegate; }
+	FOnEnterDrag& OnEnterDrag() { return OnEnterDragDelegate; }
+	FOnLeaveDrag& OnLeaveDrag() { return OnLeaveDragDelegate; }
+	FOnCancelDrag& OnCancelDrag() { return OnCancelDragDelegate; }
 	FOnDrop& OnAcceptDrop() { return OnDropDelegate; }
 
 private:
@@ -242,18 +247,20 @@ private:
 	void UpdateAutoScalingForNodes();
 
 	bool bIsDragDropping = false;
+	TSharedPtr<FDragDropOperation> DragDropOpFromOutside;
 	bool bIsOverlay = false;
 	int32 PaddingLeft = 0;
 	int32 PaddingRight = 0;
 	int32 PaddingTop = 0;
 	int32 PaddingBottom = 0;
 	int32 PaddingInterNode = 0;
-	TAttribute<float> BackgroundAlpha;
-	TSharedPtr<FFloatAttribute> FadeBackgroundAlpha;
 	FSchematicGraphModel* GraphData = nullptr;
 	FOnNodeClicked OnNodeClickedDelegate;
 	FOnBeginDrag OnBeginDragDelegate;
 	FOnEndDrag OnEndDragDelegate;
+	FOnEnterDrag OnEnterDragDelegate;
+	FOnLeaveDrag OnLeaveDragDelegate;
+	FOnCancelDrag OnCancelDragDelegate;
 	FOnDrop OnDropDelegate;
 	TArray<FGuid> NodesTopLeft;
 	TArray<FGuid> NodesTopRight;
@@ -289,6 +296,7 @@ private:
 	mutable TArray<FVector2d> NodeCenterByIndex;
 	mutable TArray<ESchematicGraphVisibility::Type> NodeVisibilityByIndex;
 	mutable TMap<FGuid, ESchematicGraphVisibility::Type> NodeVisibilityByGuid;
+	mutable TOptional<FGuid> DropTarget;
 };
 
 #endif
\ No newline at end of file
diff --git a/Engine/Source/Editor/AnimationEditorWidgets/Public/SchematicGraphPanel/SchematicGraphModel.h b/Engine/Source/Editor/AnimationEditorWidgets/Public/SchematicGraphPanel/SchematicGraphModel.h
--- a/Engine/Source/Editor/AnimationEditorWidgets/Public/SchematicGraphPanel/SchematicGraphModel.h
+++ b/Engine/Source/Editor/AnimationEditorWidgets/Public/SchematicGraphPanel/SchematicGraphModel.h
@@ -7,6 +7,8 @@
 #include "SchematicGraphNode.h"
 #include "SchematicGraphLink.h"
 
+class SSchematicGraphPanel;
+
 class ANIMATIONEDITORWIDGETS_API FSchematicGraphModel : public TSharedFromThis<FSchematicGraphModel>
 {
 public:
@@ -14,6 +16,7 @@ public:
 	SCHEMATICGRAPHELEMENT_BODY_BASE(FSchematicGraphModel)
 	
 	virtual void Reset();
+	virtual void ApplyToPanel(SSchematicGraphPanel* InPanel) {};
 	
 	template<typename NodeType = FSchematicGraphNode>
 	NodeType* AddNode(bool bNotify = true)
@@ -251,7 +254,7 @@ public:
 	virtual ESchematicGraphVisibility::Type GetVisibilityForLink(const FSchematicGraphLink* InLink) const;
 
 	virtual bool IsAutoGroupingEnabled() const { return true; }
-	virtual float GetAutoGroupingDistance() const { return 32.f; }
+	virtual float GetAutoGroupingDistance() const { return 4.f; }
 	virtual FSchematicGraphGroupNode* AddAutoGroupNode(); 
 
 	FOnSchematicGraphNodeAdded& OnNodeAdded() { return OnNodeAddedDelegate; }
