Control Rig: Change Module Hierarchy tab label

#rnx

diff --git a/Engine/Plugins/Animation/ControlRig/Source/ControlRigEditor/Private/Editor/ModularRigModelTabSummoner.cpp b/Engine/Plugins/Animation/ControlRig/Source/ControlRigEditor/Private/Editor/ModularRigModelTabSummoner.cpp
--- a/Engine/Plugins/Animation/ControlRig/Source/ControlRigEditor/Private/Editor/ModularRigModelTabSummoner.cpp
+++ b/Engine/Plugins/Animation/ControlRig/Source/ControlRigEditor/Private/Editor/ModularRigModelTabSummoner.cpp
@@ -14,11 +14,11 @@ FModularRigModelTabSummoner::FModularRigModelTabSummoner(const TSharedRef<FContr
 	: FWorkflowTabFactory(TabID, InControlRigEditor)
 	, ControlRigEditor(InControlRigEditor)
 {
-	TabLabel = LOCTEXT("ModularRigHierarchyTabLabel", "Modular Rig Hierarchy");
+	TabLabel = LOCTEXT("ModularRigHierarchyTabLabel", "Module Hierarchy");
 	TabIcon = FSlateIcon(FControlRigEditorStyle::Get().GetStyleSetName(), "ModularRigHierarchy.TabIcon");
 
-	ViewMenuDescription = LOCTEXT("ModularRigHierarchy_ViewMenu_Desc", "Modular Rig Hierarchy");
-	ViewMenuTooltip = LOCTEXT("ModularRigHierarchy_ViewMenu_ToolTip", "Show the Modular Rig Hierarchy tab");
+	ViewMenuDescription = LOCTEXT("ModularRigHierarchy_ViewMenu_Desc", "Module Hierarchy");
+	ViewMenuTooltip = LOCTEXT("ModularRigHierarchy_ViewMenu_ToolTip", "Show the Module Hierarchy tab");
 }
 
 FTabSpawnerEntry& FModularRigModelTabSummoner::RegisterTabSpawner(TSharedRef<FTabManager> InTabManager, const FApplicationMode* CurrentApplicationMode) const
