Control Rig: Turning module hierarchy into column based tree view

#rnx

diff --git a/Engine/Plugins/Animation/ControlRig/Source/ControlRigEditor/Private/Editor/SModularRigModel.cpp b/Engine/Plugins/Animation/ControlRig/Source/ControlRigEditor/Private/Editor/SModularRigModel.cpp
--- a/Engine/Plugins/Animation/ControlRig/Source/ControlRigEditor/Private/Editor/SModularRigModel.cpp
+++ b/Engine/Plugins/Animation/ControlRig/Source/ControlRigEditor/Private/Editor/SModularRigModel.cpp
@@ -103,6 +103,32 @@ void SModularRigModel::Construct(const FArguments& InArgs, TSharedRef<FControlRi
 	Delegates.OnVerifyModuleNameChanged = FOnModularRigTreeVerifyElementNameChanged::CreateSP(this, &SModularRigModel::HandleVerifyNameChanged);
 	Delegates.OnResolveConnector = FOnModularRigTreeResolveConnector::CreateSP(this, &SModularRigModel::HandleConnectorResolved);
 	Delegates.OnDisconnectConnector = FOnModularRigTreeDisconnectConnector::CreateSP(this, &SModularRigModel::HandleConnectorDisconnect);
+
+	HeaderRowWidget = SNew(SHeaderRow)
+		.Visibility(EVisibility::Visible);
+
+	HeaderRowWidget->AddColumn(
+		SHeaderRow::Column(SModularRigTreeView::Column_Module)
+		.DefaultLabel(FText::FromName(SModularRigTreeView::Column_Module))
+		.HAlignCell(HAlign_Left)
+		.HAlignHeader(HAlign_Left)
+		.VAlignCell(VAlign_Center)
+	);
+	HeaderRowWidget->AddColumn(
+		SHeaderRow::Column(SModularRigTreeView::Column_Connector)
+		.DefaultLabel(FText::FromName(SModularRigTreeView::Column_Connector))
+		.HAlignCell(HAlign_Left)
+		.HAlignHeader(HAlign_Left)
+		.VAlignCell(VAlign_Center)
+	);
+	HeaderRowWidget->AddColumn(
+		SHeaderRow::Column(SModularRigTreeView::Column_Buttons)
+		.DefaultLabel(FText::FromName(SModularRigTreeView::Column_Buttons))
+		.ManualWidth(60)
+		.HAlignCell(HAlign_Left)
+		.HAlignHeader(HAlign_Left)
+		.VAlignCell(VAlign_Center)
+	);
 	
 	ChildSlot
 	[
@@ -119,6 +145,7 @@ void SModularRigModel::Construct(const FArguments& InArgs, TSharedRef<FControlRi
 				.BorderImage(FAppStyle::GetBrush("SCSEditor.TreePanel"))
 				[
 					SAssignNew(TreeView, SModularRigTreeView)
+					.HeaderRow(HeaderRowWidget)
 					.RigTreeDelegates(Delegates)
 					.AutoScrollEnabled(true)
 				]
@@ -352,17 +379,17 @@ void SModularRigModel::CreateContextMenu()
 				{
 					const FControlRigModularRigCommands& Commands = FControlRigModularRigCommands::Get(); 
 				
-					FToolMenuSection& ElementsSection = InMenu->AddSection(TEXT("Elements"), LOCTEXT("ElementsHeader", "Elements"));
-					ElementsSection.AddSubMenu(TEXT("New"), LOCTEXT("New", "New"), LOCTEXT("New_ToolTip", "Create New Elements"),
+					FToolMenuSection& ModulesSection = InMenu->AddSection(TEXT("Modules"), LOCTEXT("ModulesHeader", "Modules"));
+					ModulesSection.AddSubMenu(TEXT("New"), LOCTEXT("New", "New"), LOCTEXT("New_ToolTip", "Create New Modules"),
 						FNewToolMenuDelegate::CreateLambda([Commands, ModelPanel](UToolMenu* InSubMenu)
 						{
 							FToolMenuSection& DefaultSection = InSubMenu->AddSection(NAME_None);
 							DefaultSection.AddMenuEntry(Commands.AddModuleItem);
 						})
 					);
-					ElementsSection.AddMenuEntry(Commands.RenameModuleItem);
-					ElementsSection.AddMenuEntry(Commands.DeleteModuleItem);
-					ElementsSection.AddMenuEntry(Commands.MirrorModuleItem);
+					ModulesSection.AddMenuEntry(Commands.RenameModuleItem);
+					ModulesSection.AddMenuEntry(Commands.DeleteModuleItem);
+					ModulesSection.AddMenuEntry(Commands.MirrorModuleItem);
 				}
 			})
 		);
diff --git a/Engine/Plugins/Animation/ControlRig/Source/ControlRigEditor/Private/Editor/SModularRigModel.h b/Engine/Plugins/Animation/ControlRig/Source/ControlRigEditor/Private/Editor/SModularRigModel.h
--- a/Engine/Plugins/Animation/ControlRig/Source/ControlRigEditor/Private/Editor/SModularRigModel.h
+++ b/Engine/Plugins/Animation/ControlRig/Source/ControlRigEditor/Private/Editor/SModularRigModel.h
@@ -115,6 +115,7 @@ private:
 
 	/** Tree view widget */
 	TSharedPtr<SModularRigTreeView> TreeView;
+	TSharedPtr<SHeaderRow> HeaderRowWidget;
 
 	TWeakObjectPtr<UControlRigBlueprint> ControlRigBlueprint;
 	TWeakObjectPtr<UModularRig> ControlRigBeingDebuggedPtr;
diff --git a/Engine/Plugins/Animation/ControlRig/Source/ControlRigEditor/Private/Editor/SModularRigTreeView.cpp b/Engine/Plugins/Animation/ControlRig/Source/ControlRigEditor/Private/Editor/SModularRigTreeView.cpp
--- a/Engine/Plugins/Animation/ControlRig/Source/ControlRigEditor/Private/Editor/SModularRigTreeView.cpp
+++ b/Engine/Plugins/Animation/ControlRig/Source/ControlRigEditor/Private/Editor/SModularRigTreeView.cpp
@@ -106,300 +106,25 @@ void SModularRigModelItem::Construct(const FArguments& InArgs, const TSharedRef<
 
 	if (InRigTreeElement->Key.IsEmpty())
 	{
-		STableRow<TSharedPtr<FModularRigTreeElement>>::Construct(
-			STableRow<TSharedPtr<FModularRigTreeElement>>::FArguments()
+		SMultiColumnTableRow<TSharedPtr<FModularRigTreeElement>>::Construct(
+			SMultiColumnTableRow<TSharedPtr<FModularRigTreeElement>>::FArguments()
 			.ShowSelection(false)
 			.OnCanAcceptDrop(Delegates.OnCanAcceptDrop)
 			.OnAcceptDrop(Delegates.OnAcceptDrop)
-			.Content()
-			[
-				SNew(SVerticalBox)
-				+ SVerticalBox::Slot()
-				.FillHeight(200.f)
-				[
-					SNew(SSpacer)
-				]
-			], OwnerTable);
+			, OwnerTable);
 		return;
 	}
 
 	const FString& ModulePath = InRigTreeElement->ModulePath;
-
 	FString ConnectorPath = FString::Printf(TEXT("%s:%s"), *ModulePath, *InRigTreeElement->ConnectorName);
-	FRigElementKey ConnectorKey(*ConnectorPath, ERigElementType::Connector);
-
-	TSharedPtr< SInlineEditableTextBlock > InlineWidget;
-	TSharedPtr< SHorizontalBox > HorizontalBox;
-	TSharedPtr<SVerticalBox> ComboButtonBox;
-
-	// todo: this should be done lazily - the matches need to be pulled when the user clicks
-	FModularRigResolveResult ConnectorMatches;
-	if (const UModularRig* ModularRig = Delegates.GetModularRig())
-	{
-		if (const FRigModuleInstance* Module = ModularRig->FindModule(ModulePath))
-		{
-			if (URigHierarchy* Hierarchy = ModularRig->GetHierarchy())
-			{
-				if (FRigConnectorElement* ConnectorElement = Cast<FRigConnectorElement>(Hierarchy->Find(ConnectorKey)))
-				{
-					const UModularRigRuleManager* RuleManager = ModularRig->GetHierarchy()->GetRuleManager();
-					ConnectorMatches = RuleManager->FindMatches(ConnectorElement, Module, ModularRig->GetElementKeyRedirector());
-				}
-			}
-		}
-	}
-
-	FRigTreeDelegates TreeDelegates;
-	TreeDelegates.OnGetHierarchy = FOnGetRigTreeHierarchy::CreateLambda([this]()
-	{
-		return Delegates.GetModularRig()->GetHierarchy();
-	});
-	TreeDelegates.OnRigTreeIsItemVisible = FOnRigTreeIsItemVisible::CreateLambda([ConnectorMatches](const FRigElementKey& InTarget)
-	{
-		return ConnectorMatches.ContainsMatch(InTarget);
-	});
-	TreeDelegates.OnGetSelection.BindLambda([this, ConnectorKey]() -> TArray<FRigElementKey>
-	{
-		FRigElementKeyRedirector Redirector = Delegates.GetModularRig()->GetElementKeyRedirector();
-		if (const FRigElementKey* Key = Redirector.FindExternalKey(ConnectorKey))
-		{
-			return {*Key};
-		}
-		return {};
-	});
-	TreeDelegates.OnSelectionChanged.BindSP(this, &SModularRigModelItem::OnConnectorTargetChanged, ConnectorKey);
+	ConnectorKey = FRigElementKey(*ConnectorPath, ERigElementType::Connector);
 
-
-	STableRow<TSharedPtr<FModularRigTreeElement>>::Construct(
-		STableRow<TSharedPtr<FModularRigTreeElement>>::FArguments()
+	SMultiColumnTableRow<TSharedPtr<FModularRigTreeElement>>::Construct(
+		SMultiColumnTableRow<TSharedPtr<FModularRigTreeElement>>::FArguments()
 		.OnDragDetected(Delegates.OnDragDetected)
 		.OnCanAcceptDrop(Delegates.OnCanAcceptDrop)
 		.OnAcceptDrop(Delegates.OnAcceptDrop)
-		.ShowWires(true)
-		.Content()
-		[
-			SAssignNew(HorizontalBox, SHorizontalBox)
-			+ SHorizontalBox::Slot()
-			[
-				SNew(SHorizontalBox)
-				+SHorizontalBox::Slot()
-				.MaxWidth(25)
-				.HAlign(HAlign_Left)
-				.VAlign(VAlign_Center)
-				.Padding(FMargin(0.f, 0.f, 3.f, 0.f))
-				[
-					SNew(SVerticalBox)
-					+SVerticalBox::Slot()
-					.MaxHeight(25)
-					[
-						SNew(SImage)
-						.Image_Lambda([this]() -> const FSlateBrush*
-						{
-							if(WeakRigTreeElement.IsValid())
-							{
-								return WeakRigTreeElement.Pin()->IconBrush;
-							}
-							return nullptr;
-						})
-						.ColorAndOpacity_Lambda([this]()
-						{
-							if(WeakRigTreeElement.IsValid())
-							{
-								return WeakRigTreeElement.Pin()->IconColor;
-							}
-							return FSlateColor::UseForeground();
-						})
-					]
-				]
-
-				+ SHorizontalBox::Slot()
-				.FillWidth(1.0)
-				.VAlign(VAlign_Center)
-				[
-					SAssignNew(InlineWidget, SInlineEditableTextBlock)
-					.Text(this, &SModularRigModelItem::GetName, true)
-					.OnVerifyTextChanged(this, &SModularRigModelItem::OnVerifyNameChanged)
-					.OnTextCommitted(this, &SModularRigModelItem::OnNameCommitted)
-					.ToolTipText(this, &SModularRigModelItem::GetItemTooltip)
-					.MultiLine(false)
-					.Font(IDetailLayoutBuilder::GetDetailFont())
-					.ColorAndOpacity_Lambda([this]()
-					{
-						if(WeakRigTreeElement.IsValid())
-						{
-							return WeakRigTreeElement.Pin()->TextColor;
-						}
-						return FSlateColor::UseForeground();
-					})
-				]
-			]
-
-			+ SHorizontalBox::Slot()
-			.HAlign(HAlign_Right)
-			.VAlign(VAlign_Center)
-			[
-				SNew( SComboButton )
-				.ContentPadding(3)
-				.MenuPlacement(MenuPlacement_BelowAnchor)
-				.OnComboBoxOpened(this, &SModularRigModelItem::PopulateConnectorTargetList, ConnectorKey)
-				.ButtonContent()
-				[
-					// Wrap in configurable box to restrain height/width of menu
-					SNew(SBox)
-					.MinDesiredWidth(200.0f)
-					.Padding(0, 0, 0, 0)
-					.HAlign(HAlign_Left)
-					[
-						SAssignNew(ComboButtonBox, SVerticalBox)
-					]
-				]
-				.MenuContent()
-				[
-					SNew(SBorder)
-					.Visibility(EVisibility::Visible)
-					.BorderImage(FAppStyle::GetBrush("Menu.Background"))
-					[
-						SAssignNew(ConnectorComboBox, SSearchableRigHierarchyTreeView)
-							.RigTreeDelegates(TreeDelegates)
-					]
-				]
-			]
-
-			// Connection buttons
-			+SHorizontalBox::Slot()
-			.AutoWidth()
-			.VAlign(VAlign_Center)
-			.Padding(0.f, 0.f, 0.f, 0.f)
-			[
-				SNew(SVerticalBox)
-				+SVerticalBox::Slot()
-				.VAlign(VAlign_Center)
-				.AutoHeight()
-				[
-					// Reset button
-					SNew(SHorizontalBox)
-					+SHorizontalBox::Slot()
-					.AutoWidth()
-					.Padding(0.f, 0.f, 0.f, 0.f)
-					[
-						SAssignNew(ResetConnectorButton, SButton)
-						.ButtonStyle( FAppStyle::Get(), "NoBorder" )
-						.ButtonColorAndOpacity_Lambda([this, ConnectorKey]()
-						{
-							return ResetConnectorButton.IsValid() && ResetConnectorButton->IsHovered()
-								? FSlateColor(FLinearColor(1.f, 1.f, 1.f, 0.8))
-								: FSlateColor(FLinearColor(1.f, 1.f, 1.f, 0.4));
-						})
-						.OnClicked_Lambda([this, ConnectorKey]()
-						{
-							Delegates.HandleDisconnectConnector(ConnectorKey);
-							return FReply::Handled();
-						})
-						.ContentPadding(1.f)
-						.ToolTipText(NSLOCTEXT("ControlRigModuleDetails", "Reset_Connector", "Reset Connector"))
-						[
-							SNew(SImage)
-							.ColorAndOpacity_Lambda( [this, ConnectorKey]()
-							{
-								return ResetConnectorButton.IsValid() && ResetConnectorButton->IsHovered()
-								? FSlateColor(FLinearColor(1.f, 1.f, 1.f, 0.8))
-								: FSlateColor(FLinearColor(1.f, 1.f, 1.f, 0.4));
-							})
-							.Image(FSlateIcon(FAppStyle::Get().GetStyleSetName(), "PropertyWindow.DiffersFromDefault").GetIcon())
-						]
-					]
-
-					// Use button
-					+SHorizontalBox::Slot()
-					.AutoWidth()
-					.Padding(0.f, 0.f, 0.f, 0.f)
-					[
-						SAssignNew(UseSelectedButton, SButton)
-						.ButtonStyle( FAppStyle::Get(), "NoBorder" )
-						.ButtonColorAndOpacity_Lambda([this]()
-						{
-							return UseSelectedButton.IsValid() && UseSelectedButton->IsHovered()
-								? FSlateColor(FLinearColor(1.f, 1.f, 1.f, 0.8))
-								: FSlateColor(FLinearColor(1.f, 1.f, 1.f, 0.4));
-						})
-						.OnClicked_Lambda([this, ConnectorKey]()
-						{
-							if (const UModularRig* ModularRig = Delegates.GetModularRig())
-							{
-								const TArray<FRigElementKey>& Selected = ModularRig->GetHierarchy()->GetSelectedKeys();
-								if (Selected.Num() > 0)
-								{
-									Delegates.HandleResolveConnector(ConnectorKey, Selected[0]);
-								}
-							}
-							return FReply::Handled();
-						})
-						.ContentPadding(1.f)
-						.ToolTipText(NSLOCTEXT("ControlRigModuleDetails", "Use_Selected", "Use Selected"))
-						[
-							SNew(SImage)
-							.ColorAndOpacity_Lambda( [this]()
-							{
-								return UseSelectedButton.IsValid() && UseSelectedButton->IsHovered()
-								? FSlateColor(FLinearColor(1.f, 1.f, 1.f, 0.8))
-								: FSlateColor(FLinearColor(1.f, 1.f, 1.f, 0.4));
-							})
-							.Image(FAppStyle::GetBrush("Icons.CircleArrowLeft"))
-						]
-					]
-					
-					// Select in hierarchy button
-					+SHorizontalBox::Slot()
-					.AutoWidth()
-					.Padding(0.f, 0.f, 0.f, 0.f)
-					[
-						SAssignNew(SelectElementButton, SButton)
-						.ButtonStyle( FAppStyle::Get(), "NoBorder" )
-						.ButtonColorAndOpacity_Lambda([this]()
-						{
-							return SelectElementButton.IsValid() && SelectElementButton->IsHovered()
-								? FSlateColor(FLinearColor(1.f, 1.f, 1.f, 0.8))
-								: FSlateColor(FLinearColor(1.f, 1.f, 1.f, 0.4));
-						})
-						.OnClicked_Lambda([this, ConnectorKey]()
-						{
-							if (const UModularRig* ModularRig = Delegates.GetModularRig())
-							{
-								const FRigElementKeyRedirector& Redirector = ModularRig->GetElementKeyRedirector();
-								if (const FRigElementKey* TargetKey = Redirector.FindExternalKey(ConnectorKey))
-								{
-									ModularRig->GetHierarchy()->GetController()->SelectElement(*TargetKey, true, true);
-								}
-							}
-							return FReply::Handled();
-						})
-						.ContentPadding(1.f)
-						.ToolTipText(NSLOCTEXT("ControlRigModuleDetails", "Select_Element", "Select Element"))
-						[
-							SNew(SImage)
-							.ColorAndOpacity_Lambda( [this]()
-							{
-								return SelectElementButton.IsValid() && SelectElementButton->IsHovered()
-								? FSlateColor(FLinearColor(1.f, 1.f, 1.f, 0.8))
-								: FSlateColor(FLinearColor(1.f, 1.f, 1.f, 0.4));
-							})
-							.Image(FAppStyle::GetBrush("Icons.Search"))
-						]
-					]
-				]
-			]
-		], OwnerTable);
-
-	InRigTreeElement->OnRenameRequested.BindSP(InlineWidget.Get(), &SInlineEditableTextBlock::EnterEditingMode);
-	const UModularRig* ModularRig = Delegates.GetModularRig();
-	const FRigElementKeyRedirector& Redirector = ModularRig->GetElementKeyRedirector();
-	FRigElementKey CurrentTargetKey;
-	if (const FRigElementKey* Key = Redirector.FindExternalKey(ConnectorKey))
-	{
-		CurrentTargetKey = *Key;
-	}
-	TPair<const FSlateBrush*, FSlateColor> IconAndColor = SRigHierarchyItem::GetBrushForElementType(ModularRig->GetHierarchy(), CurrentTargetKey);
-	PopulateConnectorCurrentTarget(ComboButtonBox, ConnectorKey, CurrentTargetKey, IconAndColor.Key, IconAndColor.Value, FText::FromName(CurrentTargetKey.Name));
+		.ShowWires(true), OwnerTable);
 }
 
 
@@ -414,7 +139,7 @@ void SModularRigModelItem::PopulateConnectorTargetList(const FRigElementKey InCo
 	{
 		if (URigHierarchy* Hierarchy = ModularRig->GetHierarchy())
 		{
-			if (FRigConnectorElement* ConnectorElement = Cast<FRigConnectorElement>(Hierarchy->Find(InConnectorKey)))
+			if (Cast<FRigConnectorElement>(Hierarchy->Find(InConnectorKey)))
 			{
 				ConnectorComboBox->GetTreeView()->RefreshTreeView(true);
 			}
@@ -523,6 +248,290 @@ bool SModularRigModelItem::OnVerifyNameChanged(const FText& InText, FText& OutEr
 	return Delegates.HandleVerifyElementNameChanged(OldPath, NewName, OutErrorMessage);
 }
 
+TSharedRef<SWidget> SModularRigModelItem::GenerateWidgetForColumn(const FName& ColumnName)
+{
+	if(ColumnName == SModularRigTreeView::Column_Module)
+	{
+		TSharedPtr< SInlineEditableTextBlock > InlineWidget;
+
+		TSharedRef<SWidget> Widget = SNew(SHorizontalBox)
+
+		+SHorizontalBox::Slot()
+		.AutoWidth()
+		.Padding(6, 0, 0, 0)
+		.VAlign(VAlign_Center)
+		[
+			SNew(SExpanderArrow, SharedThis(this))
+			.IndentAmount(12)
+			.ShouldDrawWires(true)
+		]
+
+		+SHorizontalBox::Slot()
+		.MaxWidth(25)
+		.HAlign(HAlign_Left)
+		.VAlign(VAlign_Center)
+		.Padding(FMargin(0.f, 0.f, 3.f, 0.f))
+		[
+			SNew(SVerticalBox)
+			+SVerticalBox::Slot()
+			.MaxHeight(25)
+			[
+				SNew(SImage)
+				.Image_Lambda([this]() -> const FSlateBrush*
+				{
+					if(WeakRigTreeElement.IsValid())
+					{
+						return WeakRigTreeElement.Pin()->IconBrush;
+					}
+					return nullptr;
+				})
+				.ColorAndOpacity_Lambda([this]()
+				{
+					if(WeakRigTreeElement.IsValid())
+					{
+						return WeakRigTreeElement.Pin()->IconColor;
+					}
+					return FSlateColor::UseForeground();
+				})
+			]
+		]
+
+		+ SHorizontalBox::Slot()
+		.AutoWidth()
+		.VAlign(VAlign_Center)
+		[
+			SAssignNew(InlineWidget, SInlineEditableTextBlock)
+			.Text(this, &SModularRigModelItem::GetName, true)
+			.OnVerifyTextChanged(this, &SModularRigModelItem::OnVerifyNameChanged)
+			.OnTextCommitted(this, &SModularRigModelItem::OnNameCommitted)
+			.ToolTipText(this, &SModularRigModelItem::GetItemTooltip)
+			.MultiLine(false)
+			//.Font(IDetailLayoutBuilder::GetDetailFont())
+			.ColorAndOpacity_Lambda([this]()
+			{
+				if(WeakRigTreeElement.IsValid())
+				{
+					return WeakRigTreeElement.Pin()->TextColor;
+				}
+				return FSlateColor::UseForeground();
+			})
+		];
+
+		if(WeakRigTreeElement.IsValid())
+		{
+			WeakRigTreeElement.Pin()->OnRenameRequested.BindSP(InlineWidget.Get(), &SInlineEditableTextBlock::EnterEditingMode);
+		}
+
+		return Widget;
+	}
+	if(ColumnName == SModularRigTreeView::Column_Connector)
+	{
+		TSharedPtr<SVerticalBox> ComboButtonBox;
+
+		FRigTreeDelegates TreeDelegates;
+		TreeDelegates.OnGetHierarchy = FOnGetRigTreeHierarchy::CreateLambda([this]()
+		{
+			return Delegates.GetModularRig()->GetHierarchy();
+		});
+		TreeDelegates.OnRigTreeIsItemVisible = FOnRigTreeIsItemVisible::CreateLambda([this](const FRigElementKey& InTarget)
+		{
+			if(!ConnectorMatches.IsSet() && WeakRigTreeElement.IsValid())
+			{
+				if (const UModularRig* ModularRig = Delegates.GetModularRig())
+				{
+					if (const FRigModuleInstance* Module = ModularRig->FindModule(WeakRigTreeElement.Pin()->ModulePath))
+					{
+						if (URigHierarchy* Hierarchy = ModularRig->GetHierarchy())
+						{
+							if (const FRigConnectorElement* ConnectorElement = Cast<FRigConnectorElement>(Hierarchy->Find(ConnectorKey)))
+							{
+								const UModularRigRuleManager* RuleManager = ModularRig->GetHierarchy()->GetRuleManager();
+								ConnectorMatches = RuleManager->FindMatches(ConnectorElement, Module, ModularRig->GetElementKeyRedirector());
+							}
+						}
+					}
+				}
+			}
+
+			if(ConnectorMatches.IsSet())
+			{
+				return ConnectorMatches.GetValue().ContainsMatch(InTarget); 
+			}
+			return true;
+		});
+		TreeDelegates.OnGetSelection.BindLambda([this]() -> TArray<FRigElementKey>
+		{
+			const FRigElementKeyRedirector Redirector = Delegates.GetModularRig()->GetElementKeyRedirector();
+			if (const FRigElementKey* Key = Redirector.FindExternalKey(ConnectorKey))
+			{
+				return {*Key};
+			}
+			return {};
+		});
+		TreeDelegates.OnSelectionChanged.BindSP(this, &SModularRigModelItem::OnConnectorTargetChanged, ConnectorKey);
+
+		TSharedRef<SWidget> Widget = SNew(SHorizontalBox)
+		+ SHorizontalBox::Slot()
+		.AutoWidth()
+		.HAlign(HAlign_Right)
+		.VAlign(VAlign_Center)
+		[
+			SNew( SComboButton )
+			.ContentPadding(3)
+			.MenuPlacement(MenuPlacement_BelowAnchor)
+			.OnComboBoxOpened(this, &SModularRigModelItem::PopulateConnectorTargetList, ConnectorKey)
+			.ButtonContent()
+			[
+				// Wrap in configurable box to restrain height/width of menu
+				SNew(SBox)
+				.MinDesiredWidth(200.0f)
+				.Padding(0, 0, 0, 0)
+				.HAlign(HAlign_Left)
+				[
+					SAssignNew(ComboButtonBox, SVerticalBox)
+				]
+			]
+			.MenuContent()
+			[
+				SNew(SBorder)
+				.Visibility(EVisibility::Visible)
+				.BorderImage(FAppStyle::GetBrush("Menu.Background"))
+				[
+					SAssignNew(ConnectorComboBox, SSearchableRigHierarchyTreeView)
+						.RigTreeDelegates(TreeDelegates)
+				]
+			]
+		];
+
+		const UModularRig* ModularRig = Delegates.GetModularRig();
+		const FRigElementKeyRedirector& Redirector = ModularRig->GetElementKeyRedirector();
+		FRigElementKey CurrentTargetKey;
+		if (const FRigElementKey* Key = Redirector.FindExternalKey(ConnectorKey))
+		{
+			CurrentTargetKey = *Key;
+		}
+		TPair<const FSlateBrush*, FSlateColor> IconAndColor = SRigHierarchyItem::GetBrushForElementType(ModularRig->GetHierarchy(), CurrentTargetKey);
+		PopulateConnectorCurrentTarget(ComboButtonBox, ConnectorKey, CurrentTargetKey, IconAndColor.Key, IconAndColor.Value, FText::FromName(CurrentTargetKey.Name));
+
+		return Widget;
+	}
+	if(ColumnName == SModularRigTreeView::Column_Buttons)
+	{
+		return SNew(SHorizontalBox)
+
+		// Reset button
+		+SHorizontalBox::Slot()
+		.AutoWidth()
+		.Padding(0.f, 0.f, 0.f, 0.f)
+		[
+			SAssignNew(ResetConnectorButton, SButton)
+			.ButtonStyle( FAppStyle::Get(), "NoBorder" )
+			.ButtonColorAndOpacity_Lambda([this]()
+			{
+				return ResetConnectorButton.IsValid() && ResetConnectorButton->IsHovered()
+					? FSlateColor(FLinearColor(1.f, 1.f, 1.f, 0.8))
+					: FSlateColor(FLinearColor(1.f, 1.f, 1.f, 0.4));
+			})
+			.OnClicked_Lambda([this]()
+			{
+				Delegates.HandleDisconnectConnector(ConnectorKey);
+				return FReply::Handled();
+			})
+			.ContentPadding(1.f)
+			.ToolTipText(NSLOCTEXT("ControlRigModuleDetails", "Reset_Connector", "Reset Connector"))
+			[
+				SNew(SImage)
+				.ColorAndOpacity_Lambda( [this]()
+				{
+					return ResetConnectorButton.IsValid() && ResetConnectorButton->IsHovered()
+					? FSlateColor(FLinearColor(1.f, 1.f, 1.f, 0.8))
+					: FSlateColor(FLinearColor(1.f, 1.f, 1.f, 0.4));
+				})
+				.Image(FSlateIcon(FAppStyle::Get().GetStyleSetName(), "PropertyWindow.DiffersFromDefault").GetIcon())
+			]
+		]
+
+		// Use button
+		+SHorizontalBox::Slot()
+		.AutoWidth()
+		.Padding(0.f, 0.f, 0.f, 0.f)
+		[
+			SAssignNew(UseSelectedButton, SButton)
+			.ButtonStyle( FAppStyle::Get(), "NoBorder" )
+			.ButtonColorAndOpacity_Lambda([this]()
+			{
+				return UseSelectedButton.IsValid() && UseSelectedButton->IsHovered()
+					? FSlateColor(FLinearColor(1.f, 1.f, 1.f, 0.8))
+					: FSlateColor(FLinearColor(1.f, 1.f, 1.f, 0.4));
+			})
+			.OnClicked_Lambda([this]()
+			{
+				if (const UModularRig* ModularRig = Delegates.GetModularRig())
+				{
+					const TArray<FRigElementKey>& Selected = ModularRig->GetHierarchy()->GetSelectedKeys();
+					if (Selected.Num() > 0)
+					{
+						Delegates.HandleResolveConnector(ConnectorKey, Selected[0]);
+					}
+				}
+				return FReply::Handled();
+			})
+			.ContentPadding(1.f)
+			.ToolTipText(NSLOCTEXT("ControlRigModuleDetails", "Use_Selected", "Use Selected"))
+			[
+				SNew(SImage)
+				.ColorAndOpacity_Lambda( [this]()
+				{
+					return UseSelectedButton.IsValid() && UseSelectedButton->IsHovered()
+					? FSlateColor(FLinearColor(1.f, 1.f, 1.f, 0.8))
+					: FSlateColor(FLinearColor(1.f, 1.f, 1.f, 0.4));
+				})
+				.Image(FAppStyle::GetBrush("Icons.CircleArrowLeft"))
+			]
+		]
+		
+		// Select in hierarchy button
+		+SHorizontalBox::Slot()
+		.AutoWidth()
+		.Padding(0.f, 0.f, 0.f, 0.f)
+		[
+			SAssignNew(SelectElementButton, SButton)
+			.ButtonStyle( FAppStyle::Get(), "NoBorder" )
+			.ButtonColorAndOpacity_Lambda([this]()
+			{
+				return SelectElementButton.IsValid() && SelectElementButton->IsHovered()
+					? FSlateColor(FLinearColor(1.f, 1.f, 1.f, 0.8))
+					: FSlateColor(FLinearColor(1.f, 1.f, 1.f, 0.4));
+			})
+			.OnClicked_Lambda([this]()
+			{
+				if (const UModularRig* ModularRig = Delegates.GetModularRig())
+				{
+					const FRigElementKeyRedirector& Redirector = ModularRig->GetElementKeyRedirector();
+					if (const FRigElementKey* TargetKey = Redirector.FindExternalKey(ConnectorKey))
+					{
+						ModularRig->GetHierarchy()->GetController()->SelectElement(*TargetKey, true, true);
+					}
+				}
+				return FReply::Handled();
+			})
+			.ContentPadding(1.f)
+			.ToolTipText(NSLOCTEXT("ControlRigModuleDetails", "Select_Element", "Select Element"))
+			[
+				SNew(SImage)
+				.ColorAndOpacity_Lambda( [this]()
+				{
+					return SelectElementButton.IsValid() && SelectElementButton->IsHovered()
+					? FSlateColor(FLinearColor(1.f, 1.f, 1.f, 0.8))
+					: FSlateColor(FLinearColor(1.f, 1.f, 1.f, 0.4));
+				})
+				.Image(FAppStyle::GetBrush("Icons.Search"))
+			]
+		];
+	}
+	return SNullWidget::NullWidget;
+}
+
 FText SModularRigModelItem::GetName(bool bUseShortName) const
 {
 	if(bUseShortName)
@@ -547,12 +556,17 @@ FText SModularRigModelItem::GetItemTooltip() const
 /// SModularRigTreeView
 ///////////////////////////////////////////////////////////
 
+const FName SModularRigTreeView::Column_Module = TEXT("Module");
+const FName SModularRigTreeView::Column_Connector = TEXT("Connector");
+const FName SModularRigTreeView::Column_Buttons = TEXT("Actions");
+
 void SModularRigTreeView::Construct(const FArguments& InArgs)
 {
 	Delegates = InArgs._RigTreeDelegates;
 	bAutoScrollEnabled = InArgs._AutoScrollEnabled;
 
 	STreeView<TSharedPtr<FModularRigTreeElement>>::FArguments SuperArgs;
+	SuperArgs.HeaderRow(InArgs._HeaderRow);
 	SuperArgs.TreeItemsSource(&RootElements);
 	SuperArgs.SelectionMode(ESelectionMode::Multi);
 	SuperArgs.OnGenerateRow(this, &SModularRigTreeView::MakeTableRowWidget, false);
diff --git a/Engine/Plugins/Animation/ControlRig/Source/ControlRigEditor/Private/Editor/SModularRigTreeView.h b/Engine/Plugins/Animation/ControlRig/Source/ControlRigEditor/Private/Editor/SModularRigTreeView.h
--- a/Engine/Plugins/Animation/ControlRig/Source/ControlRigEditor/Private/Editor/SModularRigTreeView.h
+++ b/Engine/Plugins/Animation/ControlRig/Source/ControlRigEditor/Private/Editor/SModularRigTreeView.h
@@ -133,7 +133,7 @@ public:
 	FSlateColor TextColor;
 };
 
-class SModularRigModelItem : public STableRow<TSharedPtr<FModularRigTreeElement>>
+class SModularRigModelItem : public SMultiColumnTableRow<TSharedPtr<FModularRigTreeElement>>
 {
 public:
 	
@@ -151,6 +151,8 @@ public:
 	void OnNameCommitted(const FText& InText, ETextCommit::Type InCommitType) const;
 	bool OnVerifyNameChanged(const FText& InText, FText& OutErrorMessage);
 
+	virtual TSharedRef<SWidget> GenerateWidgetForColumn(const FName& ColumnName) override;
+
 private:
 	TWeakPtr<FModularRigTreeElement> WeakRigTreeElement;
  	FModularRigTreeDelegates Delegates;
@@ -158,6 +160,8 @@ private:
 	TSharedPtr<SButton> ResetConnectorButton;
 	TSharedPtr<SButton> UseSelectedButton;
 	TSharedPtr<SButton> SelectElementButton;
+	FRigElementKey ConnectorKey;
+	TOptional<FModularRigResolveResult> ConnectorMatches;
 
 	FText GetName(bool bUseShortName) const;
 	FText GetItemTooltip() const;
@@ -169,9 +173,14 @@ class SModularRigTreeView : public STreeView<TSharedPtr<FModularRigTreeElement>>
 {
 public:
 
+	static const FName Column_Module;
+	static const FName Column_Connector;
+	static const FName Column_Buttons;
+
 	SLATE_BEGIN_ARGS(SModularRigTreeView)
 		: _AutoScrollEnabled(false)
 	{}
+		SLATE_ARGUMENT( TSharedPtr<SHeaderRow>, HeaderRow )
 		SLATE_ARGUMENT(FModularRigTreeDelegates, RigTreeDelegates)
 		SLATE_ARGUMENT(bool, AutoScrollEnabled)
 	SLATE_END_ARGS()
