Control Rig: FIx GetAnimChannel not working when in a namespace
#rnx

diff --git a/Engine/Plugins/Animation/ControlRig/Source/ControlRig/Private/Units/Hierarchy/RigUnit_ControlChannel.cpp b/Engine/Plugins/Animation/ControlRig/Source/ControlRig/Private/Units/Hierarchy/RigUnit_ControlChannel.cpp
--- a/Engine/Plugins/Animation/ControlRig/Source/ControlRig/Private/Units/Hierarchy/RigUnit_ControlChannel.cpp
+++ b/Engine/Plugins/Animation/ControlRig/Source/ControlRig/Private/Units/Hierarchy/RigUnit_ControlChannel.cpp
@@ -22,13 +22,16 @@ bool FRigUnit_GetAnimationChannelBase::UpdateCache(const URigHierarchy* InHierar
 	{
 		if(const FRigControlElement* ControlElement = InHierarchy->Find<FRigControlElement>(FRigElementKey(Control, ERigElementType::Control)))
 		{
+			FString Namespace, ChannelName = Channel.ToString();
+			URigHierarchy::SplitNameSpace(ChannelName, &Namespace, &ChannelName);
+			
 			for(const FRigBaseElement* Child : InHierarchy->GetChildren(ControlElement))
 			{
 				if(const FRigControlElement* ChildControl = Cast<FRigControlElement>(Child))
 				{
 					if(ChildControl->IsAnimationChannel())
 					{
-						if(ChildControl->GetDisplayName() == Channel)
+						if(ChildControl->GetDisplayName().ToString().Equals(ChannelName))
 						{
 							Key = ChildControl->GetKey();
 							Hash = ExpectedHash;
