Control Rig: FIx layered rigs with bad rotation values
#rnx

diff --git a/Engine/Plugins/Animation/ControlRig/Source/ControlRig/Private/ControlRig.cpp b/Engine/Plugins/Animation/ControlRig/Source/ControlRig/Private/ControlRig.cpp
--- a/Engine/Plugins/Animation/ControlRig/Source/ControlRig/Private/ControlRig.cpp
+++ b/Engine/Plugins/Animation/ControlRig/Source/ControlRig/Private/ControlRig.cpp
@@ -461,6 +461,7 @@ void UControlRig::Evaluate_AnyThread()
 
 						// Add them to find the final value
 						FTransform FinalTransform = AdditiveTransform * PreviousTransform;
+						FinalTransform.SetRotation(FinalTransform.GetRotation().GetNormalized());
 
 						FRigControlValue FinalValue;
 						FinalValue.SetFromTransform(FinalTransform, Control->Settings.ControlType, Control->Settings.PrimaryAxis);
@@ -1982,7 +1983,8 @@ FRigControlValue UControlRig::GetControlValue(FRigControlElement* InControl, con
 			const FRigPoseElement& AnimPose = ControlsAfterBackwardsSolve[ControlIndex];
 			const FRigControlValue& CurrentValue = GetHierarchy()->GetControlValue(InControl, InValueType, false);
 			const FTransform FinalTransform = CurrentValue.GetAsTransform(InControl->Settings.ControlType, InControl->Settings.PrimaryAxis);
-			const FTransform AdditiveTransform = FinalTransform * AnimPose.LocalTransform.Inverse();
+			FTransform AdditiveTransform = FinalTransform * AnimPose.LocalTransform.Inverse();
+			AdditiveTransform.SetRotation(AdditiveTransform.GetRotation().GetNormalized());
 			FRigControlValue AdditiveValue;
 			AdditiveValue.SetFromTransform(AdditiveTransform, InControl->Settings.ControlType, InControl->Settings.PrimaryAxis);
 			return AdditiveValue;
