1. Verify that selecting a Rig Unit from the right click context window adds the Rig Unit to the Rig Graph
2. Verify that opening a Control Rig Blueprint opens it in the Control Rig editor
3. Verify that executing a Rig Unit adds the rig unit to the Execution Stack window
#rnx

diff --git a/Engine/Plugins/Animation/ControlRig/Source/ControlRigEditor/Private/Tests/ControlRigEditor.spec.cpp b/Engine/Plugins/Animation/ControlRig/Source/ControlRigEditor/Private/Tests/ControlRigEditor.spec.cpp
new file mode 100644
--- /dev/null
+++ b/Engine/Plugins/Animation/ControlRig/Source/ControlRigEditor/Private/Tests/ControlRigEditor.spec.cpp
@@ -0,0 +1,133 @@
+// Copyright Epic Games, Inc. All Rights Reserved.
+
+#include "Misc/AutomationTest.h"
+
+#if WITH_EDITOR
+
+#include "Editor.h"
+
+#if WITH_DEV_AUTOMATION_TESTS
+
+#include "ControlRigBlueprint.h"
+#include "Graph/ControlRigGraph.h"
+#include "Graph/ControlRigGraphNode.h"
+#include "Units/Hierarchy/RigUnit_GetTransform.h"
+#include "Subsystems/EditorAssetSubsystem.h"
+#include "Subsystems/AssetEditorSubsystem.h"
+#include "Units/Highlevel/Hierarchy/RigUnit_TwoBoneIKSimple.h"
+
+
+BEGIN_DEFINE_SPEC(
+	FControlRigEditorSpec, "Editor.Content",
+	EAutomationTestFlags::EditorContext | EAutomationTestFlags::EngineFilter)
+	
+	UControlRigBlueprint* ControlRigBlueprint = nullptr;
+	UControlRigGraph* RigGraph = nullptr;
+	URigVMController* RigController = nullptr;
+	TStrongObjectPtr<URigVMUnitNode> ModelNode;
+
+	bool FindNode(URigVM* const RigVM, const FString &InNodeName) const;
+
+END_DEFINE_SPEC(FControlRigEditorSpec)
+
+void FControlRigEditorSpec::Define()
+{
+	Describe("Control Rig", [this]()
+		{
+			BeforeEach([this]()
+				{
+					FString AssetName = "BasicControls_CtrlRig";
+					FString AssetPath = FString::Printf(TEXT("/Game/Animation/ControlRig/%s"), *AssetName);
+
+					UObject* Asset = GEditor->GetEditorSubsystem<UEditorAssetSubsystem>()->LoadAsset(*AssetPath);
+					ControlRigBlueprint = CastChecked<UControlRigBlueprint>(Asset);
+
+					UEdGraph* EdGraph = ControlRigBlueprint->GetEdGraph(ControlRigBlueprint->GetModel());
+					RigGraph = CastChecked<UControlRigGraph>(EdGraph);
+					RigController = ControlRigBlueprint->GetController(RigGraph);
+
+					const bool bShowProgressWindow = false;
+					GEditor->GetEditorSubsystem<UAssetEditorSubsystem>()->OpenEditorForAsset(ControlRigBlueprint, EToolkitMode::Standalone, TSharedPtr<IToolkitHost>(), bShowProgressWindow);
+				});
+
+			It("Verify that selecting a Rig Unit from the right click context window adds the Rig Unit to the Rig Graph", [this]()
+				{
+					int32 NodesCount = RigGraph->Nodes.Num();
+
+					ModelNode.Reset(RigController->AddUnitNode(FRigUnit_GetTransform::StaticStruct(), TEXT("Execute"), FVector2D::ZeroVector, TEXT("GetTransform"), false));
+
+					ControlRigBlueprint->RecompileVM();
+
+					URigVMEdGraphNode* NewNode = CastChecked<URigVMEdGraphNode>(RigGraph->FindNodeForModelNodeName(ModelNode->GetFName()));
+
+					TestTrue(TEXT("Failed to add the 'Get Transform Rig Unit' to the RigGraph"), NewNode != nullptr);
+					TestTrue(TEXT("Unexpected node count in RigGraph"), NodesCount + 1 == RigGraph->Nodes.Num());
+				});
+
+			It("Verify that opening a Control Rig Blueprint opens it in the Control Rig editor", [this]()
+				{
+					TArray<IAssetEditorInstance*> EditorInstances = GEditor->GetEditorSubsystem<UAssetEditorSubsystem>()->FindEditorsForAssetAndSubObjects(ControlRigBlueprint);
+					TestTrue(TEXT("No asset editor window found"), !EditorInstances.IsEmpty());
+
+					FName EditorName = EditorInstances.Last()->GetEditorName();
+					TestTrue(TEXT("Wrong editor type for asset"), EditorName.IsEqual("ControlRigEditor"));
+				});
+
+			It("Executing a Rig Unit adds the Rig Unit to the Execution Stack", [this]()
+				{
+					const FString InNodeName = "BoneIK";
+					const int32 LinksNum = RigController->GetGraph()->GetLinks().Num();
+
+					ModelNode.Reset(RigController->AddUnitNode(FRigUnit_TwoBoneIKSimplePerItem::StaticStruct(), TEXT("Execute"), FVector2D::ZeroVector, InNodeName, false));
+					bool bLinkAdded = RigController->AddLink(TEXT("SetTransform_1_1_1.ExecuteContext"), TEXT("BoneIK.ExecuteContext"));
+					TestTrue(TEXT("Failed to link nodes"), bLinkAdded);
+
+					ControlRigBlueprint->RecompileVM();
+
+					TestTrue(TEXT("Unexpected number of links"), LinksNum == RigController->GetGraph()->GetLinks().Num() - 1);
+
+					TStrongObjectPtr<UControlRig> ControlRig(ControlRigBlueprint->CreateControlRig());
+					URigVM* const RigVM = ControlRig->GetVM();
+
+					bool bNodeExist = FindNode(RigVM, InNodeName);
+					TestTrue(FString::Printf(TEXT("Can not find '%s' Rig Unit in the Execution Stack"), *InNodeName), bNodeExist);
+
+					RigController->BreakLink(TEXT("SetTransform_1_1_1.ExecuteContext"), TEXT("BoneIK.ExecuteContext"));
+				});
+
+			AfterEach([this]()
+				{
+					if (ModelNode.IsValid())
+					{
+						RigController->RemoveNode(ModelNode.Get(), false);
+						ControlRigBlueprint->RecompileVM();
+						ControlRigBlueprint->GetPackage()->SetDirtyFlag(false);
+					}
+
+					GEditor->GetEditorSubsystem<UAssetEditorSubsystem>()->CloseAllEditorsForAsset(ControlRigBlueprint);
+				});
+		});
+}
+
+bool FControlRigEditorSpec::FindNode(URigVM* const RigVM, const FString& InNodeName) const
+{
+	const FRigVMInstructionArray& Instructions = RigVM->GetInstructions();
+	const FRigVMByteCode& ByteCode = RigVM->GetByteCode();
+
+	for (int32 InstructionIndex = 0; InstructionIndex < Instructions.Num(); InstructionIndex++)
+	{
+		URigVMNode* const Node = Cast<URigVMNode>(ByteCode.GetSubjectForInstruction(InstructionIndex));
+		if (Node)
+		{
+			if (Node->GetName().Equals(InNodeName))
+			{
+				return true;
+			}
+		}
+	}
+
+	return false;
+}
+
+#endif
+#endif
\ No newline at end of file
