Control Rig: Fix cannot add Add/Or rules
#rnx

diff --git a/Engine/Plugins/Animation/ControlRig/Source/ControlRigEditor/Private/ControlRigElementDetails.cpp b/Engine/Plugins/Animation/ControlRig/Source/ControlRigEditor/Private/ControlRigElementDetails.cpp
--- a/Engine/Plugins/Animation/ControlRig/Source/ControlRigEditor/Private/ControlRigElementDetails.cpp
+++ b/Engine/Plugins/Animation/ControlRig/Source/ControlRigEditor/Private/ControlRigElementDetails.cpp
@@ -5567,16 +5567,16 @@ void FRigConnectionRuleDetails::OnPickedStruct(const UScriptStruct* ChosenStruct
 		RuleStash.ScriptStructPath = ChosenStruct->GetPathName();
 		RuleStash.ExportedText = TEXT("()");
 		Storage.Reset();
+		RuleStash.Get(Storage);
 	}
 	OnRuleContentChanged();
-	PropertyUtilities->ForceRefresh();
 }
 
 FText FRigConnectionRuleDetails::OnGetStructTextValue() const
 {
 	const UScriptStruct* ScriptStruct = RuleStash.GetScriptStruct();
 	return ScriptStruct
-		? FText::AsCultureInvariant(ScriptStruct->GetName())
+		? FText::AsCultureInvariant(ScriptStruct->GetDisplayNameText())
 		: LOCTEXT("None", "None");
 }
 
@@ -5600,6 +5600,7 @@ void FRigConnectionRuleDetails::OnRuleContentChanged()
 	{
 		(void)StructPropertyHandle->SetPerObjectValue(Index, Content, EPropertyValueSetFlags::DefaultFlags);
 	}
+	StructPropertyHandle->GetParentHandle()->NotifyPostChange(EPropertyChangeType::ValueSet);
 }
 
 #undef LOCTEXT_NAMESPACE
