Sequencer: Move SequenceTimeUnit to MovieScene module and rename as MovieSceneTimeUnit

Added redirector
Tested that a SequenceTimeUnit is automatically upgraded to a MovieSceneTimeUnit

diff --git a/Engine/Config/BaseEngine.ini b/Engine/Config/BaseEngine.ini
--- a/Engine/Config/BaseEngine.ini
+++ b/Engine/Config/BaseEngine.ini
@@ -1230,6 +1230,9 @@ DefaultWorldPartitionHLODLayer=/Engine/Maps/Templates/HLODs/HLODLayer_Instanced.
 
 +EnumRedirects=(OldName="ETargetingTraceType", ValueChanges=(("Sweep", "Sphere")) )
 
++EnumRedirects=(OldName="/Script/SequencerScripting.ESequenceTimeUnit", NewName="/Script/MovieScene.EMovieSceneTimeUnit")
+
+
 ; HLOD actors now store their source actors using FWorldPartitionRuntimeCellObjectMapping rather than a custom FHLODSubActor class
 +StructRedirects=(OldName="/Script/Engine.HLODSubActor",							NewName="/Script/Engine.WorldPartitionRuntimeCellObjectMapping")
 +PropertyRedirects=(OldName="WorldPartitionRuntimeCellObjectMapping.ActorPackage",	NewName="Package")
diff --git a/Engine/Plugins/Animation/ControlRig/Source/ControlRigEditor/Private/ControlRigSequencerEditorLibrary.cpp b/Engine/Plugins/Animation/ControlRig/Source/ControlRigEditor/Private/ControlRigSequencerEditorLibrary.cpp
--- a/Engine/Plugins/Animation/ControlRig/Source/ControlRigEditor/Private/ControlRigSequencerEditorLibrary.cpp
+++ b/Engine/Plugins/Animation/ControlRig/Source/ControlRigEditor/Private/ControlRigSequencerEditorLibrary.cpp
@@ -534,7 +534,7 @@ TArray <UTickableConstraint*> UControlRigSequencerEditorLibrary::GetConstraintsF
 	return Constraints;
 }
 
-bool UControlRigSequencerEditorLibrary::Compensate(UTickableConstraint* InConstraint, FFrameNumber InTime, ESequenceTimeUnit TimeUnit)
+bool UControlRigSequencerEditorLibrary::Compensate(UTickableConstraint* InConstraint, FFrameNumber InTime, EMovieSceneTimeUnit TimeUnit)
 {
 	TWeakPtr<ISequencer> WeakSequencer = GetSequencerFromAsset();
 	if (WeakSequencer.IsValid() == false)
@@ -547,7 +547,7 @@ bool UControlRigSequencerEditorLibrary::Compensate(UTickableConstraint* InConstr
 		TSharedPtr<ISequencer>  Sequencer = WeakSequencer.Pin();
 		FFrameRate TickResolution = Sequencer->GetFocusedTickResolution();
 		FFrameRate DisplayRate = Sequencer->GetFocusedDisplayRate();
-		if (TimeUnit == ESequenceTimeUnit::DisplayRate)
+		if (TimeUnit == EMovieSceneTimeUnit::DisplayRate)
 		{
 			InTime = FFrameRate::TransformTime(FFrameTime(InTime, 0), DisplayRate, TickResolution).RoundToFrame();
 		}
@@ -582,7 +582,7 @@ bool UControlRigSequencerEditorLibrary::CompensateAll(UTickableConstraint* InCon
 	return false;
 }
 
-bool UControlRigSequencerEditorLibrary::SetConstraintActiveKey(UTickableConstraint* InConstraint, bool bActive, FFrameNumber InTime, ESequenceTimeUnit TimeUnit)
+bool UControlRigSequencerEditorLibrary::SetConstraintActiveKey(UTickableConstraint* InConstraint, bool bActive, FFrameNumber InTime, EMovieSceneTimeUnit TimeUnit)
 {
 	TWeakPtr<ISequencer> WeakSequencer = GetSequencerFromAsset();
 	if (WeakSequencer.IsValid() == false)
@@ -597,7 +597,7 @@ bool UControlRigSequencerEditorLibrary::SetConstraintActiveKey(UTickableConstrai
 		FFrameRate TickResolution = Sequencer->GetFocusedTickResolution();
 		FFrameRate DisplayRate = Sequencer->GetFocusedDisplayRate();
 
-		if (TimeUnit == ESequenceTimeUnit::DisplayRate)
+		if (TimeUnit == EMovieSceneTimeUnit::DisplayRate)
 		{
 			InTime = FFrameRate::TransformTime(FFrameTime(InTime, 0), DisplayRate, TickResolution).RoundToFrame();
 		}
@@ -610,7 +610,7 @@ bool UControlRigSequencerEditorLibrary::SetConstraintActiveKey(UTickableConstrai
 	return false;
 }
 
-bool UControlRigSequencerEditorLibrary::GetConstraintKeys(UTickableConstraint* InConstraint, UMovieSceneSection* ConstraintSection, TArray<bool>& OutBools, TArray<FFrameNumber>& OutFrames, ESequenceTimeUnit TimeUnit)
+bool UControlRigSequencerEditorLibrary::GetConstraintKeys(UTickableConstraint* InConstraint, UMovieSceneSection* ConstraintSection, TArray<bool>& OutBools, TArray<FFrameNumber>& OutFrames, EMovieSceneTimeUnit TimeUnit)
 {
 	if (InConstraint == nullptr)
 	{
@@ -661,7 +661,7 @@ bool UControlRigSequencerEditorLibrary::GetConstraintKeys(UTickableConstraint* I
 			for (FFrameNumber Frame : OurKeyTimes)
 			{
 				ConstraintAndChannel->ActiveChannel.Evaluate(FFrameTime(Frame, 0.0f), OutBools[Index]);
-				if (TimeUnit == ESequenceTimeUnit::DisplayRate)
+				if (TimeUnit == EMovieSceneTimeUnit::DisplayRate)
 				{
 					Frame = FFrameRate::TransformTime(FFrameTime(Frame, 0), TickResolution,DisplayRate).RoundToFrame();
 				}
@@ -684,7 +684,7 @@ bool UControlRigSequencerEditorLibrary::GetConstraintKeys(UTickableConstraint* I
 }
 
 
-bool UControlRigSequencerEditorLibrary::MoveConstraintKey(UTickableConstraint* Constraint, UMovieSceneSection* ConstraintSection, FFrameNumber InTime, FFrameNumber InNewTime, ESequenceTimeUnit TimeUnit)
+bool UControlRigSequencerEditorLibrary::MoveConstraintKey(UTickableConstraint* Constraint, UMovieSceneSection* ConstraintSection, FFrameNumber InTime, FFrameNumber InNewTime, EMovieSceneTimeUnit TimeUnit)
 {
 	TWeakPtr<ISequencer> WeakSequencer = GetSequencerFromAsset();
 	if (WeakSequencer.IsValid() == false)
@@ -713,7 +713,7 @@ bool UControlRigSequencerEditorLibrary::MoveConstraintKey(UTickableConstraint* C
 	FFrameRate TickResolution = Sequencer->GetFocusedTickResolution();
 	FFrameRate DisplayRate = Sequencer->GetFocusedDisplayRate();
 
-	if (TimeUnit == ESequenceTimeUnit::DisplayRate)
+	if (TimeUnit == EMovieSceneTimeUnit::DisplayRate)
 	{
 		InTime = FFrameRate::TransformTime(FFrameTime(InTime, 0), DisplayRate, TickResolution).RoundToFrame();
 		InNewTime = FFrameRate::TransformTime(FFrameTime(InNewTime, 0), DisplayRate, TickResolution).RoundToFrame();
@@ -742,7 +742,7 @@ bool UControlRigSequencerEditorLibrary::MoveConstraintKey(UTickableConstraint* C
 	return true;
 }
 
-bool UControlRigSequencerEditorLibrary::DeleteConstraintKey(UTickableConstraint* Constraint, UMovieSceneSection* ConstraintSection, FFrameNumber InTime, ESequenceTimeUnit TimeUnit)
+bool UControlRigSequencerEditorLibrary::DeleteConstraintKey(UTickableConstraint* Constraint, UMovieSceneSection* ConstraintSection, FFrameNumber InTime, EMovieSceneTimeUnit TimeUnit)
 {
 	TWeakPtr<ISequencer> WeakSequencer = GetSequencerFromAsset();
 	if (WeakSequencer.IsValid() == false)
@@ -771,7 +771,7 @@ bool UControlRigSequencerEditorLibrary::DeleteConstraintKey(UTickableConstraint*
 	FFrameRate TickResolution = Sequencer->GetFocusedTickResolution();
 	FFrameRate DisplayRate = Sequencer->GetFocusedDisplayRate();
 
-	if (TimeUnit == ESequenceTimeUnit::DisplayRate)
+	if (TimeUnit == EMovieSceneTimeUnit::DisplayRate)
 	{
 		InTime = FFrameRate::TransformTime(FFrameTime(InTime, 0), DisplayRate, TickResolution).RoundToFrame();
 	}
@@ -797,7 +797,7 @@ bool UControlRigSequencerEditorLibrary::DeleteConstraintKey(UTickableConstraint*
 	return true;
 }
 
-bool UControlRigSequencerEditorLibrary::BakeConstraint(UWorld* World, UTickableConstraint* Constraint, const TArray<FFrameNumber>& Frames, ESequenceTimeUnit TimeUnit)
+bool UControlRigSequencerEditorLibrary::BakeConstraint(UWorld* World, UTickableConstraint* Constraint, const TArray<FFrameNumber>& Frames, EMovieSceneTimeUnit TimeUnit)
 {
 	if (!World)
 	{
@@ -828,7 +828,7 @@ bool UControlRigSequencerEditorLibrary::BakeConstraint(UWorld* World, UTickableC
 		int32 Index = 0;
 		for (FFrameNumber Frame : Frames) 
 		{
-			if (TimeUnit == ESequenceTimeUnit::DisplayRate)
+			if (TimeUnit == EMovieSceneTimeUnit::DisplayRate)
 			{
 				Frame = FFrameRate::TransformTime(FFrameTime(Frame, 0), MovieScene->GetDisplayRate(), MovieScene->GetTickResolution()).RoundToFrame();
 			}
@@ -892,14 +892,14 @@ bool UControlRigSequencerEditorLibrary::BakeConstraints(UWorld* World, TArray<UT
 
 
 bool UControlRigSequencerEditorLibrary::SnapControlRig(ULevelSequence* LevelSequence, FFrameNumber StartFrame, FFrameNumber EndFrame, const FControlRigSnapperSelection& ChildrenToSnap,
-	const FControlRigSnapperSelection& ParentToSnap, const UControlRigSnapSettings* SnapSettings, ESequenceTimeUnit TimeUnit)
+	const FControlRigSnapperSelection& ParentToSnap, const UControlRigSnapSettings* SnapSettings, EMovieSceneTimeUnit TimeUnit)
 {
 	if (LevelSequence == nullptr || LevelSequence->GetMovieScene() == nullptr)
 	{
 		return false;
 	}
 	FControlRigSnapper Snapper;
-	if (TimeUnit == ESequenceTimeUnit::DisplayRate)
+	if (TimeUnit == EMovieSceneTimeUnit::DisplayRate)
 	{
 		StartFrame = FFrameRate::TransformTime(FFrameTime(StartFrame, 0), LevelSequence->GetMovieScene()->GetDisplayRate(), LevelSequence->GetMovieScene()->GetTickResolution()).RoundToFrame();
 		EndFrame = FFrameRate::TransformTime(FFrameTime(EndFrame, 0), LevelSequence->GetMovieScene()->GetDisplayRate(), LevelSequence->GetMovieScene()->GetTickResolution()).RoundToFrame();
@@ -908,7 +908,7 @@ bool UControlRigSequencerEditorLibrary::SnapControlRig(ULevelSequence* LevelSequ
 	return Snapper.SnapIt(StartFrame, EndFrame, ChildrenToSnap, ParentToSnap,SnapSettings);
 }
 
-FTransform UControlRigSequencerEditorLibrary::GetActorWorldTransform(ULevelSequence* LevelSequence, AActor* Actor, FFrameNumber Frame, ESequenceTimeUnit TimeUnit)
+FTransform UControlRigSequencerEditorLibrary::GetActorWorldTransform(ULevelSequence* LevelSequence, AActor* Actor, FFrameNumber Frame, EMovieSceneTimeUnit TimeUnit)
 {
 	TArray<FFrameNumber> Frames;
 	Frames.Add(Frame);
@@ -919,14 +919,14 @@ FTransform UControlRigSequencerEditorLibrary::GetActorWorldTransform(ULevelSeque
 	}
 	return FTransform::Identity;
 }
-static void ConvertFramesToTickResolution(UMovieScene* MovieScene, const TArray<FFrameNumber>& InFrames, ESequenceTimeUnit TimeUnit, TArray<FFrameNumber>& OutFrames)
+static void ConvertFramesToTickResolution(UMovieScene* MovieScene, const TArray<FFrameNumber>& InFrames, EMovieSceneTimeUnit TimeUnit, TArray<FFrameNumber>& OutFrames)
 {
 	OutFrames.SetNum(InFrames.Num());
 	int32 Index = 0;
 	for (const FFrameNumber& Frame : InFrames)
 	{
 		FFrameTime FrameTime(Frame);
-		if (TimeUnit == ESequenceTimeUnit::DisplayRate)
+		if (TimeUnit == EMovieSceneTimeUnit::DisplayRate)
 		{
 			FrameTime = FFrameRate::TransformTime(FrameTime, MovieScene->GetDisplayRate(), MovieScene->GetTickResolution());
 		}
@@ -934,7 +934,7 @@ static void ConvertFramesToTickResolution(UMovieScene* MovieScene, const TArray<
 	}
 }
 
-TArray<FTransform> UControlRigSequencerEditorLibrary::GetActorWorldTransforms(ULevelSequence* LevelSequence,AActor* Actor, const TArray<FFrameNumber>& InFrames, ESequenceTimeUnit TimeUnit)
+TArray<FTransform> UControlRigSequencerEditorLibrary::GetActorWorldTransforms(ULevelSequence* LevelSequence,AActor* Actor, const TArray<FFrameNumber>& InFrames, EMovieSceneTimeUnit TimeUnit)
 {
 	TWeakPtr<ISequencer> WeakSequencer = GetSequencerFromAsset();
 	TArray<FTransform> OutWorldTransforms;
@@ -951,7 +951,7 @@ TArray<FTransform> UControlRigSequencerEditorLibrary::GetActorWorldTransforms(UL
 }
 
 FTransform UControlRigSequencerEditorLibrary::GetSkeletalMeshComponentWorldTransform(ULevelSequence* LevelSequence,USkeletalMeshComponent* SkeletalMeshComponent, FFrameNumber Frame, 
-	ESequenceTimeUnit TimeUnit, FName SocketName)
+	EMovieSceneTimeUnit TimeUnit, FName SocketName)
 {
 	TArray<FFrameNumber> Frames;
 	Frames.Add(Frame);
@@ -964,7 +964,7 @@ FTransform UControlRigSequencerEditorLibrary::GetSkeletalMeshComponentWorldTrans
 }
 
 TArray<FTransform> UControlRigSequencerEditorLibrary::GetSkeletalMeshComponentWorldTransforms(ULevelSequence* LevelSequence,USkeletalMeshComponent* SkeletalMeshComponent, const TArray<FFrameNumber>& InFrames,
-	ESequenceTimeUnit TimeUnit, FName SocketName)
+	EMovieSceneTimeUnit TimeUnit, FName SocketName)
 {
 	TWeakPtr<ISequencer> WeakSequencer = GetSequencerFromAsset();
 	TArray<FTransform> OutWorldTransforms;
@@ -986,7 +986,7 @@ TArray<FTransform> UControlRigSequencerEditorLibrary::GetSkeletalMeshComponentWo
 }
 
 FTransform UControlRigSequencerEditorLibrary::GetControlRigWorldTransform(ULevelSequence* LevelSequence,UControlRig* ControlRig, FName ControlName, FFrameNumber Frame,
-	ESequenceTimeUnit TimeUnit)
+	EMovieSceneTimeUnit TimeUnit)
 {
 	TArray<FFrameNumber> Frames;
 	Frames.Add(Frame);
@@ -999,7 +999,7 @@ FTransform UControlRigSequencerEditorLibrary::GetControlRigWorldTransform(ULevel
 }
 
 TArray<FTransform> UControlRigSequencerEditorLibrary::GetControlRigWorldTransforms(ULevelSequence* LevelSequence,UControlRig* ControlRig, FName ControlName, const TArray<FFrameNumber>& InFrames,
-	ESequenceTimeUnit TimeUnit)
+	EMovieSceneTimeUnit TimeUnit)
 {
 	TWeakPtr<ISequencer> WeakSequencer = GetSequencerFromAsset();
 	TArray<FTransform> OutWorldTransforms;
@@ -1031,7 +1031,7 @@ TArray<FTransform> UControlRigSequencerEditorLibrary::GetControlRigWorldTransfor
 
 
 static void LocalSetControlRigWorldTransforms(ULevelSequence* LevelSequence,UControlRig* ControlRig, FName ControlName, EControlRigSetKey SetKey, const TArray<FFrameNumber>& InFrames, 
-	const TArray<FTransform>& WorldTransforms, ESequenceTimeUnit TimeUnit)
+	const TArray<FTransform>& WorldTransforms, EMovieSceneTimeUnit TimeUnit)
 {
 	TWeakPtr<ISequencer> WeakSequencer = GetSequencerFromAsset();
 	if (WeakSequencer.IsValid() && ControlRig)
@@ -1085,7 +1085,7 @@ static void LocalSetControlRigWorldTransforms(ULevelSequence* LevelSequence,UCon
 }
 
 void UControlRigSequencerEditorLibrary::SetControlRigWorldTransform(ULevelSequence* LevelSequence, UControlRig* ControlRig, FName ControlName, FFrameNumber Frame, const FTransform& WorldTransform,
-	ESequenceTimeUnit TimeUnit, bool bSetKey)
+	EMovieSceneTimeUnit TimeUnit, bool bSetKey)
 {
 	EControlRigSetKey SetKey = bSetKey ? EControlRigSetKey::Always : EControlRigSetKey::DoNotCare;
 	TArray<FFrameNumber> Frames;
@@ -1098,7 +1098,7 @@ void UControlRigSequencerEditorLibrary::SetControlRigWorldTransform(ULevelSequen
 }
 
 void UControlRigSequencerEditorLibrary::SetControlRigWorldTransforms(ULevelSequence* LevelSequence, UControlRig* ControlRig, FName ControlName, const TArray<FFrameNumber>& Frames, const TArray<FTransform>& WorldTransforms,
-	ESequenceTimeUnit TimeUnit)
+	EMovieSceneTimeUnit TimeUnit)
 {
 	LocalSetControlRigWorldTransforms(LevelSequence, ControlRig, ControlName, EControlRigSetKey::Always, Frames, WorldTransforms,TimeUnit);
 }
@@ -1311,7 +1311,7 @@ bool UControlRigSequencerEditorLibrary::BakeToControlRig(UWorld* World, ULevelSe
 }
 
 bool UControlRigSequencerEditorLibrary::LoadAnimSequenceIntoControlRigSection(UMovieSceneSection* MovieSceneSection, UAnimSequence* AnimSequence, USkeletalMeshComponent* SkelMeshComp,
-	FFrameNumber InStartFrame, ESequenceTimeUnit TimeUnit,bool bKeyReduce, float Tolerance, EMovieSceneKeyInterpolation Interpolation)
+	FFrameNumber InStartFrame, EMovieSceneTimeUnit TimeUnit,bool bKeyReduce, float Tolerance, EMovieSceneKeyInterpolation Interpolation)
 {
 	if (MovieSceneSection == nullptr || AnimSequence == nullptr || SkelMeshComp == nullptr)
 	{
@@ -1324,7 +1324,7 @@ bool UControlRigSequencerEditorLibrary::LoadAnimSequenceIntoControlRigSection(UM
 	}
 	if (UMovieSceneControlRigParameterSection* Section = Cast<UMovieSceneControlRigParameterSection>(MovieSceneSection))
 	{
-		if (TimeUnit == ESequenceTimeUnit::DisplayRate)
+		if (TimeUnit == EMovieSceneTimeUnit::DisplayRate)
 		{
 			InStartFrame = FFrameRate::TransformTime(FFrameTime(InStartFrame, 0), MovieScene->GetDisplayRate(),MovieScene->GetTickResolution()).RoundToFrame();
 		}
@@ -1334,7 +1334,7 @@ bool UControlRigSequencerEditorLibrary::LoadAnimSequenceIntoControlRigSection(UM
 }
 
 static bool LocalGetControlRigControlValues(IMovieScenePlayer* Player, UMovieSceneSequence* MovieSceneSequence, FMovieSceneSequenceIDRef Template, FMovieSceneSequenceTransform& RootToLocalTransform,
-	UControlRig* ControlRig, const FName& ControlName, ESequenceTimeUnit TimeUnit,
+	UControlRig* ControlRig, const FName& ControlName, EMovieSceneTimeUnit TimeUnit,
 	const TArray<FFrameNumber>& InFrames, TArray<FRigControlValue>& OutValues)
 {
 	if (Player == nullptr || MovieSceneSequence == nullptr || ControlRig == nullptr)
@@ -1355,7 +1355,7 @@ static bool LocalGetControlRigControlValues(IMovieScenePlayer* Player, UMovieSce
 		{
 			const FFrameNumber& FrameNumber = InFrames[Index];
 			FFrameTime GlobalTime(FrameNumber);
-			if (TimeUnit == ESequenceTimeUnit::DisplayRate)
+			if (TimeUnit == EMovieSceneTimeUnit::DisplayRate)
 			{
 				GlobalTime = FFrameRate::TransformTime(GlobalTime, MovieScene->GetDisplayRate(), MovieScene->GetTickResolution());
 			}
@@ -1370,7 +1370,7 @@ static bool LocalGetControlRigControlValues(IMovieScenePlayer* Player, UMovieSce
 	return true;
 }
 
-static bool GetControlRigValues(ISequencer* Sequencer, UControlRig* ControlRig, const FName& ControlName, ESequenceTimeUnit TimeUnit,
+static bool GetControlRigValues(ISequencer* Sequencer, UControlRig* ControlRig, const FName& ControlName, EMovieSceneTimeUnit TimeUnit,
 	const TArray<FFrameNumber>& Frames,  TArray<FRigControlValue>& OutValues)
 {
 	if (Sequencer->GetFocusedMovieSceneSequence())
@@ -1385,7 +1385,7 @@ static bool GetControlRigValues(ISequencer* Sequencer, UControlRig* ControlRig,
 	return false;
 }
 
-static bool GetControlRigValue(ISequencer* Sequencer, UControlRig* ControlRig, const FName& ControlName, ESequenceTimeUnit TimeUnit,
+static bool GetControlRigValue(ISequencer* Sequencer, UControlRig* ControlRig, const FName& ControlName, EMovieSceneTimeUnit TimeUnit,
 	const FFrameNumber Frame, FRigControlValue& OutValue)
 {
 	if (Sequencer->GetFocusedMovieSceneSequence())
@@ -1406,7 +1406,7 @@ static bool GetControlRigValue(ISequencer* Sequencer, UControlRig* ControlRig, c
 	return false;
 }
 
-static bool GetControlRigValues(UWorld* World, ULevelSequence* LevelSequence, UControlRig* ControlRig, const FName& ControlName, ESequenceTimeUnit TimeUnit,
+static bool GetControlRigValues(UWorld* World, ULevelSequence* LevelSequence, UControlRig* ControlRig, const FName& ControlName, EMovieSceneTimeUnit TimeUnit,
 	const TArray<FFrameNumber>& Frames, TArray<FRigControlValue>& OutValues)
 {
 	if (LevelSequence)
@@ -1425,7 +1425,7 @@ static bool GetControlRigValues(UWorld* World, ULevelSequence* LevelSequence, UC
 }
 
 
-float UControlRigSequencerEditorLibrary::GetLocalControlRigFloat(ULevelSequence* LevelSequence,UControlRig* ControlRig, FName ControlName, FFrameNumber Frame, ESequenceTimeUnit TimeUnit)
+float UControlRigSequencerEditorLibrary::GetLocalControlRigFloat(ULevelSequence* LevelSequence,UControlRig* ControlRig, FName ControlName, FFrameNumber Frame, EMovieSceneTimeUnit TimeUnit)
 {
 	float Value = 0.0f;
 	TWeakPtr<ISequencer> WeakSequencer = GetSequencerFromAsset();
@@ -1449,7 +1449,7 @@ float UControlRigSequencerEditorLibrary::GetLocalControlRigFloat(ULevelSequence*
 	return Value;
 }
 
-TArray<float> UControlRigSequencerEditorLibrary::GetLocalControlRigFloats(ULevelSequence* LevelSequence,UControlRig* ControlRig, FName ControlName, const TArray<FFrameNumber>& Frames, ESequenceTimeUnit TimeUnit)
+TArray<float> UControlRigSequencerEditorLibrary::GetLocalControlRigFloats(ULevelSequence* LevelSequence,UControlRig* ControlRig, FName ControlName, const TArray<FFrameNumber>& Frames, EMovieSceneTimeUnit TimeUnit)
 {
 	TArray<float> Values;
 	TWeakPtr<ISequencer> WeakSequencer = GetSequencerFromAsset();
@@ -1479,7 +1479,7 @@ TArray<float> UControlRigSequencerEditorLibrary::GetLocalControlRigFloats(ULevel
 	return Values;
 }
 
-void UControlRigSequencerEditorLibrary::SetLocalControlRigFloat(ULevelSequence* LevelSequence, UControlRig* ControlRig, FName ControlName, FFrameNumber Frame, float Value, ESequenceTimeUnit TimeUnit, bool bSetKey)
+void UControlRigSequencerEditorLibrary::SetLocalControlRigFloat(ULevelSequence* LevelSequence, UControlRig* ControlRig, FName ControlName, FFrameNumber Frame, float Value, EMovieSceneTimeUnit TimeUnit, bool bSetKey)
 {
 	if (LevelSequence == nullptr || ControlRig == nullptr)
 	{
@@ -1496,7 +1496,7 @@ void UControlRigSequencerEditorLibrary::SetLocalControlRigFloat(ULevelSequence*
 	if (UMovieScene* MovieScene = LevelSequence->GetMovieScene())
 	{
 		FFrameRate TickResolution = MovieScene->GetTickResolution();
-		if (TimeUnit == ESequenceTimeUnit::DisplayRate)
+		if (TimeUnit == EMovieSceneTimeUnit::DisplayRate)
 		{
 			Frame = FFrameRate::TransformTime(FFrameTime(Frame, 0), MovieScene->GetDisplayRate(), MovieScene->GetTickResolution()).RoundToFrame();
 		}
@@ -1508,7 +1508,7 @@ void UControlRigSequencerEditorLibrary::SetLocalControlRigFloat(ULevelSequence*
 }
 
 void UControlRigSequencerEditorLibrary::SetLocalControlRigFloats(ULevelSequence* LevelSequence, UControlRig* ControlRig, FName ControlName, const TArray<FFrameNumber>& Frames, const TArray<float> Values,
-	ESequenceTimeUnit TimeUnit)
+	EMovieSceneTimeUnit TimeUnit)
 {
 	if (LevelSequence == nullptr || ControlRig == nullptr || (Frames.Num() != Values.Num()))
 	{
@@ -1530,7 +1530,7 @@ void UControlRigSequencerEditorLibrary::SetLocalControlRigFloats(ULevelSequence*
 		for (int32 Index = 0; Index < Frames.Num(); ++Index)
 		{
 			FFrameNumber Frame = Frames[Index];
-			if (TimeUnit == ESequenceTimeUnit::DisplayRate)
+			if (TimeUnit == EMovieSceneTimeUnit::DisplayRate)
 			{
 				Frame = FFrameRate::TransformTime(FFrameTime(Frame, 0), MovieScene->GetDisplayRate(), MovieScene->GetTickResolution()).RoundToFrame();
 			}
@@ -1542,7 +1542,7 @@ void UControlRigSequencerEditorLibrary::SetLocalControlRigFloats(ULevelSequence*
 }
 
 
-bool UControlRigSequencerEditorLibrary::GetLocalControlRigBool(ULevelSequence* LevelSequence, UControlRig* ControlRig, FName ControlName, FFrameNumber Frame, ESequenceTimeUnit TimeUnit)
+bool UControlRigSequencerEditorLibrary::GetLocalControlRigBool(ULevelSequence* LevelSequence, UControlRig* ControlRig, FName ControlName, FFrameNumber Frame, EMovieSceneTimeUnit TimeUnit)
 {
 	bool Value = true;
 	TWeakPtr<ISequencer> WeakSequencer = GetSequencerFromAsset();
@@ -1566,7 +1566,7 @@ bool UControlRigSequencerEditorLibrary::GetLocalControlRigBool(ULevelSequence* L
 	return Value;
 }
 
-TArray<bool> UControlRigSequencerEditorLibrary::GetLocalControlRigBools(ULevelSequence* LevelSequence, UControlRig* ControlRig, FName ControlName, const TArray<FFrameNumber>& Frames, ESequenceTimeUnit TimeUnit)
+TArray<bool> UControlRigSequencerEditorLibrary::GetLocalControlRigBools(ULevelSequence* LevelSequence, UControlRig* ControlRig, FName ControlName, const TArray<FFrameNumber>& Frames, EMovieSceneTimeUnit TimeUnit)
 {
 	TArray<bool> Values;
 	TWeakPtr<ISequencer> WeakSequencer = GetSequencerFromAsset();
@@ -1595,7 +1595,7 @@ TArray<bool> UControlRigSequencerEditorLibrary::GetLocalControlRigBools(ULevelSe
 	return Values;
 }
 
-void UControlRigSequencerEditorLibrary::SetLocalControlRigBool(ULevelSequence* LevelSequence, UControlRig* ControlRig, FName ControlName, FFrameNumber Frame, bool Value, ESequenceTimeUnit TimeUnit, bool bSetKey)
+void UControlRigSequencerEditorLibrary::SetLocalControlRigBool(ULevelSequence* LevelSequence, UControlRig* ControlRig, FName ControlName, FFrameNumber Frame, bool Value, EMovieSceneTimeUnit TimeUnit, bool bSetKey)
 {
 	if (LevelSequence == nullptr || ControlRig == nullptr)
 	{
@@ -1612,7 +1612,7 @@ void UControlRigSequencerEditorLibrary::SetLocalControlRigBool(ULevelSequence* L
 			}
 		}
 		FFrameRate TickResolution = MovieScene->GetTickResolution();
-		if (TimeUnit == ESequenceTimeUnit::DisplayRate)
+		if (TimeUnit == EMovieSceneTimeUnit::DisplayRate)
 		{
 			Frame = FFrameRate::TransformTime(FFrameTime(Frame, 0), MovieScene->GetDisplayRate(), MovieScene->GetTickResolution()).RoundToFrame();
 		}
@@ -1624,7 +1624,7 @@ void UControlRigSequencerEditorLibrary::SetLocalControlRigBool(ULevelSequence* L
 }
 
 void UControlRigSequencerEditorLibrary::SetLocalControlRigBools(ULevelSequence* LevelSequence, UControlRig* ControlRig, FName ControlName, const TArray<FFrameNumber>& Frames,
-	const TArray<bool> Values, ESequenceTimeUnit TimeUnit)
+	const TArray<bool> Values, EMovieSceneTimeUnit TimeUnit)
 {
 	if (LevelSequence == nullptr || ControlRig == nullptr || (Frames.Num() != Values.Num()))
 	{
@@ -1646,7 +1646,7 @@ void UControlRigSequencerEditorLibrary::SetLocalControlRigBools(ULevelSequence*
 		for (int32 Index = 0; Index < Frames.Num(); ++Index)
 		{
 			FFrameNumber Frame = Frames[Index];
-			if (TimeUnit == ESequenceTimeUnit::DisplayRate)
+			if (TimeUnit == EMovieSceneTimeUnit::DisplayRate)
 			{
 				Frame = FFrameRate::TransformTime(FFrameTime(Frame, 0), MovieScene->GetDisplayRate(), MovieScene->GetTickResolution()).RoundToFrame();
 			}
@@ -1657,7 +1657,7 @@ void UControlRigSequencerEditorLibrary::SetLocalControlRigBools(ULevelSequence*
 	}
 }
 
-int32 UControlRigSequencerEditorLibrary::GetLocalControlRigInt(ULevelSequence* LevelSequence, UControlRig* ControlRig, FName ControlName, FFrameNumber Frame, ESequenceTimeUnit TimeUnit)
+int32 UControlRigSequencerEditorLibrary::GetLocalControlRigInt(ULevelSequence* LevelSequence, UControlRig* ControlRig, FName ControlName, FFrameNumber Frame, EMovieSceneTimeUnit TimeUnit)
 {
 	int32 Value = 0;
 	TWeakPtr<ISequencer> WeakSequencer = GetSequencerFromAsset();
@@ -1681,7 +1681,7 @@ int32 UControlRigSequencerEditorLibrary::GetLocalControlRigInt(ULevelSequence* L
 	return Value;
 }
 
-TArray<int32> UControlRigSequencerEditorLibrary::GetLocalControlRigInts(ULevelSequence* LevelSequence, UControlRig* ControlRig, FName ControlName, const TArray<FFrameNumber>& Frames, ESequenceTimeUnit TimeUnit)
+TArray<int32> UControlRigSequencerEditorLibrary::GetLocalControlRigInts(ULevelSequence* LevelSequence, UControlRig* ControlRig, FName ControlName, const TArray<FFrameNumber>& Frames, EMovieSceneTimeUnit TimeUnit)
 {
 	TArray<int32> Values;
 	TWeakPtr<ISequencer> WeakSequencer = GetSequencerFromAsset();
@@ -1710,7 +1710,7 @@ TArray<int32> UControlRigSequencerEditorLibrary::GetLocalControlRigInts(ULevelSe
 	return Values;
 }
 
-void UControlRigSequencerEditorLibrary::SetLocalControlRigInt(ULevelSequence* LevelSequence, UControlRig* ControlRig, FName ControlName, FFrameNumber Frame, int32 Value, ESequenceTimeUnit TimeUnit,bool bSetKey)
+void UControlRigSequencerEditorLibrary::SetLocalControlRigInt(ULevelSequence* LevelSequence, UControlRig* ControlRig, FName ControlName, FFrameNumber Frame, int32 Value, EMovieSceneTimeUnit TimeUnit,bool bSetKey)
 {
 	if (LevelSequence == nullptr || ControlRig == nullptr)
 	{
@@ -1727,7 +1727,7 @@ void UControlRigSequencerEditorLibrary::SetLocalControlRigInt(ULevelSequence* Le
 	if (UMovieScene* MovieScene = LevelSequence->GetMovieScene())
 	{
 		FFrameRate TickResolution = MovieScene->GetTickResolution();
-		if (TimeUnit == ESequenceTimeUnit::DisplayRate)
+		if (TimeUnit == EMovieSceneTimeUnit::DisplayRate)
 		{
 			Frame = FFrameRate::TransformTime(FFrameTime(Frame, 0), MovieScene->GetDisplayRate(), MovieScene->GetTickResolution()).RoundToFrame();
 		}
@@ -1739,7 +1739,7 @@ void UControlRigSequencerEditorLibrary::SetLocalControlRigInt(ULevelSequence* Le
 }
 
 void UControlRigSequencerEditorLibrary::SetLocalControlRigInts(ULevelSequence* LevelSequence, UControlRig* ControlRig, FName ControlName, const TArray<FFrameNumber>& Frames,
-	const TArray<int32> Values , ESequenceTimeUnit TimeUnit)
+	const TArray<int32> Values , EMovieSceneTimeUnit TimeUnit)
 {
 	if (LevelSequence == nullptr || ControlRig == nullptr || (Frames.Num() != Values.Num()))
 	{
@@ -1761,7 +1761,7 @@ void UControlRigSequencerEditorLibrary::SetLocalControlRigInts(ULevelSequence* L
 		for (int32 Index = 0; Index < Frames.Num(); ++Index)
 		{
 			FFrameNumber Frame = Frames[Index];
-			if (TimeUnit == ESequenceTimeUnit::DisplayRate)
+			if (TimeUnit == EMovieSceneTimeUnit::DisplayRate)
 			{
 				Frame = FFrameRate::TransformTime(FFrameTime(Frame, 0), MovieScene->GetDisplayRate(), MovieScene->GetTickResolution()).RoundToFrame();
 			}
@@ -1773,7 +1773,7 @@ void UControlRigSequencerEditorLibrary::SetLocalControlRigInts(ULevelSequence* L
 }
 
 
-FVector2D UControlRigSequencerEditorLibrary::GetLocalControlRigVector2D(ULevelSequence* LevelSequence, UControlRig* ControlRig, FName ControlName, FFrameNumber Frame, ESequenceTimeUnit TimeUnit)
+FVector2D UControlRigSequencerEditorLibrary::GetLocalControlRigVector2D(ULevelSequence* LevelSequence, UControlRig* ControlRig, FName ControlName, FFrameNumber Frame, EMovieSceneTimeUnit TimeUnit)
 {
 	FVector2D Value;
 	TWeakPtr<ISequencer> WeakSequencer = GetSequencerFromAsset();
@@ -1798,7 +1798,7 @@ FVector2D UControlRigSequencerEditorLibrary::GetLocalControlRigVector2D(ULevelSe
 	return Value;
 }
 
-TArray<FVector2D> UControlRigSequencerEditorLibrary::GetLocalControlRigVector2Ds(ULevelSequence* LevelSequence, UControlRig* ControlRig, FName ControlName, const TArray<FFrameNumber>& Frames, ESequenceTimeUnit TimeUnit)
+TArray<FVector2D> UControlRigSequencerEditorLibrary::GetLocalControlRigVector2Ds(ULevelSequence* LevelSequence, UControlRig* ControlRig, FName ControlName, const TArray<FFrameNumber>& Frames, EMovieSceneTimeUnit TimeUnit)
 {
 	TArray<FVector2D> Values;
 	TWeakPtr<ISequencer> WeakSequencer = GetSequencerFromAsset();
@@ -1828,7 +1828,7 @@ TArray<FVector2D> UControlRigSequencerEditorLibrary::GetLocalControlRigVector2Ds
 
 }
 
-void UControlRigSequencerEditorLibrary::SetLocalControlRigVector2D(ULevelSequence* LevelSequence, UControlRig* ControlRig, FName ControlName, FFrameNumber Frame, FVector2D Value, ESequenceTimeUnit TimeUnit,bool bSetKey)
+void UControlRigSequencerEditorLibrary::SetLocalControlRigVector2D(ULevelSequence* LevelSequence, UControlRig* ControlRig, FName ControlName, FFrameNumber Frame, FVector2D Value, EMovieSceneTimeUnit TimeUnit,bool bSetKey)
 {
 	if (LevelSequence == nullptr || ControlRig == nullptr)
 	{
@@ -1845,7 +1845,7 @@ void UControlRigSequencerEditorLibrary::SetLocalControlRigVector2D(ULevelSequenc
 	if (UMovieScene* MovieScene = LevelSequence->GetMovieScene())
 	{
 		FFrameRate TickResolution = MovieScene->GetTickResolution();
-		if (TimeUnit == ESequenceTimeUnit::DisplayRate)
+		if (TimeUnit == EMovieSceneTimeUnit::DisplayRate)
 		{
 			Frame = FFrameRate::TransformTime(FFrameTime(Frame, 0), MovieScene->GetDisplayRate(), MovieScene->GetTickResolution()).RoundToFrame();
 		}
@@ -1857,7 +1857,7 @@ void UControlRigSequencerEditorLibrary::SetLocalControlRigVector2D(ULevelSequenc
 }
 
 void UControlRigSequencerEditorLibrary::SetLocalControlRigVector2Ds(ULevelSequence* LevelSequence, UControlRig* ControlRig, FName ControlName, const TArray<FFrameNumber>& Frames, 
-	const TArray<FVector2D> Values, ESequenceTimeUnit TimeUnit)
+	const TArray<FVector2D> Values, EMovieSceneTimeUnit TimeUnit)
 {
 	if (LevelSequence == nullptr || ControlRig == nullptr || (Frames.Num() != Values.Num()))
 	{
@@ -1879,7 +1879,7 @@ void UControlRigSequencerEditorLibrary::SetLocalControlRigVector2Ds(ULevelSequen
 		for (int32 Index = 0; Index < Frames.Num(); ++Index)
 		{
 			FFrameNumber Frame = Frames[Index];
-			if (TimeUnit == ESequenceTimeUnit::DisplayRate)
+			if (TimeUnit == EMovieSceneTimeUnit::DisplayRate)
 			{
 				Frame = FFrameRate::TransformTime(FFrameTime(Frame, 0), MovieScene->GetDisplayRate(), MovieScene->GetTickResolution()).RoundToFrame();
 			}
@@ -1891,7 +1891,7 @@ void UControlRigSequencerEditorLibrary::SetLocalControlRigVector2Ds(ULevelSequen
 }
 
 
-FVector UControlRigSequencerEditorLibrary::GetLocalControlRigPosition(ULevelSequence* LevelSequence, UControlRig* ControlRig, FName ControlName, FFrameNumber Frame, ESequenceTimeUnit TimeUnit)
+FVector UControlRigSequencerEditorLibrary::GetLocalControlRigPosition(ULevelSequence* LevelSequence, UControlRig* ControlRig, FName ControlName, FFrameNumber Frame, EMovieSceneTimeUnit TimeUnit)
 {
 	FVector Value;
 	TWeakPtr<ISequencer> WeakSequencer = GetSequencerFromAsset();
@@ -1915,7 +1915,7 @@ FVector UControlRigSequencerEditorLibrary::GetLocalControlRigPosition(ULevelSequ
 	return Value;
 }
 
-TArray<FVector> UControlRigSequencerEditorLibrary::GetLocalControlRigPositions(ULevelSequence* LevelSequence, UControlRig* ControlRig, FName ControlName, const TArray<FFrameNumber>& Frames, ESequenceTimeUnit TimeUnit)
+TArray<FVector> UControlRigSequencerEditorLibrary::GetLocalControlRigPositions(ULevelSequence* LevelSequence, UControlRig* ControlRig, FName ControlName, const TArray<FFrameNumber>& Frames, EMovieSceneTimeUnit TimeUnit)
 {
 	TArray<FVector> Values;
 	TWeakPtr<ISequencer> WeakSequencer = GetSequencerFromAsset();
@@ -1944,7 +1944,7 @@ TArray<FVector> UControlRigSequencerEditorLibrary::GetLocalControlRigPositions(U
 	return Values;
 }
 
-void UControlRigSequencerEditorLibrary::SetLocalControlRigPosition(ULevelSequence* LevelSequence, UControlRig* ControlRig, FName ControlName, FFrameNumber Frame, FVector Value, ESequenceTimeUnit TimeUnit,bool bSetKey)
+void UControlRigSequencerEditorLibrary::SetLocalControlRigPosition(ULevelSequence* LevelSequence, UControlRig* ControlRig, FName ControlName, FFrameNumber Frame, FVector Value, EMovieSceneTimeUnit TimeUnit,bool bSetKey)
 {
 	if (LevelSequence == nullptr || ControlRig == nullptr)
 	{
@@ -1961,7 +1961,7 @@ void UControlRigSequencerEditorLibrary::SetLocalControlRigPosition(ULevelSequenc
 	if (UMovieScene* MovieScene = LevelSequence->GetMovieScene())
 	{
 		FFrameRate TickResolution = MovieScene->GetTickResolution();
-		if (TimeUnit == ESequenceTimeUnit::DisplayRate)
+		if (TimeUnit == EMovieSceneTimeUnit::DisplayRate)
 		{
 			Frame = FFrameRate::TransformTime(FFrameTime(Frame, 0), MovieScene->GetDisplayRate(), MovieScene->GetTickResolution()).RoundToFrame();
 		}
@@ -1973,7 +1973,7 @@ void UControlRigSequencerEditorLibrary::SetLocalControlRigPosition(ULevelSequenc
 }
 
 void UControlRigSequencerEditorLibrary::SetLocalControlRigPositions(ULevelSequence* LevelSequence, UControlRig* ControlRig, FName ControlName, const TArray<FFrameNumber>& Frames, 
-	const TArray<FVector> Values, ESequenceTimeUnit TimeUnit)
+	const TArray<FVector> Values, EMovieSceneTimeUnit TimeUnit)
 {
 	if (LevelSequence == nullptr || ControlRig == nullptr || (Frames.Num() != Values.Num()))
 	{
@@ -1995,7 +1995,7 @@ void UControlRigSequencerEditorLibrary::SetLocalControlRigPositions(ULevelSequen
 		for (int32 Index = 0; Index < Frames.Num(); ++Index)
 		{
 			FFrameNumber Frame = Frames[Index];
-			if (TimeUnit == ESequenceTimeUnit::DisplayRate)
+			if (TimeUnit == EMovieSceneTimeUnit::DisplayRate)
 			{
 				Frame = FFrameRate::TransformTime(FFrameTime(Frame, 0), MovieScene->GetDisplayRate(), MovieScene->GetTickResolution()).RoundToFrame();
 			}
@@ -2007,7 +2007,7 @@ void UControlRigSequencerEditorLibrary::SetLocalControlRigPositions(ULevelSequen
 }
 
 
-FRotator UControlRigSequencerEditorLibrary::GetLocalControlRigRotator(ULevelSequence* LevelSequence, UControlRig* ControlRig, FName ControlName, FFrameNumber Frame, ESequenceTimeUnit TimeUnit)
+FRotator UControlRigSequencerEditorLibrary::GetLocalControlRigRotator(ULevelSequence* LevelSequence, UControlRig* ControlRig, FName ControlName, FFrameNumber Frame, EMovieSceneTimeUnit TimeUnit)
 {
 	FRotator Value;
 	TWeakPtr<ISequencer> WeakSequencer = GetSequencerFromAsset();
@@ -2031,7 +2031,7 @@ FRotator UControlRigSequencerEditorLibrary::GetLocalControlRigRotator(ULevelSequ
 	return Value;
 }
 
-TArray<FRotator> UControlRigSequencerEditorLibrary::GetLocalControlRigRotators(ULevelSequence* LevelSequence, UControlRig* ControlRig, FName ControlName, const TArray<FFrameNumber>& Frames, ESequenceTimeUnit TimeUnit)
+TArray<FRotator> UControlRigSequencerEditorLibrary::GetLocalControlRigRotators(ULevelSequence* LevelSequence, UControlRig* ControlRig, FName ControlName, const TArray<FFrameNumber>& Frames, EMovieSceneTimeUnit TimeUnit)
 {
 	TArray<FRotator> Values;
 	TWeakPtr<ISequencer> WeakSequencer = GetSequencerFromAsset();
@@ -2060,7 +2060,7 @@ TArray<FRotator> UControlRigSequencerEditorLibrary::GetLocalControlRigRotators(U
 	return Values;
 }
 
-void UControlRigSequencerEditorLibrary::SetLocalControlRigRotator(ULevelSequence* LevelSequence, UControlRig* ControlRig, FName ControlName, FFrameNumber Frame, FRotator Value, ESequenceTimeUnit TimeUnit, bool bSetKey)
+void UControlRigSequencerEditorLibrary::SetLocalControlRigRotator(ULevelSequence* LevelSequence, UControlRig* ControlRig, FName ControlName, FFrameNumber Frame, FRotator Value, EMovieSceneTimeUnit TimeUnit, bool bSetKey)
 {
 	if (LevelSequence == nullptr || ControlRig == nullptr)
 	{
@@ -2077,7 +2077,7 @@ void UControlRigSequencerEditorLibrary::SetLocalControlRigRotator(ULevelSequence
 	if (UMovieScene* MovieScene = LevelSequence->GetMovieScene())
 	{
 		FFrameRate TickResolution = MovieScene->GetTickResolution();
-		if (TimeUnit == ESequenceTimeUnit::DisplayRate)
+		if (TimeUnit == EMovieSceneTimeUnit::DisplayRate)
 		{
 			Frame = FFrameRate::TransformTime(FFrameTime(Frame, 0), MovieScene->GetDisplayRate(), MovieScene->GetTickResolution()).RoundToFrame();
 		}
@@ -2089,7 +2089,7 @@ void UControlRigSequencerEditorLibrary::SetLocalControlRigRotator(ULevelSequence
 }
 
 void UControlRigSequencerEditorLibrary::SetLocalControlRigRotators(ULevelSequence* LevelSequence, UControlRig* ControlRig, FName ControlName, const TArray<FFrameNumber>& Frames,
-	const TArray<FRotator> Values, ESequenceTimeUnit TimeUnit)
+	const TArray<FRotator> Values, EMovieSceneTimeUnit TimeUnit)
 {
 	if (LevelSequence == nullptr || ControlRig == nullptr || (Frames.Num() != Values.Num()))
 	{
@@ -2111,7 +2111,7 @@ void UControlRigSequencerEditorLibrary::SetLocalControlRigRotators(ULevelSequenc
 		for (int32 Index = 0; Index < Frames.Num(); ++Index)
 		{
 			FFrameNumber Frame = Frames[Index];
-			if (TimeUnit == ESequenceTimeUnit::DisplayRate)
+			if (TimeUnit == EMovieSceneTimeUnit::DisplayRate)
 			{
 				Frame = FFrameRate::TransformTime(FFrameTime(Frame, 0), MovieScene->GetDisplayRate(), MovieScene->GetTickResolution()).RoundToFrame();
 			}
@@ -2123,7 +2123,7 @@ void UControlRigSequencerEditorLibrary::SetLocalControlRigRotators(ULevelSequenc
 }
 
 
-FVector UControlRigSequencerEditorLibrary::GetLocalControlRigScale(ULevelSequence* LevelSequence, UControlRig* ControlRig, FName ControlName, FFrameNumber Frame,  ESequenceTimeUnit TimeUnit)
+FVector UControlRigSequencerEditorLibrary::GetLocalControlRigScale(ULevelSequence* LevelSequence, UControlRig* ControlRig, FName ControlName, FFrameNumber Frame,  EMovieSceneTimeUnit TimeUnit)
 {
 	FVector Value;
 	TWeakPtr<ISequencer> WeakSequencer = GetSequencerFromAsset();
@@ -2147,7 +2147,7 @@ FVector UControlRigSequencerEditorLibrary::GetLocalControlRigScale(ULevelSequenc
 	return Value;
 }
 
-TArray<FVector>UControlRigSequencerEditorLibrary::GetLocalControlRigScales(ULevelSequence* LevelSequence, UControlRig* ControlRig, FName ControlName, const TArray<FFrameNumber>& Frames, ESequenceTimeUnit TimeUnit)
+TArray<FVector>UControlRigSequencerEditorLibrary::GetLocalControlRigScales(ULevelSequence* LevelSequence, UControlRig* ControlRig, FName ControlName, const TArray<FFrameNumber>& Frames, EMovieSceneTimeUnit TimeUnit)
 {
 	TArray<FVector> Values;
 	TWeakPtr<ISequencer> WeakSequencer = GetSequencerFromAsset();
@@ -2176,7 +2176,7 @@ TArray<FVector>UControlRigSequencerEditorLibrary::GetLocalControlRigScales(ULeve
 	return Values;
 }
 
-void UControlRigSequencerEditorLibrary::SetLocalControlRigScale(ULevelSequence* LevelSequence, UControlRig* ControlRig, FName ControlName, FFrameNumber Frame, FVector Value, ESequenceTimeUnit TimeUnit, bool bSetKey)
+void UControlRigSequencerEditorLibrary::SetLocalControlRigScale(ULevelSequence* LevelSequence, UControlRig* ControlRig, FName ControlName, FFrameNumber Frame, FVector Value, EMovieSceneTimeUnit TimeUnit, bool bSetKey)
 {
 	if (LevelSequence == nullptr || ControlRig == nullptr)
 	{
@@ -2193,7 +2193,7 @@ void UControlRigSequencerEditorLibrary::SetLocalControlRigScale(ULevelSequence*
 	if (UMovieScene* MovieScene = LevelSequence->GetMovieScene())
 	{
 		FFrameRate TickResolution = MovieScene->GetTickResolution();
-		if (TimeUnit == ESequenceTimeUnit::DisplayRate)
+		if (TimeUnit == EMovieSceneTimeUnit::DisplayRate)
 		{
 			Frame = FFrameRate::TransformTime(FFrameTime(Frame, 0), MovieScene->GetDisplayRate(), MovieScene->GetTickResolution()).RoundToFrame();
 		}
@@ -2205,7 +2205,7 @@ void UControlRigSequencerEditorLibrary::SetLocalControlRigScale(ULevelSequence*
 }
 
 void UControlRigSequencerEditorLibrary::SetLocalControlRigScales(ULevelSequence* LevelSequence, UControlRig* ControlRig, FName ControlName, const TArray<FFrameNumber>& Frames, 
-	const TArray<FVector> Values, ESequenceTimeUnit TimeUnit)
+	const TArray<FVector> Values, EMovieSceneTimeUnit TimeUnit)
 {
 	if (LevelSequence == nullptr || ControlRig == nullptr || (Frames.Num() != Values.Num()))
 	{
@@ -2227,7 +2227,7 @@ void UControlRigSequencerEditorLibrary::SetLocalControlRigScales(ULevelSequence*
 		for (int32 Index = 0; Index < Frames.Num(); ++Index)
 		{
 			FFrameNumber Frame = Frames[Index];
-			if (TimeUnit == ESequenceTimeUnit::DisplayRate)
+			if (TimeUnit == EMovieSceneTimeUnit::DisplayRate)
 			{
 				Frame = FFrameRate::TransformTime(FFrameTime(Frame, 0), MovieScene->GetDisplayRate(), MovieScene->GetTickResolution()).RoundToFrame();
 			}
@@ -2239,7 +2239,7 @@ void UControlRigSequencerEditorLibrary::SetLocalControlRigScales(ULevelSequence*
 }
 
 
-FEulerTransform UControlRigSequencerEditorLibrary::GetLocalControlRigEulerTransform(ULevelSequence* LevelSequence, UControlRig* ControlRig, FName ControlName, FFrameNumber Frame, ESequenceTimeUnit TimeUnit)
+FEulerTransform UControlRigSequencerEditorLibrary::GetLocalControlRigEulerTransform(ULevelSequence* LevelSequence, UControlRig* ControlRig, FName ControlName, FFrameNumber Frame, EMovieSceneTimeUnit TimeUnit)
 {
 	FEulerTransform Value = FEulerTransform::Identity;
 	TWeakPtr<ISequencer> WeakSequencer = GetSequencerFromAsset();
@@ -2263,7 +2263,7 @@ FEulerTransform UControlRigSequencerEditorLibrary::GetLocalControlRigEulerTransf
 	return Value;
 }
 
-TArray<FEulerTransform> UControlRigSequencerEditorLibrary::GetLocalControlRigEulerTransforms(ULevelSequence* LevelSequence, UControlRig* ControlRig, FName ControlName, const TArray<FFrameNumber>& Frames, ESequenceTimeUnit TimeUnit)
+TArray<FEulerTransform> UControlRigSequencerEditorLibrary::GetLocalControlRigEulerTransforms(ULevelSequence* LevelSequence, UControlRig* ControlRig, FName ControlName, const TArray<FFrameNumber>& Frames, EMovieSceneTimeUnit TimeUnit)
 {
 	TArray<FEulerTransform> Values;
 	TWeakPtr<ISequencer> WeakSequencer = GetSequencerFromAsset();
@@ -2293,7 +2293,7 @@ TArray<FEulerTransform> UControlRigSequencerEditorLibrary::GetLocalControlRigEul
 }
 
 void UControlRigSequencerEditorLibrary::SetLocalControlRigEulerTransform(ULevelSequence* LevelSequence, UControlRig* ControlRig, FName ControlName, FFrameNumber Frame, FEulerTransform Value,
-	ESequenceTimeUnit TimeUnit, bool bSetKey)
+	EMovieSceneTimeUnit TimeUnit, bool bSetKey)
 {
 	if (LevelSequence == nullptr || ControlRig == nullptr)
 	{
@@ -2310,7 +2310,7 @@ void UControlRigSequencerEditorLibrary::SetLocalControlRigEulerTransform(ULevelS
 	if (UMovieScene* MovieScene = LevelSequence->GetMovieScene())
 	{
 		FFrameRate TickResolution = MovieScene->GetTickResolution();
-		if (TimeUnit == ESequenceTimeUnit::DisplayRate)
+		if (TimeUnit == EMovieSceneTimeUnit::DisplayRate)
 		{
 			Frame = FFrameRate::TransformTime(FFrameTime(Frame, 0), MovieScene->GetDisplayRate(), MovieScene->GetTickResolution()).RoundToFrame();
 		}
@@ -2327,7 +2327,7 @@ void UControlRigSequencerEditorLibrary::SetLocalControlRigEulerTransform(ULevelS
 }
 
 void UControlRigSequencerEditorLibrary::SetLocalControlRigEulerTransforms(ULevelSequence* LevelSequence, UControlRig* ControlRig, FName ControlName, const TArray<FFrameNumber>& Frames,
-	const TArray<FEulerTransform> Values, ESequenceTimeUnit TimeUnit)
+	const TArray<FEulerTransform> Values, EMovieSceneTimeUnit TimeUnit)
 {
 	if (LevelSequence == nullptr || ControlRig == nullptr || (Frames.Num() != Values.Num()))
 	{
@@ -2349,7 +2349,7 @@ void UControlRigSequencerEditorLibrary::SetLocalControlRigEulerTransforms(ULevel
 		for (int32 Index = 0; Index < Frames.Num(); ++Index)
 		{
 			FFrameNumber Frame = Frames[Index];
-			if (TimeUnit == ESequenceTimeUnit::DisplayRate)
+			if (TimeUnit == EMovieSceneTimeUnit::DisplayRate)
 			{
 				Frame = FFrameRate::TransformTime(FFrameTime(Frame, 0), MovieScene->GetDisplayRate(), MovieScene->GetTickResolution()).RoundToFrame();
 			}
@@ -2367,7 +2367,7 @@ void UControlRigSequencerEditorLibrary::SetLocalControlRigEulerTransforms(ULevel
 }
 
 
-FTransformNoScale UControlRigSequencerEditorLibrary::GetLocalControlRigTransformNoScale(ULevelSequence* LevelSequence, UControlRig* ControlRig, FName ControlName, FFrameNumber Frame, ESequenceTimeUnit TimeUnit)
+FTransformNoScale UControlRigSequencerEditorLibrary::GetLocalControlRigTransformNoScale(ULevelSequence* LevelSequence, UControlRig* ControlRig, FName ControlName, FFrameNumber Frame, EMovieSceneTimeUnit TimeUnit)
 {
 	FTransformNoScale Value = FTransformNoScale::Identity;
 	TWeakPtr<ISequencer> WeakSequencer = GetSequencerFromAsset();
@@ -2391,7 +2391,7 @@ FTransformNoScale UControlRigSequencerEditorLibrary::GetLocalControlRigTransform
 	return Value;
 }
 
-TArray<FTransformNoScale> UControlRigSequencerEditorLibrary::GetLocalControlRigTransformNoScales(ULevelSequence* LevelSequence, UControlRig* ControlRig, FName ControlName, const TArray<FFrameNumber>& Frames, ESequenceTimeUnit TimeUnit)
+TArray<FTransformNoScale> UControlRigSequencerEditorLibrary::GetLocalControlRigTransformNoScales(ULevelSequence* LevelSequence, UControlRig* ControlRig, FName ControlName, const TArray<FFrameNumber>& Frames, EMovieSceneTimeUnit TimeUnit)
 {
 	TArray<FTransformNoScale> Values;
 	TWeakPtr<ISequencer> WeakSequencer = GetSequencerFromAsset();
@@ -2421,7 +2421,7 @@ TArray<FTransformNoScale> UControlRigSequencerEditorLibrary::GetLocalControlRigT
 }
 
 void UControlRigSequencerEditorLibrary::SetLocalControlRigTransformNoScale(ULevelSequence* LevelSequence, UControlRig* ControlRig, FName ControlName, FFrameNumber Frame, FTransformNoScale Value,
-	ESequenceTimeUnit TimeUnit,bool bSetKey)
+	EMovieSceneTimeUnit TimeUnit,bool bSetKey)
 {
 	if (LevelSequence == nullptr || ControlRig == nullptr)
 	{
@@ -2438,7 +2438,7 @@ void UControlRigSequencerEditorLibrary::SetLocalControlRigTransformNoScale(ULeve
 	if (UMovieScene* MovieScene = LevelSequence->GetMovieScene())
 	{
 		FFrameRate TickResolution = MovieScene->GetTickResolution();
-		if (TimeUnit == ESequenceTimeUnit::DisplayRate)
+		if (TimeUnit == EMovieSceneTimeUnit::DisplayRate)
 		{
 			Frame = FFrameRate::TransformTime(FFrameTime(Frame, 0), MovieScene->GetDisplayRate(), MovieScene->GetTickResolution()).RoundToFrame();
 		}
@@ -2450,7 +2450,7 @@ void UControlRigSequencerEditorLibrary::SetLocalControlRigTransformNoScale(ULeve
 }
 
 void UControlRigSequencerEditorLibrary::SetLocalControlRigTransformNoScales(ULevelSequence* LevelSequence, UControlRig* ControlRig, FName ControlName, const TArray<FFrameNumber>& Frames, 
-	const TArray<FTransformNoScale> Values, ESequenceTimeUnit TimeUnit)
+	const TArray<FTransformNoScale> Values, EMovieSceneTimeUnit TimeUnit)
 {
 	if (LevelSequence == nullptr || ControlRig == nullptr || (Frames.Num() != Values.Num()))
 	{
@@ -2472,7 +2472,7 @@ void UControlRigSequencerEditorLibrary::SetLocalControlRigTransformNoScales(ULev
 		for (int32 Index = 0; Index < Frames.Num(); ++Index)
 		{
 			FFrameNumber Frame = Frames[Index];
-			if (TimeUnit == ESequenceTimeUnit::DisplayRate)
+			if (TimeUnit == EMovieSceneTimeUnit::DisplayRate)
 			{
 				Frame = FFrameRate::TransformTime(FFrameTime(Frame, 0), MovieScene->GetDisplayRate(), MovieScene->GetTickResolution()).RoundToFrame();
 			}
@@ -2484,7 +2484,7 @@ void UControlRigSequencerEditorLibrary::SetLocalControlRigTransformNoScales(ULev
 }
 
 
-FTransform UControlRigSequencerEditorLibrary::GetLocalControlRigTransform(ULevelSequence* LevelSequence, UControlRig* ControlRig, FName ControlName, FFrameNumber Frame, ESequenceTimeUnit TimeUnit)
+FTransform UControlRigSequencerEditorLibrary::GetLocalControlRigTransform(ULevelSequence* LevelSequence, UControlRig* ControlRig, FName ControlName, FFrameNumber Frame, EMovieSceneTimeUnit TimeUnit)
 {
 	FTransform Value = FTransform::Identity;
 	TWeakPtr<ISequencer> WeakSequencer = GetSequencerFromAsset();
@@ -2508,7 +2508,7 @@ FTransform UControlRigSequencerEditorLibrary::GetLocalControlRigTransform(ULevel
 	return Value;
 }
 
-TArray<FTransform> UControlRigSequencerEditorLibrary::GetLocalControlRigTransforms(ULevelSequence* LevelSequence, UControlRig* ControlRig, FName ControlName, const TArray<FFrameNumber>& Frames, ESequenceTimeUnit TimeUnit)
+TArray<FTransform> UControlRigSequencerEditorLibrary::GetLocalControlRigTransforms(ULevelSequence* LevelSequence, UControlRig* ControlRig, FName ControlName, const TArray<FFrameNumber>& Frames, EMovieSceneTimeUnit TimeUnit)
 {
 	TArray<FTransform> Values;
 	TWeakPtr<ISequencer> WeakSequencer = GetSequencerFromAsset();
@@ -2538,7 +2538,7 @@ TArray<FTransform> UControlRigSequencerEditorLibrary::GetLocalControlRigTransfor
 }
 
 void UControlRigSequencerEditorLibrary::SetLocalControlRigTransform(ULevelSequence* LevelSequence, UControlRig* ControlRig, FName ControlName, FFrameNumber Frame, FTransform Value,
-	ESequenceTimeUnit TimeUnit,bool bSetKey)
+	EMovieSceneTimeUnit TimeUnit,bool bSetKey)
 {
 	if (LevelSequence == nullptr || ControlRig == nullptr)
 	{
@@ -2555,7 +2555,7 @@ void UControlRigSequencerEditorLibrary::SetLocalControlRigTransform(ULevelSequen
 	if (UMovieScene* MovieScene = LevelSequence->GetMovieScene())
 	{
 		FFrameRate TickResolution = MovieScene->GetTickResolution();
-		if (TimeUnit == ESequenceTimeUnit::DisplayRate)
+		if (TimeUnit == EMovieSceneTimeUnit::DisplayRate)
 		{
 			Frame = FFrameRate::TransformTime(FFrameTime(Frame, 0), MovieScene->GetDisplayRate(), MovieScene->GetTickResolution()).RoundToFrame();
 		}
@@ -2567,7 +2567,7 @@ void UControlRigSequencerEditorLibrary::SetLocalControlRigTransform(ULevelSequen
 }
 
 void UControlRigSequencerEditorLibrary::SetLocalControlRigTransforms(ULevelSequence* LevelSequence, UControlRig* ControlRig, FName ControlName, const TArray<FFrameNumber>& Frames, 
-	const TArray<FTransform> Values, ESequenceTimeUnit TimeUnit)
+	const TArray<FTransform> Values, EMovieSceneTimeUnit TimeUnit)
 {
 	if (LevelSequence == nullptr || ControlRig == nullptr || (Frames.Num() != Values.Num()))
 	{
@@ -2589,7 +2589,7 @@ void UControlRigSequencerEditorLibrary::SetLocalControlRigTransforms(ULevelSeque
 		for (int32 Index = 0; Index < Frames.Num(); ++Index)
 		{
 			FFrameNumber Frame = Frames[Index];
-			if (TimeUnit == ESequenceTimeUnit::DisplayRate)
+			if (TimeUnit == EMovieSceneTimeUnit::DisplayRate)
 			{
 				Frame = FFrameRate::TransformTime(FFrameTime(Frame, 0), MovieScene->GetDisplayRate(), MovieScene->GetTickResolution()).RoundToFrame();
 			}
@@ -2704,7 +2704,7 @@ bool UControlRigSequencerEditorLibrary::CollapseControlRigAnimLayers(ULevelSeque
 	return bValid;
 }
 
-bool UControlRigSequencerEditorLibrary::SetControlRigSpace(ULevelSequence* LevelSequence, UControlRig* ControlRig, FName ControlName, const FRigElementKey& InSpaceKey, FFrameNumber InTime, ESequenceTimeUnit TimeUnit)
+bool UControlRigSequencerEditorLibrary::SetControlRigSpace(ULevelSequence* LevelSequence, UControlRig* ControlRig, FName ControlName, const FRigElementKey& InSpaceKey, FFrameNumber InTime, EMovieSceneTimeUnit TimeUnit)
 {
 	TWeakPtr<ISequencer> WeakSequencer = GetSequencerFromAsset();
 	bool bValid = false;
@@ -2718,7 +2718,7 @@ bool UControlRigSequencerEditorLibrary::SetControlRigSpace(ULevelSequence* Level
 			FSpaceChannelAndSection SpaceChannelAndSection = FControlRigSpaceChannelHelpers::FindSpaceChannelAndSectionForControl(ControlRig, ControlName, Sequencer.Get(), true /*bCreateIfNeeded*/);
 			if (SpaceChannelAndSection.SpaceChannel)
 			{
-				if (TimeUnit == ESequenceTimeUnit::DisplayRate)
+				if (TimeUnit == EMovieSceneTimeUnit::DisplayRate)
 				{
 					InTime = FFrameRate::TransformTime(FFrameTime(InTime, 0), LevelSequence->GetMovieScene()->GetDisplayRate(), LevelSequence->GetMovieScene()->GetTickResolution()).RoundToFrame();
 				}
@@ -2740,7 +2740,7 @@ bool UControlRigSequencerEditorLibrary::SetControlRigSpace(ULevelSequence* Level
 	return bValid;
 }
 
-bool UControlRigSequencerEditorLibrary::BakeControlRigSpace(ULevelSequence* InSequence, UControlRig* InControlRig, const TArray<FName>& InControlNames, FRigSpacePickerBakeSettings InSettings, ESequenceTimeUnit TimeUnit)
+bool UControlRigSequencerEditorLibrary::BakeControlRigSpace(ULevelSequence* InSequence, UControlRig* InControlRig, const TArray<FName>& InControlNames, FRigSpacePickerBakeSettings InSettings, EMovieSceneTimeUnit TimeUnit)
 {
 	TWeakPtr<ISequencer> WeakSequencer = GetSequencerFromAsset();
 	bool bValid = false;
@@ -2751,7 +2751,7 @@ bool UControlRigSequencerEditorLibrary::BakeControlRigSpace(ULevelSequence* InSe
 		const FFrameRate TickResolution = Sequencer->GetFocusedTickResolution();
 		const FFrameRate& FrameRate = Sequencer->GetFocusedDisplayRate();
 		FFrameNumber FrameRateInFrameNumber = TickResolution.AsFrameNumber(FrameRate.AsInterval());
-		if (TimeUnit == ESequenceTimeUnit::DisplayRate)
+		if (TimeUnit == EMovieSceneTimeUnit::DisplayRate)
 		{
 			InSettings.Settings.StartFrame = FFrameRate::TransformTime(FFrameTime(InSettings.Settings.StartFrame, 0), FrameRate, TickResolution).RoundToFrame();
 			InSettings.Settings.EndFrame = FFrameRate::TransformTime(FFrameTime(InSettings.Settings.EndFrame, 0), FrameRate, TickResolution).RoundToFrame();
@@ -2775,7 +2775,7 @@ bool UControlRigSequencerEditorLibrary::BakeControlRigSpace(ULevelSequence* InSe
 	return bValid;
 }
 
-bool UControlRigSequencerEditorLibrary::DeleteControlRigSpace(ULevelSequence* LevelSequence, UControlRig* ControlRig, FName ControlName, FFrameNumber InTime, ESequenceTimeUnit TimeUnit)
+bool UControlRigSequencerEditorLibrary::DeleteControlRigSpace(ULevelSequence* LevelSequence, UControlRig* ControlRig, FName ControlName, FFrameNumber InTime, EMovieSceneTimeUnit TimeUnit)
 {
 	TWeakPtr<ISequencer> WeakSequencer = GetSequencerFromAsset();
 	bool bValid = false;
@@ -2789,7 +2789,7 @@ bool UControlRigSequencerEditorLibrary::DeleteControlRigSpace(ULevelSequence* Le
 			if (SpaceChannelAndSection.SpaceChannel)
 			{
 
-				if (TimeUnit == ESequenceTimeUnit::DisplayRate)
+				if (TimeUnit == EMovieSceneTimeUnit::DisplayRate)
 				{
 					InTime = FFrameRate::TransformTime(FFrameTime(InTime, 0), LevelSequence->GetMovieScene()->GetDisplayRate(), LevelSequence->GetMovieScene()->GetTickResolution()).RoundToFrame();
 				}
@@ -2832,7 +2832,7 @@ bool UControlRigSequencerEditorLibrary::DeleteControlRigSpace(ULevelSequence* Le
 	return bValid;
 }
 
-bool UControlRigSequencerEditorLibrary::MoveControlRigSpace(ULevelSequence* LevelSequence, UControlRig* ControlRig, FName ControlName, FFrameNumber InTime, FFrameNumber InNewTime, ESequenceTimeUnit TimeUnit)
+bool UControlRigSequencerEditorLibrary::MoveControlRigSpace(ULevelSequence* LevelSequence, UControlRig* ControlRig, FName ControlName, FFrameNumber InTime, FFrameNumber InNewTime, EMovieSceneTimeUnit TimeUnit)
 {
 	TWeakPtr<ISequencer> WeakSequencer = GetSequencerFromAsset();
 	bool bValid = false;
@@ -2846,7 +2846,7 @@ bool UControlRigSequencerEditorLibrary::MoveControlRigSpace(ULevelSequence* Leve
 			if (SpaceChannelAndSection.SpaceChannel)
 			{
 
-				if (TimeUnit == ESequenceTimeUnit::DisplayRate)
+				if (TimeUnit == EMovieSceneTimeUnit::DisplayRate)
 				{
 					InTime = FFrameRate::TransformTime(FFrameTime(InTime, 0), LevelSequence->GetMovieScene()->GetDisplayRate(), LevelSequence->GetMovieScene()->GetTickResolution()).RoundToFrame();
 					InNewTime = FFrameRate::TransformTime(FFrameTime(InNewTime, 0), LevelSequence->GetMovieScene()->GetDisplayRate(), LevelSequence->GetMovieScene()->GetTickResolution()).RoundToFrame();
diff --git a/Engine/Plugins/Animation/ControlRig/Source/ControlRigEditor/Public/ControlRigSequencerEditorLibrary.h b/Engine/Plugins/Animation/ControlRig/Source/ControlRigEditor/Public/ControlRigSequencerEditorLibrary.h
--- a/Engine/Plugins/Animation/ControlRig/Source/ControlRigEditor/Public/ControlRigSequencerEditorLibrary.h
+++ b/Engine/Plugins/Animation/ControlRig/Source/ControlRigEditor/Public/ControlRigSequencerEditorLibrary.h
@@ -15,7 +15,7 @@
 #include "EulerTransform.h"
 #include "MovieSceneToolsUserSettings.h"
 #include "Tools/ControlRigSnapSettings.h"
-#include "SequenceTimeUnit.h"
+#include "MovieSceneTimeUnit.h"
 #include "RigSpacePickerBakeSettings.h"
 #include "ControlRigSequencerEditorLibrary.generated.h"
 
@@ -126,7 +126,7 @@ public:
 	*/
 	UFUNCTION(BlueprintCallable, Category = "Editor Scripting | Sequencer Tools | Control Rig")
 	static bool LoadAnimSequenceIntoControlRigSection(UMovieSceneSection* MovieSceneSection, UAnimSequence* AnimSequence, USkeletalMeshComponent* SkelMeshComp,
-		FFrameNumber InStartFrame, ESequenceTimeUnit TimeUnit = ESequenceTimeUnit::DisplayRate, bool bKeyReduce = false, float Tolerance = 0.001f,
+		FFrameNumber InStartFrame, EMovieSceneTimeUnit TimeUnit = EMovieSceneTimeUnit::DisplayRate, bool bKeyReduce = false, float Tolerance = 0.001f,
 		EMovieSceneKeyInterpolation Interpolation = EMovieSceneKeyInterpolation::SmartAuto);
 
 	/**
@@ -153,7 +153,7 @@ public:
 	* @return Returns True if successful, False otherwise
 	*/
 	UFUNCTION(BlueprintCallable, Category = "Editor Scripting | Sequencer Tools | Control Rig")
-	static bool BakeConstraint(UWorld* World, UTickableConstraint* Constraint, const TArray<FFrameNumber>& Frames, ESequenceTimeUnit TimeUnit = ESequenceTimeUnit::DisplayRate);
+	static bool BakeConstraint(UWorld* World, UTickableConstraint* Constraint, const TArray<FFrameNumber>& Frames, EMovieSceneTimeUnit TimeUnit = EMovieSceneTimeUnit::DisplayRate);
 
 	/**
 	* Bake the constraint to keys based on the passed in settings. This will use the open sequencer to bake. See ConstraintsScriptingLibrary to get the list of available constraints
@@ -186,7 +186,7 @@ public:
 	* @return Returns true if we got the keys from this constraint
 	*/
 	UFUNCTION(BlueprintCallable, Category = "Editor Scripting | Sequencer Tools | Control Rig")
-	static bool GetConstraintKeys(UTickableConstraint* InConstraint, UMovieSceneSection* ConstraintSection,TArray<bool>& OutBools, TArray<FFrameNumber>& OutFrames, ESequenceTimeUnit TimeUnit = ESequenceTimeUnit::DisplayRate);
+	static bool GetConstraintKeys(UTickableConstraint* InConstraint, UMovieSceneSection* ConstraintSection,TArray<bool>& OutBools, TArray<FFrameNumber>& OutFrames, EMovieSceneTimeUnit TimeUnit = EMovieSceneTimeUnit::DisplayRate);
 	
 	/**
 	* Set the constraint active key in the current open Sequencer
@@ -197,7 +197,7 @@ public:
 	* @return Returns true if we set the constraint to be the passed in value, false if not. We may not do so if the value is the same.
 	*/
 	UFUNCTION(BlueprintCallable, Category = "Editor Scripting | Sequencer Tools | Control Rig")
-	static bool SetConstraintActiveKey(UTickableConstraint* InConstraint, bool bActive, FFrameNumber InFrame, ESequenceTimeUnit TimeUnit = ESequenceTimeUnit::DisplayRate);
+	static bool SetConstraintActiveKey(UTickableConstraint* InConstraint, bool bActive, FFrameNumber InFrame, EMovieSceneTimeUnit TimeUnit = EMovieSceneTimeUnit::DisplayRate);
 
 	/**
 	* Get all constraints for this object, which is described by a transformable handle
@@ -216,7 +216,7 @@ public:
 	* @return Will return false if function fails, for example if there is no key at this time it will fail, or if the new time is invalid it could fail also
 	*/
 	UFUNCTION(BlueprintCallable, Category = "Editor Scripting | Sequencer Tools | Control Rig")
-	static bool MoveConstraintKey(UTickableConstraint* Constraint, UMovieSceneSection* ConstraintSection, FFrameNumber InTime, FFrameNumber InNewTime, ESequenceTimeUnit TimeUnit = ESequenceTimeUnit::DisplayRate);
+	static bool MoveConstraintKey(UTickableConstraint* Constraint, UMovieSceneSection* ConstraintSection, FFrameNumber InTime, FFrameNumber InNewTime, EMovieSceneTimeUnit TimeUnit = EMovieSceneTimeUnit::DisplayRate);
 
 	/** Delete the Key for the Constraint at the specified time. 
 	* @param InConstraint The constraint whose key to move
@@ -226,7 +226,7 @@ public:
 	* @return Will return false if function fails,  for example if there is no key at this time it will fail.
 	*/
 	UFUNCTION(BlueprintCallable, Category = "Editor Scripting | Sequencer Tools | Control Rig")
-	static bool DeleteConstraintKey(UTickableConstraint* Constraint, UMovieSceneSection* ConstraintSection, FFrameNumber InTime, ESequenceTimeUnit TimeUnit = ESequenceTimeUnit::DisplayRate);
+	static bool DeleteConstraintKey(UTickableConstraint* Constraint, UMovieSceneSection* ConstraintSection, FFrameNumber InTime, EMovieSceneTimeUnit TimeUnit = EMovieSceneTimeUnit::DisplayRate);
 	/**
 	* Compensate constraint at the specfied time 
 	* @param InConstraint The constraint to compensate
@@ -235,7 +235,7 @@ public:
 	* @return Returns true if it can compensate
 	*/
 	UFUNCTION(BlueprintCallable, Category = "Editor Scripting | Sequencer Tools | Control Rig")
-	static bool Compensate(UTickableConstraint* InConstraint,  FFrameNumber InTime, ESequenceTimeUnit TimeUnit = ESequenceTimeUnit::DisplayRate);
+	static bool Compensate(UTickableConstraint* InConstraint,  FFrameNumber InTime, EMovieSceneTimeUnit TimeUnit = EMovieSceneTimeUnit::DisplayRate);
 
 	/**
 	* Compensate constraint at all keys
@@ -277,7 +277,7 @@ public:
 	*/
 	UFUNCTION(BlueprintCallable, Category = "Editor Scripting | Sequencer Tools | Control Rig")
 	static bool SnapControlRig(ULevelSequence* LevelSequence, FFrameNumber StartFrame, FFrameNumber EndFrame, const FControlRigSnapperSelection& ChildrenToSnap,
-		const FControlRigSnapperSelection& ParentToSnap, const UControlRigSnapSettings* SnapSettings, ESequenceTimeUnit TimeUnit = ESequenceTimeUnit::DisplayRate);
+		const FControlRigSnapperSelection& ParentToSnap, const UControlRigSnapSettings* SnapSettings, EMovieSceneTimeUnit TimeUnit = EMovieSceneTimeUnit::DisplayRate);
 	
 	/**
 	* Get Actors World Transform at a specific time
@@ -288,7 +288,7 @@ public:
 	* @return Returns World Transform
 	*/
 	UFUNCTION(BlueprintCallable, Category = "Editor Scripting | Sequencer Tools | Control Rig")
-	static FTransform GetActorWorldTransform(ULevelSequence* LevelSequence,AActor* Actor, FFrameNumber Frame, ESequenceTimeUnit TimeUnit = ESequenceTimeUnit::DisplayRate);
+	static FTransform GetActorWorldTransform(ULevelSequence* LevelSequence,AActor* Actor, FFrameNumber Frame, EMovieSceneTimeUnit TimeUnit = EMovieSceneTimeUnit::DisplayRate);
 
 	/**
 	* Get Actors World Transforms at specific times
@@ -299,7 +299,7 @@ public:
 	* @return Returns World Transforms
 	*/
 	UFUNCTION(BlueprintCallable, Category = "Editor Scripting | Sequencer Tools | Control Rig")
-	static TArray<FTransform> GetActorWorldTransforms(ULevelSequence* LevelSequence, AActor* Actor, const TArray<FFrameNumber>& Frames, ESequenceTimeUnit TimeUnit = ESequenceTimeUnit::DisplayRate);
+	static TArray<FTransform> GetActorWorldTransforms(ULevelSequence* LevelSequence, AActor* Actor, const TArray<FFrameNumber>& Frames, EMovieSceneTimeUnit TimeUnit = EMovieSceneTimeUnit::DisplayRate);
 
 	/**
 	* Get SkeletalMeshComponent World Transform at a specific time
@@ -312,7 +312,7 @@ public:
 	*/
 	UFUNCTION(BlueprintCallable, Category = "Editor Scripting | Sequencer Tools | Control Rig")
 	static FTransform GetSkeletalMeshComponentWorldTransform(ULevelSequence* LevelSequence, USkeletalMeshComponent* SkeletalMeshComponent, FFrameNumber Frame,
-		ESequenceTimeUnit TimeUnit = ESequenceTimeUnit::DisplayRate,FName ReferenceName = NAME_None);
+		EMovieSceneTimeUnit TimeUnit = EMovieSceneTimeUnit::DisplayRate,FName ReferenceName = NAME_None);
 
 	/**
 	* Get SkeletalMeshComponents World Transforms at specific times
@@ -325,7 +325,7 @@ public:
 	*/
 	UFUNCTION(BlueprintCallable, Category = "Editor Scripting | Sequencer Tools | Control Rig")
 	static TArray<FTransform> GetSkeletalMeshComponentWorldTransforms(ULevelSequence* LevelSequence, USkeletalMeshComponent* SkeletalMeshComponent, const TArray<FFrameNumber>& Frames, 
-		ESequenceTimeUnit TimeUnit = ESequenceTimeUnit::DisplayRate,FName ReferenceName = NAME_None);
+		EMovieSceneTimeUnit TimeUnit = EMovieSceneTimeUnit::DisplayRate,FName ReferenceName = NAME_None);
 
 	/**
 	* Get ControlRig Control's World Transform at a specific time
@@ -338,7 +338,7 @@ public:
 	*/
 	UFUNCTION(BlueprintCallable, Category = "Editor Scripting | Sequencer Tools | Control Rig")
 	static FTransform GetControlRigWorldTransform(ULevelSequence* LevelSequence, UControlRig* ControlRig, FName ControlName, FFrameNumber Frame,
-		ESequenceTimeUnit TimeUnit = ESequenceTimeUnit::DisplayRate);
+		EMovieSceneTimeUnit TimeUnit = EMovieSceneTimeUnit::DisplayRate);
 
 	/**
 	* Get ControlRig Control's World Transforms at specific times
@@ -351,7 +351,7 @@ public:
 	*/
 	UFUNCTION(BlueprintCallable, Category = "Editor Scripting | Sequencer Tools | Control Rig")
 	static TArray<FTransform> GetControlRigWorldTransforms(ULevelSequence* LevelSequence, UControlRig* ControlRig, FName ControlName, const TArray<FFrameNumber>& Frames,
-		ESequenceTimeUnit TimeUnit = ESequenceTimeUnit::DisplayRate);
+		EMovieSceneTimeUnit TimeUnit = EMovieSceneTimeUnit::DisplayRate);
 
 	/**
 	* Set ControlRig Control's World Transform at a specific time
@@ -365,7 +365,7 @@ public:
 	*/
 	UFUNCTION(BlueprintCallable, Category = "Editor Scripting | Sequencer Tools | Control Rig")
 	static void SetControlRigWorldTransform(ULevelSequence* LevelSequence, UControlRig* ControlRig, FName ControlName, FFrameNumber Frame, const FTransform& WorldTransform,
-		ESequenceTimeUnit TimeUnit = ESequenceTimeUnit::DisplayRate, bool bSetKey = true);
+		EMovieSceneTimeUnit TimeUnit = EMovieSceneTimeUnit::DisplayRate, bool bSetKey = true);
 
 	/**
 	* Set ControlRig Control's World Transforms at a specific times.
@@ -378,7 +378,7 @@ public:
 	*/
 	UFUNCTION(BlueprintCallable, Category = "Editor Scripting | Sequencer Tools | Control Rig")
 	static void SetControlRigWorldTransforms(ULevelSequence* LevelSequence, UControlRig* ControlRig, FName ControlName, const TArray<FFrameNumber>& Frames, 
-		 const TArray<FTransform>& WorldTransforms, ESequenceTimeUnit TimeUnit = ESequenceTimeUnit::DisplayRate);
+		 const TArray<FTransform>& WorldTransforms, EMovieSceneTimeUnit TimeUnit = EMovieSceneTimeUnit::DisplayRate);
 	
 	/**
 	* Get ControlRig Control's float value at a specific time
@@ -391,7 +391,7 @@ public:
 	*/
 	UFUNCTION(BlueprintCallable, Category = "Editor Scripting | Sequencer Tools | Control Rig")
 	static float GetLocalControlRigFloat(ULevelSequence* LevelSequence, UControlRig* ControlRig, FName ControlName, FFrameNumber Frame,
-		ESequenceTimeUnit TimeUnit = ESequenceTimeUnit::DisplayRate);
+		EMovieSceneTimeUnit TimeUnit = EMovieSceneTimeUnit::DisplayRate);
 
 	/**
 	* Get ControlRig Control's float values at specific times
@@ -404,7 +404,7 @@ public:
 	*/
 	UFUNCTION(BlueprintCallable, Category = "Editor Scripting | Sequencer Tools | Control Rig")
 	static TArray<float> GetLocalControlRigFloats(ULevelSequence* LevelSequence, UControlRig* ControlRig, FName ControlName, const TArray<FFrameNumber>& Frames,
-		ESequenceTimeUnit TimeUnit = ESequenceTimeUnit::DisplayRate);
+		EMovieSceneTimeUnit TimeUnit = EMovieSceneTimeUnit::DisplayRate);
 
 	/**
 	* Set ControlRig Control's float value at specific time
@@ -418,7 +418,7 @@ public:
 	*/
 	UFUNCTION(BlueprintCallable, Category = "Editor Scripting | Sequencer Tools | Control Rig")
 	static void SetLocalControlRigFloat(ULevelSequence* LevelSequence, UControlRig* ControlRig, FName ControlName, FFrameNumber Frame, float Value,
-		ESequenceTimeUnit TimeUnit = ESequenceTimeUnit::DisplayRate,bool bSetKey = true);
+		EMovieSceneTimeUnit TimeUnit = EMovieSceneTimeUnit::DisplayRate,bool bSetKey = true);
 
 	/**
 	* Set ControlRig Control's float values at specific times
@@ -431,7 +431,7 @@ public:
 	*/
 	UFUNCTION(BlueprintCallable, Category = "Editor Scripting | Sequencer Tools | Control Rig")
 	static void SetLocalControlRigFloats(ULevelSequence* LevelSequence, UControlRig* ControlRig, FName ControlName, const TArray<FFrameNumber>& Frames, const TArray<float> Values,
-		ESequenceTimeUnit TimeUnit = ESequenceTimeUnit::DisplayRate);
+		EMovieSceneTimeUnit TimeUnit = EMovieSceneTimeUnit::DisplayRate);
 
 	/**
 	* Get ControlRig Control's bool value at a specific time
@@ -444,7 +444,7 @@ public:
 	*/
 	UFUNCTION(BlueprintCallable, Category = "Editor Scripting | Sequencer Tools | Control Rig")
 	static bool GetLocalControlRigBool(ULevelSequence* LevelSequence, UControlRig* ControlRig, FName ControlName, FFrameNumber Frame,
-		ESequenceTimeUnit TimeUnit = ESequenceTimeUnit::DisplayRate);
+		EMovieSceneTimeUnit TimeUnit = EMovieSceneTimeUnit::DisplayRate);
 
 	/**
 	* Get ControlRig Control's bool values at specific times
@@ -457,7 +457,7 @@ public:
 	*/
 	UFUNCTION(BlueprintCallable, Category = "Editor Scripting | Sequencer Tools | Control Rig")
 	static TArray<bool> GetLocalControlRigBools(ULevelSequence* LevelSequence, UControlRig* ControlRig, FName ControlName, const TArray<FFrameNumber>& Frames,
-		ESequenceTimeUnit TimeUnit = ESequenceTimeUnit::DisplayRate);
+		EMovieSceneTimeUnit TimeUnit = EMovieSceneTimeUnit::DisplayRate);
 
 	/**
 	* Set ControlRig Control's bool value at specific time
@@ -471,7 +471,7 @@ public:
 	*/
 	UFUNCTION(BlueprintCallable, Category = "Editor Scripting | Sequencer Tools | Control Rig")
 	static void SetLocalControlRigBool(ULevelSequence* LevelSequence, UControlRig* ControlRig, FName ControlName, FFrameNumber Frame, bool Value,
-		ESequenceTimeUnit TimeUnit = ESequenceTimeUnit::DisplayRate, bool bSetKey = true);
+		EMovieSceneTimeUnit TimeUnit = EMovieSceneTimeUnit::DisplayRate, bool bSetKey = true);
 
 	/**
 	* Set ControlRig Control's bool values at specific times
@@ -484,7 +484,7 @@ public:
 */
 	UFUNCTION(BlueprintCallable, Category = "Editor Scripting | Sequencer Tools | Control Rig")
 	static void SetLocalControlRigBools(ULevelSequence* LevelSequence, UControlRig* ControlRig, FName ControlName, const TArray<FFrameNumber>& Frames,
-		const TArray<bool> Values, ESequenceTimeUnit TimeUnit = ESequenceTimeUnit::DisplayRate);
+		const TArray<bool> Values, EMovieSceneTimeUnit TimeUnit = EMovieSceneTimeUnit::DisplayRate);
 
 	/**
 	* Get ControlRig Control's integer value at a specific time
@@ -497,7 +497,7 @@ public:
 	*/
 	UFUNCTION(BlueprintCallable, Category = "Editor Scripting | Sequencer Tools | Control Rig")
 	static int32 GetLocalControlRigInt(ULevelSequence* LevelSequence, UControlRig* ControlRig, FName ControlName, FFrameNumber Frame, 
-		ESequenceTimeUnit TimeUnit = ESequenceTimeUnit::DisplayRate);
+		EMovieSceneTimeUnit TimeUnit = EMovieSceneTimeUnit::DisplayRate);
 
 	/**
 	* Get ControlRig Control's integer values at specific times
@@ -510,7 +510,7 @@ public:
 	*/
 	UFUNCTION(BlueprintCallable, Category = "Editor Scripting | Sequencer Tools | Control Rig")
 	static TArray<int32> GetLocalControlRigInts(ULevelSequence* LevelSequence, UControlRig* ControlRig, FName ControlName, const TArray<FFrameNumber>& Frames,
-		ESequenceTimeUnit TimeUnit = ESequenceTimeUnit::DisplayRate);
+		EMovieSceneTimeUnit TimeUnit = EMovieSceneTimeUnit::DisplayRate);
 
 	/**
 	* Set ControlRig Control's int value at specific time
@@ -524,7 +524,7 @@ public:
 	*/
 	UFUNCTION(BlueprintCallable, Category = "Editor Scripting | Sequencer Tools | Control Rig")
 	static void SetLocalControlRigInt(ULevelSequence* LevelSequence, UControlRig* ControlRig, FName ControlName, FFrameNumber Frame, int32 Value,
-		ESequenceTimeUnit TimeUnit = ESequenceTimeUnit::DisplayRate,bool bSetKey = true);
+		EMovieSceneTimeUnit TimeUnit = EMovieSceneTimeUnit::DisplayRate,bool bSetKey = true);
 
 
 	/**
@@ -538,7 +538,7 @@ public:
 	*/
 	UFUNCTION(BlueprintCallable, Category = "Editor Scripting | Sequencer Tools | Control Rig")
 	static void SetLocalControlRigInts(ULevelSequence* LevelSequence, UControlRig* ControlRig, FName ControlName, const TArray<FFrameNumber>& Frames, const TArray<int32> Values,
-		ESequenceTimeUnit TimeUnit = ESequenceTimeUnit::DisplayRate);
+		EMovieSceneTimeUnit TimeUnit = EMovieSceneTimeUnit::DisplayRate);
 
 	/**
 	* Get ControlRig Control's Vector2D value at a specific time
@@ -551,7 +551,7 @@ public:
 	*/
 	UFUNCTION(BlueprintCallable, Category = "Editor Scripting | Sequencer Tools | Control Rig")
 	static FVector2D GetLocalControlRigVector2D(ULevelSequence* LevelSequence, UControlRig* ControlRig, FName ControlName, FFrameNumber Frame,
-		ESequenceTimeUnit TimeUnit = ESequenceTimeUnit::DisplayRate);
+		EMovieSceneTimeUnit TimeUnit = EMovieSceneTimeUnit::DisplayRate);
 
 	/**
 	* Get ControlRig Control's Vector2D values at specific times
@@ -564,7 +564,7 @@ public:
 	*/
 	UFUNCTION(BlueprintCallable, Category = "Editor Scripting | Sequencer Tools | Control Rig")
 	static TArray<FVector2D> GetLocalControlRigVector2Ds(ULevelSequence* LevelSequence, UControlRig* ControlRig, FName ControlName, const TArray<FFrameNumber>& Frames,
-		ESequenceTimeUnit TimeUnit = ESequenceTimeUnit::DisplayRate);
+		EMovieSceneTimeUnit TimeUnit = EMovieSceneTimeUnit::DisplayRate);
 
 	/**
 	* Set ControlRig Control's Vector2D value at specific time
@@ -578,7 +578,7 @@ public:
 	*/
 	UFUNCTION(BlueprintCallable, Category = "Editor Scripting | Sequencer Tools | Control Rig")
 	static void SetLocalControlRigVector2D(ULevelSequence* LevelSequence, UControlRig* ControlRig, FName ControlName, FFrameNumber Frame, FVector2D Value,
-		ESequenceTimeUnit TimeUnit = ESequenceTimeUnit::DisplayRate,bool bSetKey = true);
+		EMovieSceneTimeUnit TimeUnit = EMovieSceneTimeUnit::DisplayRate,bool bSetKey = true);
 
 
 	/**
@@ -592,7 +592,7 @@ public:
 	*/
 	UFUNCTION(BlueprintCallable, Category = "Editor Scripting | Sequencer Tools | Control Rig")
 	static void SetLocalControlRigVector2Ds(ULevelSequence* LevelSequence, UControlRig* ControlRig, FName ControlName, const TArray<FFrameNumber>& Frames, const TArray<FVector2D> Values,
-		ESequenceTimeUnit TimeUnit = ESequenceTimeUnit::DisplayRate);
+		EMovieSceneTimeUnit TimeUnit = EMovieSceneTimeUnit::DisplayRate);
 
 	/**
 	* Get ControlRig Control's Position value at a specific time
@@ -605,7 +605,7 @@ public:
 	*/
 	UFUNCTION(BlueprintCallable, Category = "Editor Scripting | Sequencer Tools | Control Rig")
 	static FVector GetLocalControlRigPosition(ULevelSequence* LevelSequence, UControlRig* ControlRig, FName ControlName, FFrameNumber Frame,
-		ESequenceTimeUnit TimeUnit = ESequenceTimeUnit::DisplayRate);
+		EMovieSceneTimeUnit TimeUnit = EMovieSceneTimeUnit::DisplayRate);
 
 	/**
 	* Get ControlRig Control's Position values at specific times
@@ -618,7 +618,7 @@ public:
 	*/
 	UFUNCTION(BlueprintCallable, Category = "Editor Scripting | Sequencer Tools | Control Rig")
 	static TArray<FVector> GetLocalControlRigPositions(ULevelSequence* LevelSequence, UControlRig* ControlRig, FName ControlName, const TArray<FFrameNumber>& Frames,
-		ESequenceTimeUnit TimeUnit = ESequenceTimeUnit::DisplayRate);
+		EMovieSceneTimeUnit TimeUnit = EMovieSceneTimeUnit::DisplayRate);
 
 	/**
 	* Set ControlRig Control's Position value at specific time
@@ -632,7 +632,7 @@ public:
 	*/
 	UFUNCTION(BlueprintCallable, Category = "Editor Scripting | Sequencer Tools | Control Rig")
 	static void SetLocalControlRigPosition(ULevelSequence* LevelSequence, UControlRig* ControlRig, FName ControlName, FFrameNumber Frame, FVector Value,
-		ESequenceTimeUnit TimeUnit = ESequenceTimeUnit::DisplayRate, bool bSetKey = true);
+		EMovieSceneTimeUnit TimeUnit = EMovieSceneTimeUnit::DisplayRate, bool bSetKey = true);
 
 	/**
 	* Set ControlRig Control's Position values at specific times
@@ -645,7 +645,7 @@ public:
 	*/
 	UFUNCTION(BlueprintCallable, Category = "Editor Scripting | Sequencer Tools | Control Rig")
 	static void SetLocalControlRigPositions(ULevelSequence* LevelSequence, UControlRig* ControlRig, FName ControlName, const TArray<FFrameNumber>& Frames, const TArray<FVector> Values,
-		ESequenceTimeUnit TimeUnit = ESequenceTimeUnit::DisplayRate);
+		EMovieSceneTimeUnit TimeUnit = EMovieSceneTimeUnit::DisplayRate);
 
 	/**
 	* Get ControlRig Control's Rotator value at a specific time
@@ -658,7 +658,7 @@ public:
 	*/
 	UFUNCTION(BlueprintCallable, Category = "Editor Scripting | Sequencer Tools | Control Rig")
 	static FRotator GetLocalControlRigRotator(ULevelSequence* LevelSequence, UControlRig* ControlRig, FName ControlName, FFrameNumber Frame,
-		ESequenceTimeUnit TimeUnit = ESequenceTimeUnit::DisplayRate);
+		EMovieSceneTimeUnit TimeUnit = EMovieSceneTimeUnit::DisplayRate);
 
 	/**
 	* Get ControlRig Control's Rotator values at specific times
@@ -671,7 +671,7 @@ public:
 	*/
 	UFUNCTION(BlueprintCallable, Category = "Editor Scripting | Sequencer Tools | Control Rig")
 	static TArray<FRotator> GetLocalControlRigRotators(ULevelSequence* LevelSequence, UControlRig* ControlRig, FName ControlName, const TArray<FFrameNumber>& Frames,
-		ESequenceTimeUnit TimeUnit = ESequenceTimeUnit::DisplayRate);
+		EMovieSceneTimeUnit TimeUnit = EMovieSceneTimeUnit::DisplayRate);
 
 	/**
 	* Set ControlRig Control's Rotator value at specific time
@@ -685,7 +685,7 @@ public:
 	*/
 	UFUNCTION(BlueprintCallable, Category = "Editor Scripting | Sequencer Tools | Control Rig")
 	static void SetLocalControlRigRotator(ULevelSequence* LevelSequence, UControlRig* ControlRig, FName ControlName, FFrameNumber Frame, FRotator Value, 
-		ESequenceTimeUnit TimeUnit = ESequenceTimeUnit::DisplayRate, bool bSetKey = true);
+		EMovieSceneTimeUnit TimeUnit = EMovieSceneTimeUnit::DisplayRate, bool bSetKey = true);
 
 	/**
 	* Set ControlRig Control's Rotator values at specific times
@@ -698,7 +698,7 @@ public:
 	*/
 	UFUNCTION(BlueprintCallable, Category = "Editor Scripting | Sequencer Tools | Control Rig")
 	static void SetLocalControlRigRotators(ULevelSequence* LevelSequence, UControlRig* ControlRig, FName ControlName, const TArray<FFrameNumber>& Frames, const TArray<FRotator> Values,
-		ESequenceTimeUnit TimeUnit = ESequenceTimeUnit::DisplayRate);
+		EMovieSceneTimeUnit TimeUnit = EMovieSceneTimeUnit::DisplayRate);
 
 	/**
 	* Get ControlRig Control's Scale value at a specific time
@@ -711,7 +711,7 @@ public:
 	*/
 	UFUNCTION(BlueprintCallable, Category = "Editor Scripting | Sequencer Tools | Control Rig")
 	static FVector GetLocalControlRigScale(ULevelSequence* LevelSequence, UControlRig* ControlRig, FName ControlName, FFrameNumber Frame,
-		ESequenceTimeUnit TimeUnit = ESequenceTimeUnit::DisplayRate);
+		EMovieSceneTimeUnit TimeUnit = EMovieSceneTimeUnit::DisplayRate);
 
 	/**
 	* Get ControlRig Control's Scale values at specific times
@@ -724,7 +724,7 @@ public:
 	*/
 	UFUNCTION(BlueprintCallable, Category = "Editor Scripting | Sequencer Tools | Control Rig")
 	static TArray<FVector> GetLocalControlRigScales(ULevelSequence* LevelSequence, UControlRig* ControlRig, FName ControlName, const TArray<FFrameNumber>& Frames,
-		ESequenceTimeUnit TimeUnit = ESequenceTimeUnit::DisplayRate);
+		EMovieSceneTimeUnit TimeUnit = EMovieSceneTimeUnit::DisplayRate);
 
 	/**
 	* Set ControlRig Control's Scale value at specific time
@@ -738,7 +738,7 @@ public:
 */
 	UFUNCTION(BlueprintCallable, Category = "Editor Scripting | Sequencer Tools | Control Rig")
 	static void SetLocalControlRigScale(ULevelSequence* LevelSequence, UControlRig* ControlRig, FName ControlName, FFrameNumber Frame, FVector Value, 
-		ESequenceTimeUnit TimeUnit = ESequenceTimeUnit::DisplayRate, bool bSetKey = true);
+		EMovieSceneTimeUnit TimeUnit = EMovieSceneTimeUnit::DisplayRate, bool bSetKey = true);
 
 	/**
 	* Set ControlRig Control's Scale values at specific times
@@ -751,7 +751,7 @@ public:
 	*/
 	UFUNCTION(BlueprintCallable, Category = "Editor Scripting | Sequencer Tools | Control Rig")
 	static void SetLocalControlRigScales(ULevelSequence* LevelSequence, UControlRig* ControlRig, FName ControlName, const TArray<FFrameNumber>& Frames, const TArray<FVector> Values,
-		ESequenceTimeUnit TimeUnit = ESequenceTimeUnit::DisplayRate);
+		EMovieSceneTimeUnit TimeUnit = EMovieSceneTimeUnit::DisplayRate);
 
 
 	/**
@@ -765,7 +765,7 @@ public:
 	*/
 	UFUNCTION(BlueprintCallable, Category = "Editor Scripting | Sequencer Tools | Control Rig")
 	static FEulerTransform GetLocalControlRigEulerTransform(ULevelSequence* LevelSequence, UControlRig* ControlRig, FName ControlName, FFrameNumber Frame,
-		ESequenceTimeUnit TimeUnit = ESequenceTimeUnit::DisplayRate);
+		EMovieSceneTimeUnit TimeUnit = EMovieSceneTimeUnit::DisplayRate);
 
 	/**
 	* Get ControlRig Control's EulerTransform values at specific times
@@ -778,7 +778,7 @@ public:
 	*/
 	UFUNCTION(BlueprintCallable, Category = "Editor Scripting | Sequencer Tools | Control Rig")
 	static TArray<FEulerTransform> GetLocalControlRigEulerTransforms(ULevelSequence* LevelSequence, UControlRig* ControlRig, FName ControlName, const TArray<FFrameNumber>& Frames,
-		ESequenceTimeUnit TimeUnit = ESequenceTimeUnit::DisplayRate);
+		EMovieSceneTimeUnit TimeUnit = EMovieSceneTimeUnit::DisplayRate);
 
 	/**
 	* Set ControlRig Control's EulerTransform value at specific time
@@ -792,7 +792,7 @@ public:
 	*/
 	UFUNCTION(BlueprintCallable, Category = "Editor Scripting | Sequencer Tools | Control Rig")
 	static void SetLocalControlRigEulerTransform(ULevelSequence* LevelSequence, UControlRig* ControlRig, FName ControlName, FFrameNumber Frame, FEulerTransform Value,
-		ESequenceTimeUnit TimeUnit = ESequenceTimeUnit::DisplayRate, bool bSetKey = true);
+		EMovieSceneTimeUnit TimeUnit = EMovieSceneTimeUnit::DisplayRate, bool bSetKey = true);
 
 	/**
 	* Set ControlRig Control's EulerTransform values at specific times
@@ -805,7 +805,7 @@ public:
 	*/
 	UFUNCTION(BlueprintCallable, Category = "Editor Scripting | Sequencer Tools | Control Rig")
 	static void SetLocalControlRigEulerTransforms(ULevelSequence* LevelSequence, UControlRig* ControlRig, FName ControlName, const TArray<FFrameNumber>& Frames, const TArray<FEulerTransform> Values,
-		ESequenceTimeUnit TimeUnit = ESequenceTimeUnit::DisplayRate);
+		EMovieSceneTimeUnit TimeUnit = EMovieSceneTimeUnit::DisplayRate);
 
 	/**
 	* Get ControlRig Control's TransformNoScale value at a specific time
@@ -818,7 +818,7 @@ public:
 	*/
 	UFUNCTION(BlueprintCallable, Category = "Editor Scripting | Sequencer Tools | Control Rig")
 	static FTransformNoScale GetLocalControlRigTransformNoScale(ULevelSequence* LevelSequence, UControlRig* ControlRig, FName ControlName, FFrameNumber Frame,
-		ESequenceTimeUnit TimeUnit = ESequenceTimeUnit::DisplayRate);
+		EMovieSceneTimeUnit TimeUnit = EMovieSceneTimeUnit::DisplayRate);
 
 	/**
 	* Get ControlRig Control's TransformNoScale values at specific times
@@ -831,7 +831,7 @@ public:
 	*/
 	UFUNCTION(BlueprintCallable, Category = "Editor Scripting | Sequencer Tools | Control Rig")
 	static TArray<FTransformNoScale> GetLocalControlRigTransformNoScales(ULevelSequence* LevelSequence, UControlRig* ControlRig, FName ControlName, const TArray<FFrameNumber>& Frames,
-		ESequenceTimeUnit TimeUnit = ESequenceTimeUnit::DisplayRate);
+		EMovieSceneTimeUnit TimeUnit = EMovieSceneTimeUnit::DisplayRate);
 
 	/**
 	* Set ControlRig Control's TransformNoScale value at specific time
@@ -845,7 +845,7 @@ public:
 	*/
 	UFUNCTION(BlueprintCallable, Category = "Editor Scripting | Sequencer Tools | Control Rig")
 	static void SetLocalControlRigTransformNoScale(ULevelSequence* LevelSequence, UControlRig* ControlRig, FName ControlName, FFrameNumber Frame, FTransformNoScale Value, 
-		ESequenceTimeUnit TimeUnit = ESequenceTimeUnit::DisplayRate, bool bSetKey = true);
+		EMovieSceneTimeUnit TimeUnit = EMovieSceneTimeUnit::DisplayRate, bool bSetKey = true);
 
 	/**
 	* Set ControlRig Control's TransformNoScale values at specific times
@@ -858,7 +858,7 @@ public:
 	*/
 	UFUNCTION(BlueprintCallable, Category = "Editor Scripting | Sequencer Tools | Control Rig")
 	static void SetLocalControlRigTransformNoScales(ULevelSequence* LevelSequence, UControlRig* ControlRig, FName ControlName, const TArray<FFrameNumber>& Frames, const TArray<FTransformNoScale> Values,
-		ESequenceTimeUnit TimeUnit = ESequenceTimeUnit::DisplayRate);
+		EMovieSceneTimeUnit TimeUnit = EMovieSceneTimeUnit::DisplayRate);
 
 	/**
 	* Get ControlRig Control's Transform value at a specific time
@@ -871,7 +871,7 @@ public:
 	*/
 	UFUNCTION(BlueprintCallable, Category = "Editor Scripting | Sequencer Tools | Control Rig")
 	static FTransform GetLocalControlRigTransform(ULevelSequence* LevelSequence, UControlRig* ControlRig, FName ControlName, FFrameNumber Frame, 
-		ESequenceTimeUnit TimeUnit = ESequenceTimeUnit::DisplayRate);
+		EMovieSceneTimeUnit TimeUnit = EMovieSceneTimeUnit::DisplayRate);
 
 	/**
 	* Get ControlRig Control's Transform values at specific times
@@ -884,7 +884,7 @@ public:
 	*/
 	UFUNCTION(BlueprintCallable, Category = "Editor Scripting | Sequencer Tools | Control Rig")
 	static TArray<FTransform> GetLocalControlRigTransforms(ULevelSequence* LevelSequence, UControlRig* ControlRig, FName ControlName, const TArray<FFrameNumber>& Frames,
-		ESequenceTimeUnit TimeUnit = ESequenceTimeUnit::DisplayRate);
+		EMovieSceneTimeUnit TimeUnit = EMovieSceneTimeUnit::DisplayRate);
 
 	/**
 	* Set ControlRig Control's Transform value at specific time
@@ -898,7 +898,7 @@ public:
 	*/
 	UFUNCTION(BlueprintCallable, Category = "Editor Scripting | Sequencer Tools | Control Rig")
 	static void SetLocalControlRigTransform(ULevelSequence* LevelSequence, UControlRig* ControlRig, FName ControlName, FFrameNumber Frame, FTransform Value, 
-		ESequenceTimeUnit TimeUnit = ESequenceTimeUnit::DisplayRate, bool bSetKey = true);
+		EMovieSceneTimeUnit TimeUnit = EMovieSceneTimeUnit::DisplayRate, bool bSetKey = true);
 
 	/**
 	* Set ControlRig Control's Transform values at specific times
@@ -911,7 +911,7 @@ public:
 	*/
 	UFUNCTION(BlueprintCallable, Category = "Editor Scripting | Sequencer Tools | Control Rig")
 	static void SetLocalControlRigTransforms(ULevelSequence* LevelSequence, UControlRig* ControlRig, FName ControlName, const TArray<FFrameNumber>& Frames, const TArray<FTransform> Values,
-		ESequenceTimeUnit TimeUnit = ESequenceTimeUnit::DisplayRate);
+		EMovieSceneTimeUnit TimeUnit = EMovieSceneTimeUnit::DisplayRate);
 
 	/*
 	 * Import FBX onto a control rig with the specified track and section
@@ -981,7 +981,7 @@ public:
 	 * @param TimeUnit Unit for the InTime, either in display rate or tick resolution
 	 */
 	UFUNCTION(BlueprintCallable, Category = "Editor Scripting | Sequencer Tools | Control Rig")
-	static bool SetControlRigSpace(ULevelSequence* InSequence, UControlRig* InControlRig, FName InControlName, const FRigElementKey& InSpaceKey, FFrameNumber InTime,  ESequenceTimeUnit TimeUnit = ESequenceTimeUnit::DisplayRate);
+	static bool SetControlRigSpace(ULevelSequence* InSequence, UControlRig* InControlRig, FName InControlName, const FRigElementKey& InSpaceKey, FFrameNumber InTime,  EMovieSceneTimeUnit TimeUnit = EMovieSceneTimeUnit::DisplayRate);
 
 	/** Bake specified Control Rig Controls to a specified Space based upon the current settings
 	* @param InSequence Sequence to bake
@@ -991,7 +991,7 @@ public:
 	* @param TimeUnit Unit for the start and end times in the InSettings parameter.
 	*/
 	UFUNCTION(BlueprintCallable, Category = "Editor Scripting | Sequencer Tools | Control Rig")
-	static bool BakeControlRigSpace(ULevelSequence* InSequence, UControlRig* InControlRig, const TArray<FName>& InControlNames, FRigSpacePickerBakeSettings InSettings, ESequenceTimeUnit TimeUnit = ESequenceTimeUnit::DisplayRate);
+	static bool BakeControlRigSpace(ULevelSequence* InSequence, UControlRig* InControlRig, const TArray<FName>& InControlNames, FRigSpacePickerBakeSettings InSettings, EMovieSceneTimeUnit TimeUnit = EMovieSceneTimeUnit::DisplayRate);
 	
 	/** Delete the Control Rig Space Key for the Control at the specified time. This will delete any attached Control Rig keys at this time and will perform any needed compensation to the new space.
 	*
@@ -1003,7 +1003,7 @@ public:
 	* @return Will return false if function fails,  for example if there is no key at this time it will fail.
 	*/
 	UFUNCTION(BlueprintCallable, Category = "Editor Scripting | Sequencer Tools | Control Rig")
-	static bool DeleteControlRigSpace(ULevelSequence* InSequence, UControlRig* InControlRig, FName InControlName, FFrameNumber InTime, ESequenceTimeUnit TimeUnit = ESequenceTimeUnit::DisplayRate);
+	static bool DeleteControlRigSpace(ULevelSequence* InSequence, UControlRig* InControlRig, FName InControlName, FFrameNumber InTime, EMovieSceneTimeUnit TimeUnit = EMovieSceneTimeUnit::DisplayRate);
 
 	/** Move the Control Rig Space Key for the Control at the specified time to the new time. This will also move any Control Rig keys at this space switch boundary.
 	*
@@ -1016,7 +1016,7 @@ public:
 	* @return Will return false if function fails, for example if there is no key at this time it will fail, or if the new time is invalid it could fail also
 	*/
 	UFUNCTION(BlueprintCallable, Category = "Editor Scripting | Sequencer Tools | Control Rig")
-	static bool MoveControlRigSpace(ULevelSequence* InSequence, UControlRig* InControlRig, FName InControlName, FFrameNumber InTime, FFrameNumber InNewTime, ESequenceTimeUnit TimeUnit = ESequenceTimeUnit::DisplayRate);
+	static bool MoveControlRigSpace(ULevelSequence* InSequence, UControlRig* InControlRig, FName InControlName, FFrameNumber InTime, FFrameNumber InNewTime, EMovieSceneTimeUnit TimeUnit = EMovieSceneTimeUnit::DisplayRate);
 
 	/** Rename the Control Rig Channels in Sequencer to the specified new control names, which should be present on the Control Rig
 	* @param InSequence Sequence to rename controls
diff --git a/Engine/Plugins/Editor/SequencerAnimTools/Source/SequencerAnimTools/Private/SequencerAnimEditPivotTool.cpp b/Engine/Plugins/Editor/SequencerAnimTools/Source/SequencerAnimTools/Private/SequencerAnimEditPivotTool.cpp
--- a/Engine/Plugins/Editor/SequencerAnimTools/Source/SequencerAnimTools/Private/SequencerAnimEditPivotTool.cpp
+++ b/Engine/Plugins/Editor/SequencerAnimTools/Source/SequencerAnimTools/Private/SequencerAnimEditPivotTool.cpp
@@ -17,6 +17,7 @@
 #include "EditorModeManager.h"
 #include "Misc/QualifiedFrameTime.h"
 #include "Modules/ModuleManager.h"
+#include "MovieSceneTimeUnit.h"
 #include "ControlRig.h"
 #include "ControlRigSequencerEditorLibrary.h"
 #include "LevelSequence.h"
@@ -455,7 +456,7 @@ bool USequencerPivotTool::SetGizmoBasedOnSelection(bool bUseSaved)
 			{
 				UControlRig* RealControlRig = Cast<UControlRig>(ControlRig.Get());
 				GizmoTransform = UControlRigSequencerEditorLibrary::GetControlRigWorldTransform(LevelSequence, RealControlRig, Name, FrameTime.RoundToFrame(),
-					ESequenceTimeUnit::TickResolution);
+					EMovieSceneTimeUnit::TickResolution);
 				AverageLocation += GizmoTransform.GetLocation();
 				NumLocations++;
 			}
@@ -480,13 +481,13 @@ bool USequencerPivotTool::SetGizmoBasedOnSelection(bool bUseSaved)
 			else
 			{
 				GizmoTransform = UControlRigSequencerEditorLibrary::GetActorWorldTransform(LevelSequence, SelectedActor, FrameTime.RoundToFrame(),
-					ESequenceTimeUnit::TickResolution);
+					EMovieSceneTimeUnit::TickResolution);
 			}
 		}
 		else
 		{
 			GizmoTransform = UControlRigSequencerEditorLibrary::GetActorWorldTransform(LevelSequence,SelectedActor, FrameTime.RoundToFrame(),
-				ESequenceTimeUnit::TickResolution);
+				EMovieSceneTimeUnit::TickResolution);
 			AverageLocation += GizmoTransform.GetLocation();
 			NumLocations++;
 		}
@@ -646,7 +647,7 @@ void USequencerPivotTool::GizmoTransformStarted(UTransformProxy* Proxy)
 						}
 						UControlRig* RealControlRig = Cast<UControlRig>(ControlRig.Get());
 						FTransform Transform = UControlRigSequencerEditorLibrary::GetControlRigWorldTransform(LevelSequence, RealControlRig, Name, FrameNumber,
-							ESequenceTimeUnit::TickResolution);
+							EMovieSceneTimeUnit::TickResolution);
 
 						FControlRigSelectionDuringDrag ControlDrag;
 						ControlDrag.LevelSequence = LevelSequence;
@@ -677,7 +678,7 @@ void USequencerPivotTool::GizmoTransformStarted(UTransformProxy* Proxy)
 			}
 
 			FTransform Transform = UControlRigSequencerEditorLibrary::GetActorWorldTransform(LevelSequence, Actor.Get(), FrameNumber,
-				ESequenceTimeUnit::TickResolution);
+				EMovieSceneTimeUnit::TickResolution);
 
 			FActorSelectonDuringDrag ActorDrag;
 			ActorDrag.LevelSequence = LevelSequence;
@@ -802,7 +803,7 @@ void USequencerPivotTool::GizmoTransformChanged(UTransformProxy* Proxy, FTransfo
 						ControlDrag.CurrentTransform.SetRotation(OptQuat);
 					}
 					UControlRigSequencerEditorLibrary::SetControlRigWorldTransform(ControlDrag.LevelSequence, RealControlRig, ControlDrag.ControlName,
-						ControlDrag.CurrentFrame, ControlDrag.CurrentTransform, ESequenceTimeUnit::TickResolution, bSetKey);
+						ControlDrag.CurrentFrame, ControlDrag.CurrentTransform, EMovieSceneTimeUnit::TickResolution, bSetKey);
 					
 				}
 				else if (bTranslationChanged)
@@ -810,13 +811,13 @@ void USequencerPivotTool::GizmoTransformChanged(UTransformProxy* Proxy, FTransfo
 					FVector DiffTranslation = StartDragTransform.GetRotation().RotateVector(Diff.GetTranslation());
 					ControlDrag.CurrentTransform.AddToTranslation(DiffTranslation);
 					UControlRigSequencerEditorLibrary::SetControlRigWorldTransform(ControlDrag.LevelSequence, RealControlRig, ControlDrag.ControlName,
-						ControlDrag.CurrentFrame, ControlDrag.CurrentTransform, ESequenceTimeUnit::TickResolution, bSetKey);
+						ControlDrag.CurrentFrame, ControlDrag.CurrentTransform, EMovieSceneTimeUnit::TickResolution, bSetKey);
 				}
 			}
 			else //last one with shift we keep locked!
 			{
 				UControlRigSequencerEditorLibrary::SetControlRigWorldTransform(ControlDrag.LevelSequence, RealControlRig, ControlDrag.ControlName,
-					ControlDrag.CurrentFrame, ControlDrag.CurrentTransform, ESequenceTimeUnit::TickResolution, bSetKey);
+					ControlDrag.CurrentFrame, ControlDrag.CurrentTransform, EMovieSceneTimeUnit::TickResolution, bSetKey);
 			}
 			
 			++Index;
diff --git a/Engine/Plugins/MovieScene/LevelSequenceEditor/Source/LevelSequenceEditor/Private/LevelSequenceEditorSubsystem.cpp b/Engine/Plugins/MovieScene/LevelSequenceEditor/Source/LevelSequenceEditor/Private/LevelSequenceEditorSubsystem.cpp
--- a/Engine/Plugins/MovieScene/LevelSequenceEditor/Source/LevelSequenceEditor/Private/LevelSequenceEditorSubsystem.cpp
+++ b/Engine/Plugins/MovieScene/LevelSequenceEditor/Source/LevelSequenceEditor/Private/LevelSequenceEditorSubsystem.cpp
@@ -919,7 +919,7 @@ void ULevelSequenceEditorSubsystem::BakeTransformInternal()
 		.OnBake_Lambda([this, &BindingProxies, Sequencer](FBakingAnimationKeySettings InSettings)
 			{
 				FMovieSceneScriptingParams Params;
-				Params.TimeUnit = ESequenceTimeUnit::TickResolution;
+				Params.TimeUnit = EMovieSceneTimeUnit::TickResolution;
 				BakeTransformWithSettings(BindingProxies, InSettings, Params);
 				return FReply::Handled();
 			});
@@ -958,7 +958,7 @@ void ULevelSequenceEditorSubsystem::BakeTransform(const TArray<FMovieSceneBindin
 	FFrameTime OutFrame = BakeOutTime;
 	FFrameTime Interval = BakeInterval;
 
-	if (Params.TimeUnit == ESequenceTimeUnit::DisplayRate)
+	if (Params.TimeUnit == EMovieSceneTimeUnit::DisplayRate)
 	{
 		InFrame = ConvertFrameTime(BakeInTime, DisplayRate, TickResolution);
 		OutFrame = ConvertFrameTime(BakeOutTime, DisplayRate, TickResolution);
@@ -974,7 +974,7 @@ void ULevelSequenceEditorSubsystem::BakeTransform(const TArray<FMovieSceneBindin
 	Settings.EndFrame = OutFrame.GetFrame();
 	Settings.BakingKeySettings = EBakingKeySettings::AllFrames;
 	FMovieSceneScriptingParams NewParams;
-	NewParams.TimeUnit = ESequenceTimeUnit::TickResolution;
+	NewParams.TimeUnit = EMovieSceneTimeUnit::TickResolution;
 	BakeTransformWithSettings(ObjectBindings, Settings, NewParams);
 }
 
@@ -1067,7 +1067,7 @@ bool ULevelSequenceEditorSubsystem::BakeTransformWithSettings(const TArray<FMovi
 
 	FBakingAnimationKeySettings SettingsInTick = InSettings;
 	
-	if (Params.TimeUnit == ESequenceTimeUnit::DisplayRate)
+	if (Params.TimeUnit == EMovieSceneTimeUnit::DisplayRate)
 	{
 		SettingsInTick.StartFrame = ConvertFrameTime(SettingsInTick.StartFrame, DisplayRate, TickResolution).GetFrame();
 		SettingsInTick.EndFrame = ConvertFrameTime(SettingsInTick.EndFrame, DisplayRate, TickResolution).GetFrame();
diff --git a/Engine/Plugins/MovieScene/LevelSequenceEditor/Source/LevelSequenceEditor/Public/LevelSequenceEditorSubsystem.h b/Engine/Plugins/MovieScene/LevelSequenceEditor/Source/LevelSequenceEditor/Public/LevelSequenceEditorSubsystem.h
--- a/Engine/Plugins/MovieScene/LevelSequenceEditor/Source/LevelSequenceEditor/Public/LevelSequenceEditorSubsystem.h
+++ b/Engine/Plugins/MovieScene/LevelSequenceEditor/Source/LevelSequenceEditor/Public/LevelSequenceEditorSubsystem.h
@@ -5,7 +5,7 @@
 
 #include "EditorSubsystem.h"
 
-#include "SequenceTimeUnit.h"
+#include "MovieSceneTimeUnit.h"
 #include "Containers/SortedMap.h"
 #include "UObject/StructOnScope.h"
 #include "UniversalObjectLocator.h"
@@ -44,7 +44,7 @@ struct FMovieSceneScriptingParams
 	FMovieSceneScriptingParams() {}
 
 	UPROPERTY(BlueprintReadWrite, Category = "Movie Scene")
-	ESequenceTimeUnit TimeUnit = ESequenceTimeUnit::DisplayRate;
+	EMovieSceneTimeUnit TimeUnit = EMovieSceneTimeUnit::DisplayRate;
 };
 
 // Helper struct for Binding Properties UI for locators.
diff --git a/Engine/Plugins/MovieScene/SequencerScripting/Source/SequencerScripting/Private/ExtensionLibraries/MovieSceneSequenceExtensions.cpp b/Engine/Plugins/MovieScene/SequencerScripting/Source/SequencerScripting/Private/ExtensionLibraries/MovieSceneSequenceExtensions.cpp
--- a/Engine/Plugins/MovieScene/SequencerScripting/Source/SequencerScripting/Private/ExtensionLibraries/MovieSceneSequenceExtensions.cpp
+++ b/Engine/Plugins/MovieScene/SequencerScripting/Source/SequencerScripting/Private/ExtensionLibraries/MovieSceneSequenceExtensions.cpp
@@ -1031,7 +1031,7 @@ void UMovieSceneSequenceExtensions::RemoveRootFolderFromSequence(UMovieSceneSequ
 #endif
 }
 
-TArray<FMovieSceneMarkedFrame> UMovieSceneSequenceExtensions::GetMarkedFramesFromSequence(UMovieSceneSequence* Sequence, ESequenceTimeUnit TimeUnit)
+TArray<FMovieSceneMarkedFrame> UMovieSceneSequenceExtensions::GetMarkedFramesFromSequence(UMovieSceneSequence* Sequence, EMovieSceneTimeUnit TimeUnit)
 {
 	if (!Sequence)
 	{
@@ -1043,7 +1043,7 @@ TArray<FMovieSceneMarkedFrame> UMovieSceneSequenceExtensions::GetMarkedFramesFro
 	if (MovieScene)
 	{
 		TArray<FMovieSceneMarkedFrame> MarkedFrames = MovieScene->GetMarkedFrames();
-		if (TimeUnit == ESequenceTimeUnit::DisplayRate)
+		if (TimeUnit == EMovieSceneTimeUnit::DisplayRate)
 		{
 			FFrameRate TickResolution = MovieScene->GetTickResolution();
 			FFrameRate DisplayRate = MovieScene->GetDisplayRate();
@@ -1060,7 +1060,7 @@ TArray<FMovieSceneMarkedFrame> UMovieSceneSequenceExtensions::GetMarkedFramesFro
 	return TArray<FMovieSceneMarkedFrame>();
 }
 
-int32 UMovieSceneSequenceExtensions::AddMarkedFrameToSequence(UMovieSceneSequence* Sequence, const FMovieSceneMarkedFrame& InMarkedFrame, ESequenceTimeUnit TimeUnit)
+int32 UMovieSceneSequenceExtensions::AddMarkedFrameToSequence(UMovieSceneSequence* Sequence, const FMovieSceneMarkedFrame& InMarkedFrame, EMovieSceneTimeUnit TimeUnit)
 {
 	if (!Sequence)
 	{
@@ -1074,7 +1074,7 @@ int32 UMovieSceneSequenceExtensions::AddMarkedFrameToSequence(UMovieSceneSequenc
 		MovieScene->Modify();
 
 		FMovieSceneMarkedFrame MarkedFrame(InMarkedFrame);
-		if (TimeUnit == ESequenceTimeUnit::DisplayRate)
+		if (TimeUnit == EMovieSceneTimeUnit::DisplayRate)
 		{
 			MarkedFrame.FrameNumber = FFrameRate::TransformTime(MarkedFrame.FrameNumber, MovieScene->GetDisplayRate(), MovieScene->GetTickResolution()).RoundToFrame();
 		}
@@ -1085,7 +1085,7 @@ int32 UMovieSceneSequenceExtensions::AddMarkedFrameToSequence(UMovieSceneSequenc
 }
 
 
-void UMovieSceneSequenceExtensions::SetMarkedFrameInSequence(UMovieSceneSequence* Sequence, int32 InMarkIndex, FFrameNumber InFrameNumber, ESequenceTimeUnit TimeUnit)
+void UMovieSceneSequenceExtensions::SetMarkedFrameInSequence(UMovieSceneSequence* Sequence, int32 InMarkIndex, FFrameNumber InFrameNumber, EMovieSceneTimeUnit TimeUnit)
 {
 	if (!Sequence)
 	{
@@ -1096,7 +1096,7 @@ void UMovieSceneSequenceExtensions::SetMarkedFrameInSequence(UMovieSceneSequence
 	UMovieScene* MovieScene = Sequence->GetMovieScene();
 	if (MovieScene)
 	{
-		if (TimeUnit == ESequenceTimeUnit::DisplayRate)
+		if (TimeUnit == EMovieSceneTimeUnit::DisplayRate)
 		{
 			InFrameNumber = FFrameRate::TransformTime(InFrameNumber, MovieScene->GetDisplayRate(), MovieScene->GetTickResolution()).RoundToFrame();
 		}
@@ -1171,7 +1171,7 @@ int32 UMovieSceneSequenceExtensions::FindMarkedFrameByLabel(UMovieSceneSequence*
 	return INDEX_NONE;
 }
 
-int32 UMovieSceneSequenceExtensions::FindMarkedFrameByFrameNumberInSequence(UMovieSceneSequence* Sequence, FFrameNumber InFrameNumber, ESequenceTimeUnit TimeUnit)
+int32 UMovieSceneSequenceExtensions::FindMarkedFrameByFrameNumberInSequence(UMovieSceneSequence* Sequence, FFrameNumber InFrameNumber, EMovieSceneTimeUnit TimeUnit)
 {
 	if (!Sequence)
 	{
@@ -1182,7 +1182,7 @@ int32 UMovieSceneSequenceExtensions::FindMarkedFrameByFrameNumberInSequence(UMov
 	UMovieScene* MovieScene = Sequence->GetMovieScene();
 	if (MovieScene)
 	{
-		if (TimeUnit == ESequenceTimeUnit::DisplayRate)
+		if (TimeUnit == EMovieSceneTimeUnit::DisplayRate)
 		{
 			InFrameNumber = FFrameRate::TransformTime(InFrameNumber, MovieScene->GetDisplayRate(), MovieScene->GetTickResolution()).RoundToFrame();
 		}
@@ -1192,7 +1192,7 @@ int32 UMovieSceneSequenceExtensions::FindMarkedFrameByFrameNumberInSequence(UMov
 	return INDEX_NONE;
 }
 
-int32 UMovieSceneSequenceExtensions::FindNextMarkedFrameInSequence(UMovieSceneSequence* Sequence, FFrameNumber InFrameNumber, bool bForward, ESequenceTimeUnit TimeUnit)
+int32 UMovieSceneSequenceExtensions::FindNextMarkedFrameInSequence(UMovieSceneSequence* Sequence, FFrameNumber InFrameNumber, bool bForward, EMovieSceneTimeUnit TimeUnit)
 {
 	if (!Sequence)
 	{
@@ -1203,7 +1203,7 @@ int32 UMovieSceneSequenceExtensions::FindNextMarkedFrameInSequence(UMovieSceneSe
 	UMovieScene* MovieScene = Sequence->GetMovieScene();
 	if (MovieScene)
 	{
-		if (TimeUnit == ESequenceTimeUnit::DisplayRate)
+		if (TimeUnit == EMovieSceneTimeUnit::DisplayRate)
 		{
 			InFrameNumber = FFrameRate::TransformTime(InFrameNumber, MovieScene->GetDisplayRate(), MovieScene->GetTickResolution()).RoundToFrame();
 		}
diff --git a/Engine/Plugins/MovieScene/SequencerScripting/Source/SequencerScripting/Private/KeysAndChannels/MovieSceneScriptingActorReference.h b/Engine/Plugins/MovieScene/SequencerScripting/Source/SequencerScripting/Private/KeysAndChannels/MovieSceneScriptingActorReference.h
--- a/Engine/Plugins/MovieScene/SequencerScripting/Source/SequencerScripting/Private/KeysAndChannels/MovieSceneScriptingActorReference.h
+++ b/Engine/Plugins/MovieScene/SequencerScripting/Source/SequencerScripting/Private/KeysAndChannels/MovieSceneScriptingActorReference.h
@@ -5,6 +5,7 @@
 #include "MovieSceneScriptingChannel.h"
 #include "Sections/MovieSceneActorReferenceSection.h"
 #include "KeyParams.h"
+#include "MovieSceneTimeUnit.h"
 
 #include "MovieSceneScriptingActorReference.generated.h"
 
@@ -27,7 +28,7 @@ public:
 	* @return			The time of this key which combines both the frame number and the sub-frame it is on. Sub-frame will be zero if you request Tick Resolution.	
 	*/
 	UFUNCTION(BlueprintCallable, Category = "Sequencer|Keys", meta = (DisplayName = "Get Time (Actor Reference)"))
-	virtual FFrameTime GetTime(ESequenceTimeUnit TimeUnit = ESequenceTimeUnit::DisplayRate) const override { return GetTimeFromChannel(KeyHandle, OwningSequence, TimeUnit); }
+	virtual FFrameTime GetTime(EMovieSceneTimeUnit TimeUnit = EMovieSceneTimeUnit::DisplayRate) const override { return GetTimeFromChannel(KeyHandle, OwningSequence, TimeUnit); }
 	
 	/**
 	* Sets the time for this key in the owning channel. Will replace any key that already exists at that frame number in this channel.
@@ -36,7 +37,7 @@ public:
 	* @param TimeUnit		Should the NewFrameNumber be interpreted as Display Rate frames or in Tick Resolution?
 	*/
 	UFUNCTION(BlueprintCallable, Category = "Sequencer|Keys", meta = (DisplayName = "Set Time (Actor Reference)"))
-	void SetTime(const FFrameNumber& NewFrameNumber, float SubFrame = 0.f, ESequenceTimeUnit TimeUnit = ESequenceTimeUnit::DisplayRate) { SetTimeInChannel(KeyHandle, OwningSequence, OwningSection, NewFrameNumber, TimeUnit, SubFrame); }
+	void SetTime(const FFrameNumber& NewFrameNumber, float SubFrame = 0.f, EMovieSceneTimeUnit TimeUnit = EMovieSceneTimeUnit::DisplayRate) { SetTimeInChannel(KeyHandle, OwningSequence, OwningSection, NewFrameNumber, TimeUnit, SubFrame); }
 
 	/**
 	* Gets the value for this key from the owning channel.
@@ -78,7 +79,7 @@ public:
 	* @return	The key that was created with the specified values at the specified time.
 	*/
 	UFUNCTION(BlueprintCallable, Category = "Sequencer|Keys", meta = (DisplayName = "Add Key (Actor Reference)"))
-	UMovieSceneScriptingActorReferenceKey* AddKey(const FFrameNumber InTime, FMovieSceneObjectBindingID NewValue, float SubFrame = 0.f, ESequenceTimeUnit TimeUnit = ESequenceTimeUnit::DisplayRate)
+	UMovieSceneScriptingActorReferenceKey* AddKey(const FFrameNumber InTime, FMovieSceneObjectBindingID NewValue, float SubFrame = 0.f, EMovieSceneTimeUnit TimeUnit = EMovieSceneTimeUnit::DisplayRate)
 	{
 		FMovieSceneActorReferenceKey ReferenceKey = FMovieSceneActorReferenceKey(NewValue);
 		return Impl::AddKeyInChannel(ChannelHandle, OwningSequence, OwningSection, InTime, ReferenceKey, SubFrame, TimeUnit, EMovieSceneKeyInterpolation::Auto);
diff --git a/Engine/Plugins/MovieScene/SequencerScripting/Source/SequencerScripting/Private/KeysAndChannels/MovieSceneScriptingBool.h b/Engine/Plugins/MovieScene/SequencerScripting/Source/SequencerScripting/Private/KeysAndChannels/MovieSceneScriptingBool.h
--- a/Engine/Plugins/MovieScene/SequencerScripting/Source/SequencerScripting/Private/KeysAndChannels/MovieSceneScriptingBool.h
+++ b/Engine/Plugins/MovieScene/SequencerScripting/Source/SequencerScripting/Private/KeysAndChannels/MovieSceneScriptingBool.h
@@ -7,6 +7,7 @@
 #include "Channels/MovieSceneBoolChannel.h"
 #include "MovieSceneScriptingChannel.h"
 #include "KeyParams.h"
+#include "MovieSceneTimeUnit.h"
 
 #include "MovieSceneScriptingBool.generated.h"
 
@@ -26,7 +27,7 @@ public:
 	* @return			The time of this key which combines both the frame number and the sub-frame it is on. Sub-frame will be zero if you request Tick Resolution.
 	*/
 	UFUNCTION(BlueprintCallable, Category = "Sequencer|Keys", meta = (DisplayName = "Get Time (Bool)"))
-	virtual FFrameTime GetTime(ESequenceTimeUnit TimeUnit = ESequenceTimeUnit::DisplayRate) const override { return GetTimeFromChannel(KeyHandle, OwningSequence, TimeUnit); }
+	virtual FFrameTime GetTime(EMovieSceneTimeUnit TimeUnit = EMovieSceneTimeUnit::DisplayRate) const override { return GetTimeFromChannel(KeyHandle, OwningSequence, TimeUnit); }
 	
 	/**
 	* Sets the time for this key in the owning channel. Will replace any key that already exists at that frame number in this channel.
@@ -35,7 +36,7 @@ public:
 	* @param TimeUnit		Should the NewFrameNumber be interpreted as Display Rate frames or in Tick Resolution?
 	*/
 	UFUNCTION(BlueprintCallable, Category = "Sequencer|Keys", meta = (DisplayName = "Set Time (Bool)"))
-	void SetTime(const FFrameNumber& NewFrameNumber, float SubFrame = 0.f, ESequenceTimeUnit TimeUnit = ESequenceTimeUnit::DisplayRate) { SetTimeInChannel(KeyHandle, OwningSequence, OwningSection, NewFrameNumber, TimeUnit, SubFrame); }
+	void SetTime(const FFrameNumber& NewFrameNumber, float SubFrame = 0.f, EMovieSceneTimeUnit TimeUnit = EMovieSceneTimeUnit::DisplayRate) { SetTimeInChannel(KeyHandle, OwningSequence, OwningSection, NewFrameNumber, TimeUnit, SubFrame); }
 
 	/**
 	* Gets the value for this key from the owning channel.
@@ -75,7 +76,7 @@ public:
 	* @return	The key that was created with the specified values at the specified time.
 	*/
 	UFUNCTION(BlueprintCallable, Category = "Sequencer|Keys", meta=(DisplayName = "Add Key (Bool)"))
-	UMovieSceneScriptingBoolKey* AddKey(const FFrameNumber& InTime, bool NewValue, float SubFrame = 0.f, ESequenceTimeUnit TimeUnit = ESequenceTimeUnit::DisplayRate)
+	UMovieSceneScriptingBoolKey* AddKey(const FFrameNumber& InTime, bool NewValue, float SubFrame = 0.f, EMovieSceneTimeUnit TimeUnit = EMovieSceneTimeUnit::DisplayRate)
 	{
 		return Impl::AddKeyInChannel(ChannelHandle, OwningSequence, OwningSection, InTime, NewValue, SubFrame, TimeUnit, EMovieSceneKeyInterpolation::Auto);
 	}
diff --git a/Engine/Plugins/MovieScene/SequencerScripting/Source/SequencerScripting/Private/KeysAndChannels/MovieSceneScriptingByte.h b/Engine/Plugins/MovieScene/SequencerScripting/Source/SequencerScripting/Private/KeysAndChannels/MovieSceneScriptingByte.h
--- a/Engine/Plugins/MovieScene/SequencerScripting/Source/SequencerScripting/Private/KeysAndChannels/MovieSceneScriptingByte.h
+++ b/Engine/Plugins/MovieScene/SequencerScripting/Source/SequencerScripting/Private/KeysAndChannels/MovieSceneScriptingByte.h
@@ -6,6 +6,7 @@
 #include "MovieSceneScriptingChannel.h"
 #include "Channels/MovieSceneByteChannel.h"
 #include "KeyParams.h"
+#include "MovieSceneTimeUnit.h"
 
 #include "MovieSceneScriptingByte.generated.h"
 
@@ -28,7 +29,7 @@ public:
 	* @return			The time of this key which combines both the frame number and the sub-frame it is on. Sub-frame will be zero if you request Tick Resolution.
 	*/
 	UFUNCTION(BlueprintCallable, Category = "Sequencer|Keys", meta = (DisplayName = "Get Time (Enum)"))
-	virtual FFrameTime GetTime(ESequenceTimeUnit TimeUnit = ESequenceTimeUnit::DisplayRate) const override { return GetTimeFromChannel(KeyHandle, OwningSequence, TimeUnit); }
+	virtual FFrameTime GetTime(EMovieSceneTimeUnit TimeUnit = EMovieSceneTimeUnit::DisplayRate) const override { return GetTimeFromChannel(KeyHandle, OwningSequence, TimeUnit); }
 	
 	/**
 	* Sets the time for this key in the owning channel. Will replace any key that already exists at that frame number in this channel.
@@ -37,7 +38,7 @@ public:
 	* @param TimeUnit		Should the NewFrameNumber be interpreted as Display Rate frames or in Tick Resolution?
 	*/
 	UFUNCTION(BlueprintCallable, Category = "Sequencer|Keys", meta = (DisplayName = "Set Time (Enum)"))
-	void SetTime(const FFrameNumber& NewFrameNumber, float SubFrame = 0.f, ESequenceTimeUnit TimeUnit = ESequenceTimeUnit::DisplayRate) { SetTimeInChannel(KeyHandle, OwningSequence, OwningSection, NewFrameNumber, TimeUnit, SubFrame); }
+	void SetTime(const FFrameNumber& NewFrameNumber, float SubFrame = 0.f, EMovieSceneTimeUnit TimeUnit = EMovieSceneTimeUnit::DisplayRate) { SetTimeInChannel(KeyHandle, OwningSequence, OwningSection, NewFrameNumber, TimeUnit, SubFrame); }
 
 	/**
 	* Gets the value for this key from the owning channel.
@@ -78,7 +79,7 @@ public:
 	* @return	The key that was created with the specified values at the specified time.
 	*/
 	UFUNCTION(BlueprintCallable, Category = "Sequencer|Keys", meta = (DisplayName = "Add Key (Enum)"))
-	UMovieSceneScriptingByteKey* AddKey(const FFrameNumber& InTime, uint8 NewValue, float SubFrame = 0.f, ESequenceTimeUnit TimeUnit = ESequenceTimeUnit::DisplayRate, EMovieSceneKeyInterpolation InInterpolation = EMovieSceneKeyInterpolation::Auto)
+	UMovieSceneScriptingByteKey* AddKey(const FFrameNumber& InTime, uint8 NewValue, float SubFrame = 0.f, EMovieSceneTimeUnit TimeUnit = EMovieSceneTimeUnit::DisplayRate, EMovieSceneKeyInterpolation InInterpolation = EMovieSceneKeyInterpolation::Auto)
 	{
 		return Impl::AddKeyInChannel(ChannelHandle, OwningSequence, OwningSection, InTime, NewValue, SubFrame, TimeUnit, InInterpolation);
 	}
diff --git a/Engine/Plugins/MovieScene/SequencerScripting/Source/SequencerScripting/Private/KeysAndChannels/MovieSceneScriptingChannel.h b/Engine/Plugins/MovieScene/SequencerScripting/Source/SequencerScripting/Private/KeysAndChannels/MovieSceneScriptingChannel.h
--- a/Engine/Plugins/MovieScene/SequencerScripting/Source/SequencerScripting/Private/KeysAndChannels/MovieSceneScriptingChannel.h
+++ b/Engine/Plugins/MovieScene/SequencerScripting/Source/SequencerScripting/Private/KeysAndChannels/MovieSceneScriptingChannel.h
@@ -5,7 +5,7 @@
 #include "Channels/MovieSceneChannelHandle.h"
 #include "Channels/MovieSceneChannelTraits.h"
 #include "SequencerScriptingRange.h"
-#include "SequenceTimeUnit.h"
+#include "MovieSceneTimeUnit.h"
 #include "ExtensionLibraries/MovieSceneSequenceExtensions.h"
 #include "MovieSceneSection.h"
 #include "MovieSceneSequence.h"
@@ -23,7 +23,7 @@ public:
 	* @param TimeUnit (Optional) Should the time be returned in Display Rate frames (possibly with a sub-frame value) or in Tick Resolution with no sub-frame values? Defaults to Display Rate.
 	* @return					 The FrameTime of this key which combines both the frame number and the sub-frame it is on. Sub-frame will be zero if you request Tick Resolution.
 	*/
-	virtual FFrameTime GetTime(ESequenceTimeUnit TimeUnit = ESequenceTimeUnit::DisplayRate) const PURE_VIRTUAL(UMovieSceneScriptingKey::GetTime, return FFrameTime(););
+	virtual FFrameTime GetTime(EMovieSceneTimeUnit TimeUnit = EMovieSceneTimeUnit::DisplayRate) const PURE_VIRTUAL(UMovieSceneScriptingKey::GetTime, return FFrameTime(););
 public:
 	FKeyHandle KeyHandle;
 	TWeakObjectPtr<UMovieSceneSequence> OwningSequence;
@@ -61,7 +61,7 @@ template<typename ChannelType, typename ScriptingKeyType, typename ScriptingKeyV
 struct TMovieSceneScriptingChannel
 {
 	static ScriptingKeyType* AddKeyInChannel(TMovieSceneChannelHandle<ChannelType> ChannelHandle, TWeakObjectPtr<UMovieSceneSequence> Sequence, 
-		TWeakObjectPtr<UMovieSceneSection> Section, const FFrameNumber InTime, ScriptingKeyValueType& NewValue, float SubFrame, ESequenceTimeUnit TimeUnit, EMovieSceneKeyInterpolation Interpolation)
+		TWeakObjectPtr<UMovieSceneSection> Section, const FFrameNumber InTime, ScriptingKeyValueType& NewValue, float SubFrame, EMovieSceneTimeUnit TimeUnit, EMovieSceneKeyInterpolation Interpolation)
 	{
 		ChannelType* Channel = ChannelHandle.Get();
 		if (Channel)
@@ -70,7 +70,7 @@ struct TMovieSceneScriptingChannel
 
 			// The Key's time is always going to be in Tick Resolution space, but the user may want to set it via Display Rate, so we convert.
 			FFrameNumber KeyTime = InTime;
-			if (TimeUnit == ESequenceTimeUnit::DisplayRate)
+			if (TimeUnit == EMovieSceneTimeUnit::DisplayRate)
 			{
 				KeyTime = FFrameRate::TransformTime(KeyTime, UMovieSceneSequenceExtensions::GetDisplayRate(Sequence.Get()), UMovieSceneSequenceExtensions::GetTickResolution(Sequence.Get())).RoundToFrame();
 			}
@@ -322,7 +322,7 @@ struct TMovieSceneScriptingChannel
 template<typename ChannelType, typename ChannelDataType>
 struct TMovieSceneScriptingKey
 {
-	FFrameTime GetTimeFromChannel(FKeyHandle KeyHandle, TWeakObjectPtr<UMovieSceneSequence> Sequence, ESequenceTimeUnit TimeUnit) const
+	FFrameTime GetTimeFromChannel(FKeyHandle KeyHandle, TWeakObjectPtr<UMovieSceneSequence> Sequence, EMovieSceneTimeUnit TimeUnit) const
 	{
 		if (!Sequence.IsValid())
 		{
@@ -337,7 +337,7 @@ struct TMovieSceneScriptingKey
 			Channel->GetKeyTime(KeyHandle, KeyTime);
 
 			// The KeyTime is always going to be in Tick Resolution space, but the user may desire it in Play Rate /w a Subframe.
-			if (TimeUnit == ESequenceTimeUnit::DisplayRate)
+			if (TimeUnit == EMovieSceneTimeUnit::DisplayRate)
 			{
 				FFrameTime DisplayRateTime = FFrameRate::TransformTime(KeyTime, UMovieSceneSequenceExtensions::GetTickResolution(Sequence.Get()), UMovieSceneSequenceExtensions::GetDisplayRate(Sequence.Get()));
 				return DisplayRateTime;
@@ -351,7 +351,7 @@ struct TMovieSceneScriptingKey
 		return FFrameTime();
 	}
 
-	void SetTimeInChannel(FKeyHandle KeyHandle, TWeakObjectPtr<UMovieSceneSequence> Sequence, TWeakObjectPtr<UMovieSceneSection> Section, const FFrameNumber NewFrameNumber, ESequenceTimeUnit TimeUnit, float SubFrame)
+	void SetTimeInChannel(FKeyHandle KeyHandle, TWeakObjectPtr<UMovieSceneSequence> Sequence, TWeakObjectPtr<UMovieSceneSection> Section, const FFrameNumber NewFrameNumber, EMovieSceneTimeUnit TimeUnit, float SubFrame)
 	{
 		if (!Sequence.IsValid())
 		{
@@ -363,7 +363,7 @@ struct TMovieSceneScriptingKey
 		SubFrame = FMath::Clamp(SubFrame, 0.f, FFrameTime::MaxSubframe);
 
 		// TickResolution doesn't support a sub-frame as you can't get finer detailed than that.
-		if (TimeUnit == ESequenceTimeUnit::TickResolution && SubFrame > 0.f)
+		if (TimeUnit == EMovieSceneTimeUnit::TickResolution && SubFrame > 0.f)
 		{
 			FFrame::KismetExecutionMessage(TEXT("SetTime called with a SubFrame specified for a Tick Resolution type time! SubFrames are only allowed for Display Rate types, ignoring..."), ELogVerbosity::Error);
 			SubFrame = 0.f;
@@ -372,7 +372,7 @@ struct TMovieSceneScriptingKey
 		FFrameNumber KeyFrameNumber = NewFrameNumber;
 		
 		// Keys are always stored in Tick Resolution so we need to potentially convert their values.
-		if (TimeUnit == ESequenceTimeUnit::DisplayRate)
+		if (TimeUnit == EMovieSceneTimeUnit::DisplayRate)
 		{
 			KeyFrameNumber = FFrameRate::TransformTime(FFrameTime(NewFrameNumber, SubFrame), UMovieSceneSequenceExtensions::GetDisplayRate(Sequence.Get()), UMovieSceneSequenceExtensions::GetTickResolution(Sequence.Get())).RoundToFrame();
 		}
diff --git a/Engine/Plugins/MovieScene/SequencerScripting/Source/SequencerScripting/Private/KeysAndChannels/MovieSceneScriptingDouble.h b/Engine/Plugins/MovieScene/SequencerScripting/Source/SequencerScripting/Private/KeysAndChannels/MovieSceneScriptingDouble.h
--- a/Engine/Plugins/MovieScene/SequencerScripting/Source/SequencerScripting/Private/KeysAndChannels/MovieSceneScriptingDouble.h
+++ b/Engine/Plugins/MovieScene/SequencerScripting/Source/SequencerScripting/Private/KeysAndChannels/MovieSceneScriptingDouble.h
@@ -6,6 +6,7 @@
 #include "MovieSceneScriptingChannel.h"
 #include "Channels/MovieSceneDoubleChannel.h"
 #include "KeyParams.h"
+#include "MovieSceneTimeUnit.h"
 
 #include "MovieSceneScriptingDouble.generated.h"
 
@@ -30,7 +31,7 @@ public:
 	* @return			The time of this key which combines both the frame number and the sub-frame it is on. Sub-frame will be zero if you request Tick Resolution.
 	*/
 	UFUNCTION(BlueprintCallable, Category = "Sequencer|Keys", meta = (DisplayName = "Get Time (Double)"))
-	virtual FFrameTime GetTime(ESequenceTimeUnit TimeUnit = ESequenceTimeUnit::DisplayRate) const override { return GetTimeFromChannel(KeyHandle, OwningSequence, TimeUnit); }
+	virtual FFrameTime GetTime(EMovieSceneTimeUnit TimeUnit = EMovieSceneTimeUnit::DisplayRate) const override { return GetTimeFromChannel(KeyHandle, OwningSequence, TimeUnit); }
 	
 	/**
 	* Sets the time for this key in the owning channel. Will replace any key that already exists at that frame number in this channel.
@@ -39,7 +40,7 @@ public:
 	* @param TimeUnit		Should the NewFrameNumber be interpreted as Display Rate frames or in Tick Resolution?
 	*/
 	UFUNCTION(BlueprintCallable, Category = "Sequencer|Keys", meta = (DisplayName = "Set Time (Double)"))
-	void SetTime(const FFrameNumber& NewFrameNumber, float SubFrame = 0.f, ESequenceTimeUnit TimeUnit = ESequenceTimeUnit::DisplayRate) { SetTimeInChannel(KeyHandle, OwningSequence, OwningSection, NewFrameNumber, TimeUnit, SubFrame); }
+	void SetTime(const FFrameNumber& NewFrameNumber, float SubFrame = 0.f, EMovieSceneTimeUnit TimeUnit = EMovieSceneTimeUnit::DisplayRate) { SetTimeInChannel(KeyHandle, OwningSequence, OwningSection, NewFrameNumber, TimeUnit, SubFrame); }
 
 	/**
 	* Gets the value for this key from the owning channel.
@@ -241,7 +242,7 @@ public:
 	* @return	The key that was created with the specified values at the specified time.
 	*/
 	UFUNCTION(BlueprintCallable, Category = "Sequencer|Keys", meta = (DisplayName = "Add Key (Double)"))
-	UMovieSceneScriptingDoubleKey* AddKey(const FFrameNumber& InTime, double NewValue, float SubFrame = 0.f, ESequenceTimeUnit TimeUnit = ESequenceTimeUnit::DisplayRate, EMovieSceneKeyInterpolation InInterpolation = EMovieSceneKeyInterpolation::Auto)
+	UMovieSceneScriptingDoubleKey* AddKey(const FFrameNumber& InTime, double NewValue, float SubFrame = 0.f, EMovieSceneTimeUnit TimeUnit = EMovieSceneTimeUnit::DisplayRate, EMovieSceneKeyInterpolation InInterpolation = EMovieSceneKeyInterpolation::Auto)
 	{
 		return Impl::AddKeyInChannel(ChannelHandle, OwningSequence, OwningSection, InTime, NewValue, SubFrame, TimeUnit, InInterpolation);
 	}
diff --git a/Engine/Plugins/MovieScene/SequencerScripting/Source/SequencerScripting/Private/KeysAndChannels/MovieSceneScriptingEvent.h b/Engine/Plugins/MovieScene/SequencerScripting/Source/SequencerScripting/Private/KeysAndChannels/MovieSceneScriptingEvent.h
--- a/Engine/Plugins/MovieScene/SequencerScripting/Source/SequencerScripting/Private/KeysAndChannels/MovieSceneScriptingEvent.h
+++ b/Engine/Plugins/MovieScene/SequencerScripting/Source/SequencerScripting/Private/KeysAndChannels/MovieSceneScriptingEvent.h
@@ -4,6 +4,7 @@
 
 #include "MovieSceneScriptingChannel.h"
 #include "KeyParams.h"
+#include "MovieSceneTimeUnit.h"
 #include "Channels/MovieSceneEventChannel.h"
 
 #include "MovieSceneScriptingEvent.generated.h"
@@ -23,7 +24,7 @@ public:
 	* @return			The time of this key which combines both the frame number and the sub-frame it is on. Sub-frame will be zero if you request Tick Resolution.
 	*/
 	UFUNCTION(BlueprintCallable, Category = "Sequencer|Keys", meta = (DisplayName = "Get Time (Event)"))
-	virtual FFrameTime GetTime(ESequenceTimeUnit TimeUnit = ESequenceTimeUnit::DisplayRate) const override { return GetTimeFromChannel(KeyHandle, OwningSequence, TimeUnit); }
+	virtual FFrameTime GetTime(EMovieSceneTimeUnit TimeUnit = EMovieSceneTimeUnit::DisplayRate) const override { return GetTimeFromChannel(KeyHandle, OwningSequence, TimeUnit); }
 	
 	/**
 	* Sets the time for this key in the owning channel. Will replace any key that already exists at that frame number in this channel.
@@ -32,7 +33,7 @@ public:
 	* @param TimeUnit		Should the NewFrameNumber be interpreted as Display Rate frames or in Tick Resolution?
 	*/
 	UFUNCTION(BlueprintCallable, Category = "Sequencer|Keys", meta = (DisplayName = "Set Time (Event)"))
-	void SetTime(const FFrameNumber& NewFrameNumber, float SubFrame = 0.f, ESequenceTimeUnit TimeUnit = ESequenceTimeUnit::DisplayRate) { SetTimeInChannel(KeyHandle, OwningSequence, OwningSection, NewFrameNumber, TimeUnit, SubFrame); }
+	void SetTime(const FFrameNumber& NewFrameNumber, float SubFrame = 0.f, EMovieSceneTimeUnit TimeUnit = EMovieSceneTimeUnit::DisplayRate) { SetTimeInChannel(KeyHandle, OwningSequence, OwningSection, NewFrameNumber, TimeUnit, SubFrame); }
 
 	/**
 	* Gets the value for this key from the owning channel.
@@ -72,7 +73,7 @@ public:
 	* @return	The key that was created with the specified values at the specified time.
 	*/
 	UFUNCTION(BlueprintCallable, Category = "Sequencer|Keys", meta = (DisplayName = "Add Key (Event)"))
-	UMovieSceneScriptingEventKey* AddKey(const FFrameNumber& InTime, FMovieSceneEvent NewValue, float SubFrame = 0.f, ESequenceTimeUnit TimeUnit = ESequenceTimeUnit::DisplayRate)
+	UMovieSceneScriptingEventKey* AddKey(const FFrameNumber& InTime, FMovieSceneEvent NewValue, float SubFrame = 0.f, EMovieSceneTimeUnit TimeUnit = EMovieSceneTimeUnit::DisplayRate)
 	{
 		return Impl::AddKeyInChannel(ChannelHandle, OwningSequence, OwningSection,InTime, NewValue, SubFrame, TimeUnit, EMovieSceneKeyInterpolation::Auto);
 	}
diff --git a/Engine/Plugins/MovieScene/SequencerScripting/Source/SequencerScripting/Private/KeysAndChannels/MovieSceneScriptingFloat.h b/Engine/Plugins/MovieScene/SequencerScripting/Source/SequencerScripting/Private/KeysAndChannels/MovieSceneScriptingFloat.h
--- a/Engine/Plugins/MovieScene/SequencerScripting/Source/SequencerScripting/Private/KeysAndChannels/MovieSceneScriptingFloat.h
+++ b/Engine/Plugins/MovieScene/SequencerScripting/Source/SequencerScripting/Private/KeysAndChannels/MovieSceneScriptingFloat.h
@@ -7,6 +7,7 @@
 #include "Channels/MovieSceneFloatChannel.h"
 #include "Channels/MovieSceneDoubleChannel.h"
 #include "KeyParams.h"
+#include "MovieSceneTimeUnit.h"
 
 #include "MovieSceneScriptingFloat.generated.h"
 
@@ -29,7 +30,7 @@ public:
 	* @return			The time of this key which combines both the frame number and the sub-frame it is on. Sub-frame will be zero if you request Tick Resolution.
 	*/
 	UFUNCTION(BlueprintCallable, Category = "Sequencer|Keys", meta = (DisplayName = "Get Time (Float)"))
-	virtual FFrameTime GetTime(ESequenceTimeUnit TimeUnit = ESequenceTimeUnit::DisplayRate) const override PURE_VIRTUAL(UMovieSceneScriptingFloatKey::GetTime, return FFrameTime(););
+	virtual FFrameTime GetTime(EMovieSceneTimeUnit TimeUnit = EMovieSceneTimeUnit::DisplayRate) const override PURE_VIRTUAL(UMovieSceneScriptingFloatKey::GetTime, return FFrameTime(););
 	
 	/**
 	* Sets the time for this key in the owning channel. Will replace any key that already exists at that frame number in this channel.
@@ -38,7 +39,7 @@ public:
 	* @param TimeUnit		Should the NewFrameNumber be interpreted as Display Rate frames or in Tick Resolution?
 	*/
 	UFUNCTION(BlueprintCallable, Category = "Sequencer|Keys", meta = (DisplayName = "Set Time (Float)"))
-	virtual void SetTime(const FFrameNumber& NewFrameNumber, float SubFrame = 0.f, ESequenceTimeUnit TimeUnit = ESequenceTimeUnit::DisplayRate) PURE_VIRTUAL(UMovieSceneScriptingFloatKey::SetTime);
+	virtual void SetTime(const FFrameNumber& NewFrameNumber, float SubFrame = 0.f, EMovieSceneTimeUnit TimeUnit = EMovieSceneTimeUnit::DisplayRate) PURE_VIRTUAL(UMovieSceneScriptingFloatKey::SetTime);
 
 	/**
 	* Gets the value for this key from the owning channel.
@@ -158,11 +159,11 @@ class UMovieSceneScriptingActualFloatKey : public UMovieSceneScriptingFloatKey,
 {
 	GENERATED_BODY()
 public:
-	virtual FFrameTime GetTime(ESequenceTimeUnit TimeUnit = ESequenceTimeUnit::DisplayRate) const override
+	virtual FFrameTime GetTime(EMovieSceneTimeUnit TimeUnit = EMovieSceneTimeUnit::DisplayRate) const override
 	{
 		return GetTimeFromChannel(KeyHandle, OwningSequence, TimeUnit);
 	}
-	virtual void SetTime(const FFrameNumber& NewFrameNumber, float SubFrame = 0.f, ESequenceTimeUnit TimeUnit = ESequenceTimeUnit::DisplayRate) override
+	virtual void SetTime(const FFrameNumber& NewFrameNumber, float SubFrame = 0.f, EMovieSceneTimeUnit TimeUnit = EMovieSceneTimeUnit::DisplayRate) override
 	{
 		SetTimeInChannel(KeyHandle, OwningSequence, OwningSection, NewFrameNumber, TimeUnit, SubFrame);
 	}
@@ -258,11 +259,11 @@ class UMovieSceneScriptingDoubleAsFloatKey : public UMovieSceneScriptingFloatKey
 {
 	GENERATED_BODY()
 public:
-	virtual FFrameTime GetTime(ESequenceTimeUnit TimeUnit = ESequenceTimeUnit::DisplayRate) const override
+	virtual FFrameTime GetTime(EMovieSceneTimeUnit TimeUnit = EMovieSceneTimeUnit::DisplayRate) const override
 	{
 		return GetTimeFromChannel(KeyHandle, OwningSequence, TimeUnit);
 	}
-	virtual void SetTime(const FFrameNumber& NewFrameNumber, float SubFrame = 0.f, ESequenceTimeUnit TimeUnit = ESequenceTimeUnit::DisplayRate) override
+	virtual void SetTime(const FFrameNumber& NewFrameNumber, float SubFrame = 0.f, EMovieSceneTimeUnit TimeUnit = EMovieSceneTimeUnit::DisplayRate) override
 	{
 		SetTimeInChannel(KeyHandle, OwningSequence, OwningSection, NewFrameNumber, TimeUnit, SubFrame);
 	}
@@ -372,7 +373,7 @@ public:
 	* @return	The key that was created with the specified values at the specified time.
 	*/
 	UFUNCTION(BlueprintCallable, Category = "Sequencer|Keys", meta = (DisplayName = "Add Key (Float)"))
-	UMovieSceneScriptingFloatKey* AddKey(const FFrameNumber& InTime, float NewValue, float SubFrame = 0.f, ESequenceTimeUnit TimeUnit = ESequenceTimeUnit::DisplayRate, EMovieSceneKeyInterpolation InInterpolation = EMovieSceneKeyInterpolation::Auto)
+	UMovieSceneScriptingFloatKey* AddKey(const FFrameNumber& InTime, float NewValue, float SubFrame = 0.f, EMovieSceneTimeUnit TimeUnit = EMovieSceneTimeUnit::DisplayRate, EMovieSceneKeyInterpolation InInterpolation = EMovieSceneKeyInterpolation::Auto)
 	{
 		if (FloatChannelHandle.Get())
 		{
diff --git a/Engine/Plugins/MovieScene/SequencerScripting/Source/SequencerScripting/Private/KeysAndChannels/MovieSceneScriptingInteger.h b/Engine/Plugins/MovieScene/SequencerScripting/Source/SequencerScripting/Private/KeysAndChannels/MovieSceneScriptingInteger.h
--- a/Engine/Plugins/MovieScene/SequencerScripting/Source/SequencerScripting/Private/KeysAndChannels/MovieSceneScriptingInteger.h
+++ b/Engine/Plugins/MovieScene/SequencerScripting/Source/SequencerScripting/Private/KeysAndChannels/MovieSceneScriptingInteger.h
@@ -6,6 +6,7 @@
 #include "MovieSceneScriptingChannel.h"
 #include "Channels/MovieSceneIntegerChannel.h"
 #include "KeyParams.h"
+#include "MovieSceneTimeUnit.h"
 
 #include "MovieSceneScriptingInteger.generated.h"
 
@@ -28,7 +29,7 @@ public:
 	* @return			The time of this key which combines both the frame number and the sub-frame it is on. Sub-frame will be zero if you request Tick Resolution.
 	*/
 	UFUNCTION(BlueprintCallable, Category = "Sequencer|Keys", meta = (DisplayName = "Get Time (Integer)"))
-	virtual FFrameTime GetTime(ESequenceTimeUnit TimeUnit = ESequenceTimeUnit::DisplayRate) const override { return GetTimeFromChannel(KeyHandle, OwningSequence, TimeUnit); }
+	virtual FFrameTime GetTime(EMovieSceneTimeUnit TimeUnit = EMovieSceneTimeUnit::DisplayRate) const override { return GetTimeFromChannel(KeyHandle, OwningSequence, TimeUnit); }
 	
 	/**
 	* Sets the time for this key in the owning channel. Will replace any key that already exists at that frame number in this channel.
@@ -37,7 +38,7 @@ public:
 	* @param TimeUnit		Should the NewFrameNumber be interpreted as Display Rate frames or in Tick Resolution?
 	*/
 	UFUNCTION(BlueprintCallable, Category = "Sequencer|Keys", meta = (DisplayName = "Set Time (Integer)"))
-	void SetTime(const FFrameNumber& NewFrameNumber, float SubFrame = 0.f, ESequenceTimeUnit TimeUnit = ESequenceTimeUnit::DisplayRate) { SetTimeInChannel(KeyHandle, OwningSequence, OwningSection, NewFrameNumber, TimeUnit, SubFrame); }
+	void SetTime(const FFrameNumber& NewFrameNumber, float SubFrame = 0.f, EMovieSceneTimeUnit TimeUnit = EMovieSceneTimeUnit::DisplayRate) { SetTimeInChannel(KeyHandle, OwningSequence, OwningSection, NewFrameNumber, TimeUnit, SubFrame); }
 
 	/**
 	* Gets the value for this key from the owning channel.
@@ -77,7 +78,7 @@ public:
 	* @return	The key that was created with the specified values at the specified time.
 	*/
 	UFUNCTION(BlueprintCallable, Category = "Sequencer|Keys", meta = (DisplayName = "Add Key (Integer)"))
-	UMovieSceneScriptingIntegerKey* AddKey(const FFrameNumber& InTime, int32 NewValue, float SubFrame = 0.f, ESequenceTimeUnit TimeUnit = ESequenceTimeUnit::DisplayRate)
+	UMovieSceneScriptingIntegerKey* AddKey(const FFrameNumber& InTime, int32 NewValue, float SubFrame = 0.f, EMovieSceneTimeUnit TimeUnit = EMovieSceneTimeUnit::DisplayRate)
 	{
 		return Impl::AddKeyInChannel(ChannelHandle, OwningSequence, OwningSection, InTime, NewValue, SubFrame, TimeUnit, EMovieSceneKeyInterpolation::Auto);
 	}
diff --git a/Engine/Plugins/MovieScene/SequencerScripting/Source/SequencerScripting/Private/KeysAndChannels/MovieSceneScriptingObjectPath.h b/Engine/Plugins/MovieScene/SequencerScripting/Source/SequencerScripting/Private/KeysAndChannels/MovieSceneScriptingObjectPath.h
--- a/Engine/Plugins/MovieScene/SequencerScripting/Source/SequencerScripting/Private/KeysAndChannels/MovieSceneScriptingObjectPath.h
+++ b/Engine/Plugins/MovieScene/SequencerScripting/Source/SequencerScripting/Private/KeysAndChannels/MovieSceneScriptingObjectPath.h
@@ -6,6 +6,7 @@
 #include "MovieSceneScriptingChannel.h"
 #include "Channels/MovieSceneObjectPathChannel.h"
 #include "KeyParams.h"
+#include "MovieSceneTimeUnit.h"
 
 #include "MovieSceneScriptingObjectPath.generated.h"
 
@@ -28,7 +29,7 @@ public:
 	 * @return			The time of this key which combines both the frame number and the sub-frame it is on. Sub-frame will be zero if you request Tick Resolution.	
 	 */
 	UFUNCTION(BlueprintCallable, Category = "Sequencer|Keys", meta = (DisplayName = "Get Time (Object Path)"))
-	virtual FFrameTime GetTime(ESequenceTimeUnit TimeUnit = ESequenceTimeUnit::DisplayRate) const override { return GetTimeFromChannel(KeyHandle, OwningSequence, TimeUnit); }
+	virtual FFrameTime GetTime(EMovieSceneTimeUnit TimeUnit = EMovieSceneTimeUnit::DisplayRate) const override { return GetTimeFromChannel(KeyHandle, OwningSequence, TimeUnit); }
 	
 	/**
 	 * Sets the time for this key in the owning channel. Will replace any key that already exists at that frame number in this channel.
@@ -37,7 +38,7 @@ public:
 	 * @param TimeUnit		Should the NewFrameNumber be interpreted as Display Rate frames or in Tick Resolution?
 	 */
 	UFUNCTION(BlueprintCallable, Category = "Sequencer|Keys", meta = (DisplayName = "Set Time (Object Path)"))
-	void SetTime(const FFrameNumber& NewFrameNumber, float SubFrame = 0.f, ESequenceTimeUnit TimeUnit = ESequenceTimeUnit::DisplayRate) { SetTimeInChannel(KeyHandle, OwningSequence, OwningSection, NewFrameNumber, TimeUnit, SubFrame); }
+	void SetTime(const FFrameNumber& NewFrameNumber, float SubFrame = 0.f, EMovieSceneTimeUnit TimeUnit = EMovieSceneTimeUnit::DisplayRate) { SetTimeInChannel(KeyHandle, OwningSequence, OwningSection, NewFrameNumber, TimeUnit, SubFrame); }
 
 	/**
 	 * Gets the value for this key from the owning channel.
@@ -79,7 +80,7 @@ public:
 	 * @return	The key that was created with the specified values at the specified time.
 	 */
 	UFUNCTION(BlueprintCallable, Category = "Sequencer|Keys", meta = (DisplayName = "Add Key (Object Path)"))
-	UMovieSceneScriptingObjectPathKey* AddKey(const FFrameNumber InTime, UObject* NewValue, float SubFrame = 0.f, ESequenceTimeUnit TimeUnit = ESequenceTimeUnit::DisplayRate)
+	UMovieSceneScriptingObjectPathKey* AddKey(const FFrameNumber InTime, UObject* NewValue, float SubFrame = 0.f, EMovieSceneTimeUnit TimeUnit = EMovieSceneTimeUnit::DisplayRate)
 	{
 		FMovieSceneObjectPathChannelKeyValue ReferenceKey = FMovieSceneObjectPathChannelKeyValue(NewValue);
 		return Impl::AddKeyInChannel(ChannelHandle, OwningSequence, OwningSection,InTime, ReferenceKey, SubFrame, TimeUnit, EMovieSceneKeyInterpolation::Auto);
diff --git a/Engine/Plugins/MovieScene/SequencerScripting/Source/SequencerScripting/Private/KeysAndChannels/MovieSceneScriptingString.h b/Engine/Plugins/MovieScene/SequencerScripting/Source/SequencerScripting/Private/KeysAndChannels/MovieSceneScriptingString.h
--- a/Engine/Plugins/MovieScene/SequencerScripting/Source/SequencerScripting/Private/KeysAndChannels/MovieSceneScriptingString.h
+++ b/Engine/Plugins/MovieScene/SequencerScripting/Source/SequencerScripting/Private/KeysAndChannels/MovieSceneScriptingString.h
@@ -5,6 +5,7 @@
 #include "MovieSceneScriptingChannel.h"
 #include "Channels/MovieSceneStringChannel.h"
 #include "KeyParams.h"
+#include "MovieSceneTimeUnit.h"
 
 #include "MovieSceneScriptingString.generated.h"
 
@@ -27,7 +28,7 @@ public:
 	* @return			The time of this key which combines both the frame number and the sub-frame it is on. Sub-frame will be zero if you request Tick Resolution.
 	*/
 	UFUNCTION(BlueprintCallable, Category = "Sequencer|Keys", meta = (DisplayName = "Get Time (String)"))
-	virtual FFrameTime GetTime(ESequenceTimeUnit TimeUnit = ESequenceTimeUnit::DisplayRate) const override { return GetTimeFromChannel(KeyHandle, OwningSequence, TimeUnit); }
+	virtual FFrameTime GetTime(EMovieSceneTimeUnit TimeUnit = EMovieSceneTimeUnit::DisplayRate) const override { return GetTimeFromChannel(KeyHandle, OwningSequence, TimeUnit); }
 	
 	/**
 	* Sets the time for this key in the owning channel. Will replace any key that already exists at that frame number in this channel.
@@ -36,7 +37,7 @@ public:
 	* @param TimeUnit		Should the NewFrameNumber be interpreted as Display Rate frames or in Tick Resolution?
 	*/
 	UFUNCTION(BlueprintCallable, Category = "Sequencer|Keys", meta = (DisplayName = "Set Time (String)"))
-	void SetTime(const FFrameNumber& NewFrameNumber, float SubFrame = 0.f, ESequenceTimeUnit TimeUnit = ESequenceTimeUnit::DisplayRate) { SetTimeInChannel(KeyHandle, OwningSequence, OwningSection, NewFrameNumber, TimeUnit, SubFrame); }
+	void SetTime(const FFrameNumber& NewFrameNumber, float SubFrame = 0.f, EMovieSceneTimeUnit TimeUnit = EMovieSceneTimeUnit::DisplayRate) { SetTimeInChannel(KeyHandle, OwningSequence, OwningSection, NewFrameNumber, TimeUnit, SubFrame); }
 
 	/**
 	* Gets the value for this key from the owning channel.
@@ -77,7 +78,7 @@ public:
 	* @return	The key that was created with the specified values at the specified time.
 	*/
 	UFUNCTION(BlueprintCallable, Category = "Sequencer|Keys", meta = (DisplayName = "Add Key (String)"))
-	UMovieSceneScriptingStringKey* AddKey(const FFrameNumber& InTime, FString NewValue, float SubFrame = 0.f, ESequenceTimeUnit TimeUnit = ESequenceTimeUnit::DisplayRate)
+	UMovieSceneScriptingStringKey* AddKey(const FFrameNumber& InTime, FString NewValue, float SubFrame = 0.f, EMovieSceneTimeUnit TimeUnit = EMovieSceneTimeUnit::DisplayRate)
 	{
 		return Impl::AddKeyInChannel(ChannelHandle, OwningSequence, OwningSection,InTime, NewValue, SubFrame, TimeUnit, EMovieSceneKeyInterpolation::Auto);
 	}
diff --git a/Engine/Plugins/MovieScene/SequencerScripting/Source/SequencerScripting/Public/ExtensionLibraries/MovieSceneSequenceExtensions.h b/Engine/Plugins/MovieScene/SequencerScripting/Source/SequencerScripting/Public/ExtensionLibraries/MovieSceneSequenceExtensions.h
--- a/Engine/Plugins/MovieScene/SequencerScripting/Source/SequencerScripting/Public/ExtensionLibraries/MovieSceneSequenceExtensions.h
+++ b/Engine/Plugins/MovieScene/SequencerScripting/Source/SequencerScripting/Public/ExtensionLibraries/MovieSceneSequenceExtensions.h
@@ -5,7 +5,7 @@
 #include "Kismet/BlueprintFunctionLibrary.h"
 #include "MovieSceneObjectBindingID.h" // for EMovieSceneObjectBindingSpace
 #include "MovieSceneTrack.h"
-#include "SequenceTimeUnit.h"
+#include "MovieSceneTimeUnit.h"
 #include "MovieSceneSequenceExtensions.generated.h"
 
 class UMovieScene;
@@ -534,11 +534,11 @@ public:
 	 * @return Return the user marked frames
 	 */
 	UFUNCTION(BlueprintCallable, Category = "Sequencer|Sequence", meta = (ScriptMethod, DisplayName = "Get Marked Frames"))
-	static TArray<FMovieSceneMarkedFrame> GetMarkedFramesFromSequence(UMovieSceneSequence* Sequence, ESequenceTimeUnit TimeUnit = ESequenceTimeUnit::DisplayRate);
+	static TArray<FMovieSceneMarkedFrame> GetMarkedFramesFromSequence(UMovieSceneSequence* Sequence, EMovieSceneTimeUnit TimeUnit = EMovieSceneTimeUnit::DisplayRate);
 
 	UE_DEPRECATED(5.4, "GetMarkedFrames is deprecated. Please use GetMarkedFrames that takes a time unit and defaults to display rate instead")
 	UFUNCTION(BlueprintCallable, Category = "Sequencer|Sequence", meta = (ScriptMethod, DeprecatedFunction, DeprecationMessage = "GetMarkedFrames is deprecated. Please use GetMarkedFrames that takes a time unit instead"))
-	static TArray<FMovieSceneMarkedFrame> GetMarkedFrames(UMovieSceneSequence* Sequence) { return GetMarkedFramesFromSequence(Sequence, ESequenceTimeUnit::TickResolution); }
+	static TArray<FMovieSceneMarkedFrame> GetMarkedFrames(UMovieSceneSequence* Sequence) { return GetMarkedFramesFromSequence(Sequence, EMovieSceneTimeUnit::TickResolution); }
 
 	/*
 	 * Add a given user marked frame.
@@ -548,11 +548,11 @@ public:
 	 * @return The index to the newly added marked frame
 	 */
 	UFUNCTION(BlueprintCallable, Category = "Sequencer|Sequence", meta = (ScriptMethod, DisplayName = "Add Marked Frame"))
-	static int32 AddMarkedFrameToSequence(UMovieSceneSequence* Sequence, const FMovieSceneMarkedFrame& InMarkedFrame, ESequenceTimeUnit TimeUnit = ESequenceTimeUnit::DisplayRate);
+	static int32 AddMarkedFrameToSequence(UMovieSceneSequence* Sequence, const FMovieSceneMarkedFrame& InMarkedFrame, EMovieSceneTimeUnit TimeUnit = EMovieSceneTimeUnit::DisplayRate);
 
 	UE_DEPRECATED(5.4, "AddMarkedFrame is deprecated. Please use AddMarkedFrame that takes a time unit and defaults to display rate instead")
 	UFUNCTION(BlueprintCallable, Category = "Sequencer|Sequence", meta = (ScriptMethod, DeprecatedFunction, DeprecationMessage = "AddMarkedFrame is deprecated. Please use AddMarkedFrame that takes a time unit instead"))
-	static int32 AddMarkedFrame(UMovieSceneSequence* Sequence, const FMovieSceneMarkedFrame& InMarkedFrame) { return AddMarkedFrameToSequence(Sequence, InMarkedFrame, ESequenceTimeUnit::TickResolution); }
+	static int32 AddMarkedFrame(UMovieSceneSequence* Sequence, const FMovieSceneMarkedFrame& InMarkedFrame) { return AddMarkedFrameToSequence(Sequence, InMarkedFrame, EMovieSceneTimeUnit::TickResolution); }
 
 	/*
 	 * Sets the frame number for the given marked frame index. Does not maintain sort. Call SortMarkedFrames
@@ -561,11 +561,11 @@ public:
 	 * @InFrameNumber The frame number to set
 	 */
 	UFUNCTION(BlueprintCallable, Category = "Sequencer|Sequence", meta = (ScriptMethod, DisplayName = "Set Marked Frame"))
-	static void SetMarkedFrameInSequence(UMovieSceneSequence* Sequence, int32 InMarkIndex, FFrameNumber InFrameNumber, ESequenceTimeUnit TimeUnit = ESequenceTimeUnit::DisplayRate);
+	static void SetMarkedFrameInSequence(UMovieSceneSequence* Sequence, int32 InMarkIndex, FFrameNumber InFrameNumber, EMovieSceneTimeUnit TimeUnit = EMovieSceneTimeUnit::DisplayRate);
 
 	UE_DEPRECATED(5.4, "SetMarkedFrame is deprecated. Please use SetMarkedFrame that takes a time unit and defaults to display rate instead")
 	UFUNCTION(BlueprintCallable, Category = "Sequencer|Sequence", meta = (ScriptMethod, DeprecatedFunction, DeprecationMessage = "SetMarkedFrame is deprecated. Please use SetMarkedFrame that takes a time unit instead"))
-	static void SetMarkedFrame(UMovieSceneSequence* Sequence, int32 InMarkIndex, FFrameNumber InFrameNumber) { SetMarkedFrameInSequence(Sequence, InMarkIndex, InFrameNumber, ESequenceTimeUnit::TickResolution); }
+	static void SetMarkedFrame(UMovieSceneSequence* Sequence, int32 InMarkIndex, FFrameNumber InFrameNumber) { SetMarkedFrameInSequence(Sequence, InMarkIndex, InFrameNumber, EMovieSceneTimeUnit::TickResolution); }
 
 	/*
 	 * Delete the user marked frame by index.
@@ -601,11 +601,11 @@ public:
 	 * @InFrameNumber The frame number of the user marked frame to find
 	 */
 	UFUNCTION(BlueprintCallable, Category = "Sequencer|Sequence", meta = (ScriptMethod, DisplayName = "Find Marked Frame By Frame Number"))
-	static int32 FindMarkedFrameByFrameNumberInSequence(UMovieSceneSequence* Sequence, FFrameNumber InFrameNumber, ESequenceTimeUnit TimeUnit = ESequenceTimeUnit::DisplayRate);
+	static int32 FindMarkedFrameByFrameNumberInSequence(UMovieSceneSequence* Sequence, FFrameNumber InFrameNumber, EMovieSceneTimeUnit TimeUnit = EMovieSceneTimeUnit::DisplayRate);
 
 	UE_DEPRECATED(5.4, "FindMarkedFrameByFrameNumber is deprecated. Please use FindMarkedFrameByFrameNumber that takes a time unit and defaults to display rate instead")
 	UFUNCTION(BlueprintCallable, Category = "Sequencer|Sequence", meta = (ScriptMethod, DeprecatedFunction, DeprecationMessage = "FindMarkedFrameByFrameNumber is deprecated. Please use FindMarkedFrameByFrameNumber that takes a time unit instead"))
-	static int32 FindMarkedFrameByFrameNumber(UMovieSceneSequence* Sequence, FFrameNumber InFrameNumber) { return FindMarkedFrameByFrameNumberInSequence(Sequence, InFrameNumber, ESequenceTimeUnit::TickResolution); }
+	static int32 FindMarkedFrameByFrameNumber(UMovieSceneSequence* Sequence, FFrameNumber InFrameNumber) { return FindMarkedFrameByFrameNumberInSequence(Sequence, InFrameNumber, EMovieSceneTimeUnit::TickResolution); }
 
 	/*
 	 * Find the next/previous user marked frame from the given frame number
@@ -614,11 +614,11 @@ public:
 	 * @bForward Find forward from the given frame number.
 	 */
 	UFUNCTION(BlueprintCallable, Category = "Sequencer|Sequence", meta = (ScriptMethod, DisplayName = "Find Next Marked Frame"))
-	static int32 FindNextMarkedFrameInSequence(UMovieSceneSequence* Sequence, FFrameNumber InFrameNumber, bool bForward, ESequenceTimeUnit TimeUnit = ESequenceTimeUnit::DisplayRate);
+	static int32 FindNextMarkedFrameInSequence(UMovieSceneSequence* Sequence, FFrameNumber InFrameNumber, bool bForward, EMovieSceneTimeUnit TimeUnit = EMovieSceneTimeUnit::DisplayRate);
 
 	UE_DEPRECATED(5.4, "FindNextMarkedFrame is deprecated. Please use FindNextMarkedFrame that takes a time unit instead")
 	UFUNCTION(BlueprintCallable, Category = "Sequencer|Sequence", meta = (ScriptMethod, DeprecatedFunction, DeprecationMessage = "FindNextMarkedFrame is deprecated. Please use FindNextMarkedFrame that takes a time unit and defaults to display rate instead"))
-	static int32 FindNextMarkedFrame(UMovieSceneSequence* Sequence, FFrameNumber InFrameNumber, bool bForward) { return FindNextMarkedFrameInSequence(Sequence, InFrameNumber, bForward, ESequenceTimeUnit::TickResolution); }
+	static int32 FindNextMarkedFrame(UMovieSceneSequence* Sequence, FFrameNumber InFrameNumber, bool bForward) { return FindNextMarkedFrameInSequence(Sequence, InFrameNumber, bForward, EMovieSceneTimeUnit::TickResolution); }
 
 	/*
 	 * Set read only
diff --git a/Engine/Plugins/MovieScene/SequencerScripting/Source/SequencerScripting/Public/SequenceTimeUnit.h b/Engine/Plugins/MovieScene/SequencerScripting/Source/SequencerScripting/Public/SequenceTimeUnit.h
--- a/Engine/Plugins/MovieScene/SequencerScripting/Source/SequencerScripting/Public/SequenceTimeUnit.h
+++ b/Engine/Plugins/MovieScene/SequencerScripting/Source/SequencerScripting/Public/SequenceTimeUnit.h
@@ -4,16 +4,10 @@
 
 #include "SequenceTimeUnit.generated.h"
 
-/**
-* Specifies which frame of reference you want to set/get time values in. This allows users to work
-* in reference space without having to manually convert back and forth all of the time.
-*/
 UENUM(BlueprintType)
-enum class ESequenceTimeUnit : uint8
+enum class UE_DEPRECATED(5.4, "ESequenceTimeUnit is deprecated, please use EMovieSceneTimeUnit") ESequenceTimeUnit : uint8
 {
-	/** Display Rate matches the values shown in the UI such as 30fps giving you 30 frames per second. Supports sub-frame values (precision limited to Tick Resolution) */
 	DisplayRate,
-	/** Tick Resolution is the internal resolution that data is actually stored in, such as 24000 giving you 24,000 frames per second. This is the smallest interval that data can be stored on. */
 	TickResolution
 };
 
diff --git a/Engine/Source/Runtime/MovieScene/Public/MovieSceneTimeUnit.h b/Engine/Source/Runtime/MovieScene/Public/MovieSceneTimeUnit.h
new file mode 100644
--- /dev/null
+++ b/Engine/Source/Runtime/MovieScene/Public/MovieSceneTimeUnit.h
@@ -0,0 +1,22 @@
+// Copyright Epic Games, Inc. All Rights Reserved.
+
+#pragma once
+
+#include "MovieSceneTimeUnit.generated.h"
+
+/**
+* Specifies which frame of reference you want to set/get time values in. This allows users to work
+* in reference space without having to manually convert back and forth all of the time.
+*/
+UENUM(BlueprintType)
+enum class EMovieSceneTimeUnit : uint8
+{
+	/** Display Rate matches the values shown in the UI such as 30fps giving you 30 frames per second. Supports sub-frame values (precision limited to Tick Resolution) */
+	DisplayRate,
+	/** Tick Resolution is the internal resolution that data is actually stored in, such as 24000 giving you 24,000 frames per second. This is the smallest interval that data can be stored on. */
+	TickResolution
+};
+
+#if UE_ENABLE_INCLUDE_ORDER_DEPRECATED_IN_5_2
+#include "CoreMinimal.h"
+#endif
