Control Rig: Allow renaming of resolved connectors in RigModules
#rnx

diff --git a/Engine/Plugins/Animation/ControlRig/Source/ControlRigDeveloper/Private/ControlRigBlueprint.cpp b/Engine/Plugins/Animation/ControlRig/Source/ControlRigDeveloper/Private/ControlRigBlueprint.cpp
--- a/Engine/Plugins/Animation/ControlRig/Source/ControlRigDeveloper/Private/ControlRigBlueprint.cpp
+++ b/Engine/Plugins/Animation/ControlRig/Source/ControlRigDeveloper/Private/ControlRigBlueprint.cpp
@@ -2432,7 +2432,16 @@ void UControlRigBlueprint::HandleHierarchyModified(ERigHierarchyNotification InN
 		case ERigHierarchyNotification::ElementRenamed:
 		{
 			Modify();
-			Influences.OnKeyRenamed(FRigElementKey(InHierarchy->GetPreviousName(InElement->GetKey()), InElement->GetType()), InElement->GetKey());
+			const FRigElementKey OldKey(InHierarchy->GetPreviousName(InElement->GetKey()), InElement->GetType());
+			Influences.OnKeyRenamed(OldKey, InElement->GetKey());
+			if (IsControlRigModule() && InElement->IsTypeOf(ERigElementType::Connector))
+			{
+				if (FRigElementKey* Target = ConnectionMap.Find(OldKey))
+				{
+					ConnectionMap.FindOrAdd(InElement->GetKey(), *Target);
+					ConnectionMap.Remove(OldKey);
+				}
+			}
 			PropagateHierarchyFromBPToInstances();
 			break;
 		}
diff --git a/Engine/Plugins/Animation/ControlRig/Source/ControlRigEditor/Private/Editor/SRigHierarchyTagWidget.cpp b/Engine/Plugins/Animation/ControlRig/Source/ControlRigEditor/Private/Editor/SRigHierarchyTagWidget.cpp
--- a/Engine/Plugins/Animation/ControlRig/Source/ControlRigEditor/Private/Editor/SRigHierarchyTagWidget.cpp
+++ b/Engine/Plugins/Animation/ControlRig/Source/ControlRigEditor/Private/Editor/SRigHierarchyTagWidget.cpp
@@ -7,6 +7,7 @@
 #include "Brushes/SlateRoundedBoxBrush.h"
 #include "DetailLayoutBuilder.h"
 #include "ControlRigDragOps.h"
+#include "Widgets/Text/SInlineEditableTextBlock.h"
 
 #define LOCTEXT_NAMESPACE "SRigHierarchyTagWidget"
 
@@ -25,6 +26,8 @@ void SRigHierarchyTagWidget::Construct(const FArguments& InArgs)
 	Identifier = InArgs._Identifier;
 	bAllowDragDrop = InArgs._AllowDragDrop;
 	OnClicked = InArgs._OnClicked;
+	OnRenamed = InArgs._OnRenamed;
+	OnVerifyRename = InArgs._OnVerifyRename;
 
 	const FMargin CombinedPadding = ContentPadding + Padding;
 
@@ -53,10 +56,13 @@ void SRigHierarchyTagWidget::Construct(const FArguments& InArgs)
 		.VAlign(VAlign_Center)
 		.Padding(TAttribute<FMargin>::CreateSP(this, &SRigHierarchyTagWidget::GetTextPadding))
 		[
-			SNew(STextBlock)
+			SNew(SInlineEditableTextBlock)
 			.Text(Text)
 			.ColorAndOpacity(InArgs._TextColor)
 			.Font(IDetailLayoutBuilder::GetDetailFont())
+			.IsReadOnly_Lambda([this](){ return !OnRenamed.IsBound(); })
+			.OnTextCommitted(this, &SRigHierarchyTagWidget::HandleElementRenamed)
+			.OnVerifyTextChanged(this, &SRigHierarchyTagWidget::HandleVerifyRename)
 		]
 	];
 }
@@ -128,4 +134,24 @@ FMargin SRigHierarchyTagWidget::GetTextPadding() const
 	return CombinedPadding;
 }
 
+void SRigHierarchyTagWidget::HandleElementRenamed(const FText& InNewName, ETextCommit::Type InCommitType)
+{
+	if (InCommitType == ETextCommit::OnEnter)
+	{
+		if (OnRenamed.IsBound())
+		{
+			OnRenamed.Execute(InNewName, InCommitType);
+		}
+	}
+}
+
+bool SRigHierarchyTagWidget::HandleVerifyRename(const FText& InText, FText& OutError)
+{
+	if (OnVerifyRename.IsBound())
+	{
+		return OnVerifyRename.Execute(InText, OutError);
+	}
+	return false;
+}
+
 #undef LOCTEXT_NAMESPACE
diff --git a/Engine/Plugins/Animation/ControlRig/Source/ControlRigEditor/Private/Editor/SRigHierarchyTagWidget.h b/Engine/Plugins/Animation/ControlRig/Source/ControlRigEditor/Private/Editor/SRigHierarchyTagWidget.h
--- a/Engine/Plugins/Animation/ControlRig/Source/ControlRigEditor/Private/Editor/SRigHierarchyTagWidget.h
+++ b/Engine/Plugins/Animation/ControlRig/Source/ControlRigEditor/Private/Editor/SRigHierarchyTagWidget.h
@@ -26,6 +26,8 @@ public:
 		, _Identifier()
 		, _AllowDragDrop(false)
 		, _OnClicked()
+		, _OnRenamed()
+		, _OnVerifyRename()
 	{}
 		SLATE_ATTRIBUTE(FText, Text)
 		SLATE_ATTRIBUTE(FText, TooltipText)
@@ -38,6 +40,8 @@ public:
 		SLATE_ATTRIBUTE(FString, Identifier)
 		SLATE_ARGUMENT(bool, AllowDragDrop)
 		SLATE_EVENT(FSimpleDelegate, OnClicked)
+		SLATE_EVENT(FOnTextCommitted, OnRenamed)
+		SLATE_EVENT(FOnVerifyTextChanged, OnVerifyRename)
 	SLATE_END_ARGS()
 	
 	void Construct(const FArguments& InArgs);
@@ -49,6 +53,8 @@ public:
 private:
 
 	FMargin GetTextPadding() const;
+	void HandleElementRenamed(const FText& InNewName, ETextCommit::Type InCommitType);
+	bool HandleVerifyRename(const FText& InText, FText& OutError);
 
 	TAttribute<FText> Text;
 	TAttribute<const FSlateBrush*> Icon;
@@ -60,6 +66,8 @@ private:
 	bool bAllowDragDrop = false;
 	FSimpleDelegate OnClicked;
 	FOnRigTreeElementKeyTagDragDetected OnElementKeyDragDetectedDelegate;
+	FOnTextCommitted OnRenamed;
+	FOnVerifyTextChanged OnVerifyRename;
 
 	friend class FRigHierarchyTagDragDropOp;
 };
diff --git a/Engine/Plugins/Animation/ControlRig/Source/ControlRigEditor/Private/Editor/SRigHierarchyTreeView.cpp b/Engine/Plugins/Animation/ControlRig/Source/ControlRigEditor/Private/Editor/SRigHierarchyTreeView.cpp
--- a/Engine/Plugins/Animation/ControlRig/Source/ControlRigEditor/Private/Editor/SRigHierarchyTreeView.cpp
+++ b/Engine/Plugins/Animation/ControlRig/Source/ControlRigEditor/Private/Editor/SRigHierarchyTreeView.cpp
@@ -591,6 +591,18 @@ bool SRigHierarchyTreeView::AddElement(const FRigBaseElement* InElement)
 									Delegates.RequestDetailsInspection(ConnectorKey);
 								});
 
+								if (!ControlRig->IsModularRig())
+								{
+									TagArguments.OnRenamed_Lambda([ConnectorKey, this](const FText& InNewName, ETextCommit::Type InCommitType)
+									{
+										Delegates.HandleRenameElement(ConnectorKey, InNewName.ToString());
+									});
+									TagArguments.OnVerifyRename_Lambda([ConnectorKey, this](const FText& InText, FText& OutError)
+									{
+										return Delegates.HandleVerifyElementNameChanged(ConnectorKey, InText.ToString(), OutError);
+									});
+								}
+
 								TargetElementPtr->Get()->Tags.Add(TagArguments);
 
 								AddConnectorResolveWarningTag(*TargetElementPtr, InElement, Hierarchy);
