Control Rig: Add option to choose enum when spawning integer control or animation channel
#rnx

diff --git a/Engine/Plugins/Animation/ControlRig/Source/ControlRig/Private/Units/Execution/RigUnit_DynamicHierarchy.cpp b/Engine/Plugins/Animation/ControlRig/Source/ControlRig/Private/Units/Execution/RigUnit_DynamicHierarchy.cpp
--- a/Engine/Plugins/Animation/ControlRig/Source/ControlRig/Private/Units/Execution/RigUnit_DynamicHierarchy.cpp
+++ b/Engine/Plugins/Animation/ControlRig/Source/ControlRig/Private/Units/Execution/RigUnit_DynamicHierarchy.cpp
@@ -514,6 +514,7 @@ void FRigUnit_HierarchyAddControlInteger_Settings::ConfigureFrom(const FRigContr
 	FRigUnit_HierarchyAddControl_Settings::ConfigureFrom(InControlElement, InSettings);
 
 	PrimaryAxis = InSettings.PrimaryAxis;
+	ControlEnum = InSettings.ControlEnum;
 
 	Proxy.ConfigureFrom(InControlElement, InSettings);
 	Limits.ConfigureFrom(InControlElement, InSettings);
@@ -526,6 +527,7 @@ void FRigUnit_HierarchyAddControlInteger_Settings::Configure(FRigControlSettings
 	
 	OutSettings.ControlType = ERigControlType::Integer;
 	OutSettings.PrimaryAxis = PrimaryAxis;
+	OutSettings.ControlEnum = ControlEnum;
 
 	Proxy.Configure(OutSettings);
 	Limits.Configure(OutSettings);
@@ -1048,6 +1050,7 @@ FRigUnit_HierarchyAddAnimationChannelInteger_Execute()
 		ControlSettings.MinimumValue = FRigControlValue::Make<int32>(MinimumValue);
 		ControlSettings.MaximumValue = FRigControlValue::Make<int32>(MaximumValue);
 		ControlSettings.DisplayName = Controller->GetHierarchy()->GetSafeNewDisplayName(Parent, Name);
+		ControlSettings.ControlEnum = ControlEnum;
 		const FRigControlValue Value = FRigControlValue::Make<int32>(InitialValue);
 		
 		FRigHierarchyControllerInstructionBracket InstructionBracket(Controller, ExecuteContext.GetInstructionIndex());
diff --git a/Engine/Plugins/Animation/ControlRig/Source/ControlRig/Public/Rigs/RigHierarchyElements.h b/Engine/Plugins/Animation/ControlRig/Source/ControlRig/Public/Rigs/RigHierarchyElements.h
--- a/Engine/Plugins/Animation/ControlRig/Source/ControlRig/Public/Rigs/RigHierarchyElements.h
+++ b/Engine/Plugins/Animation/ControlRig/Source/ControlRig/Public/Rigs/RigHierarchyElements.h
@@ -1191,7 +1191,7 @@ struct CONTROLRIG_API FRigControlSettings
 	UPROPERTY(BlueprintReadWrite, Category = Control)
 	bool bIsTransientControl;
 
-	/** If the control is 4transient and only visible in the control rig editor */
+	/** If the control is integer it can use this enum to choose values */
 	UPROPERTY(BlueprintReadOnly, EditAnywhere, Category = Control)
 	TObjectPtr<UEnum> ControlEnum;
 
diff --git a/Engine/Plugins/Animation/ControlRig/Source/ControlRig/Public/Units/Execution/RigUnit_DynamicHierarchy.h b/Engine/Plugins/Animation/ControlRig/Source/ControlRig/Public/Units/Execution/RigUnit_DynamicHierarchy.h
--- a/Engine/Plugins/Animation/ControlRig/Source/ControlRig/Public/Units/Execution/RigUnit_DynamicHierarchy.h
+++ b/Engine/Plugins/Animation/ControlRig/Source/ControlRig/Public/Units/Execution/RigUnit_DynamicHierarchy.h
@@ -652,6 +652,9 @@ struct CONTROLRIG_API FRigUnit_HierarchyAddControlInteger_Settings : public FRig
 
 	UPROPERTY(EditAnywhere, BlueprintReadWrite, Category = "Settings")
 	ERigControlAxis PrimaryAxis;
+
+	UPROPERTY(EditAnywhere, BlueprintReadWrite, Category = "Settings")
+	TObjectPtr<UEnum> ControlEnum;
 	
 	UPROPERTY(EditAnywhere, BlueprintReadWrite, Category = "Settings")
 	FRigUnit_HierarchyAddControlInteger_LimitSettings Limits;;
@@ -1293,7 +1296,7 @@ struct CONTROLRIG_API FRigUnit_HierarchyAddAnimationChannelScaleFloat : public F
  * Adds a new animation channel to the hierarchy
  * Note: This node only runs as part of the construction event.
  */
-USTRUCT(meta=(DisplayName="Spawn Integer Animation Channel", TemplateName="SpawnAnimationChannel", Keywords="Construction,Create,New,AddAnimationChannel,NewAnimationChannel,CreateAnimationChannel,AddChannel,NewChannel,CreateChannel,SpawnChannel", Varying))
+USTRUCT(meta=(DisplayName="Spawn Integer Animation Channel", Keywords="Construction,Create,New,AddAnimationChannel,NewAnimationChannel,CreateAnimationChannel,AddChannel,NewChannel,CreateChannel,SpawnChannel", Varying))
 struct CONTROLRIG_API FRigUnit_HierarchyAddAnimationChannelInteger : public FRigUnit_HierarchyAddElement
 {
 	GENERATED_BODY()
@@ -1330,6 +1333,12 @@ struct CONTROLRIG_API FRigUnit_HierarchyAddAnimationChannelInteger : public FRig
 	UPROPERTY(meta = (Input))
 	FRigUnit_HierarchyAddAnimationChannelSingleLimitSettings LimitsEnabled;
 
+	/*
+	 * The enum to use to find valid values
+	 */
+	UPROPERTY(meta = (Input))
+	TObjectPtr<UEnum> ControlEnum;
+
 	RIGVM_METHOD()
 	virtual void Execute() override;
 };
