Control Rig: Updated Modular Rig Editor Layout

#rnx

diff --git a/Engine/Plugins/Animation/ControlRig/Source/ControlRigEditor/Private/Editor/ControlRigEditorMode.cpp b/Engine/Plugins/Animation/ControlRig/Source/ControlRigEditor/Private/Editor/ControlRigEditorMode.cpp
--- a/Engine/Plugins/Animation/ControlRig/Source/ControlRigEditor/Private/Editor/ControlRigEditorMode.cpp
+++ b/Engine/Plugins/Animation/ControlRig/Source/ControlRigEditor/Private/Editor/ControlRigEditorMode.cpp
@@ -165,63 +165,42 @@ FModularRigEditorMode::FModularRigEditorMode(const TSharedRef<FControlRigEditor>
 					->SetSizeCoefficient(0.2f)
 					->Split
 					(
-						//	Left top - viewport
+						// Left top - Module Library
 						FTabManager::NewStack()
+						->SetHideTabWell(false)
 						->SetSizeCoefficient(0.5f)
-						->SetHideTabWell(true)
-						->AddTab(FPersonaTabs::PreviewViewportID, ETabState::OpenedTab)
-					
+						->AddTab(FRigModuleAssetBrowserTabSummoner::TabID, ETabState::OpenedTab)
+						->SetForegroundTab(FBlueprintEditorTabs::DetailsID)
 					)
 					->Split
 					(
 						//	Left bottom - rig/hierarchy/modules
-						FTabManager::NewSplitter()
-						->SetOrientation(Orient_Horizontal)
+						FTabManager::NewStack()
+						->SetHideTabWell(false)
 						->SetSizeCoefficient(0.5f)
-						->Split
-						(
-							// Left bottom left
-							FTabManager::NewStack()
-							->SetSizeCoefficient(0.5f)
-							->SetForegroundTab(FRigHierarchyTabSummoner::TabID)
-							->AddTab(FBlueprintEditorTabs::CompilerResultsID, ETabState::ClosedTab)
-							->AddTab(FRigHierarchyTabSummoner::TabID, ETabState::OpenedTab)
-							->AddTab(FRigCurveContainerTabSummoner::TabID, ETabState::OpenedTab)
-							->AddTab(FBlueprintEditorTabs::MyBlueprintID, ETabState::OpenedTab)
-						)
-						->Split(
-							// Left bottom right
-							FTabManager::NewStack()
-							->SetSizeCoefficient(0.5f)
-							->AddTab(FModularRigModelTabSummoner::TabID, ETabState::OpenedTab)
-						)
+						->SetForegroundTab(FModularRigModelTabSummoner::TabID)
+						->AddTab(FModularRigModelTabSummoner::TabID, ETabState::OpenedTab)
+						->AddTab(FBlueprintEditorTabs::MyBlueprintID, ETabState::OpenedTab)
 					)
 				)
 				->Split
+				(
+					//	Center - viewport
+					FTabManager::NewStack()
+					->SetSizeCoefficient(0.6f)
+					->SetHideTabWell(true)
+					->AddTab(FPersonaTabs::PreviewViewportID, ETabState::OpenedTab)
+				)
+				->Split
 				(
 					// Right side
-					FTabManager::NewSplitter()
-					->SetOrientation(Orient_Vertical)
+					FTabManager::NewStack()
 					->SetSizeCoefficient(0.2f)
-					->Split
-					(
-						// Right top
-						FTabManager::NewStack()
-						->SetHideTabWell(false)
-						->SetSizeCoefficient(0.5f )
-						->AddTab(FBlueprintEditorTabs::DetailsID, ETabState::OpenedTab)
-						->AddTab(FPersonaTabs::AdvancedPreviewSceneSettingsID, ETabState::OpenedTab)
-						->AddTab(FRigAnimAttributeTabSummoner::TabID, ETabState::OpenedTab)
-						->SetForegroundTab(FBlueprintEditorTabs::DetailsID)
-					)
-					->Split(
-						// Right bottom
-						FTabManager::NewStack()
-						->SetHideTabWell(false)
-						->SetSizeCoefficient(0.5f)
-						->AddTab(FRigModuleAssetBrowserTabSummoner::TabID, ETabState::OpenedTab)
-						->SetForegroundTab(FBlueprintEditorTabs::DetailsID)
-					)
+					->SetHideTabWell(false)
+					->SetForegroundTab(FBlueprintEditorTabs::DetailsID)
+					->AddTab(FBlueprintEditorTabs::DetailsID, ETabState::OpenedTab)
+					->AddTab(FPersonaTabs::AdvancedPreviewSceneSettingsID, ETabState::OpenedTab)
+					->AddTab(FRigHierarchyTabSummoner::TabID, ETabState::OpenedTab)
 				)
 			)
 		);
diff --git a/Engine/Plugins/Animation/ControlRig/Source/ControlRigEditor/Private/Editor/RigModuleAssetBrowserTabSummoner.cpp b/Engine/Plugins/Animation/ControlRig/Source/ControlRigEditor/Private/Editor/RigModuleAssetBrowserTabSummoner.cpp
--- a/Engine/Plugins/Animation/ControlRig/Source/ControlRigEditor/Private/Editor/RigModuleAssetBrowserTabSummoner.cpp
+++ b/Engine/Plugins/Animation/ControlRig/Source/ControlRigEditor/Private/Editor/RigModuleAssetBrowserTabSummoner.cpp
@@ -16,11 +16,11 @@ FRigModuleAssetBrowserTabSummoner::FRigModuleAssetBrowserTabSummoner(const TShar
 {
 	bIsSingleton = true; // only allow a single instance of this tab
 	
-	TabLabel = LOCTEXT("RigModuleAssetBrowserTabLabel", "Rig Module Asset Browser");
+	TabLabel = LOCTEXT("RigModuleAssetBrowserTabLabel", "Module Assets");
 	TabIcon = FSlateIcon(FAppStyle::GetAppStyleSetName(), "ContentBrowser.TabIcon");
 
-	ViewMenuDescription = LOCTEXT("RigModuleAssetBrowser_ViewMenu_Desc", "Rig Module Asset Browser");
-	ViewMenuTooltip = LOCTEXT("RigModuleAssetBrowser_ViewMenu_ToolTip", "Show the Rig Module Asset Browser Tab");
+	ViewMenuDescription = LOCTEXT("RigModuleAssetBrowser_ViewMenu_Desc", "Module Assets");
+	ViewMenuTooltip = LOCTEXT("RigModuleAssetBrowser_ViewMenu_ToolTip", "Show the Module Assets Tab");
 }
 
 TSharedPtr<SToolTip> FRigModuleAssetBrowserTabSummoner::CreateTabToolTipWidget(const FWorkflowTabSpawnInfo& Info) const
diff --git a/Engine/Plugins/Animation/ControlRig/Source/ControlRigEditor/Private/Editor/SRigModuleAssetBrowser.cpp b/Engine/Plugins/Animation/ControlRig/Source/ControlRigEditor/Private/Editor/SRigModuleAssetBrowser.cpp
--- a/Engine/Plugins/Animation/ControlRig/Source/ControlRigEditor/Private/Editor/SRigModuleAssetBrowser.cpp
+++ b/Engine/Plugins/Animation/ControlRig/Source/ControlRigEditor/Private/Editor/SRigModuleAssetBrowser.cpp
@@ -51,6 +51,7 @@ void SRigModuleAssetBrowser::RefreshView()
 	AssetPickerConfig.bAllowRename = false;
 	AssetPickerConfig.bForceShowPluginContent = true;
 	AssetPickerConfig.bForceShowEngineContent = true;
+	AssetPickerConfig.InitialThumbnailSize = EThumbnailSize::Small;
 
 	// hide all asset registry columns by default (we only really want the name and path)
 	UObject* DefaultControlRigBlueprint = UControlRigBlueprint::StaticClass()->GetDefaultObject();
diff --git a/Engine/Source/Editor/ContentBrowser/Private/SAssetPicker.cpp b/Engine/Source/Editor/ContentBrowser/Private/SAssetPicker.cpp
--- a/Engine/Source/Editor/ContentBrowser/Private/SAssetPicker.cpp
+++ b/Engine/Source/Editor/ContentBrowser/Private/SAssetPicker.cpp
@@ -358,6 +358,7 @@ void SAssetPicker::Construct( const FArguments& InArgs )
 		.HiddenColumnNames(InArgs._AssetPickerConfig.HiddenColumnNames)
 		.CustomColumns(InArgs._AssetPickerConfig.CustomColumns)
 		.OnSearchOptionsChanged(this, &SAssetPicker::HandleSearchSettingsChanged)
+		.InitialThumbnailSize(InArgs._AssetPickerConfig.InitialThumbnailSize)
 	];
 
 
diff --git a/Engine/Source/Editor/ContentBrowser/Public/IContentBrowserSingleton.h b/Engine/Source/Editor/ContentBrowser/Public/IContentBrowserSingleton.h
--- a/Engine/Source/Editor/ContentBrowser/Public/IContentBrowserSingleton.h
+++ b/Engine/Source/Editor/ContentBrowser/Public/IContentBrowserSingleton.h
@@ -187,6 +187,9 @@ struct FAssetPickerConfig
 	/** The default scale for thumbnails. [0-1] range */
 	TAttribute< float > ThumbnailScale;
 
+	/** Initial thumbnail size */
+	EThumbnailSize InitialThumbnailSize;
+
 	/** Only display results in these collections */
 	TArray<FCollectionNameType> Collections;
 
@@ -313,6 +316,7 @@ struct FAssetPickerConfig
 		: SelectionMode( ESelectionMode::Multi )
 		, ThumbnailLabel( EThumbnailLabel::ClassName )
 		, ThumbnailScale(0.1f)
+		, InitialThumbnailSize(EThumbnailSize::Medium)
 		, InitialAssetViewType(EAssetViewType::Tile)
 		, bFocusSearchBoxWhenOpened(true)
 		, bAllowNullSelection(false)
