Schematic: Add context menu per node + control rig schema to allow disconnect through context menu

#rnx

diff --git a/Engine/Plugins/Animation/ControlRig/Source/ControlRigEditor/Private/Editor/ControlRigSchematicModel.cpp b/Engine/Plugins/Animation/ControlRig/Source/ControlRigEditor/Private/Editor/ControlRigSchematicModel.cpp
--- a/Engine/Plugins/Animation/ControlRig/Source/ControlRigEditor/Private/Editor/ControlRigSchematicModel.cpp
+++ b/Engine/Plugins/Animation/ControlRig/Source/ControlRigEditor/Private/Editor/ControlRigSchematicModel.cpp
@@ -1093,6 +1093,81 @@ bool FControlRigSchematicModel::GetForwardedNodeForDrag(FGuid& InOutGuid) const
 	return Super::GetForwardedNodeForDrag(InOutGuid);
 }
 
+bool FControlRigSchematicModel::GetContextMenuForNode(const FSchematicGraphNode* InNode, FMenuBuilder& OutMenu) const
+{
+	bool bSuccess = false;
+	if(Super::GetContextMenuForNode(InNode, OutMenu))
+	{
+		bSuccess = true;
+	}
+	
+	if(ControlRigBlueprint.IsValid())
+	{
+		if(const FControlRigSchematicRigElementKeyNode* ElementKeyNode = Cast<FControlRigSchematicRigElementKeyNode>(InNode))
+		{
+			const URigHierarchy* Hierarchy = ControlRigBlueprint->GetDebuggedControlRig()->GetHierarchy();
+			check(Hierarchy);
+			const FModularRigConnections& Connections = ControlRigBlueprint->ModularRigModel.Connections;
+			const TArray<FRigElementKey>& Connectors = Connections.FindConnectorsFromTarget(ElementKeyNode->Key);
+			if(!Connectors.IsEmpty())
+			{
+				OutMenu.BeginSection(TEXT("DisconnectConnectors"), LOCTEXT("DisconnectConnectors", "Disconnect"));
+
+				// note: this is a copy on purpose since it is passed into the lambda
+				for(const FRigElementKey Connector : Connectors)
+				{
+					FString Label = *Connector.Name.ToString();
+
+					const FName ModulePath = Hierarchy->GetModulePathFName(Connector);
+					if(!ModulePath.IsNone())
+					{
+						if(const FRigModuleReference* Module = ControlRigBlueprint->ModularRigModel.FindModule(ModulePath.ToString()))
+						{
+							const FName DesiredName = Hierarchy->GetNameMetadata(Connector, URigHierarchy::DesiredNameMetadataName, NAME_None);
+							if(!DesiredName.IsNone())
+							{
+								Label = DesiredName.ToString();
+
+								if(const FRigConnectorElement* PrimaryConnector = Module->FindPrimaryConnector(Hierarchy))
+								{
+									if(PrimaryConnector->GetKey() == Connector)
+									{
+										Label = Module->GetShortName();
+									}
+									else
+									{
+										Label = URigHierarchy::JoinNameSpace(Module->GetShortName(), Label);
+									}
+								}
+							}
+						}
+					}
+					
+					const FText Description = FText::FromString(FString::Printf(TEXT("Disconnect %s"), *Label));
+					
+					OutMenu.AddMenuEntry(Description, Description, FSlateIcon(), FUIAction(
+						FExecuteAction::CreateLambda([this, Connector]()
+						{
+							if(ControlRigBlueprint.IsValid())
+							{
+								if (UModularRigController* Controller = ControlRigBlueprint->GetModularRigController())
+								{
+									FScopedTransaction Transaction(LOCTEXT("DisconnectConnector", "Disconnect Connector"));
+									ControlRigBlueprint->Modify();
+									Controller->DisconnectConnector(Connector);
+								}
+							}								
+						})
+					));
+				}
+				OutMenu.EndSection();
+				bSuccess = true;
+			}
+		}
+	}
+	return bSuccess;
+}
+
 TArray<FRigElementKey> FControlRigSchematicModel::GetElementKeysFromDragDropEvent(const FDragDropOperation& InDragDropOperation, const UControlRig* InControlRig)
 {
 	TArray<FRigElementKey> DraggedKeys;
diff --git a/Engine/Plugins/Animation/ControlRig/Source/ControlRigEditor/Private/Editor/ControlRigSchematicModel.h b/Engine/Plugins/Animation/ControlRig/Source/ControlRigEditor/Private/Editor/ControlRigSchematicModel.h
--- a/Engine/Plugins/Animation/ControlRig/Source/ControlRigEditor/Private/Editor/ControlRigSchematicModel.h
+++ b/Engine/Plugins/Animation/ControlRig/Source/ControlRigEditor/Private/Editor/ControlRigSchematicModel.h
@@ -106,6 +106,7 @@ public:
 	virtual ESchematicGraphVisibility::Type GetVisibilityForTag(const FSchematicGraphTag* InTag) const override;
 	virtual const FText GetToolTipForTag(const FSchematicGraphTag* InTag) const override;
 	virtual bool GetForwardedNodeForDrag(FGuid& InOutGuid) const override;
+	virtual bool GetContextMenuForNode(const FSchematicGraphNode* InNode, FMenuBuilder& OutMenu) const override;
 
 	static TArray<FRigElementKey> GetElementKeysFromDragDropEvent(const FDragDropOperation& InDragDropOperation, const UControlRig* InControlRig);
 
@@ -122,7 +123,6 @@ private:
 	void HandleSchematicDrop(SSchematicGraphPanel* InPanel, SSchematicGraphNode* InNode, const FDragDropEvent& InDragDropEvent);
 	void HandlePostConstruction(UControlRig* Subject, const FName& InEventName);
 	bool IsConnectorResolved(const FRigElementKey& InConnectorKey, FRigElementKey* OutKey = nullptr) const;
-
 	void OnShowCandidatesForConnector(const FRigElementKey& InConnectorKey);
 	void OnShowCandidatesForConnector(const FRigModuleConnector* InModuleConnector);
 	void OnShowCandidatesForMatches(const FModularRigResolveResult& InMatches);
diff --git a/Engine/Source/Editor/AnimationEditorWidgets/Private/SchematicGraphPanel/SSchematicGraphPanel.cpp b/Engine/Source/Editor/AnimationEditorWidgets/Private/SchematicGraphPanel/SSchematicGraphPanel.cpp
--- a/Engine/Source/Editor/AnimationEditorWidgets/Private/SchematicGraphPanel/SSchematicGraphPanel.cpp
+++ b/Engine/Source/Editor/AnimationEditorWidgets/Private/SchematicGraphPanel/SSchematicGraphPanel.cpp
@@ -461,6 +461,27 @@ FReply SSchematicGraphNode::OnMouseButtonDown(const FGeometry& MyGeometry, const
 		OnClickedDelegate.ExecuteIfBound(this, MouseEvent);
 	}
 
+	if (MouseEvent.GetPressedButtons().Contains(EKeys::RightMouseButton))
+	{
+		if(SchematicGraphPanel && GetNodeData())
+		{
+			if(const FSchematicGraphModel* Graph = SchematicGraphPanel->GetSchematicGraph())
+			{
+				FMenuBuilder MenuBuilder(true, nullptr);
+				if(Graph->GetContextMenuForNode(GetNodeData(), MenuBuilder))
+				{
+					TSharedPtr<SWidget> MenuContent = MenuBuilder.MakeWidget();
+					if ( MenuContent.IsValid() )
+					{
+						FVector2D SummonLocation = MouseEvent.GetScreenSpacePosition();
+						FWidgetPath WidgetPath = MouseEvent.GetEventPath() != nullptr ? *MouseEvent.GetEventPath() : FWidgetPath();
+						FSlateApplication::Get().PushMenu(AsShared(), WidgetPath, MenuContent.ToSharedRef(), SummonLocation, FPopupTransitionEffect(FPopupTransitionEffect::ContextMenu));
+					}
+				}
+			}
+		}
+	}
+
 	const bool bFadedOut = IsFadedOut();
 	
 	if(NodeData)
diff --git a/Engine/Source/Editor/AnimationEditorWidgets/Private/SchematicGraphPanel/SchematicGraphModel.cpp b/Engine/Source/Editor/AnimationEditorWidgets/Private/SchematicGraphPanel/SchematicGraphModel.cpp
--- a/Engine/Source/Editor/AnimationEditorWidgets/Private/SchematicGraphPanel/SchematicGraphModel.cpp
+++ b/Engine/Source/Editor/AnimationEditorWidgets/Private/SchematicGraphPanel/SchematicGraphModel.cpp
@@ -439,6 +439,11 @@ bool FSchematicGraphModel::IsDragSupportedForNode(const FSchematicGraphNode* InN
 	return InNode->IsDragSupported();
 }
 
+bool FSchematicGraphModel::GetContextMenuForNode(const FSchematicGraphNode* InNode, FMenuBuilder& OutMenu) const
+{
+	return false;
+}
+
 FLinearColor FSchematicGraphModel::GetBackgroundColorForTag(const FGuid& InNodeGuid, const FGuid& InTagGuid) const
 {
 	if(const FSchematicGraphNode* Node = FindNode(InNodeGuid))
diff --git a/Engine/Source/Editor/AnimationEditorWidgets/Public/SchematicGraphPanel/SchematicGraphModel.h b/Engine/Source/Editor/AnimationEditorWidgets/Public/SchematicGraphPanel/SchematicGraphModel.h
--- a/Engine/Source/Editor/AnimationEditorWidgets/Public/SchematicGraphPanel/SchematicGraphModel.h
+++ b/Engine/Source/Editor/AnimationEditorWidgets/Public/SchematicGraphPanel/SchematicGraphModel.h
@@ -6,6 +6,7 @@
 #include "SchematicGraphTag.h"
 #include "SchematicGraphNode.h"
 #include "SchematicGraphLink.h"
+#include "Framework/MultiBox/MultiBoxBuilder.h"
 
 class SSchematicGraphPanel;
 
@@ -113,6 +114,7 @@ public:
 	virtual TOptional<bool> GetInteractivityForChildNode(const FSchematicGraphNode* InParentNode, const FSchematicGraphNode* InChildNode) const;
 	bool IsDragSupportedForNode(const FGuid& InNodeGuid) const;
 	virtual bool IsDragSupportedForNode(const FSchematicGraphNode* InNode) const;
+	virtual bool GetContextMenuForNode(const FSchematicGraphNode* InNode, FMenuBuilder& OutMenu) const;
 
 	FLinearColor GetBackgroundColorForTag(const FGuid& InNodeGuid, const FGuid& InTagGuid) const;
 	virtual FLinearColor GetBackgroundColorForTag(const FSchematicGraphTag* InTag) const;
