New Anim Details. Main code is the new objects in AnimDetailsProxy that handle support for control rig/sequencer, having multiple objects, and shared selection state with sequencer. CR doesn't support templates so there is more copypasta thn I would like but will be cleaning this up more over the weekend and over the next few weeks. ControlRigControlsProxy now just contains the base class and the main detail object. ControlRigProxyCustomizations holds most of the cusmtomizations.
Some stuff like the key proprerty was put in the control rig editor since they own the details/panel and the proxies and it avoid us have too many duplicate set of delegates avoiding dangling delegate issues.

diff --git a/Engine/Plugins/Animation/ControlRig/Source/ControlRig/Private/Sequencer/ControlRigSequencerHelpers.cpp b/Engine/Plugins/Animation/ControlRig/Source/ControlRig/Private/Sequencer/ControlRigSequencerHelpers.cpp
--- a/Engine/Plugins/Animation/ControlRig/Source/ControlRig/Private/Sequencer/ControlRigSequencerHelpers.cpp
+++ b/Engine/Plugins/Animation/ControlRig/Source/ControlRig/Private/Sequencer/ControlRigSequencerHelpers.cpp
@@ -3,7 +3,12 @@
 #include "ControlRig.h"
 #include "Rigs/RigHierarchyElements.h"
 #include "Channels/MovieSceneChannelProxy.h"
-#include "Sequencer/MovieSceneControlRigParameterSection.h"
+#include "Sequencer/MovieSceneControlRigParameterTrack.h"
+#include "Channels/MovieSceneIntegerChannel.h"
+#include "Channels/MovieSceneByteChannel.h"
+#include "Channels/MovieSceneBoolChannel.h"
+#include "MovieSceneSequence.h"
+#include "MovieScene.h"
 
 TPair<const FChannelMapInfo*, int32> FControlRigSequencerHelpers::GetInfoAndNumFloatChannels(
 	const UControlRig* InControlRig,
@@ -24,6 +29,11 @@ TPair<const FChannelMapInfo*, int32> FControlRigSequencerHelpers::GetInfoAndNumF
 		case ERigControlType::Transform:
 		case ERigControlType::EulerTransform:
 			return 9;
+		case ERigControlType::Vector2D:
+			return 2;
+		case ERigControlType::Float:
+		case ERigControlType::ScaleFloat:
+			return 1;
 		default:
 			break;
 		}
@@ -64,3 +74,118 @@ TArrayView<FMovieSceneFloatChannel*>  FControlRigSequencerHelpers::GetFloatChann
 	return FloatChannels.Slice(ChannelStartIndex, NumChannels);
 }
 
+
+TArrayView<FMovieSceneBoolChannel*> FControlRigSequencerHelpers::GetBoolChannels(const UControlRig* InControlRig,
+	const FName& InControlName, const UMovieSceneSection* InSection)
+{
+	static const TArrayView<FMovieSceneBoolChannel*> EmptyChannelsView;
+	const UMovieSceneControlRigParameterSection* CRSection = Cast<UMovieSceneControlRigParameterSection>(InSection);
+	if (CRSection == nullptr)
+	{
+		return EmptyChannelsView;
+	}
+	const FRigControlElement* ControlElement = InControlRig ? InControlRig->FindControl(InControlName) : nullptr;
+	if (ControlElement == nullptr)
+	{
+		return EmptyChannelsView;
+	}
+	if (ControlElement->Settings.ControlType == ERigControlType::Bool)
+	{
+		int32 NumChannels = 1;
+		const FChannelMapInfo* ChannelInfo = CRSection ? CRSection->ControlChannelMap.Find(InControlName) : nullptr;
+		if (ChannelInfo == nullptr)
+		{
+			return EmptyChannelsView;
+		}
+		const TArrayView<FMovieSceneBoolChannel*> Channels = CRSection->GetChannelProxy().GetChannels<FMovieSceneBoolChannel>();
+		const int32 ChannelStartIndex = ChannelInfo->ChannelIndex;
+		return Channels.Slice(ChannelStartIndex, NumChannels);
+
+	}
+	return EmptyChannelsView;
+}
+
+
+
+TArrayView<FMovieSceneByteChannel*> FControlRigSequencerHelpers::GetByteChannels(const UControlRig* InControlRig,
+	const FName& InControlName, const UMovieSceneSection* InSection)
+{
+	static const TArrayView<FMovieSceneByteChannel*> EmptyChannelsView;
+	const UMovieSceneControlRigParameterSection* CRSection = Cast<UMovieSceneControlRigParameterSection>(InSection);
+	if (CRSection == nullptr)
+	{
+		return EmptyChannelsView;
+	}
+	const FRigControlElement* ControlElement = InControlRig ? InControlRig->FindControl(InControlName) : nullptr;
+	if (ControlElement == nullptr)
+	{
+		return EmptyChannelsView;
+	}
+	if (ControlElement->Settings.ControlType == ERigControlType::Integer && ControlElement->Settings.ControlEnum)
+	{
+		int32 NumChannels = 1;
+		const FChannelMapInfo* ChannelInfo = CRSection ? CRSection->ControlChannelMap.Find(InControlName) : nullptr;
+		if (ChannelInfo == nullptr)
+		{
+			return EmptyChannelsView;
+		}
+		const TArrayView<FMovieSceneByteChannel*> Channels = CRSection->GetChannelProxy().GetChannels<FMovieSceneByteChannel>();
+		const int32 ChannelStartIndex = ChannelInfo->ChannelIndex;
+		return Channels.Slice(ChannelStartIndex, NumChannels);
+
+	}
+	return EmptyChannelsView;
+}
+
+TArrayView<FMovieSceneIntegerChannel*> FControlRigSequencerHelpers::GetIntegerChannels(const UControlRig* InControlRig,
+	const FName& InControlName, const UMovieSceneSection* InSection)
+{
+	static const TArrayView<FMovieSceneIntegerChannel*> EmptyChannelsView;
+	const UMovieSceneControlRigParameterSection* CRSection = Cast<UMovieSceneControlRigParameterSection>(InSection);
+	if (CRSection == nullptr)
+	{
+		return EmptyChannelsView;
+	}
+	const FRigControlElement* ControlElement = InControlRig ? InControlRig->FindControl(InControlName) : nullptr;
+	if (ControlElement == nullptr)
+	{
+		return EmptyChannelsView;
+	}
+	if (ControlElement->Settings.ControlType == ERigControlType::Integer && ControlElement->Settings.ControlEnum == nullptr)
+	{
+		int32 NumChannels = 1;
+		const FChannelMapInfo* ChannelInfo = CRSection ? CRSection->ControlChannelMap.Find(InControlName) : nullptr;
+		if (ChannelInfo == nullptr)
+		{
+			return EmptyChannelsView;
+		}
+		const TArrayView<FMovieSceneIntegerChannel*> Channels = CRSection->GetChannelProxy().GetChannels<FMovieSceneIntegerChannel>();
+		const int32 ChannelStartIndex = ChannelInfo->ChannelIndex;
+		return Channels.Slice(ChannelStartIndex, NumChannels);
+
+	}
+	return EmptyChannelsView;
+}
+
+UMovieSceneControlRigParameterTrack* FControlRigSequencerHelpers::FindControlRigTrack(UMovieSceneSequence* MovieSceneSequence, const UControlRig* ControlRig)
+{
+	UMovieSceneControlRigParameterTrack* Track = nullptr;
+	if (MovieSceneSequence == nullptr)
+	{
+		return Track;
+	}
+	UMovieScene* MovieScene = MovieSceneSequence->GetMovieScene();
+	const TArray<FMovieSceneBinding>& Bindings = MovieScene->GetBindings();
+	bool bRecreateCurves = false;
+	TArray<TPair<UControlRig*, FName>> ControlRigPairsToReselect;
+	for (const FMovieSceneBinding& Binding : Bindings)
+	{
+		UMovieSceneControlRigParameterTrack* ControlRigParameterTrack = Cast<UMovieSceneControlRigParameterTrack>(MovieScene->FindTrack(UMovieSceneControlRigParameterTrack::StaticClass(), Binding.GetObjectGuid(), NAME_None));
+		if (ControlRigParameterTrack && ControlRigParameterTrack->GetControlRig() == ControlRig)
+		{
+			Track = ControlRigParameterTrack;
+			break;
+		}
+	}
+	return Track;
+}
diff --git a/Engine/Plugins/Animation/ControlRig/Source/ControlRig/Public/Sequencer/ControlRigSequencerHelpers.h b/Engine/Plugins/Animation/ControlRig/Source/ControlRig/Public/Sequencer/ControlRigSequencerHelpers.h
--- a/Engine/Plugins/Animation/ControlRig/Source/ControlRig/Public/Sequencer/ControlRigSequencerHelpers.h
+++ b/Engine/Plugins/Animation/ControlRig/Source/ControlRig/Public/Sequencer/ControlRigSequencerHelpers.h
@@ -8,8 +8,12 @@ struct FChannelMapInfo;
 class UControlRig;
 class UMovieSceneControlRigParameterSection;
 struct FMovieSceneFloatChannel;
+struct FMovieSceneBoolChannel;
+struct FMovieSceneIntegerChannel;
+struct FMovieSceneByteChannel;
 class UMovieSceneSection;
-class ISequencer;
+class UMovieSceneControlRigParameterTrack;
+class UMovieSceneSequence;
 
 struct CONTROLRIG_API FControlRigSequencerHelpers
 {
@@ -21,7 +25,16 @@ struct CONTROLRIG_API FControlRigSequencerHelpers
 	static TArrayView<FMovieSceneFloatChannel*> GetFloatChannels(const UControlRig* InControlRig,
 		const FName& InControlName, const UMovieSceneSection* InSection);
 
-	static UMovieSceneControlRigParameterSection* GetControlRigSection(ISequencer* InSequencer, const UControlRig* InControlRig);
+	static TArrayView<FMovieSceneBoolChannel*> GetBoolChannels(const UControlRig* InControlRig,
+		const FName& InControlName, const UMovieSceneSection* InSection);
+
+	static TArrayView<FMovieSceneByteChannel*> GetByteChannels(const UControlRig* InControlRig,
+		const FName& InControlName, const UMovieSceneSection* InSection);
+
+	static TArrayView<FMovieSceneIntegerChannel*> GetIntegerChannels(const UControlRig* InControlRig,
+		const FName& InControlName, const UMovieSceneSection* InSection);
+
+	static UMovieSceneControlRigParameterTrack*  FindControlRigTrack(UMovieSceneSequence* InSequencer, const UControlRig* InControlRig);
 
 };
 
diff --git a/Engine/Plugins/Animation/ControlRig/Source/ControlRigEditor/ControlRigEditor.Build.cs b/Engine/Plugins/Animation/ControlRig/Source/ControlRigEditor/ControlRigEditor.Build.cs
--- a/Engine/Plugins/Animation/ControlRig/Source/ControlRigEditor/ControlRigEditor.Build.cs
+++ b/Engine/Plugins/Animation/ControlRig/Source/ControlRigEditor/ControlRigEditor.Build.cs
@@ -22,10 +22,11 @@ namespace UnrealBuildTool.Rules
                     "Core",
                     "CoreUObject",
 					"CurveEditor",
+					"DetailCustomizations",
 					"Slate",
                     "SlateCore",
                     "InputCore",
-                    "Engine",
+					"Engine",
 					"EditorFramework",
 					"UnrealEd",
                     "KismetCompiler",
diff --git a/Engine/Plugins/Animation/ControlRig/Source/ControlRigEditor/Private/ControlRigEditorModule.cpp b/Engine/Plugins/Animation/ControlRig/Source/ControlRigEditor/Private/ControlRigEditorModule.cpp
--- a/Engine/Plugins/Animation/ControlRig/Source/ControlRigEditor/Private/ControlRigEditorModule.cpp
+++ b/Engine/Plugins/Animation/ControlRig/Source/ControlRigEditor/Private/ControlRigEditorModule.cpp
@@ -21,6 +21,7 @@
 #include "EditorModeManager.h"
 #include "Sequencer/MovieSceneControlRigParameterSection.h"
 #include "Sequencer/MovieSceneControlRigSectionDetailsCustomization.h"
+#include "EditMode/ControlsProxyDetailCustomization.h"
 #include "EditMode/ControlRigEditModeCommands.h"
 #include "Framework/MultiBox/MultiBoxBuilder.h"
 #include "ToolMenus.h"
@@ -129,12 +130,112 @@ void FControlRigEditorModule::StartupModule()
 
 	ClassesToUnregisterOnShutdown.Add(UControlRig::StaticClass()->GetFName());
 
+	// proxies
+	const FName PropertyEditorModuleName("PropertyEditor");
+	FPropertyEditorModule& PropertyModule = FModuleManager::GetModuleChecked<FPropertyEditorModule>(PropertyEditorModuleName);
+
+	ClassesToUnregisterOnShutdown.Add(UAnimDetailControlsProxyFloat::StaticClass()->GetFName());
+	PropertyEditorModule.RegisterCustomClassLayout(UAnimDetailControlsProxyFloat::StaticClass()->GetFName(), FOnGetDetailCustomizationInstance::CreateLambda([]() {
+		FName CategoryName = FName(TEXT("Float"));
+		return FAnimDetailProxyDetails::MakeInstance(CategoryName);
+		}));
+
+	{	
+		TSharedRef<FPropertySection> Section = RegisterPropertySection(PropertyModule, "AnimDetailControlsProxyFloat", "General", LOCTEXT("General", "General"));
+		Section->AddCategory("Attributes");
+	}
+
+	ClassesToUnregisterOnShutdown.Add(UAnimDetailControlsProxyBool::StaticClass()->GetFName());
+	PropertyEditorModule.RegisterCustomClassLayout(UAnimDetailControlsProxyBool::StaticClass()->GetFName(), FOnGetDetailCustomizationInstance::CreateLambda([]() {
+		FName CategoryName = FName(TEXT("Bool"));
+		return FAnimDetailProxyDetails::MakeInstance(CategoryName);
+		}));
+
+	{
+		TSharedRef<FPropertySection> Section = RegisterPropertySection(PropertyModule, "AnimDetailControlsProxyBool", "General", LOCTEXT("General", "General"));
+		Section->AddCategory("Attributes");
+	}
+
+	ClassesToUnregisterOnShutdown.Add(UAnimDetailControlsProxyTransform::StaticClass()->GetFName());
+	PropertyEditorModule.RegisterCustomClassLayout(UAnimDetailControlsProxyTransform::StaticClass()->GetFName(), FOnGetDetailCustomizationInstance::CreateLambda([]() {
+		FName CategoryName = FName(TEXT("Transform"));
+		return FAnimDetailProxyDetails::MakeInstance(CategoryName);
+		}));
+
+	{
+		TSharedRef<FPropertySection> Section = RegisterPropertySection(PropertyModule, "AnimDetailControlsProxyTransform", "General", LOCTEXT("General", "General"));
+		Section->AddCategory("Attributes");
+	}
+
+	ClassesToUnregisterOnShutdown.Add(UAnimDetailControlsProxyLocation::StaticClass()->GetFName());
+	PropertyEditorModule.RegisterCustomClassLayout(UAnimDetailControlsProxyLocation::StaticClass()->GetFName(), FOnGetDetailCustomizationInstance::CreateLambda([]() {
+		FName CategoryName = FName(TEXT("Location"));
+		return FAnimDetailProxyDetails::MakeInstance(CategoryName);
+		}));
+
+	{
+		TSharedRef<FPropertySection> Section = RegisterPropertySection(PropertyModule, "AnimDetailControlsProxyLocation", "General", LOCTEXT("General", "General"));
+		Section->AddCategory("Attributes");
+	}
+
+	ClassesToUnregisterOnShutdown.Add(UAnimDetailControlsProxyRotation::StaticClass()->GetFName());
+	PropertyEditorModule.RegisterCustomClassLayout(UAnimDetailControlsProxyRotation::StaticClass()->GetFName(), FOnGetDetailCustomizationInstance::CreateLambda([]() {
+		FName CategoryName = FName(TEXT("Rotation"));
+		return FAnimDetailProxyDetails::MakeInstance(CategoryName);
+		}));
+
+	{
+		TSharedRef<FPropertySection> Section = RegisterPropertySection(PropertyModule, "AnimDetailControlsProxyRotation", "General", LOCTEXT("General", "General"));
+		Section->AddCategory("Attributes");
+	}
+
+	ClassesToUnregisterOnShutdown.Add(UAnimDetailControlsProxyScale::StaticClass()->GetFName());
+	PropertyEditorModule.RegisterCustomClassLayout(UAnimDetailControlsProxyScale::StaticClass()->GetFName(), FOnGetDetailCustomizationInstance::CreateLambda([]() {
+		FName CategoryName = FName(TEXT("Scale"));
+		return FAnimDetailProxyDetails::MakeInstance(CategoryName);
+		}));
+	{
+		TSharedRef<FPropertySection> Section = RegisterPropertySection(PropertyModule, "AnimDetailControlsProxyScale", "General", LOCTEXT("General", "General"));
+		Section->AddCategory("Attributes");
+	}
+
+	ClassesToUnregisterOnShutdown.Add(UAnimDetailControlsProxyVector2D::StaticClass()->GetFName());
+	PropertyEditorModule.RegisterCustomClassLayout(UAnimDetailControlsProxyVector2D::StaticClass()->GetFName(), FOnGetDetailCustomizationInstance::CreateLambda([]() {
+		FName CategoryName = FName(TEXT("Vector2D"));
+		return FAnimDetailProxyDetails::MakeInstance(CategoryName);
+		}));
+
+	{
+		TSharedRef<FPropertySection> Section = RegisterPropertySection(PropertyModule, "AnimDetailControlsProxyVector2D", "General", LOCTEXT("General", "General"));
+		Section->AddCategory("Attributes");
+	}
+	
 	// same as ClassesToUnregisterOnShutdown but for properties, there is none right now
 	PropertiesToUnregisterOnShutdown.Reset();
-
 	PropertiesToUnregisterOnShutdown.Add(FRigVMCompileSettings::StaticStruct()->GetFName());
 	PropertyEditorModule.RegisterCustomPropertyTypeLayout(PropertiesToUnregisterOnShutdown.Last(), FOnGetPropertyTypeCustomizationInstance::CreateStatic(&FRigVMCompileSettingsDetails::MakeInstance));
 
+	PropertiesToUnregisterOnShutdown.Add(FAnimDetailProxyFloat::StaticStruct()->GetFName());
+	PropertyEditorModule.RegisterCustomPropertyTypeLayout(PropertiesToUnregisterOnShutdown.Last(), FOnGetPropertyTypeCustomizationInstance::CreateStatic(&FAnimDetailValueCustomization::MakeInstance));
+
+	PropertiesToUnregisterOnShutdown.Add(FAnimDetailProxyInteger::StaticStruct()->GetFName());
+	PropertyEditorModule.RegisterCustomPropertyTypeLayout(PropertiesToUnregisterOnShutdown.Last(), FOnGetPropertyTypeCustomizationInstance::CreateStatic(&FAnimDetailValueCustomization::MakeInstance));
+
+	PropertiesToUnregisterOnShutdown.Add(FAnimDetailProxyBool::StaticStruct()->GetFName());
+	PropertyEditorModule.RegisterCustomPropertyTypeLayout(PropertiesToUnregisterOnShutdown.Last(), FOnGetPropertyTypeCustomizationInstance::CreateStatic(&FAnimDetailValueCustomization::MakeInstance));
+
+	PropertiesToUnregisterOnShutdown.Add(FAnimDetailProxyVector2D::StaticStruct()->GetFName());
+	PropertyEditorModule.RegisterCustomPropertyTypeLayout(PropertiesToUnregisterOnShutdown.Last(), FOnGetPropertyTypeCustomizationInstance::CreateStatic(&FAnimDetailValueCustomization::MakeInstance));
+
+	PropertiesToUnregisterOnShutdown.Add(FAnimDetailProxyLocation::StaticStruct()->GetFName());
+	PropertyEditorModule.RegisterCustomPropertyTypeLayout(PropertiesToUnregisterOnShutdown.Last(), FOnGetPropertyTypeCustomizationInstance::CreateStatic(&FAnimDetailValueCustomization::MakeInstance));
+
+	PropertiesToUnregisterOnShutdown.Add(FAnimDetailProxyRotation::StaticStruct()->GetFName());
+	PropertyEditorModule.RegisterCustomPropertyTypeLayout(PropertiesToUnregisterOnShutdown.Last(), FOnGetPropertyTypeCustomizationInstance::CreateStatic(&FAnimDetailValueCustomization::MakeInstance));
+
+	PropertiesToUnregisterOnShutdown.Add(FAnimDetailProxyScale::StaticStruct()->GetFName());
+	PropertyEditorModule.RegisterCustomPropertyTypeLayout(PropertiesToUnregisterOnShutdown.Last(), FOnGetPropertyTypeCustomizationInstance::CreateStatic(&FAnimDetailValueCustomization::MakeInstance));
+	
 	PropertiesToUnregisterOnShutdown.Add(FRigVMPythonSettings::StaticStruct()->GetFName());
 	PropertyEditorModule.RegisterCustomPropertyTypeLayout(PropertiesToUnregisterOnShutdown.Last(), FOnGetPropertyTypeCustomizationInstance::CreateStatic(&FControlRigPythonLogDetails::MakeInstance));
 
@@ -292,6 +393,8 @@ void FControlRigEditorModule::ShutdownModule()
 		}
 	}
 	WorkflowHandles.Reset();
+
+	UnregisterPropertySectionMappings();
 }
 
 UClass* FControlRigEditorModule::GetRigVMBlueprintClass() const
@@ -299,6 +402,33 @@ UClass* FControlRigEditorModule::GetRigVMBlueprintClass() const
 	return UControlRigBlueprint::StaticClass();
 }
 
+TSharedRef<FPropertySection> FControlRigEditorModule::RegisterPropertySection(FPropertyEditorModule& PropertyModule, FName ClassName, FName SectionName, FText DisplayName)
+{
+	TSharedRef<FPropertySection> PropertySection = PropertyModule.FindOrCreateSection(ClassName, SectionName, DisplayName);
+	RegisteredPropertySections.Add(ClassName, SectionName);
+
+	return PropertySection;
+}
+
+void FControlRigEditorModule::UnregisterPropertySectionMappings()
+{
+	const FName PropertyEditorModuleName("PropertyEditor");
+	FPropertyEditorModule* PropertyModule = FModuleManager::GetModulePtr<FPropertyEditorModule>(PropertyEditorModuleName);
+
+	if (!PropertyModule)
+	{
+		return;
+	}
+
+	for (TMultiMap<FName, FName>::TIterator PropertySectionIterator = RegisteredPropertySections.CreateIterator(); PropertySectionIterator; ++PropertySectionIterator)
+	{
+		PropertyModule->RemoveSection(PropertySectionIterator->Key, PropertySectionIterator->Value);
+		PropertySectionIterator.RemoveCurrent();
+	}
+
+	RegisteredPropertySections.Empty();
+}
+
 void FControlRigEditorModule::GetNodeContextMenuActions(IRigVMClientHost* RigVMClientHost,
 	const URigVMEdGraphNode* EdGraphNode, URigVMNode* ModelNode, UToolMenu* Menu) const
 {
diff --git a/Engine/Plugins/Animation/ControlRig/Source/ControlRigEditor/Private/ControlRigEditorModule.h b/Engine/Plugins/Animation/ControlRig/Source/ControlRigEditor/Private/ControlRigEditorModule.h
--- a/Engine/Plugins/Animation/ControlRig/Source/ControlRigEditor/Private/ControlRigEditorModule.h
+++ b/Engine/Plugins/Animation/ControlRig/Source/ControlRigEditor/Private/ControlRigEditorModule.h
@@ -21,6 +21,8 @@ class FUICommandList;
 class UMovieSceneTrack;
 class FRigVMEdGraphPanelNodeFactory;
 class FControlRigGraphPanelPinFactory;
+class FPropertySection;
+class FPropertyEditorModule;
 
 class FControlRigEditorModule : public IControlRigEditorModule
 {
@@ -51,6 +53,11 @@ public:
 	void GetDirectManipulationMenuActions(IRigVMClientHost* RigVMClientHost, URigVMNode* InNode, URigVMPin* ModelPin, UToolMenu* Menu) const;
 
 private:
+	//property sections
+	TSharedRef<FPropertySection> RegisterPropertySection(FPropertyEditorModule& PropertyModule, FName ClassName, FName SectionName, FText DisplayName);
+	void UnregisterPropertySectionMappings();
+	TMultiMap<FName, FName> RegisteredPropertySections;
+
 
 	/** Handle for our sequencer control rig parameter track editor */
 	FDelegateHandle ControlRigParameterTrackCreateEditorHandle;
