New Anim Details. Main code is the new objects in AnimDetailsProxy that handle support for control rig/sequencer, having multiple objects, and shared selection state with sequencer. CR doesn't support

diff --git a/Engine/Plugins/Animation/ControlRig/Source/ControlRigEditor/Private/EditMode/AnimDetailsProxy.cpp b/Engine/Plugins/Animation/ControlRig/Source/ControlRigEditor/Private/EditMode/AnimDetailsProxy.cpp
new file mode 100644
--- /dev/null
+++ b/Engine/Plugins/Animation/ControlRig/Source/ControlRigEditor/Private/EditMode/AnimDetailsProxy.cpp
@@ -0,0 +1,2962 @@
+﻿// Copyright Epic Games, Inc. All Rights Reserved.
+
+
+#include "AnimDetailsProxy.h"
+#include "EditorModeManager.h"
+#include "EditMode/ControlRigEditMode.h"
+#include "ControlRig.h"
+#include "Rigs/RigHierarchy.h"
+#include "Sequencer/MovieSceneControlRigParameterSection.h"
+#include "Sequencer/MovieSceneControlRigParameterTrack.h"
+#include "Components/SkeletalMeshComponent.h"
+#include "Subsystems/AssetEditorSubsystem.h"
+#include "MovieSceneCommonHelpers.h"
+#include "PropertyHandle.h"
+#include "Sequencer/ControlRigSequencerHelpers.h"
+#include "ISequencer.h"
+#include "CurveEditor.h"
+#include "CurveModel.h"
+#include "MovieScene.h"
+#include "Tracks/MovieScenePropertyTrack.h"
+#include "MovieSceneCommonHelpers.h"
+#include "UnrealEdGlobals.h"
+#include "UnrealEdMisc.h"
+#include "Editor/UnrealEdEngine.h"
+#include "ScopedTransaction.h"
+#include "Channels/MovieSceneBoolChannel.h"
+#include "Channels/MovieSceneChannelProxy.h"
+#include "Channels/MovieSceneDoubleChannel.h"
+#include "Channels/MovieSceneIntegerChannel.h"
+#include "Channels/MovieSceneByteChannel.h"
+#include "DetailLayoutBuilder.h"
+#include "DetailCategoryBuilder.h"
+#include "DetailWidgetRow.h"
+#include "SEnumCombo.h"
+#include "MVVM/ViewModels/ChannelModel.h"
+#include "MVVM/ViewModels/SectionModel.h"
+#include "MVVM/ViewModels/SequencerEditorViewModel.h"
+#include "MVVM/SectionModelStorageExtension.h"
+#include "MVVM/ViewModels/SequenceModel.h"
+#include "MVVM/ViewModels/TrackAreaViewModel.h"
+#include "MVVM/ViewModels/TrackModel.h"
+#include "IKeyArea.h"
+#include "SequencerAddKeyOperation.h"
+
+#include UE_INLINE_GENERATED_CPP_BY_NAME(AnimDetailsProxy)
+
+
+static void KeyTrack(TSharedPtr<ISequencer>& Sequencer, UControlRigControlsProxy* Proxy, UMovieScenePropertyTrack* Track, EControlRigContextChannelToKey ChannelToKey)
+{
+	using namespace UE::Sequencer;
+
+	const FFrameNumber Time = Sequencer->GetLocalTime().Time.FloorToFrame();
+
+	float Weight = 0.0;
+	UMovieSceneSection* Section = Track->FindOrExtendSection(Time, Weight);
+	FScopedTransaction PropertyChangedTransaction(NSLOCTEXT("ControlRig", "KeyProperty", "Key Property"), !GIsTransacting);
+	if (!Section || !Section->TryModify())
+	{
+		PropertyChangedTransaction.Cancel();
+		return;
+	}
+
+	FSectionModelStorageExtension* SectionModelStorage = Sequencer->GetViewModel()->GetRootModel()->CastDynamic<FSectionModelStorageExtension>();
+	check(SectionModelStorage);
+
+	TSharedPtr<FSectionModel> SectionHandle = SectionModelStorage->FindModelForSection(Section);
+	if (SectionHandle)
+	{
+		TArray<TSharedRef<IKeyArea>> KeyAreas;
+		TParentFirstChildIterator<FChannelGroupModel> KeyAreaNodes = SectionHandle->GetParentTrackModel().AsModel()->GetDescendantsOfType<FChannelGroupModel>();
+		for (const TViewModelPtr<FChannelGroupModel>& KeyAreaNode : KeyAreaNodes)
+		{
+			for (const TWeakViewModelPtr<FChannelModel>& Channel : KeyAreaNode->GetChannels())
+			{
+				if (TSharedPtr< FChannelModel> ChannelPtr = Channel.Pin())
+				{
+					EControlRigContextChannelToKey ThisChannelToKey = Proxy->GetChannelToKeyFromChannelName(ChannelPtr->GetChannelName().ToString());
+					if ((int32)ChannelToKey & (int32)ThisChannelToKey)
+					{
+						KeyAreas.Add(ChannelPtr->GetKeyArea().ToSharedRef());
+					}
+				}
+			}
+		}	
+		TSharedPtr<FTrackModel> TrackModel = SectionHandle->FindAncestorOfType<FTrackModel>();
+		FAddKeyOperation::FromKeyAreas(TrackModel->GetTrackEditor().Get(), KeyAreas).Commit(Time, *Sequencer);
+	}				
+}
+
+//////UAnimDetailControlsKeyedProxy////////
+
+static EAnimDetailSelectionState CachePropertySelection(TWeakPtr<FCurveEditor>& InCurveEditor, UControlRigControlsProxy* Proxy, const FName& PropertyName)
+{
+	int32 TotalSelected = 0;
+	int32 TotalNum = 0; //need to add proxies + items that are valid
+	if (InCurveEditor.IsValid())
+	{
+		TSharedPtr<FCurveEditor> CurveEditor = InCurveEditor.Pin();
+		const TMap<FCurveModelID, TUniquePtr<FCurveModel>>& Curves = CurveEditor->GetCurves();
+		if (Proxy)
+		{
+			for (const TPair<UControlRig*, FControlRigProxyItem>& Items : Proxy->ControlRigItems)
+			{
+				if (UControlRig* ControlRig = Items.Value.ControlRig.Get())
+				{
+					for (FRigControlElement* ControlElement : Items.Value.ControlElements)
+					{
+						if (ControlElement)
+						{
+							++TotalNum;
+							EControlRigContextChannelToKey ChannelToKey = Proxy->GetChannelToKeyFromPropertyName(PropertyName);
+							for (const TPair<FCurveModelID, TUniquePtr<FCurveModel>>& Pair : Curves) //horribly slow
+							{
+								FCurveEditorTreeItemID TreeItemID = CurveEditor->GetTreeIDFromCurveID(Pair.Key);
+								if (Pair.Value.IsValid() && CurveEditor->GetTreeSelectionState(TreeItemID) == ECurveEditorTreeSelectionState::Explicit)
+								{
+									UObject* OwnerObject = Pair.Value->GetOwningObject();
+									UMovieSceneControlRigParameterSection* Section = Cast< UMovieSceneControlRigParameterSection>(OwnerObject);
+									if (!Section || Section->GetControlRig() != ControlRig)
+									{
+										continue;
+									}
+									UObject* ColorObject = nullptr;
+									FString Name;
+									Pair.Value->GetCurveColorObjectAndName(&ColorObject, Name);
+									TArray<FString> StringArray;
+									Name.ParseIntoArray(StringArray, TEXT("."));
+									if (StringArray.Num() > 1)
+									{
+										if (StringArray[0] == ControlElement->GetDisplayName() ||
+											StringArray[1] == ControlElement->GetDisplayName()) //nested controls will be 2nd
+										{
+											FString ChannelName;
+											if (StringArray.Num() == 3)
+											{
+												ChannelName = StringArray[1] + "." + StringArray[2];
+											}
+											else if (StringArray.Num() == 2)
+											{
+												ChannelName = StringArray[1];
+											}
+
+											EControlRigContextChannelToKey ChannelToKeyFromCurve = Proxy->GetChannelToKeyFromChannelName(ChannelName);
+											if (ChannelToKeyFromCurve == ChannelToKey)
+											{
+												TotalSelected++;
+												break;
+											}
+										}
+									}
+								}
+							}
+						}
+					}
+				}
+			}
+			for (const TPair<UObject*, FSequencerProxyItem>& SItems : Proxy->SequencerItems)
+			{
+				if (UObject* Object = SItems.Key)
+				{
+					for (const FBindingAndTrack& Element : SItems.Value.Bindings)
+					{
+						++TotalNum;
+						EControlRigContextChannelToKey ChannelToKey = Proxy->GetChannelToKeyFromPropertyName(PropertyName);
+						for (const TPair<FCurveModelID, TUniquePtr<FCurveModel>>& Pair : Curves) //horribly slow
+						{
+							FCurveEditorTreeItemID TreeItemID = CurveEditor->GetTreeIDFromCurveID(Pair.Key);
+							if (Pair.Value.IsValid() && CurveEditor->GetTreeSelectionState(TreeItemID) == ECurveEditorTreeSelectionState::Explicit)
+							{
+								UObject* OwnerObject = Pair.Value->GetOwningObject();
+								UMovieSceneSection* Section = Cast< UMovieSceneSection>(OwnerObject);
+								if (!Section || Section->GetTypedOuter<UMovieSceneTrack>() != Element.WeakTrack.Get())
+								{
+									continue;
+								}
+								UObject* ColorObject = nullptr;
+								FString Name;
+								Pair.Value->GetCurveColorObjectAndName(&ColorObject, Name);
+								TArray<FString> StringArray;
+								Name.ParseIntoArray(StringArray, TEXT("."));
+								if (StringArray.Num() > 1)
+								{
+									FString ChannelName;
+									if (StringArray.Num() == 2)
+									{
+										ChannelName = StringArray[0] + "." + StringArray[1];
+									}
+									else if (StringArray.Num() == 1)
+									{
+										ChannelName = StringArray[0];
+									}
+
+									EControlRigContextChannelToKey ChannelToKeyFromCurve = Proxy->GetChannelToKeyFromChannelName(ChannelName);
+									if (ChannelToKeyFromCurve == ChannelToKey)
+									{
+										TotalSelected++;
+										break;
+									}
+								}
+							}
+						}
+					}
+				}
+			}
+		}
+	}
+	if (TotalSelected == TotalNum)
+	{
+		return EAnimDetailSelectionState::All;
+	}
+	else if (TotalSelected == 0)
+	{
+		return EAnimDetailSelectionState::None;
+	}
+	return EAnimDetailSelectionState::Partial;
+}
+
+void UAnimDetailControlsKeyedProxy::SetKey(TSharedPtr<ISequencer>& Sequencer, const IPropertyHandle& KeyedPropertyHandle)
+{
+	for (const TPair<UControlRig*, FControlRigProxyItem>& Items : ControlRigItems)
+	{
+		if (UControlRig* ControlRig = Items.Value.ControlRig.Get())
+		{
+			for (FRigControlElement* ControlElement : Items.Value.ControlElements)
+			{
+				if (ControlElement && ControlRig->GetHierarchy()->Contains(FRigElementKey(ControlElement->GetKey().Name, ERigElementType::Control)))
+				{
+					FRigControlModifiedContext Context;
+					Context.SetKey = EControlRigSetKey::Always;
+					FName PropertyName = KeyedPropertyHandle.GetProperty()->GetFName();
+					Context.KeyMask = (uint32)GetChannelToKeyFromPropertyName(PropertyName);
+					SetControlRigElementValueFromCurrent(ControlRig, ControlElement, Context);
+				}
+			}
+		}
+	}
+	for (const TPair<UObject*, FSequencerProxyItem>& SItems : SequencerItems)
+	{
+		if (SItems.Key)
+		{
+			TArray<UObject*> ObjectsToKey = { SItems.Key };
+			for (const FBindingAndTrack& Element : SItems.Value.Bindings)
+			{
+				FName PropertyName = KeyedPropertyHandle.GetProperty()->GetFName();
+				EControlRigContextChannelToKey ChannelToKey = GetChannelToKeyFromPropertyName(PropertyName);
+				if (UMovieScenePropertyTrack* PropertyTrack = Cast<UMovieScenePropertyTrack>(Element.WeakTrack.Get()))
+				{
+					KeyTrack(Sequencer, this, PropertyTrack, ChannelToKey);
+				}
+			}
+		}
+	}
+}
+static EPropertyKeyedStatus GetChannelKeyStatus(FMovieSceneChannel* Channel, EPropertyKeyedStatus SectionKeyedStatus, const TRange<FFrameNumber>& Range, int32& EmptyChannelCount)
+{
+	if (!Channel)
+	{
+		return SectionKeyedStatus;
+	}
+
+	if (Channel->GetNumKeys() == 0)
+	{
+		++EmptyChannelCount;
+		return SectionKeyedStatus;
+	}
+
+	SectionKeyedStatus = FMath::Max(SectionKeyedStatus, EPropertyKeyedStatus::KeyedInOtherFrame);
+
+	TArray<FFrameNumber> KeyTimes;
+	Channel->GetKeys(Range, &KeyTimes, nullptr);
+	if (KeyTimes.IsEmpty())
+	{
+		++EmptyChannelCount;
+	}
+	else
+	{
+		SectionKeyedStatus = FMath::Max(SectionKeyedStatus, EPropertyKeyedStatus::PartiallyKeyed);
+	}
+	return SectionKeyedStatus;
+}
+
+static EPropertyKeyedStatus GetKeyedStatusInSection(const UControlRig* ControlRig, const FName& ControlName, UMovieSceneControlRigParameterSection* Section, const TRange<FFrameNumber>& Range, EControlRigContextChannelToKey ChannelToKey)
+{
+	int32 EmptyChannelCount = 0;
+	EPropertyKeyedStatus SectionKeyedStatus = EPropertyKeyedStatus::NotKeyed;
+	const FRigControlElement* ControlElement = ControlRig ? ControlRig->FindControl(ControlName) : nullptr;
+	if (ControlElement == nullptr)
+	{
+		return SectionKeyedStatus;
+	}
+	switch (ControlElement->Settings.ControlType)
+	{
+	case ERigControlType::Bool:
+	{
+		TArrayView<FMovieSceneBoolChannel*> BoolChannels = FControlRigSequencerHelpers::GetBoolChannels(ControlRig, ControlElement->GetKey().Name, Section);
+		for (FMovieSceneChannel* Channel : BoolChannels)
+		{
+			SectionKeyedStatus = GetChannelKeyStatus(Channel, SectionKeyedStatus, Range, EmptyChannelCount);
+		}
+		break;
+	}
+	case ERigControlType::Integer:
+	{
+		TArrayView<FMovieSceneIntegerChannel*> IntegarChannels = FControlRigSequencerHelpers::GetIntegerChannels(ControlRig, ControlElement->GetKey().Name, Section);
+		for (FMovieSceneChannel* Channel : IntegarChannels)
+		{
+			SectionKeyedStatus = GetChannelKeyStatus(Channel, SectionKeyedStatus, Range, EmptyChannelCount);
+		}
+		TArrayView<FMovieSceneByteChannel*>  EnumChannels = FControlRigSequencerHelpers::GetByteChannels(ControlRig, ControlElement->GetKey().Name, Section);
+		for (FMovieSceneChannel* Channel : EnumChannels)
+		{
+			SectionKeyedStatus = GetChannelKeyStatus(Channel, SectionKeyedStatus, Range, EmptyChannelCount);
+		}
+		break;
+	}
+	case ERigControlType::Position:
+	case ERigControlType::Transform:
+	case ERigControlType::TransformNoScale:
+	case ERigControlType::EulerTransform:
+	case ERigControlType::Float:
+	case ERigControlType::ScaleFloat:
+	case ERigControlType::Vector2D:
+	{
+		int32 IChannelToKey = (int32)ChannelToKey;
+		TArrayView<FMovieSceneFloatChannel*> FloatChannels = FControlRigSequencerHelpers::GetFloatChannels(ControlRig, ControlElement->GetKey().Name, Section);
+		int32 Num = 0;
+		if (FloatChannels.Num() > Num)
+		{
+			if (IChannelToKey & int32(EControlRigContextChannelToKey::TranslationX))
+			{
+				SectionKeyedStatus = GetChannelKeyStatus(FloatChannels[Num], SectionKeyedStatus, Range, EmptyChannelCount);
+			}
+		}
+		else
+		{
+			break;
+		}
+		++Num;
+		if (FloatChannels.Num() > Num)
+		{
+			if (IChannelToKey & int32(EControlRigContextChannelToKey::TranslationY))
+			{
+				SectionKeyedStatus = GetChannelKeyStatus(FloatChannels[Num], SectionKeyedStatus, Range, EmptyChannelCount);
+			}
+		}
+		else
+		{
+			break;
+		}
+		++Num;
+		if (FloatChannels.Num() > Num)
+		{
+			if (IChannelToKey & int32(EControlRigContextChannelToKey::TranslationZ))
+			{
+				SectionKeyedStatus = GetChannelKeyStatus(FloatChannels[Num], SectionKeyedStatus, Range, EmptyChannelCount);
+			}
+		}
+		else
+		{
+			break;
+		}
+		++Num;
+		if (FloatChannels.Num() > Num)
+		{
+			if (IChannelToKey & int32(EControlRigContextChannelToKey::RotationX))
+			{
+				SectionKeyedStatus = GetChannelKeyStatus(FloatChannels[Num], SectionKeyedStatus, Range, EmptyChannelCount);
+			}
+		}
+		else
+		{
+			break;
+		}
+		++Num;
+		if (FloatChannels.Num() > Num)
+		{
+			if (IChannelToKey & int32(EControlRigContextChannelToKey::RotationY))
+			{
+				SectionKeyedStatus = GetChannelKeyStatus(FloatChannels[Num], SectionKeyedStatus, Range, EmptyChannelCount);
+			}
+		}
+		else
+		{
+			break;
+		}
+		++Num;
+		if (FloatChannels.Num() > Num)
+		{
+			if (IChannelToKey & int32(EControlRigContextChannelToKey::RotationZ))
+			{
+				SectionKeyedStatus = GetChannelKeyStatus(FloatChannels[Num], SectionKeyedStatus, Range, EmptyChannelCount);
+			}
+		}
+		else
+		{
+			break;
+		}
+		++Num;
+		if (FloatChannels.Num() > Num)
+		{
+			if (IChannelToKey & int32(EControlRigContextChannelToKey::ScaleX))
+			{
+				SectionKeyedStatus = GetChannelKeyStatus(FloatChannels[Num], SectionKeyedStatus, Range, EmptyChannelCount);
+			}
+		}
+		else
+		{
+			break;
+		}
+		++Num;
+		if (FloatChannels.Num() > Num)
+		{
+			if (IChannelToKey & int32(EControlRigContextChannelToKey::ScaleY))
+			{
+				SectionKeyedStatus = GetChannelKeyStatus(FloatChannels[Num], SectionKeyedStatus, Range, EmptyChannelCount);
+			}
+		}
+		else
+		{
+			break;
+		}
+		++Num;
+		if (FloatChannels.Num() > Num)
+		{
+			if (IChannelToKey & int32(EControlRigContextChannelToKey::ScaleZ))
+			{
+				SectionKeyedStatus = GetChannelKeyStatus(FloatChannels[Num], SectionKeyedStatus, Range, EmptyChannelCount);
+			}
+		}
+		break;
+	}
+	case ERigControlType::Scale:
+	{
+
+		int32 IChannelToKey = (int32)ChannelToKey;
+		TArrayView<FMovieSceneFloatChannel*> FloatChannels = FControlRigSequencerHelpers::GetFloatChannels(ControlRig, ControlElement->GetKey().Name, Section);
+		int32 Num = 0;
+		if (FloatChannels.Num() > Num)
+		{
+			if (IChannelToKey & int32(EControlRigContextChannelToKey::ScaleX))
+			{
+				SectionKeyedStatus = GetChannelKeyStatus(FloatChannels[Num], SectionKeyedStatus, Range, EmptyChannelCount);
+			}
+		}
+		else
+		{
+			break;
+		}
+		++Num;
+		if (FloatChannels.Num() > Num)
+		{
+			if (IChannelToKey & int32(EControlRigContextChannelToKey::ScaleY))
+			{
+				SectionKeyedStatus = GetChannelKeyStatus(FloatChannels[Num], SectionKeyedStatus, Range, EmptyChannelCount);
+			}
+		}
+		else
+		{
+			break;
+		}
+		++Num;
+		if (FloatChannels.Num() > Num)
+		{
+			if (IChannelToKey & int32(EControlRigContextChannelToKey::ScaleZ))
+			{
+				SectionKeyedStatus = GetChannelKeyStatus(FloatChannels[Num], SectionKeyedStatus, Range, EmptyChannelCount);
+			}
+		}
+		break;
+	}
+	case ERigControlType::Rotator:
+	{
+		int32 IChannelToKey = (int32)ChannelToKey;
+		TArrayView<FMovieSceneFloatChannel*> FloatChannels = FControlRigSequencerHelpers::GetFloatChannels(ControlRig, ControlElement->GetKey().Name, Section);
+		int32 Num = 0;
+		if (FloatChannels.Num() > Num)
+		{
+			if (IChannelToKey & int32(EControlRigContextChannelToKey::RotationX))
+			{
+				SectionKeyedStatus = GetChannelKeyStatus(FloatChannels[Num], SectionKeyedStatus, Range, EmptyChannelCount);
+			}
+		}
+		else
+		{
+			break;
+		}
+		++Num;
+		if (FloatChannels.Num() > Num)
+		{
+			if (IChannelToKey & int32(EControlRigContextChannelToKey::RotationY))
+			{
+				SectionKeyedStatus = GetChannelKeyStatus(FloatChannels[Num], SectionKeyedStatus, Range, EmptyChannelCount);
+			}
+		}
+		else
+		{
+			break;
+		}
+		++Num;
+		if (FloatChannels.Num() > Num)
+		{
+			if (IChannelToKey & int32(EControlRigContextChannelToKey::RotationZ))
+			{
+				SectionKeyedStatus = GetChannelKeyStatus(FloatChannels[Num], SectionKeyedStatus, Range, EmptyChannelCount);
+			}
+		}
+		else
+		{
+			break;
+		}
+		break;
+	}
+	}
+	if (EmptyChannelCount == 0 && SectionKeyedStatus == EPropertyKeyedStatus::PartiallyKeyed)
+	{
+		SectionKeyedStatus = EPropertyKeyedStatus::KeyedInFrame;
+	}
+	return SectionKeyedStatus;
+}
+
+static EPropertyKeyedStatus GetKeyedStatusInTrack(const UControlRig* ControlRig, const FName& ControlName, UMovieSceneControlRigParameterTrack* Track, const TRange<FFrameNumber>& Range, EControlRigContextChannelToKey ChannelToKey)
+{
+	EPropertyKeyedStatus SectionKeyedStatus = EPropertyKeyedStatus::NotKeyed;
+	int32 EmptyChannelCount = 0;
+	const FRigControlElement* ControlElement = ControlRig ? ControlRig->FindControl(ControlName) : nullptr;
+	if (ControlElement == nullptr)
+	{
+		return SectionKeyedStatus;
+	}
+
+	for (UMovieSceneSection* BaseSection : Track->GetAllSections())
+	{
+		UMovieSceneControlRigParameterSection* Section = Cast<UMovieSceneControlRigParameterSection>(BaseSection);
+		if (!Section)
+		{
+			continue;
+		}
+		EPropertyKeyedStatus NewSectionKeyedStatus = GetKeyedStatusInSection(ControlRig, ControlName, Section, Range, ChannelToKey);
+		SectionKeyedStatus = FMath::Max(SectionKeyedStatus, NewSectionKeyedStatus);
+
+		// Maximum Status Reached no need to iterate further
+		if (SectionKeyedStatus == EPropertyKeyedStatus::KeyedInFrame)
+		{
+			return SectionKeyedStatus;
+		}
+	}
+
+	return SectionKeyedStatus;
+}
+
+static EPropertyKeyedStatus GetKeyedStatusInSection(const UMovieSceneSection* Section, const TRange<FFrameNumber>& Range,
+	EControlRigContextChannelToKey ChannelToKey, int32 MaxNumIndices)
+{
+	EPropertyKeyedStatus SectionKeyedStatus = EPropertyKeyedStatus::NotKeyed;
+
+	FMovieSceneChannelProxy& ChannelProxy = Section->GetChannelProxy();
+
+	int32 EmptyChannelCount = 0;
+	TArray<int32, TFixedAllocator<3>> ChannelIndices;
+	switch (ChannelToKey)
+	{
+	case EControlRigContextChannelToKey::Translation:
+		ChannelIndices = { 0, 1, 2 };
+		break;
+	case EControlRigContextChannelToKey::TranslationX:
+		ChannelIndices = { 0 };
+		break;
+	case EControlRigContextChannelToKey::TranslationY:
+		ChannelIndices = { 1 };
+		break;
+	case EControlRigContextChannelToKey::TranslationZ:
+		ChannelIndices = { 2 };
+		break;
+	case EControlRigContextChannelToKey::Rotation:
+		ChannelIndices = { 3, 4, 5 };
+		break;
+	case EControlRigContextChannelToKey::RotationX:
+		ChannelIndices = { 3 };
+		break;
+	case EControlRigContextChannelToKey::RotationY:
+		ChannelIndices = { 4 };
+		break;
+	case EControlRigContextChannelToKey::RotationZ:
+		ChannelIndices = { 5 };
+		break;
+	case EControlRigContextChannelToKey::Scale:
+		ChannelIndices = { 6, 7, 8 };
+		break;
+	case EControlRigContextChannelToKey::ScaleX:
+		ChannelIndices = { 6 };
+		break;
+	case EControlRigContextChannelToKey::ScaleY:
+		ChannelIndices = { 7 };
+		break;
+	case EControlRigContextChannelToKey::ScaleZ:
+		ChannelIndices = { 8 };
+		break;
+	}
+	for (const FMovieSceneChannelEntry& ChannelEntry : ChannelProxy.GetAllEntries())
+	{
+		if (ChannelEntry.GetChannelTypeName() != FMovieSceneDoubleChannel::StaticStruct()->GetFName() &&
+			ChannelEntry.GetChannelTypeName() != FMovieSceneBoolChannel::StaticStruct()->GetFName() &&
+			ChannelEntry.GetChannelTypeName() != FMovieSceneIntegerChannel::StaticStruct()->GetFName() &&
+			ChannelEntry.GetChannelTypeName() != FMovieSceneByteChannel::StaticStruct()->GetFName())
+		{
+			continue;
+		}
+
+		TConstArrayView<FMovieSceneChannel*> Channels = ChannelEntry.GetChannels();
+
+		int32 ChannelIndex = 0;
+		for (FMovieSceneChannel* Channel : Channels)
+		{
+			if (ChannelIndex >= MaxNumIndices)
+			{
+				break;
+			}
+			if (ChannelIndices.Contains(ChannelIndex++) == false)
+			{
+				continue;
+			}
+			if (Channel->GetNumKeys() == 0)
+			{
+				++EmptyChannelCount;
+				continue;
+			}
+
+			SectionKeyedStatus = FMath::Max(SectionKeyedStatus, EPropertyKeyedStatus::KeyedInOtherFrame);
+
+			TArray<FFrameNumber> KeyTimes;
+			Channel->GetKeys(Range, &KeyTimes, nullptr);
+			if (KeyTimes.IsEmpty())
+			{
+				++EmptyChannelCount;
+			}
+			else
+			{
+				SectionKeyedStatus = FMath::Max(SectionKeyedStatus, EPropertyKeyedStatus::PartiallyKeyed);
+			}
+		}
+	}
+
+	if (EmptyChannelCount == 0 && SectionKeyedStatus == EPropertyKeyedStatus::PartiallyKeyed)
+	{
+		SectionKeyedStatus = EPropertyKeyedStatus::KeyedInFrame;
+	}
+	return SectionKeyedStatus;
+}
+
+static EPropertyKeyedStatus GetKeyedStatusInTrack(const UMovieScenePropertyTrack* Track, const TRange<FFrameNumber>& Range,
+	EControlRigContextChannelToKey ChannelToKey, int32 MaxNumIndices)
+{
+	EPropertyKeyedStatus SectionKeyedStatus = EPropertyKeyedStatus::NotKeyed;
+	for (UMovieSceneSection* BaseSection : Track->GetAllSections())
+	{
+		if (!BaseSection)
+		{
+			continue;
+		}
+		EPropertyKeyedStatus NewSectionKeyedStatus = GetKeyedStatusInSection(BaseSection, Range, ChannelToKey, MaxNumIndices);
+		SectionKeyedStatus = FMath::Max(SectionKeyedStatus, NewSectionKeyedStatus);
+
+		// Maximum Status Reached no need to iterate further
+		if (SectionKeyedStatus == EPropertyKeyedStatus::KeyedInFrame)
+		{
+			return SectionKeyedStatus;
+		}
+	}
+
+	return SectionKeyedStatus;
+}
+
+EPropertyKeyedStatus UAnimDetailControlsKeyedProxy::GetPropertyKeyedStatus(TSharedPtr<ISequencer>& Sequencer, const IPropertyHandle& PropertyHandle) const
+{
+	EPropertyKeyedStatus KeyedStatus = EPropertyKeyedStatus::NotKeyed;
+
+	if (Sequencer.IsValid() == false || Sequencer->GetFocusedMovieSceneSequence() == nullptr)
+	{
+		return KeyedStatus;
+	}
+	const TRange<FFrameNumber> FrameRange = TRange<FFrameNumber>(Sequencer->GetLocalTime().Time.FrameNumber);
+	FName PropertyName = PropertyHandle.GetProperty()->GetFName();
+	EControlRigContextChannelToKey ChannelToKey = GetChannelToKeyFromPropertyName(PropertyName);
+	for (const TPair<UControlRig*, FControlRigProxyItem>& Items : ControlRigItems)
+	{
+		if (UControlRig* ControlRig = Items.Value.ControlRig.Get())
+		{
+			UMovieSceneControlRigParameterTrack* Track = FControlRigSequencerHelpers::FindControlRigTrack(Sequencer->GetFocusedMovieSceneSequence(), ControlRig);
+			if (Track == nullptr)
+			{
+				continue;
+			}
+			for (FRigControlElement* ControlElement : Items.Value.ControlElements)
+			{
+				EPropertyKeyedStatus NewKeyedStatus = GetKeyedStatusInTrack(ControlRig, ControlElement->GetKey().Name, Track, FrameRange, ChannelToKey);
+				KeyedStatus = FMath::Max(KeyedStatus, NewKeyedStatus);
+			}
+		}
+	}
+	int32 MaxNumIndices = 1;
+	if (IsA<UAnimDetailControlsProxyTransform>())
+	{
+		MaxNumIndices = 9;
+	}
+	else if (IsA<UAnimDetailControlsProxyVector2D>())
+	{
+		MaxNumIndices = 2;
+	}
+	for (const TPair<UObject*, FSequencerProxyItem>& Items : SequencerItems)
+	{
+		for (const FBindingAndTrack& Element : Items.Value.Bindings)
+		{
+			UMovieScenePropertyTrack* Track = Cast<UMovieScenePropertyTrack>(Element.WeakTrack.Get());
+			if (Track == nullptr)
+			{
+				continue;
+			}
+			EPropertyKeyedStatus NewKeyedStatus = GetKeyedStatusInTrack(Track, FrameRange, ChannelToKey, MaxNumIndices);
+			KeyedStatus = FMath::Max(KeyedStatus, NewKeyedStatus);
+		}
+	}
+	return KeyedStatus;
+}
+
+#if WITH_EDITOR
+void UAnimDetailControlsKeyedProxy::PostEditUndo()
+{
+	FRigControlModifiedContext Context;
+	Context.SetKey = EControlRigSetKey::Never;
+	for (const TPair<UControlRig*, FControlRigProxyItem>& Items : ControlRigItems)
+	{
+		if (UControlRig* ControlRig = Items.Value.ControlRig.Get())
+		{
+			for (FRigControlElement* ControlElement : Items.Value.ControlElements)
+			{
+				if (ControlElement && ControlRig->GetHierarchy()->Contains(FRigElementKey(ControlElement->GetKey().Name, ERigElementType::Control)))
+				{
+					ControlRig->SelectControl(ControlElement->GetKey().Name, bSelected);
+					SetControlRigElementValueFromCurrent(ControlRig, ControlElement, Context);
+				}
+			}
+		}
+	}
+	for (TPair <UObject*, FSequencerProxyItem>& SItems : SequencerItems)
+	{
+		//we do this backwards so ValueChanged later is set up correctly since that iterates in the other direction
+		for (FBindingAndTrack& Binding: SItems.Value.Bindings)
+		{
+			SetBindingValueFromCurrent(SItems.Key, Binding.Binding, Context);
+		}
+	}
+}
+#endif
+
+//////UAnimDetailControlsProxyTransform////////
+
+
+static void SetValuesFromContext(const FEulerTransform& EulerTransform, const FRigControlModifiedContext& Context, FVector& TLocation, FRotator& TRotation, FVector& TScale)
+{
+	EControlRigContextChannelToKey ChannelsToKey = (EControlRigContextChannelToKey)Context.KeyMask;
+	if (EnumHasAnyFlags(ChannelsToKey, EControlRigContextChannelToKey::TranslationX) == false)
+	{
+		TLocation.X = EulerTransform.Location.X;
+	}
+	if (EnumHasAnyFlags(ChannelsToKey, EControlRigContextChannelToKey::TranslationY) == false)
+	{
+		TLocation.Y = EulerTransform.Location.Y;
+	}
+	if (EnumHasAnyFlags(ChannelsToKey, EControlRigContextChannelToKey::TranslationZ) == false)
+	{
+		TLocation.Z = EulerTransform.Location.Z;
+	}
+	if (EnumHasAnyFlags(ChannelsToKey, EControlRigContextChannelToKey::RotationX) == false)
+	{
+		TRotation.Roll = EulerTransform.Rotation.Roll;
+	}
+	if (EnumHasAnyFlags(ChannelsToKey, EControlRigContextChannelToKey::RotationZ) == false)
+	{
+		TRotation.Yaw = EulerTransform.Rotation.Yaw;
+	}
+	if (EnumHasAnyFlags(ChannelsToKey, EControlRigContextChannelToKey::RotationY) == false)
+	{
+		TRotation.Pitch = EulerTransform.Rotation.Pitch;
+	}
+	if (EnumHasAnyFlags(ChannelsToKey, EControlRigContextChannelToKey::ScaleX) == false)
+	{
+		TScale.X = EulerTransform.Scale.X;
+	}
+	if (EnumHasAnyFlags(ChannelsToKey, EControlRigContextChannelToKey::ScaleY) == false)
+	{
+		TScale.Y = EulerTransform.Scale.Y;
+	}
+	if (EnumHasAnyFlags(ChannelsToKey, EControlRigContextChannelToKey::ScaleZ) == false)
+	{
+		TScale.Z = EulerTransform.Scale.Z;
+	}
+}
+
+static FEulerTransform GetCurrentValue(UControlRig* ControlRig, FRigControlElement* ControlElement)
+{
+	FEulerTransform EulerTransform = FEulerTransform::Identity;
+	switch (ControlElement->Settings.ControlType)
+	{
+	case ERigControlType::Transform:
+	{
+		const FTransform NewTransform = ControlRig->GetControlValue(ControlElement, ERigControlValueType::Current).Get<FRigControlValue::FTransform_Float>().ToTransform();
+		EulerTransform = FEulerTransform(NewTransform);
+		break;
+
+	}
+	case ERigControlType::TransformNoScale:
+	{
+		const FTransformNoScale NewTransform = ControlRig->GetControlValue(ControlElement, ERigControlValueType::Current).Get<FRigControlValue::FTransformNoScale_Float>().ToTransform();
+		EulerTransform.Location = NewTransform.Location;
+		EulerTransform.Rotation = FRotator(NewTransform.Rotation);
+		break;
+
+	}
+	case ERigControlType::EulerTransform:
+	{
+		EulerTransform = ControlRig->GetControlValue(ControlElement, ERigControlValueType::Current).Get<FRigControlValue::FEulerTransform_Float>().ToTransform();
+		break;
+	}
+	};
+	EulerTransform.Rotation = ControlRig->GetHierarchy()->GetControlPreferredRotator(ControlElement);
+	return EulerTransform;
+}
+
+static void GetActorAndSceneComponentFromObject(UObject* Object, AActor*& OutActor, USceneComponent*& OutSceneComponent)
+{
+	OutActor = Cast<AActor>(Object);
+	if (OutActor != nullptr && OutActor->GetRootComponent())
+	{
+		OutSceneComponent = OutActor->GetRootComponent();
+	}
+	else
+	{
+		// If the object wasn't an actor attempt to get it directly as a scene component and then get the actor from there.
+		OutSceneComponent = Cast<USceneComponent>(Object);
+		if (OutSceneComponent != nullptr)
+		{
+			OutActor = Cast<AActor>(OutSceneComponent->GetOuter());
+		}
+	}
+}
+
+static FEulerTransform GetCurrentValue(UObject* InObject, TSharedPtr<FTrackInstancePropertyBindings>& Binding)
+{
+	const FStructProperty* TransformProperty = Binding.IsValid() ? CastField<FStructProperty>(Binding->GetProperty(*InObject)) : nullptr;
+	if (TransformProperty)
+	{
+		if (TransformProperty->Struct == TBaseStructure<FTransform>::Get())
+		{
+			if (TOptional<FTransform> Transform = Binding->GetOptionalValue<FTransform>(*InObject))
+			{
+				return FEulerTransform(Transform.GetValue());
+			}
+		}
+		else if (TransformProperty->Struct == TBaseStructure<FEulerTransform>::Get())
+		{
+			if (TOptional<FEulerTransform> EulerTransform = Binding->GetOptionalValue<FEulerTransform>(*InObject))
+			{
+				return EulerTransform.GetValue();
+			}
+		}
+	}
+	AActor* ActorThatChanged = nullptr;
+	USceneComponent* SceneComponentThatChanged = nullptr;
+	GetActorAndSceneComponentFromObject(InObject, ActorThatChanged, SceneComponentThatChanged);
+
+	FEulerTransform EulerTransform = FEulerTransform::Identity;
+	if (SceneComponentThatChanged)
+	{
+		FTransform Transform(SceneComponentThatChanged->GetRelativeTransform());
+		EulerTransform = FEulerTransform(Transform);
+	}
+	return EulerTransform;
+}
+
+void UAnimDetailControlsProxyTransform::SetBindingValueFromCurrent(UObject* InObject, TSharedPtr<FTrackInstancePropertyBindings>& Binding, FRigControlModifiedContext& Context)
+{
+	if (InObject && Binding.IsValid())
+	{
+		FVector TLocation = Location.ToVector();
+		FRotator TRotation = Rotation.ToRotator();
+		FVector TScale = Scale.ToVector();
+		FEulerTransform EulerTransform = GetCurrentValue(InObject, Binding);
+		SetValuesFromContext(EulerTransform, Context, TLocation, TRotation, TScale);
+		FTransform RealTransform(TRotation, TLocation, TScale);
+
+		const FStructProperty* TransformProperty = Binding.IsValid() ? CastField<FStructProperty>(Binding->GetProperty(*InObject)) : nullptr;
+		if (TransformProperty)
+		{
+			if (TransformProperty->Struct == TBaseStructure<FTransform>::Get())
+			{
+			}
+			else if (TransformProperty->Struct == TBaseStructure<FEulerTransform>::Get())
+			{
+				EulerTransform = FEulerTransform(TLocation, TRotation, TScale);
+				Binding->SetCurrentValue<FTransform>(*InObject, RealTransform);
+			}
+		}
+
+		AActor* ActorThatChanged = nullptr;
+		USceneComponent* SceneComponentThatChanged = nullptr;
+		GetActorAndSceneComponentFromObject(InObject, ActorThatChanged, SceneComponentThatChanged);
+		if (SceneComponentThatChanged)
+		{
+			SceneComponentThatChanged->SetRelativeTransform(RealTransform, false, nullptr, ETeleportType::None);
+
+			// Force the location and rotation values to avoid Rot->Quat->Rot conversions
+			SceneComponentThatChanged->SetRelativeLocation_Direct(TLocation);
+			SceneComponentThatChanged->SetRelativeRotation_Direct(TRotation);
+			GUnrealEd->UpdatePivotLocationForSelection();
+		}
+	}
+}
+
+void UAnimDetailControlsProxyTransform::SetControlRigElementValueFromCurrent(UControlRig* ControlRig, FRigControlElement* ControlElement, const FRigControlModifiedContext& Context)
+{
+	if (ControlElement && ControlRig)
+	{
+		FVector TLocation = Location.ToVector();
+		FRotator TRotation = Rotation.ToRotator();
+		FVector TScale = Scale.ToVector();
+		FEulerTransform EulerTransform = GetCurrentValue(ControlRig, ControlElement);
+		SetValuesFromContext(EulerTransform, Context, TLocation, TRotation, TScale);
+		switch (ControlElement->Settings.ControlType)
+		{
+		case ERigControlType::Transform:
+		{
+			FTransform RealTransform(TRotation, TLocation, TScale);
+			FVector EulerAngle(TRotation.Roll, TRotation.Pitch, TRotation.Yaw);
+			ControlRig->GetHierarchy()->SetControlSpecifiedEulerAngle(ControlElement, EulerAngle);
+			ControlRig->SetControlValue<FRigControlValue::FTransform_Float>(ControlElement->GetKey().Name, RealTransform, true, Context, false);
+			ControlRig->GetHierarchy()->SetControlSpecifiedEulerAngle(ControlElement, EulerAngle);
+			break;
+
+		}
+		case ERigControlType::TransformNoScale:
+		{
+			FTransformNoScale NoScale(TLocation, TRotation.Quaternion());
+			ControlRig->SetControlValue<FRigControlValue::FTransformNoScale_Float>(ControlElement->GetKey().Name, NoScale, true, Context, false);
+			break;
+
+		}
+		case ERigControlType::EulerTransform:
+		{
+			URigHierarchy* Hierarchy = ControlRig->GetHierarchy();
+			FVector EulerAngle(TRotation.Roll, TRotation.Pitch, TRotation.Yaw);
+			FQuat Quat = Hierarchy->GetControlQuaternion(ControlElement, EulerAngle);
+			Hierarchy->SetControlSpecifiedEulerAngle(ControlElement, EulerAngle);
+			FRotator UERotator(Quat);
+			FEulerTransform UETransform(UERotator, TLocation, TScale);
+			UETransform.Rotation = UERotator;
+			ControlRig->SetControlValue<FRigControlValue::FEulerTransform_Float>(ControlElement->GetKey().Name, UETransform, true, Context, false);
+			Hierarchy->SetControlSpecifiedEulerAngle(ControlElement, EulerAngle);
+			break;
+		}
+		}
+		ControlRig->Evaluate_AnyThread();
+	}
+}
+
+bool UAnimDetailControlsProxyTransform::PropertyIsOnProxy(FProperty* Property, FProperty* MemberProperty) 
+{
+	return ((Property && Property->GetFName() == GET_MEMBER_NAME_CHECKED(UAnimDetailControlsProxyTransform, Location)) ||
+		(MemberProperty && MemberProperty->GetFName() == GET_MEMBER_NAME_CHECKED(UAnimDetailControlsProxyTransform, Location)) ||
+		(Property && Property->GetFName() == GET_MEMBER_NAME_CHECKED(UAnimDetailControlsProxyTransform, Rotation)) ||
+		(MemberProperty && MemberProperty->GetFName() == GET_MEMBER_NAME_CHECKED(UAnimDetailControlsProxyTransform, Rotation)) ||
+		(Property && Property->GetFName() == GET_MEMBER_NAME_CHECKED(UAnimDetailControlsProxyTransform, Scale)) ||
+		(MemberProperty && MemberProperty->GetFName() == GET_MEMBER_NAME_CHECKED(UAnimDetailControlsProxyTransform, Scale))
+		);
+}
+
+void UAnimDetailControlsProxyTransform::ClearMultipleFlags()
+{
+	Location.State.bXMultiple = false;
+	Location.State.bYMultiple = false;
+	Location.State.bZMultiple = false;
+	Rotation.State.bXMultiple = false;
+	Rotation.State.bYMultiple = false;
+	Rotation.State.bZMultiple = false;
+	Scale.State.bXMultiple = false;
+	Scale.State.bYMultiple = false;
+	Scale.State.bZMultiple = false;
+}
+
+void UAnimDetailControlsProxyTransform::SetMultipleFlags(const FEulerTransform& ValA, const FEulerTransform& ValB)
+{
+	if (FMath::IsNearlyEqual(ValA.Location.X, ValB.Location.X) == false)
+	{
+		Location.State.bXMultiple = true;
+	}
+	if (FMath::IsNearlyEqual(ValA.Location.Y, ValB.Location.Y) == false)
+	{
+		Location.State.bYMultiple = true;
+	}
+	if (FMath::IsNearlyEqual(ValA.Location.Z, ValB.Location.Z) == false)
+	{
+		Location.State.bZMultiple = true;
+	}
+	if (FMath::IsNearlyEqual(ValA.Rotation.Roll, ValB.Rotation.Roll) == false)
+	{
+		Rotation.State.bXMultiple = true;
+	}
+	if (FMath::IsNearlyEqual(ValA.Rotation.Pitch, ValB.Rotation.Pitch) == false)
+	{
+		Rotation.State.bYMultiple = true;
+	}
+	if (FMath::IsNearlyEqual(ValA.Rotation.Yaw, ValB.Rotation.Yaw) == false)
+	{
+		Rotation.State.bZMultiple = true;
+	}
+	if (FMath::IsNearlyEqual(ValA.Scale.X, ValB.Scale.X) == false)
+	{
+		Scale.State.bXMultiple = true;
+	}
+	if (FMath::IsNearlyEqual(ValA.Scale.Y, ValB.Scale.Y) == false)
+	{
+		Scale.State.bYMultiple = true;
+	}
+	if (FMath::IsNearlyEqual(ValA.Scale.Z, ValB.Scale.Z) == false)
+	{
+		Scale.State.bZMultiple = true;
+	}
+}
+
+void UAnimDetailControlsProxyTransform::ValueChanged()
+{
+	TOptional<FEulerTransform> LastEulerTransform;
+	ClearMultipleFlags();
+
+	for (const TPair<UControlRig*, FControlRigProxyItem>& Items : ControlRigItems)
+	{
+		if (UControlRig* ControlRig = Items.Value.ControlRig.Get())
+		{
+			for (FRigControlElement* ControlElement : Items.Value.ControlElements)
+			{
+				FEulerTransform EulerTransform = GetCurrentValue(ControlRig, ControlElement);
+				if (LastEulerTransform.IsSet())
+				{
+					const FEulerTransform LastVal = LastEulerTransform.GetValue();
+					SetMultipleFlags(LastVal, EulerTransform);
+				}
+				else
+				{
+					LastEulerTransform = EulerTransform;
+				}
+			}
+		}
+	}
+	for (TPair <UObject*, FSequencerProxyItem>& SItems : SequencerItems)
+	{
+		//we do this backwards so ValueChanged later is set up correctly since that iterates in the other direction
+		for (int32 Index = SItems.Value.Bindings.Num() - 1; Index >= 0; --Index)
+		{
+			FBindingAndTrack& Binding = SItems.Value.Bindings[Index];
+			if (Binding.Binding.IsValid())
+			{
+				FEulerTransform EulerTransform = GetCurrentValue(SItems.Key, Binding.Binding);
+				if (LastEulerTransform.IsSet())
+				{
+					const FEulerTransform LastVal = LastEulerTransform.GetValue();
+					SetMultipleFlags(LastVal, EulerTransform);
+				}
+				else
+				{
+					LastEulerTransform = EulerTransform;
+				}
+			}
+		}
+	}
+
+	const FEulerTransform LastVal = LastEulerTransform.IsSet() ? LastEulerTransform.GetValue() : FEulerTransform::Identity;
+
+	//set from values, note that the fact if they were multiples or not was already set so we need to set them on these values.
+	FAnimDetailProxyLocation TLocation(LastVal.Location, Location.State);
+	FAnimDetailProxyRotation TRotation(LastVal.Rotation, Rotation.State);
+	FAnimDetailProxyScale TScale(LastVal.Scale, Scale.State);
+
+	const FName LocationName("Location");
+	FTrackInstancePropertyBindings LocationBinding(LocationName, LocationName.ToString());
+	LocationBinding.CallFunction<FAnimDetailProxyLocation>(*this, TLocation);
+
+	const FName RotationName("Rotation");
+	FTrackInstancePropertyBindings RotationBinding(RotationName, RotationName.ToString());
+	RotationBinding.CallFunction<FAnimDetailProxyRotation>(*this, TRotation);
+
+	const FName ScaleName("Scale");
+	FTrackInstancePropertyBindings ScaleBinding(ScaleName, ScaleName.ToString());
+	ScaleBinding.CallFunction<FAnimDetailProxyScale>(*this, TScale);
+}
+
+bool UAnimDetailControlsProxyTransform::IsMultiple(const FName& PropertyName) const
+{
+	if (PropertyName == GET_MEMBER_NAME_CHECKED(FAnimDetailProxyLocation, LX))
+	{
+		return Location.State.bXMultiple;
+	}
+	else if (PropertyName == GET_MEMBER_NAME_CHECKED(FAnimDetailProxyLocation, LY))
+	{
+		return Location.State.bYMultiple;
+	}
+	else if (PropertyName == GET_MEMBER_NAME_CHECKED(FAnimDetailProxyLocation, LZ))
+	{
+		return Location.State.bZMultiple;
+	}
+	else if (PropertyName == GET_MEMBER_NAME_CHECKED(FAnimDetailProxyRotation, RX))
+	{
+		return Rotation.State.bXMultiple;
+	}
+	else if (PropertyName == GET_MEMBER_NAME_CHECKED(FAnimDetailProxyRotation, RY))
+	{
+		return Rotation.State.bYMultiple;
+	}
+	else if (PropertyName == GET_MEMBER_NAME_CHECKED(FAnimDetailProxyRotation, RZ))
+	{
+		return Rotation.State.bZMultiple;
+	}
+	else if (PropertyName == GET_MEMBER_NAME_CHECKED(FAnimDetailProxyScale, SX))
+	{
+		return Scale.State.bXMultiple;
+	}
+	else if (PropertyName == GET_MEMBER_NAME_CHECKED(FAnimDetailProxyScale, SY))
+	{
+		return Scale.State.bYMultiple;
+	}
+	else if (PropertyName == GET_MEMBER_NAME_CHECKED(FAnimDetailProxyScale, SZ))
+	{
+		return Scale.State.bZMultiple;
+	}
+
+	return false;
+}
+
+EControlRigContextChannelToKey UAnimDetailControlsProxyTransform::GetChannelToKeyFromPropertyName(const FName& PropertyName) const
+{
+	if (PropertyName == GET_MEMBER_NAME_CHECKED(UAnimDetailControlsProxyTransform, Location))
+	{
+		return EControlRigContextChannelToKey::Translation;
+	}
+	else if (PropertyName == GET_MEMBER_NAME_CHECKED(FAnimDetailProxyLocation, LX))
+	{
+		return EControlRigContextChannelToKey::TranslationX;
+	}
+	else if (PropertyName == GET_MEMBER_NAME_CHECKED(FAnimDetailProxyLocation, LY))
+	{
+		return EControlRigContextChannelToKey::TranslationY;
+	}
+	else if (PropertyName == GET_MEMBER_NAME_CHECKED(FAnimDetailProxyLocation, LZ))
+	{
+		return EControlRigContextChannelToKey::TranslationZ;
+	}
+	else if (PropertyName == GET_MEMBER_NAME_CHECKED(UAnimDetailControlsProxyTransform, Rotation))
+	{
+		return EControlRigContextChannelToKey::Rotation;
+	}
+	else if (PropertyName == GET_MEMBER_NAME_CHECKED(FAnimDetailProxyRotation, RX))
+	{
+		return EControlRigContextChannelToKey::RotationX;
+	}
+	else if (PropertyName == GET_MEMBER_NAME_CHECKED(FAnimDetailProxyRotation, RY))
+	{
+		return EControlRigContextChannelToKey::RotationY;
+	}
+	else if (PropertyName == GET_MEMBER_NAME_CHECKED(FAnimDetailProxyRotation, RZ))
+	{
+		return EControlRigContextChannelToKey::RotationZ;
+	}
+	else if (PropertyName == GET_MEMBER_NAME_CHECKED(UAnimDetailControlsProxyTransform, Scale))
+	{
+		return EControlRigContextChannelToKey::Scale;
+	}
+	else if (PropertyName == GET_MEMBER_NAME_CHECKED(FAnimDetailProxyScale, SX))
+	{
+		return EControlRigContextChannelToKey::ScaleX;
+	}
+	else if (PropertyName == GET_MEMBER_NAME_CHECKED(FAnimDetailProxyScale, SY))
+	{
+		return EControlRigContextChannelToKey::ScaleY;
+	}
+	else if (PropertyName == GET_MEMBER_NAME_CHECKED(FAnimDetailProxyScale, SZ))
+	{
+		return EControlRigContextChannelToKey::ScaleZ;
+	}
+
+	return EControlRigContextChannelToKey::AllTransform;
+}
+
+EControlRigContextChannelToKey UAnimDetailControlsProxyTransform::GetChannelToKeyFromChannelName(const FString& InChannelName) const
+{
+	if (InChannelName == TEXT("Location.X"))
+	{
+		return EControlRigContextChannelToKey::TranslationX;
+	}
+	else if (InChannelName == TEXT("Location.Y"))
+	{
+		return EControlRigContextChannelToKey::TranslationY;
+	}
+	else if (InChannelName == TEXT("Location.Z"))
+	{
+		return EControlRigContextChannelToKey::TranslationZ;
+	}
+	else if (InChannelName == TEXT("Rotation.X") || InChannelName == TEXT("Rotation.Roll"))
+	{
+		return EControlRigContextChannelToKey::RotationX;
+	}
+	else if (InChannelName == TEXT("Rotation.Y") || InChannelName == TEXT("Rotation.Pitch"))
+	{
+		return EControlRigContextChannelToKey::RotationY;
+	}
+	else if (InChannelName == TEXT("Rotation.Z") || InChannelName == TEXT("Rotation.Yaw"))
+	{
+		return EControlRigContextChannelToKey::RotationZ;
+	}
+	else if (InChannelName == TEXT("Scale.X"))
+	{
+		return EControlRigContextChannelToKey::ScaleX;
+	}
+	else if (InChannelName == TEXT("Scale.Y"))
+	{
+		return EControlRigContextChannelToKey::ScaleY;
+	}
+	else if (InChannelName == TEXT("Scale.Z"))
+	{
+		return EControlRigContextChannelToKey::ScaleZ;
+	}
+	return EControlRigContextChannelToKey::AllTransform;
+}
+
+void UAnimDetailControlsProxyTransform::GetChannelSelectionState(TWeakPtr<FCurveEditor>& CurveEditor, FAnimDetailVectorSelection& LocationSelectionCache, FAnimDetailVectorSelection& RotationSelectionCache,
+	FAnimDetailVectorSelection& ScaleSelectionCache)
+{
+	LocationSelectionCache.XSelected = CachePropertySelection(CurveEditor, this, GET_MEMBER_NAME_CHECKED(FAnimDetailProxyLocation, LX));
+	LocationSelectionCache.YSelected = CachePropertySelection(CurveEditor, this, GET_MEMBER_NAME_CHECKED(FAnimDetailProxyLocation, LY));
+	LocationSelectionCache.ZSelected = CachePropertySelection(CurveEditor, this, GET_MEMBER_NAME_CHECKED(FAnimDetailProxyLocation, LZ));
+
+	RotationSelectionCache.XSelected = CachePropertySelection(CurveEditor, this, GET_MEMBER_NAME_CHECKED(FAnimDetailProxyRotation, RX));
+	RotationSelectionCache.YSelected = CachePropertySelection(CurveEditor, this, GET_MEMBER_NAME_CHECKED(FAnimDetailProxyRotation, RY));
+	RotationSelectionCache.ZSelected = CachePropertySelection(CurveEditor, this, GET_MEMBER_NAME_CHECKED(FAnimDetailProxyRotation, RZ));
+
+	ScaleSelectionCache.XSelected = CachePropertySelection(CurveEditor, this, GET_MEMBER_NAME_CHECKED(FAnimDetailProxyScale, SX));
+	ScaleSelectionCache.YSelected = CachePropertySelection(CurveEditor, this, GET_MEMBER_NAME_CHECKED(FAnimDetailProxyScale, SY));
+	ScaleSelectionCache.ZSelected = CachePropertySelection(CurveEditor, this, GET_MEMBER_NAME_CHECKED(FAnimDetailProxyScale, SZ));
+}
+
+
+//////UAnimDetailControlsProxyLocation////////
+static void SetLocationValuesFromContext(UControlRig* ControlRig, FRigControlElement* ControlElement, const FRigControlModifiedContext& Context, FVector3f& TLocation)
+{
+	FRigControlValue ControlValue = ControlRig->GetControlValue(ControlElement, ERigControlValueType::Current);
+	FVector3f Value = ControlValue.Get<FVector3f>();
+
+	EControlRigContextChannelToKey ChannelsToKey = (EControlRigContextChannelToKey)Context.KeyMask;
+	if (EnumHasAnyFlags(ChannelsToKey, EControlRigContextChannelToKey::TranslationX) == false)
+	{
+		TLocation.X = Value.X;
+	}
+	if (EnumHasAnyFlags(ChannelsToKey, EControlRigContextChannelToKey::TranslationY) == false)
+	{
+		TLocation.Y = Value.Y;
+	}
+	if (EnumHasAnyFlags(ChannelsToKey, EControlRigContextChannelToKey::TranslationZ) == false)
+	{
+		TLocation.Z = Value.Z;
+	}
+}
+
+void UAnimDetailControlsProxyLocation::SetControlRigElementValueFromCurrent(UControlRig* ControlRig, FRigControlElement* ControlElement, const FRigControlModifiedContext& Context)
+{
+	if (ControlElement && ControlRig && ControlElement->Settings.ControlType == ERigControlType::Position)
+	{
+		FVector3f TLocation = Location.ToVector3f();
+		SetLocationValuesFromContext(ControlRig, ControlElement, Context, TLocation);
+
+		ControlRig->SetControlValue<FVector3f>(ControlElement->GetKey().Name, TLocation, true, Context, false);
+		ControlRig->Evaluate_AnyThread();
+	}
+}
+
+bool UAnimDetailControlsProxyLocation::PropertyIsOnProxy(FProperty* Property, FProperty* MemberProperty)
+{
+	return ((Property && Property->GetFName() == GET_MEMBER_NAME_CHECKED(UAnimDetailControlsProxyLocation, Location)) ||
+		(MemberProperty && MemberProperty->GetFName() == GET_MEMBER_NAME_CHECKED(UAnimDetailControlsProxyLocation, Location))
+		);
+}
+
+void UAnimDetailControlsProxyLocation::ClearMultipleFlags()
+{
+	Location.State.bXMultiple = false;
+	Location.State.bYMultiple = false;
+	Location.State.bZMultiple = false;
+}
+
+void UAnimDetailControlsProxyLocation::SetMultipleFlags(const FVector3f& ValA, const FVector3f& ValB)
+{
+	if (FMath::IsNearlyEqual(ValA.X, ValB.X) == false)
+	{
+		Location.State.bXMultiple = true;
+	}
+	if (FMath::IsNearlyEqual(ValA.Y, ValB.Y) == false)
+	{
+		Location.State.bYMultiple = true;
+	}
+	if (FMath::IsNearlyEqual(ValA.Z, ValB.Z) == false)
+	{
+		Location.State.bZMultiple = true;
+	}
+}
+
+void UAnimDetailControlsProxyLocation::ValueChanged()
+{
+	TOptional<FVector3f> LastValue;
+	ClearMultipleFlags();
+
+	for (const TPair<UControlRig*, FControlRigProxyItem>& Items : ControlRigItems)
+	{
+		if (UControlRig* ControlRig = Items.Value.ControlRig.Get())
+		{
+			for (FRigControlElement* ControlElement : Items.Value.ControlElements)
+			{
+				FVector3f Value = FVector3f::ZeroVector;
+				if (ControlElement->Settings.ControlType == ERigControlType::Position)
+				{
+					FRigControlValue ControlValue = ControlRig->GetControlValue(ControlElement, ERigControlValueType::Current);
+					Value = ControlValue.Get<FVector3f>();
+				}
+				if (LastValue.IsSet())
+				{
+					const FVector3f LastVal = LastValue.GetValue();
+					SetMultipleFlags(LastVal, Value);
+				}
+				else
+				{
+					LastValue = Value;
+				}
+			}
+		}
+	}
+
+	const FVector3f LastVal = LastValue.IsSet() ? LastValue.GetValue() : FVector3f::ZeroVector;
+
+	//set from values, note that the fact if they were multiples or not was already set so we need to set them on these values.
+	FAnimDetailProxyLocation TLocation(LastVal, Location.State);
+
+	const FName LocationName("Location");
+	FTrackInstancePropertyBindings LocationBinding(LocationName, LocationName.ToString());
+	LocationBinding.CallFunction<FAnimDetailProxyLocation>(*this, TLocation);
+
+}
+
+bool UAnimDetailControlsProxyLocation::IsMultiple(const FName& PropertyName) const
+{
+	if (PropertyName == GET_MEMBER_NAME_CHECKED(FAnimDetailProxyLocation, LX))
+	{
+		return Location.State.bXMultiple;
+	}
+	else if (PropertyName == GET_MEMBER_NAME_CHECKED(FAnimDetailProxyLocation, LY))
+	{
+		return Location.State.bYMultiple;
+	}
+	else if (PropertyName == GET_MEMBER_NAME_CHECKED(FAnimDetailProxyLocation, LZ))
+	{
+		return Location.State.bZMultiple;
+	}
+
+	return false;
+}
+
+EControlRigContextChannelToKey UAnimDetailControlsProxyLocation::GetChannelToKeyFromPropertyName(const FName& PropertyName) const
+{
+	if (PropertyName == GET_MEMBER_NAME_CHECKED(FAnimDetailProxyLocation, LX))
+	{
+		return EControlRigContextChannelToKey::TranslationX;
+	}
+	else if (PropertyName == GET_MEMBER_NAME_CHECKED(FAnimDetailProxyLocation, LY))
+	{
+		return EControlRigContextChannelToKey::TranslationY;
+	}
+	else if (PropertyName == GET_MEMBER_NAME_CHECKED(FAnimDetailProxyLocation, LZ))
+	{
+		return EControlRigContextChannelToKey::TranslationZ;
+	}
+
+	return EControlRigContextChannelToKey::AllTransform;
+}
+
+EControlRigContextChannelToKey UAnimDetailControlsProxyLocation::GetChannelToKeyFromChannelName(const FString& InChannelName) const
+{
+	if (InChannelName == TEXT("X"))
+	{
+		return EControlRigContextChannelToKey::TranslationX;
+	}
+	else if (InChannelName == TEXT("Y"))
+	{
+		return EControlRigContextChannelToKey::TranslationY;
+	}
+	else if (InChannelName == TEXT("Z"))
+	{
+		return EControlRigContextChannelToKey::TranslationZ;
+	}
+	return EControlRigContextChannelToKey::AllTransform;
+}
+
+void UAnimDetailControlsProxyLocation::GetChannelSelectionState(TWeakPtr<FCurveEditor>& CurveEditor, FAnimDetailVectorSelection& LocationSelectionCache, FAnimDetailVectorSelection& RotationSelectionCache,
+	FAnimDetailVectorSelection& ScaleSelectionCache)
+{
+	LocationSelectionCache.XSelected = CachePropertySelection(CurveEditor, this, GET_MEMBER_NAME_CHECKED(FAnimDetailProxyLocation, LX));
+	LocationSelectionCache.YSelected = CachePropertySelection(CurveEditor, this, GET_MEMBER_NAME_CHECKED(FAnimDetailProxyLocation, LY));
+	LocationSelectionCache.ZSelected = CachePropertySelection(CurveEditor, this, GET_MEMBER_NAME_CHECKED(FAnimDetailProxyLocation, LZ));
+}
+
+//////UAnimDetailControlsProxyRotation////////
+
+static void SetRotationValuesFromContext(UControlRig* ControlRig, FRigControlElement* ControlElement, const FRigControlModifiedContext& Context, FVector3f& Val)
+{
+	FRigControlValue ControlValue = ControlRig->GetControlValue(ControlElement, ERigControlValueType::Current);
+	FVector3f Value = ControlValue.Get<FVector3f>();
+
+	EControlRigContextChannelToKey ChannelsToKey = (EControlRigContextChannelToKey)Context.KeyMask;
+	if (EnumHasAnyFlags(ChannelsToKey, EControlRigContextChannelToKey::RotationX) == false)
+	{
+		Val.X = Value.X;
+	}
+	if (EnumHasAnyFlags(ChannelsToKey, EControlRigContextChannelToKey::RotationZ) == false)
+	{
+		Val.Y = Value.Y;
+	}
+	if (EnumHasAnyFlags(ChannelsToKey, EControlRigContextChannelToKey::RotationY) == false)
+	{
+		Val.Z = Value.Z;
+	}
+}
+
+void UAnimDetailControlsProxyRotation::SetControlRigElementValueFromCurrent(UControlRig* ControlRig, FRigControlElement* ControlElement, const FRigControlModifiedContext& Context)
+{
+	if (ControlElement && ControlRig && ControlElement->Settings.ControlType == ERigControlType::Rotator)
+	{
+		FVector3f Val = Rotation.ToVector3f();
+		SetRotationValuesFromContext(ControlRig, ControlElement, Context, Val);
+		ControlRig->SetControlValue<FVector3f>(ControlElement->GetKey().Name, Val, true, Context, false);
+		ControlRig->Evaluate_AnyThread();
+	}
+}
+
+bool UAnimDetailControlsProxyRotation::PropertyIsOnProxy(FProperty* Property, FProperty* MemberProperty)
+{
+	return ((Property && Property->GetFName() == GET_MEMBER_NAME_CHECKED(UAnimDetailControlsProxyRotation, Rotation)) ||
+		(MemberProperty && MemberProperty->GetFName() == GET_MEMBER_NAME_CHECKED(UAnimDetailControlsProxyRotation, Rotation))
+		);
+}
+
+void UAnimDetailControlsProxyRotation::ClearMultipleFlags()
+{
+	Rotation.State.bXMultiple = false;
+	Rotation.State.bYMultiple = false;
+	Rotation.State.bZMultiple = false;
+}
+
+void UAnimDetailControlsProxyRotation::SetMultipleFlags(const FVector3f& ValA, const FVector3f& ValB)
+{
+	if (FMath::IsNearlyEqual(ValA.X, ValB.X) == false)
+	{
+		Rotation.State.bXMultiple = true;
+	}
+	if (FMath::IsNearlyEqual(ValA.Y, ValB.Y) == false)
+	{
+		Rotation.State.bYMultiple = true;
+	}
+	if (FMath::IsNearlyEqual(ValA.Z, ValB.Z) == false)
+	{
+		Rotation.State.bZMultiple = true;
+	}
+}
+
+void UAnimDetailControlsProxyRotation::ValueChanged()
+{
+	TOptional<FVector3f> LastValue;
+	ClearMultipleFlags();
+
+	for (const TPair<UControlRig*, FControlRigProxyItem>& Items : ControlRigItems)
+	{
+		if (UControlRig* ControlRig = Items.Value.ControlRig.Get())
+		{
+			for (FRigControlElement* ControlElement : Items.Value.ControlElements)
+			{
+				FVector3f Value = FVector3f::ZeroVector;
+				if (ControlElement->Settings.ControlType == ERigControlType::Rotator)
+				{
+					FRigControlValue ControlValue = ControlRig->GetControlValue(ControlElement, ERigControlValueType::Current);
+					Value = ControlValue.Get<FVector3f>();
+				}
+				if (LastValue.IsSet())
+				{
+					const FVector3f LastVal = LastValue.GetValue();
+					SetMultipleFlags(LastVal, Value);
+				}
+				else
+				{
+					LastValue = Value;
+				}
+			}
+		}
+	}
+
+	const FVector3f LastVal = LastValue.IsSet() ? LastValue.GetValue() : FVector3f::ZeroVector;
+
+	//set from values, note that the fact if they were multiples or not was already set so we need to set them on these values.
+	FAnimDetailProxyRotation Val(LastVal, Rotation.State);
+
+	const FName PropName("Rotation");
+	FTrackInstancePropertyBindings Binding(PropName, PropName.ToString());
+	Binding.CallFunction<FAnimDetailProxyRotation>(*this, Val);
+
+}
+
+bool UAnimDetailControlsProxyRotation::IsMultiple(const FName& PropertyName) const
+{
+	if (PropertyName == GET_MEMBER_NAME_CHECKED(FAnimDetailProxyRotation, RX))
+	{
+		return Rotation.State.bXMultiple;
+	}
+	else if (PropertyName == GET_MEMBER_NAME_CHECKED(FAnimDetailProxyRotation, RY))
+	{
+		return Rotation.State.bYMultiple;
+	}
+	else if (PropertyName == GET_MEMBER_NAME_CHECKED(FAnimDetailProxyRotation, RZ))
+	{
+		return Rotation.State.bZMultiple;
+	}
+	return false;
+}
+
+EControlRigContextChannelToKey UAnimDetailControlsProxyRotation::GetChannelToKeyFromPropertyName(const FName& PropertyName) const
+{
+	if (PropertyName == GET_MEMBER_NAME_CHECKED(FAnimDetailProxyRotation, RX))
+	{
+		return EControlRigContextChannelToKey::RotationX;
+	}
+	else if (PropertyName == GET_MEMBER_NAME_CHECKED(FAnimDetailProxyRotation, RY))
+	{
+		return EControlRigContextChannelToKey::RotationY;
+	}
+	else if (PropertyName == GET_MEMBER_NAME_CHECKED(FAnimDetailProxyRotation, RZ))
+	{
+		return EControlRigContextChannelToKey::RotationZ;
+	}
+
+	return EControlRigContextChannelToKey::AllTransform;
+}
+
+EControlRigContextChannelToKey UAnimDetailControlsProxyRotation::GetChannelToKeyFromChannelName(const FString& InChannelName) const
+{
+	if (InChannelName == TEXT("X"))
+	{
+		return EControlRigContextChannelToKey::RotationX;
+	}
+	else if (InChannelName == TEXT("Y"))
+	{
+		return EControlRigContextChannelToKey::RotationY;
+	}
+	else if (InChannelName == TEXT("Z"))
+	{
+		return EControlRigContextChannelToKey::RotationZ;
+	}
+	return EControlRigContextChannelToKey::AllTransform;
+}
+
+void UAnimDetailControlsProxyRotation::GetChannelSelectionState(TWeakPtr<FCurveEditor>& CurveEditor, FAnimDetailVectorSelection& LocationSelectionCache, FAnimDetailVectorSelection& RotationSelectionCache,
+	FAnimDetailVectorSelection& ScaleSelectionCache)
+{
+	RotationSelectionCache.XSelected = CachePropertySelection(CurveEditor, this, GET_MEMBER_NAME_CHECKED(FAnimDetailProxyRotation, RX));
+	RotationSelectionCache.YSelected = CachePropertySelection(CurveEditor, this, GET_MEMBER_NAME_CHECKED(FAnimDetailProxyRotation, RY));
+	RotationSelectionCache.ZSelected = CachePropertySelection(CurveEditor, this, GET_MEMBER_NAME_CHECKED(FAnimDetailProxyRotation, RZ));
+}
+
+//////UAnimDetailControlsProxyScale////////
+
+static void SetScaleValuesFromContext(UControlRig* ControlRig, FRigControlElement* ControlElement, const FRigControlModifiedContext& Context, FVector3f& TScale)
+{
+	FRigControlValue ControlValue = ControlRig->GetControlValue(ControlElement, ERigControlValueType::Current);
+	FVector3f Value = ControlValue.Get<FVector3f>();
+
+	EControlRigContextChannelToKey ChannelsToKey = (EControlRigContextChannelToKey)Context.KeyMask;
+
+	if (EnumHasAnyFlags(ChannelsToKey, EControlRigContextChannelToKey::ScaleX) == false)
+	{
+		TScale.X = Value.X;
+	}
+	if (EnumHasAnyFlags(ChannelsToKey, EControlRigContextChannelToKey::ScaleY) == false)
+	{
+		TScale.Y = Value.Y;
+	}
+	if (EnumHasAnyFlags(ChannelsToKey, EControlRigContextChannelToKey::ScaleZ) == false)
+	{
+		TScale.Z = Value.Z;
+	}
+}
+void UAnimDetailControlsProxyScale::SetControlRigElementValueFromCurrent(UControlRig* ControlRig, FRigControlElement* ControlElement, const FRigControlModifiedContext& Context)
+{
+	if (ControlElement && ControlRig && ControlElement->Settings.ControlType == ERigControlType::Scale)
+	{
+		FVector3f Val = Scale.ToVector3f();
+		SetScaleValuesFromContext(ControlRig, ControlElement, Context, Val);
+		ControlRig->SetControlValue<FVector3f>(ControlElement->GetKey().Name, Val, true, Context, false);
+		ControlRig->Evaluate_AnyThread();
+	}
+}
+
+bool UAnimDetailControlsProxyScale::PropertyIsOnProxy(FProperty* Property, FProperty* MemberProperty)
+{
+	return ((Property && Property->GetFName() == GET_MEMBER_NAME_CHECKED(UAnimDetailControlsProxyScale, Scale)) ||
+		(MemberProperty && MemberProperty->GetFName() == GET_MEMBER_NAME_CHECKED(UAnimDetailControlsProxyScale, Scale))
+		);
+}
+
+void UAnimDetailControlsProxyScale::ClearMultipleFlags()
+{
+	Scale.State.bXMultiple = false;
+	Scale.State.bYMultiple = false;
+	Scale.State.bZMultiple = false;
+}
+
+void UAnimDetailControlsProxyScale::SetMultipleFlags(const FVector3f& ValA, const FVector3f& ValB)
+{
+	if (FMath::IsNearlyEqual(ValA.X, ValB.X) == false)
+	{
+		Scale.State.bXMultiple = true;
+	}
+	if (FMath::IsNearlyEqual(ValA.Y, ValB.Y) == false)
+	{
+		Scale.State.bYMultiple = true;
+	}
+	if (FMath::IsNearlyEqual(ValA.Z, ValB.Z) == false)
+	{
+		Scale.State.bZMultiple = true;
+	}
+}
+
+void UAnimDetailControlsProxyScale::ValueChanged()
+{
+	TOptional<FVector3f> LastValue;
+	ClearMultipleFlags();
+
+	for (const TPair<UControlRig*, FControlRigProxyItem>& Items : ControlRigItems)
+	{
+		if (UControlRig* ControlRig = Items.Value.ControlRig.Get())
+		{
+			for (FRigControlElement* ControlElement : Items.Value.ControlElements)
+			{
+				FVector3f Value = FVector3f::ZeroVector;
+				if (ControlElement->Settings.ControlType == ERigControlType::Scale)
+				{
+					FRigControlValue ControlValue = ControlRig->GetControlValue(ControlElement, ERigControlValueType::Current);
+					Value = ControlValue.Get<FVector3f>();
+				}
+				if (LastValue.IsSet())
+				{
+					const FVector3f LastVal = LastValue.GetValue();
+					SetMultipleFlags(LastVal, Value);
+				}
+				else
+				{
+					LastValue = Value;
+				}
+			}
+		}
+	}
+
+	const FVector3f LastVal = LastValue.IsSet() ? LastValue.GetValue() : FVector3f::ZeroVector;
+	//set from values, note that the fact if they were multiples or not was already set so we need to set them on these values.
+	FAnimDetailProxyScale Value(LastVal, Scale.State);
+
+	const FName PropName("Scale");
+	FTrackInstancePropertyBindings Binding(PropName, PropName.ToString());
+	Binding.CallFunction<FAnimDetailProxyScale>(*this, Value);
+
+}
+
+bool UAnimDetailControlsProxyScale::IsMultiple(const FName& PropertyName) const
+{
+	if (PropertyName == GET_MEMBER_NAME_CHECKED(FAnimDetailProxyScale, SX))
+	{
+		return Scale.State.bXMultiple;
+	}
+	else if (PropertyName == GET_MEMBER_NAME_CHECKED(FAnimDetailProxyScale, SY))
+	{
+		return Scale.State.bYMultiple;
+	}
+	else if (PropertyName == GET_MEMBER_NAME_CHECKED(FAnimDetailProxyScale, SZ))
+	{
+		return Scale.State.bZMultiple;
+	}
+
+	return false;
+}
+
+EControlRigContextChannelToKey UAnimDetailControlsProxyScale::GetChannelToKeyFromPropertyName(const FName& PropertyName) const
+{
+	if (PropertyName == GET_MEMBER_NAME_CHECKED(FAnimDetailProxyScale, SX))
+	{
+		return EControlRigContextChannelToKey::ScaleX;
+	}
+	else if (PropertyName == GET_MEMBER_NAME_CHECKED(FAnimDetailProxyScale, SY))
+	{
+		return EControlRigContextChannelToKey::ScaleY;
+	}
+	else if (PropertyName == GET_MEMBER_NAME_CHECKED(FAnimDetailProxyScale, SZ))
+	{
+		return EControlRigContextChannelToKey::ScaleZ;
+	}
+
+	return EControlRigContextChannelToKey::AllTransform;
+}
+
+EControlRigContextChannelToKey UAnimDetailControlsProxyScale::GetChannelToKeyFromChannelName(const FString& InChannelName) const
+{
+	if (InChannelName == TEXT("X"))
+	{
+		return EControlRigContextChannelToKey::ScaleX;
+	}
+	else if (InChannelName == TEXT("Y"))
+	{
+		return EControlRigContextChannelToKey::ScaleY;
+	}
+	else if (InChannelName == TEXT("Z"))
+	{
+		return EControlRigContextChannelToKey::ScaleZ;
+	}
+	return EControlRigContextChannelToKey::AllTransform;
+}
+
+void UAnimDetailControlsProxyScale::GetChannelSelectionState(TWeakPtr<FCurveEditor>& CurveEditor, FAnimDetailVectorSelection& LocationSelectionCache, FAnimDetailVectorSelection& RotationSelectionCache,
+	FAnimDetailVectorSelection& ScaleSelectionCache)
+{
+	ScaleSelectionCache.XSelected = CachePropertySelection(CurveEditor, this, GET_MEMBER_NAME_CHECKED(FAnimDetailProxyScale, SX));
+	ScaleSelectionCache.YSelected = CachePropertySelection(CurveEditor, this, GET_MEMBER_NAME_CHECKED(FAnimDetailProxyScale, SY));
+	ScaleSelectionCache.ZSelected = CachePropertySelection(CurveEditor, this, GET_MEMBER_NAME_CHECKED(FAnimDetailProxyScale, SZ));
+}
+
+//////UAnimDetailControlsProxyVector2D////////
+static void SetVector2DValuesFromContext(UControlRig* ControlRig, FRigControlElement* ControlElement, const FRigControlModifiedContext& Context, FVector2D& Val)
+{
+	FRigControlValue ControlValue = ControlRig->GetControlValue(ControlElement, ERigControlValueType::Current);
+	FVector3f Value = ControlValue.Get<FVector3f>();
+
+	EControlRigContextChannelToKey ChannelsToKey = (EControlRigContextChannelToKey)Context.KeyMask;
+	if (EnumHasAnyFlags(ChannelsToKey, EControlRigContextChannelToKey::RotationX) == false)
+	{
+		Val.X = Value.X;
+	}
+	if (EnumHasAnyFlags(ChannelsToKey, EControlRigContextChannelToKey::RotationZ) == false)
+	{
+		Val.Y = Value.Y;
+	}
+}
+void UAnimDetailControlsProxyVector2D::SetControlRigElementValueFromCurrent(UControlRig* ControlRig, FRigControlElement* ControlElement, const FRigControlModifiedContext& Context)
+{
+	if (ControlElement && ControlRig && ControlElement->Settings.ControlType == ERigControlType::Vector2D)
+	{
+		FVector2D Val = Vector2D.ToVector2D();
+		SetVector2DValuesFromContext(ControlRig, ControlElement, Context, Val);
+		ControlRig->SetControlValue<FVector2D>(ControlElement->GetKey().Name, Val, true, Context, false);
+		ControlRig->Evaluate_AnyThread();
+	}
+}
+
+bool UAnimDetailControlsProxyVector2D::PropertyIsOnProxy(FProperty* Property, FProperty* MemberProperty)
+{
+	return ((Property && Property->GetFName() == GET_MEMBER_NAME_CHECKED(UAnimDetailControlsProxyVector2D, Vector2D)) ||
+		(MemberProperty && MemberProperty->GetFName() == GET_MEMBER_NAME_CHECKED(UAnimDetailControlsProxyVector2D, Vector2D))
+		);
+}
+
+void UAnimDetailControlsProxyVector2D::ClearMultipleFlags()
+{
+	Vector2D.State.bXMultiple = false;
+	Vector2D.State.bYMultiple = false;
+}
+
+void UAnimDetailControlsProxyVector2D::SetMultipleFlags(const FVector2D& ValA, const FVector2D& ValB)
+{
+	if (FMath::IsNearlyEqual(ValA.X, ValB.X) == false)
+	{
+		Vector2D.State.bXMultiple = true;
+	}
+	if (FMath::IsNearlyEqual(ValA.Y, ValB.Y) == false)
+	{
+		Vector2D.State.bYMultiple = true;
+	}
+}
+
+void UAnimDetailControlsProxyVector2D::ValueChanged()
+{
+	TOptional<FVector2D> LastValue;
+	ClearMultipleFlags();
+
+	for (const TPair<UControlRig*, FControlRigProxyItem>& Items : ControlRigItems)
+	{
+		if (UControlRig* ControlRig = Items.Value.ControlRig.Get())
+		{
+			for (FRigControlElement* ControlElement : Items.Value.ControlElements)
+			{
+				FVector2D Value = FVector2D::ZeroVector;
+				if (ControlElement->Settings.ControlType == ERigControlType::Vector2D)
+				{
+					FRigControlValue ControlValue = ControlRig->GetControlValue(ControlElement, ERigControlValueType::Current);
+					//FVector2D version deleted for some reason so need to convert
+					FVector3f Val = ControlValue.Get<FVector3f>();
+					Value = FVector2D(Val.X, Val.Y);
+				}
+				if (LastValue.IsSet())
+				{
+					const FVector2D LastVal = LastValue.GetValue();
+					SetMultipleFlags(LastVal, Value);
+				}
+				else
+				{
+					LastValue = Value;
+				}
+			}
+		}
+	}
+
+	const FVector2D LastVal = LastValue.IsSet() ? LastValue.GetValue() : FVector2D::ZeroVector;
+
+	//set from values, note that the fact if they were multiples or not was already set so we need to set them on these values.
+	FAnimDetailProxyVector2D Value(LastVal, Vector2D.State);
+
+	const FName PropName("Vector2D");
+	FTrackInstancePropertyBindings Binding(PropName, PropName.ToString());
+	Binding.CallFunction<FAnimDetailProxyVector2D>(*this, Value);
+}
+
+bool UAnimDetailControlsProxyVector2D::IsMultiple(const FName& PropertyName) const
+{
+	if (PropertyName == GET_MEMBER_NAME_CHECKED(FAnimDetailProxyVector2D, X))
+	{
+		return Vector2D.State.bXMultiple;
+	}
+	else if (PropertyName == GET_MEMBER_NAME_CHECKED(FAnimDetailProxyVector2D, Y))
+	{
+		return Vector2D.State.bYMultiple;
+	}
+
+	return false;
+}
+
+EControlRigContextChannelToKey UAnimDetailControlsProxyVector2D::GetChannelToKeyFromPropertyName(const FName& PropertyName) const
+{
+	if (PropertyName == GET_MEMBER_NAME_CHECKED(FAnimDetailProxyVector2D, X))
+	{
+		return EControlRigContextChannelToKey::TranslationX;
+	}
+	else if (PropertyName == GET_MEMBER_NAME_CHECKED(FAnimDetailProxyVector2D, Y))
+	{
+		return EControlRigContextChannelToKey::TranslationY;
+	}
+
+	return EControlRigContextChannelToKey::AllTransform;
+}
+
+EControlRigContextChannelToKey UAnimDetailControlsProxyVector2D::GetChannelToKeyFromChannelName(const FString& InChannelName) const
+{
+	if (InChannelName == TEXT("X"))
+	{
+		return EControlRigContextChannelToKey::TranslationX;
+	}
+	else if (InChannelName == TEXT("Y"))
+	{
+		return EControlRigContextChannelToKey::TranslationY;
+	}
+	return EControlRigContextChannelToKey::AllTransform;
+}
+
+void UAnimDetailControlsProxyVector2D::GetChannelSelectionState(TWeakPtr<FCurveEditor>& CurveEditor, FAnimDetailVectorSelection& LocationSelectionCache, FAnimDetailVectorSelection& RotationSelectionCache,
+	FAnimDetailVectorSelection& ScaleSelectionCache)
+{
+	LocationSelectionCache.XSelected = CachePropertySelection(CurveEditor, this, GET_MEMBER_NAME_CHECKED(FAnimDetailProxyVector2D, X));
+	LocationSelectionCache.YSelected = CachePropertySelection(CurveEditor, this, GET_MEMBER_NAME_CHECKED(FAnimDetailProxyVector2D, Y));
+}
+
+//////UAnimDetailControlsProxyFloat////////
+
+void UAnimDetailControlsProxyFloat::UpdatePropertyNames(IDetailLayoutBuilder& DetailBuilder)
+{
+	TSharedPtr<IPropertyHandle> ValuePropertyHandle = DetailBuilder.GetProperty(GET_MEMBER_NAME_CHECKED(UAnimDetailControlsProxyFloat, Float), GetClass());
+	if (ValuePropertyHandle)
+	{
+		ValuePropertyHandle->SetPropertyDisplayName(FText::FromName(GetName()));
+	}
+	if (bIsIndividual)
+	{
+		const FString PropertyPath = FString::Printf(TEXT("%s.%s"), GET_MEMBER_NAME_STRING_CHECKED(UAnimDetailControlsProxyFloat, Float), GET_MEMBER_NAME_STRING_CHECKED(FAnimDetailProxyFloat, Float));
+		ValuePropertyHandle = DetailBuilder.GetProperty(FName(*PropertyPath), GetClass());
+		if (ValuePropertyHandle)
+		{
+			ValuePropertyHandle->SetPropertyDisplayName(FText::FromName(GetName()));
+		}
+	}
+}
+
+void UAnimDetailControlsProxyFloat::SetControlRigElementValueFromCurrent(UControlRig* ControlRig, FRigControlElement* ControlElement, const FRigControlModifiedContext& Context)
+{
+	if (ControlElement && ControlRig && (ControlElement->Settings.ControlType == ERigControlType::Float ||
+		ControlElement->Settings.ControlType == ERigControlType::ScaleFloat))
+	{
+		float Val = Float.Float;
+		ControlRig->SetControlValue<float>(ControlElement->GetKey().Name, Val, true, Context, false);
+		ControlRig->Evaluate_AnyThread();
+	}
+}
+
+bool UAnimDetailControlsProxyFloat::PropertyIsOnProxy(FProperty* Property, FProperty* MemberProperty)
+{
+	return ((Property && Property->GetFName() == GET_MEMBER_NAME_CHECKED(FAnimDetailProxyFloat, Float)) ||
+		(MemberProperty && MemberProperty->GetFName() == GET_MEMBER_NAME_CHECKED(UAnimDetailControlsProxyFloat, Float))
+		);
+}
+
+void UAnimDetailControlsProxyFloat::ClearMultipleFlags()
+{
+	State.bMultiple = false;
+}
+
+void UAnimDetailControlsProxyFloat::SetMultipleFlags(const float& ValA, const float& ValB)
+{
+	if (FMath::IsNearlyEqual(ValA, ValB) == false)
+	{
+		State.bMultiple = true;
+	}
+}
+
+void UAnimDetailControlsProxyFloat::ValueChanged()
+{
+	TOptional<float> LastValue;
+	ClearMultipleFlags();
+
+	for (const TPair<UControlRig*, FControlRigProxyItem>& Items : ControlRigItems)
+	{
+		if (UControlRig* ControlRig = Items.Value.ControlRig.Get())
+		{
+			for (FRigControlElement* ControlElement : Items.Value.ControlElements)
+			{
+				float Value = 0.0f;
+				if (ControlElement->Settings.ControlType == ERigControlType::Float ||
+					(ControlElement->Settings.ControlType == ERigControlType::ScaleFloat))
+				{
+					FRigControlValue ControlValue = ControlRig->GetControlValue(ControlElement, ERigControlValueType::Current);
+					Value = ControlValue.Get<float>();
+				}
+				if (LastValue.IsSet())
+				{
+					const float LastVal = LastValue.GetValue();
+					SetMultipleFlags(LastVal, Value);
+				}
+				else
+				{
+					LastValue = Value;
+				}
+			}
+		}
+	}
+	for (TPair <UObject*, FSequencerProxyItem>& SItems : SequencerItems)
+	{
+		for (FBindingAndTrack& Binding: SItems.Value.Bindings)
+		{
+			if (Binding.Binding.IsValid())
+			{
+				if (TOptional<double> Value = Binding.Binding->GetOptionalValue<double>(*SItems.Key))
+				{
+					if (LastValue.IsSet())
+					{
+						const float LastVal = LastValue.GetValue();
+						SetMultipleFlags(LastVal, Value.GetValue());
+					}
+					else
+					{
+						LastValue = Value.GetValue();
+					}
+				}
+				
+			}
+		}
+	}
+
+	const float LastVal = LastValue.IsSet() ? LastValue.GetValue() : 0.0f;
+
+	//set from values, note that the fact if they were multiples or not was already set so we need to set them on these values.
+
+	const FName PropName("Float");
+	FTrackInstancePropertyBindings Binding(PropName, PropName.ToString());
+
+	FAnimDetailProxyFloat Value;
+	Value.Float = LastVal;
+	Binding.CallFunction<FAnimDetailProxyFloat>(*this, Value);
+}
+
+bool UAnimDetailControlsProxyFloat::IsMultiple(const FName& PropertyName) const
+{
+	if (PropertyName == GET_MEMBER_NAME_CHECKED(FAnimDetailProxyFloat, Float))
+	{
+		return State.bMultiple;
+	}
+
+	return false;
+}
+
+EControlRigContextChannelToKey UAnimDetailControlsProxyFloat::GetChannelToKeyFromPropertyName(const FName& PropertyName) const
+{
+	if (PropertyName == GET_MEMBER_NAME_CHECKED(FAnimDetailProxyFloat, Float))
+	{
+		return EControlRigContextChannelToKey::TranslationX;
+	}
+
+	return EControlRigContextChannelToKey::AllTransform;
+}
+
+EControlRigContextChannelToKey UAnimDetailControlsProxyFloat::GetChannelToKeyFromChannelName(const FString& InChannelName) const
+{
+	if (InChannelName == TEXT("Float"))
+	{
+		return EControlRigContextChannelToKey::TranslationX;
+	}
+	TArray<FRigControlElement*> Elements = GetControlElements();
+	for (FRigControlElement* Element : Elements)
+	{
+		if (Element && Element->GetDisplayName() == InChannelName)
+		{
+			return EControlRigContextChannelToKey::TranslationX;
+		}
+	}
+	return EControlRigContextChannelToKey::AllTransform;
+}
+
+void UAnimDetailControlsProxyFloat::GetChannelSelectionState(TWeakPtr<FCurveEditor>& CurveEditor, FAnimDetailVectorSelection& LocationSelectionCache, FAnimDetailVectorSelection& RotationSelectionCache,
+	FAnimDetailVectorSelection& ScaleSelectionCache)
+{
+	LocationSelectionCache.XSelected = CachePropertySelection(CurveEditor, this, GET_MEMBER_NAME_CHECKED(FAnimDetailProxyFloat, Float));
+}
+
+void UAnimDetailControlsProxyFloat::SetBindingValueFromCurrent(UObject* InObject, TSharedPtr<FTrackInstancePropertyBindings>& Binding, FRigControlModifiedContext& Context)
+{
+	if (InObject && Binding.IsValid())
+	{
+		Binding->SetCurrentValue<double>(*InObject, Float.Float);
+	}
+}
+//////UAnimDetailControlsProxyBool////////
+
+void UAnimDetailControlsProxyBool::UpdatePropertyNames(IDetailLayoutBuilder& DetailBuilder)
+{
+	TSharedPtr<IPropertyHandle> ValuePropertyHandle = DetailBuilder.GetProperty(GET_MEMBER_NAME_CHECKED(UAnimDetailControlsProxyBool, Bool), GetClass());
+	if (ValuePropertyHandle)
+	{
+		ValuePropertyHandle->SetPropertyDisplayName(FText::FromName(GetName()));
+	}
+	if (bIsIndividual)
+	{
+		const FString PropertyPath = FString::Printf(TEXT("%s.%s"), GET_MEMBER_NAME_STRING_CHECKED(UAnimDetailControlsProxyBool, Bool), GET_MEMBER_NAME_STRING_CHECKED(FAnimDetailProxyBool, Bool));
+		ValuePropertyHandle = DetailBuilder.GetProperty(FName(*PropertyPath), GetClass());
+		if (ValuePropertyHandle)
+		{
+			ValuePropertyHandle->SetPropertyDisplayName(FText::FromName(GetName()));
+		}
+	}
+}
+
+void UAnimDetailControlsProxyBool::SetControlRigElementValueFromCurrent(UControlRig* ControlRig, FRigControlElement* ControlElement, const FRigControlModifiedContext& Context)
+{
+	if (ControlElement && ControlRig && (ControlElement->Settings.ControlType == ERigControlType::Bool))
+	{
+		bool Val = Bool.Bool;
+		ControlRig->SetControlValue<bool>(ControlElement->GetKey().Name, Val, true, Context, false);
+		ControlRig->Evaluate_AnyThread();
+	}
+}
+
+bool UAnimDetailControlsProxyBool::PropertyIsOnProxy(FProperty* Property, FProperty* MemberProperty)
+{
+	return ((Property && Property->GetFName() == GET_MEMBER_NAME_CHECKED(FAnimDetailProxyBool, Bool)) ||
+		(MemberProperty && MemberProperty->GetFName() == GET_MEMBER_NAME_CHECKED(UAnimDetailControlsProxyBool, Bool))
+		);
+}
+
+void UAnimDetailControlsProxyBool::ClearMultipleFlags()
+{
+	State.bMultiple = false;
+}
+
+void UAnimDetailControlsProxyBool::SetMultipleFlags(const bool& ValA, const bool& ValB)
+{
+	if (ValA == ValB)
+	{
+		State.bMultiple = true;
+	}
+}
+
+void UAnimDetailControlsProxyBool::ValueChanged()
+{
+	TOptional<bool> LastValue;
+	ClearMultipleFlags();
+
+	for (const TPair<UControlRig*, FControlRigProxyItem>& Items : ControlRigItems)
+	{
+		if (UControlRig* ControlRig = Items.Value.ControlRig.Get())
+		{
+			for (FRigControlElement* ControlElement : Items.Value.ControlElements)
+			{
+				bool Value = false;
+				if (ControlElement->Settings.ControlType == ERigControlType::Bool)
+				{
+					FRigControlValue ControlValue = ControlRig->GetControlValue(ControlElement, ERigControlValueType::Current);
+					Value = ControlValue.Get<bool>();
+				}
+				if (LastValue.IsSet())
+				{
+					const bool LastVal = LastValue.GetValue();
+					SetMultipleFlags(LastVal, Value);
+				}
+				else
+				{
+					LastValue = Value;
+				}
+			}
+		}
+	}
+	for (TPair <UObject*, FSequencerProxyItem>& SItems : SequencerItems)
+	{
+		for (FBindingAndTrack& Binding : SItems.Value.Bindings)
+		{
+			if (Binding.Binding.IsValid())
+			{
+				if (TOptional<bool> Value = Binding.Binding->GetOptionalValue<bool>(*SItems.Key))
+				{
+					if (LastValue.IsSet())
+					{
+						const bool LastVal = LastValue.GetValue();
+						SetMultipleFlags(LastVal, Value.GetValue());
+					}
+					else
+					{
+						LastValue = Value.GetValue();
+					}
+				}
+
+			}
+		}
+	}
+
+	const bool LastVal = LastValue.IsSet() ? LastValue.GetValue() : 0.0f;
+
+	//set from values, note that the fact if they were multiples or not was already set so we need to set them on these values.
+
+	const FName PropName("Bool");
+	FTrackInstancePropertyBindings Binding(PropName, PropName.ToString());
+
+	FAnimDetailProxyBool Value;
+	Value.Bool = LastVal;
+	Binding.CallFunction<FAnimDetailProxyBool>(*this, Value);
+
+}
+
+bool UAnimDetailControlsProxyBool::IsMultiple(const FName& PropertyName) const
+{
+	if (PropertyName == GET_MEMBER_NAME_CHECKED(FAnimDetailProxyBool, Bool))
+	{
+		return State.bMultiple;
+	}
+
+	return false;
+}
+
+EControlRigContextChannelToKey UAnimDetailControlsProxyBool::GetChannelToKeyFromPropertyName(const FName& PropertyName) const
+{
+	if (PropertyName == GET_MEMBER_NAME_CHECKED(FAnimDetailProxyBool, Bool))
+	{
+		return EControlRigContextChannelToKey::TranslationX;
+	}
+
+	return EControlRigContextChannelToKey::AllTransform;
+}
+
+EControlRigContextChannelToKey UAnimDetailControlsProxyBool::GetChannelToKeyFromChannelName(const FString& InChannelName) const
+{
+	if (InChannelName == TEXT("Bool"))
+	{
+		return EControlRigContextChannelToKey::TranslationX;
+	}
+	TArray<FRigControlElement*> Elements = GetControlElements();
+	for (FRigControlElement* Element : Elements)
+	{
+		if (Element && Element->GetDisplayName() == InChannelName)
+		{
+			return EControlRigContextChannelToKey::TranslationX;
+		}
+	}
+	return EControlRigContextChannelToKey::AllTransform;
+}
+
+void UAnimDetailControlsProxyBool::GetChannelSelectionState(TWeakPtr<FCurveEditor>& CurveEditor, FAnimDetailVectorSelection& LocationSelectionCache, FAnimDetailVectorSelection& RotationSelectionCache,
+	FAnimDetailVectorSelection& ScaleSelectionCache)
+{
+	LocationSelectionCache.XSelected = CachePropertySelection(CurveEditor, this, GET_MEMBER_NAME_CHECKED(FAnimDetailProxyBool, Bool));
+}
+
+void UAnimDetailControlsProxyBool::SetBindingValueFromCurrent(UObject* InObject, TSharedPtr<FTrackInstancePropertyBindings>& Binding, FRigControlModifiedContext& Context)
+{
+	if (InObject && Binding.IsValid())
+	{
+		Binding->SetCurrentValue<bool>(*InObject, Bool.Bool);
+	}
+}
+//////UAnimDetailControlsProxyInteger////////
+
+void UAnimDetailControlsProxyInteger::UpdatePropertyNames(IDetailLayoutBuilder& DetailBuilder)
+{
+	TSharedPtr<IPropertyHandle> ValuePropertyHandle = DetailBuilder.GetProperty(GET_MEMBER_NAME_CHECKED(UAnimDetailControlsProxyInteger, Integer), GetClass());
+	if (ValuePropertyHandle)
+	{
+		ValuePropertyHandle->SetPropertyDisplayName(FText::FromName(GetName()));
+	}
+	if (bIsIndividual)
+	{
+		const FString PropertyPath = FString::Printf(TEXT("%s.%s"), GET_MEMBER_NAME_STRING_CHECKED(UAnimDetailControlsProxyInteger, Integer), GET_MEMBER_NAME_STRING_CHECKED(FAnimDetailProxyInteger, Integer));
+		ValuePropertyHandle = DetailBuilder.GetProperty(FName(*PropertyPath), GetClass());
+		if (ValuePropertyHandle)
+		{
+			ValuePropertyHandle->SetPropertyDisplayName(FText::FromName(GetName()));
+		}
+	}
+}
+
+void UAnimDetailControlsProxyInteger::SetControlRigElementValueFromCurrent(UControlRig* ControlRig, FRigControlElement* ControlElement, const FRigControlModifiedContext& Context)
+{
+	if (ControlElement && ControlRig && (ControlElement->Settings.ControlType == ERigControlType::Integer &&
+		ControlElement->Settings.ControlEnum == nullptr))
+	{
+		int32 Val = (int32)Integer.Integer;
+		ControlRig->SetControlValue<int32>(ControlElement->GetKey().Name, Val, true, Context, false);
+		ControlRig->Evaluate_AnyThread();
+	}
+}
+
+bool UAnimDetailControlsProxyInteger::PropertyIsOnProxy(FProperty* Property, FProperty* MemberProperty)
+{
+	return ((Property && Property->GetFName() == GET_MEMBER_NAME_CHECKED(FAnimDetailProxyInteger, Integer)) ||
+		(MemberProperty && MemberProperty->GetFName() == GET_MEMBER_NAME_CHECKED(UAnimDetailControlsProxyInteger, Integer))
+		);
+}
+
+void UAnimDetailControlsProxyInteger::ClearMultipleFlags()
+{
+	State.bMultiple = false;
+}
+
+void UAnimDetailControlsProxyInteger::SetMultipleFlags(const int64& ValA, const int64& ValB)
+{
+	if (ValA != ValB)
+	{
+		State.bMultiple = true;
+	}
+}
+
+void UAnimDetailControlsProxyInteger::ValueChanged()
+{
+	TOptional<int64> LastValue;
+	ClearMultipleFlags();
+
+	for (const TPair<UControlRig*, FControlRigProxyItem>& Items : ControlRigItems)
+	{
+		if (UControlRig* ControlRig = Items.Value.ControlRig.Get())
+		{
+			for (FRigControlElement* ControlElement : Items.Value.ControlElements)
+			{
+				int64 Value = 0;
+				if (ControlElement->Settings.ControlType == ERigControlType::Integer)
+				{
+					FRigControlValue ControlValue = ControlRig->GetControlValue(ControlElement, ERigControlValueType::Current);
+					Value = ControlValue.Get<int32>();
+				}
+				if (LastValue.IsSet())
+				{
+					const int64 LastVal = LastValue.GetValue();
+					SetMultipleFlags(LastVal, Value);
+				}
+				else
+				{
+					LastValue = Value;
+				}
+			}
+		}
+		for (TPair <UObject*, FSequencerProxyItem>& SItems : SequencerItems)
+		{
+			for (FBindingAndTrack& Binding : SItems.Value.Bindings)
+			{
+				if (Binding.Binding.IsValid())
+				{
+					if (TOptional<int64> Value = Binding.Binding->GetOptionalValue<int64>(*SItems.Key))
+					{
+						if (LastValue.IsSet())
+						{
+							const float LastVal = LastValue.GetValue();
+							SetMultipleFlags(LastVal, Value.GetValue());
+						}
+						else
+						{
+							LastValue = Value.GetValue();
+						}
+					}
+
+				}
+			}
+		}
+	}
+
+	const int64 LastVal = LastValue.IsSet() ? LastValue.GetValue() : 0;
+
+	//set from values, note that the fact if they were multiples or not was already set so we need to set them on these values.
+
+	const FName PropName("Integer");
+	FTrackInstancePropertyBindings Binding(PropName, PropName.ToString());
+	FAnimDetailProxyInteger Value;
+	Value.Integer = LastVal;
+	Binding.CallFunction<FAnimDetailProxyInteger>(*this, Value);
+}
+
+bool UAnimDetailControlsProxyInteger::IsMultiple(const FName& PropertyName) const
+{
+	if (PropertyName == GET_MEMBER_NAME_CHECKED(FAnimDetailProxyInteger, Integer))
+	{
+		return State.bMultiple;
+	}
+	return false;
+}
+
+EControlRigContextChannelToKey UAnimDetailControlsProxyInteger::GetChannelToKeyFromPropertyName(const FName& PropertyName) const
+{
+	if (PropertyName == GET_MEMBER_NAME_CHECKED(FAnimDetailProxyInteger, Integer))
+	{
+		return EControlRigContextChannelToKey::TranslationX;
+	}
+
+	return EControlRigContextChannelToKey::AllTransform;
+}
+
+EControlRigContextChannelToKey UAnimDetailControlsProxyInteger::GetChannelToKeyFromChannelName(const FString& InChannelName) const
+{
+	if (InChannelName == TEXT("Integer"))
+	{
+		return EControlRigContextChannelToKey::TranslationX;
+	}
+	TArray<FRigControlElement*> Elements = GetControlElements();
+	for (FRigControlElement* Element : Elements)
+	{
+		if (Element && Element->GetDisplayName() == InChannelName)
+		{
+			return EControlRigContextChannelToKey::TranslationX;
+		}
+	}
+	return EControlRigContextChannelToKey::AllTransform;
+}
+
+void UAnimDetailControlsProxyInteger::GetChannelSelectionState(TWeakPtr<FCurveEditor>& CurveEditor, FAnimDetailVectorSelection& LocationSelectionCache, FAnimDetailVectorSelection& RotationSelectionCache,
+	FAnimDetailVectorSelection& ScaleSelectionCache)
+{
+	LocationSelectionCache.XSelected = CachePropertySelection(CurveEditor, this, GET_MEMBER_NAME_CHECKED(FAnimDetailProxyInteger, Integer));
+}
+
+void UAnimDetailControlsProxyInteger::SetBindingValueFromCurrent(UObject* InObject, TSharedPtr<FTrackInstancePropertyBindings>& Binding, FRigControlModifiedContext& Context)
+{
+	if (InObject && Binding.IsValid())
+	{
+		Binding->SetCurrentValue<int64>(*InObject, Integer.Integer);
+	}
+}
+//////UAnimDetailControlsProxyEnum////////
+
+void UAnimDetailControlsProxyEnum::UpdatePropertyNames(IDetailLayoutBuilder& DetailBuilder)
+{
+	TSharedPtr<IPropertyHandle> ValuePropertyHandle = DetailBuilder.GetProperty(GET_MEMBER_NAME_CHECKED(UAnimDetailControlsProxyEnum, Enum), GetClass());
+	if (ValuePropertyHandle)
+	{
+		ValuePropertyHandle->SetPropertyDisplayName(FText::FromName(GetName()));
+	}
+	if (bIsIndividual)
+	{
+		const FString PropertyPath = FString::Printf(TEXT("%s.%s"), GET_MEMBER_NAME_STRING_CHECKED(UAnimDetailControlsProxyEnum, Enum), GET_MEMBER_NAME_STRING_CHECKED(FControlRigEnumControlProxyValue, EnumIndex));
+		ValuePropertyHandle = DetailBuilder.GetProperty(FName(*PropertyPath), GetClass());
+		if (ValuePropertyHandle)
+		{
+			ValuePropertyHandle->SetPropertyDisplayName(FText::FromName(GetName()));
+		}
+	}
+}
+
+void UAnimDetailControlsProxyEnum::SetControlRigElementValueFromCurrent(UControlRig* ControlRig, FRigControlElement* ControlElement, const FRigControlModifiedContext& Context)
+{
+	if (ControlElement && ControlRig && (ControlElement->Settings.ControlType == ERigControlType::Integer &&
+		ControlElement->Settings.ControlEnum != nullptr))
+	{
+		int32 Val = Enum.EnumIndex;
+		ControlRig->SetControlValue<int32>(ControlElement->GetKey().Name, Val, true, Context, false);
+		ControlRig->Evaluate_AnyThread();
+	}
+}
+
+bool UAnimDetailControlsProxyEnum::PropertyIsOnProxy(FProperty* Property, FProperty* MemberProperty)
+{
+	return ((Property && Property->GetFName() == GET_MEMBER_NAME_CHECKED(FControlRigEnumControlProxyValue, EnumIndex)) ||
+		(MemberProperty && MemberProperty->GetFName() == GET_MEMBER_NAME_CHECKED(UAnimDetailControlsProxyEnum, Enum))
+	);
+}
+
+void UAnimDetailControlsProxyEnum::ClearMultipleFlags()
+{
+	State.bMultiple = false;
+}
+
+void UAnimDetailControlsProxyEnum::SetMultipleFlags(const int32& ValA, const int32& ValB)
+{
+	if (ValA != ValB)
+	{
+		State.bMultiple = true;
+	}
+}
+
+void UAnimDetailControlsProxyEnum::ValueChanged()
+{
+	TOptional<int32> LastValue;
+	ClearMultipleFlags();
+	TObjectPtr<UEnum> EnumType = nullptr;
+	for (const TPair<UControlRig*, FControlRigProxyItem>& Items : ControlRigItems)
+	{
+		if (UControlRig* ControlRig = Items.Value.ControlRig.Get())
+		{
+			for (FRigControlElement* ControlElement : Items.Value.ControlElements)
+			{
+				int32 Value = 0;
+				if (ControlElement->Settings.ControlType == ERigControlType::Integer &&
+					ControlElement->Settings.ControlEnum)
+				{
+					EnumType = ControlElement->Settings.ControlEnum;
+					FRigControlValue ControlValue = ControlRig->GetControlValue(ControlElement, ERigControlValueType::Current);
+					Value = ControlValue.Get<int32>();
+				}
+
+				if (LastValue.IsSet())
+				{
+					const int32 LastVal = LastValue.GetValue();
+					SetMultipleFlags(LastVal, Value);
+				}
+				else
+				{
+					LastValue = Value;
+				}
+			}
+		}
+	}
+
+	const int32 LastVal = LastValue.IsSet() ? LastValue.GetValue() : 0;
+
+	//set from values, note that the fact if they were multiples or not was already set so we need to set them on these values.
+	int32 Value(LastVal);
+
+	if (EnumType)
+	{
+		const FName PropertyName("Enum");
+		FTrackInstancePropertyBindings Binding(PropertyName, PropertyName.ToString());
+		FControlRigEnumControlProxyValue Val;
+		Val.EnumType = EnumType;
+		Val.EnumIndex = Value;
+		Binding.CallFunction<FControlRigEnumControlProxyValue>(*this, Val);
+	}
+
+}
+
+bool UAnimDetailControlsProxyEnum::IsMultiple(const FName& PropertyName) const
+{
+	if (PropertyName == GET_MEMBER_NAME_CHECKED(FControlRigEnumControlProxyValue, EnumIndex))
+	{
+		return State.bMultiple;
+	}
+
+	return false;
+}
+
+EControlRigContextChannelToKey UAnimDetailControlsProxyEnum::GetChannelToKeyFromPropertyName(const FName& PropertyName) const
+{
+	if (PropertyName == GET_MEMBER_NAME_CHECKED(FControlRigEnumControlProxyValue, EnumIndex))
+	{
+		return EControlRigContextChannelToKey::TranslationX;
+	}
+
+	return EControlRigContextChannelToKey::AllTransform;
+}
+
+EControlRigContextChannelToKey UAnimDetailControlsProxyEnum::GetChannelToKeyFromChannelName(const FString& InChannelName) const
+{
+	if (InChannelName == TEXT("EnumIndex"))
+	{
+		return EControlRigContextChannelToKey::TranslationX;
+	}
+	TArray<FRigControlElement*> Elements = GetControlElements();
+	for (FRigControlElement* Element : Elements)
+	{
+		if (Element && Element->GetDisplayName() == InChannelName)
+		{
+			return EControlRigContextChannelToKey::TranslationX;
+		}
+	}
+	return EControlRigContextChannelToKey::AllTransform;
+}
+
+void UAnimDetailControlsProxyEnum::GetChannelSelectionState(TWeakPtr<FCurveEditor>& CurveEditor, FAnimDetailVectorSelection& LocationSelectionCache, FAnimDetailVectorSelection& RotationSelectionCache,
+	FAnimDetailVectorSelection& ScaleSelectionCache)
+{
+	LocationSelectionCache.XSelected = CachePropertySelection(CurveEditor, this, GET_MEMBER_NAME_CHECKED(FControlRigEnumControlProxyValue, EnumIndex));
+}
+
+//////UControlDetailPanelControlProxies////////
+UControlRigDetailPanelControlProxies::UControlRigDetailPanelControlProxies()
+{
+	FCoreUObjectDelegates::OnObjectPropertyChanged.AddUObject(this, &UControlRigDetailPanelControlProxies::OnPostPropertyChanged);
+}
+
+UControlRigDetailPanelControlProxies::~UControlRigDetailPanelControlProxies()
+{
+	FCoreUObjectDelegates::OnObjectPropertyChanged.RemoveAll(this);
+}
+
+UControlRigControlsProxy* UControlRigDetailPanelControlProxies::FindProxy(UControlRig* ControlRig, FRigControlElement* ControlElement) const
+{
+	const FNameToProxyMap* ControlRigProxies = ControlRigOnlyProxies.Find(ControlRig);
+	if (ControlRigProxies)
+	{
+		TObjectPtr<UControlRigControlsProxy> const* Proxy = ControlRigProxies->NameToProxy.Find(ControlElement->GetFName());
+		if (Proxy && Proxy[0])
+		{
+			return Proxy[0];
+		}
+	}
+	return nullptr;
+}
+
+UControlRigControlsProxy* UControlRigDetailPanelControlProxies::FindProxy(UObject* InObject, FName PropertyName) const
+{
+	const FNameToProxyMap* SequencerProxies = SequencerOnlyProxies.Find(InObject);
+	if (SequencerProxies)
+	{
+		TObjectPtr<UControlRigControlsProxy> const* Proxy = SequencerProxies->NameToProxy.Find(PropertyName);
+		if (Proxy && Proxy[0])
+		{
+			return Proxy[0];
+		}
+	}
+	return nullptr;
+}
+
+void UControlRigDetailPanelControlProxies::OnPostPropertyChanged(UObject* InObject, FPropertyChangedEvent& InPropertyChangedEvent)
+{
+	const FName MemberPropertyName = InPropertyChangedEvent.MemberProperty != nullptr ? InPropertyChangedEvent.MemberProperty->GetFName() : NAME_None;
+	/*
+	const bool bTransformationChanged =
+		(MemberPropertyName == USceneComponent::GetRelativeLocationPropertyName() ||
+			MemberPropertyName == USceneComponent::GetRelativeRotationPropertyName() ||
+			MemberPropertyName == USceneComponent::GetRelativeScale3DPropertyName());
+
+	if (InObject && bTransformationChanged)
+	{
+		OnTransformChanged(*InObject);
+	}
+	*/
+}
+
+void UControlRigDetailPanelControlProxies::ResetSequencerProxies(TMap<ERigControlType, FSequencerProxyPerType>& ProxyPerType)
+{
+	//Remove
+	SelectedSequencerProxies.SetNum(0);
+	for (TPair<ERigControlType, FSequencerProxyPerType>& Pair : ProxyPerType)
+	{
+		for (TPair<UObject*, TArray<FBindingAndTrack>>& PerTypePair : Pair.Value.Bindings)
+		{
+			for (FBindingAndTrack& Binding : PerTypePair.Value)
+			{
+				if (UControlRigControlsProxy* Proxy = AddProxy(PerTypePair.Key, Pair.Key, Binding.WeakTrack, Binding.Binding))
+				{
+					SelectedSequencerProxies.Add(Proxy);
+				}
+			}
+		}
+	}
+}
+
+const TArray<UControlRigControlsProxy*> UControlRigDetailPanelControlProxies::GetAllSelectedProxies() const
+{
+	TArray<UControlRigControlsProxy*> SelectedProxies(SelectedControlRigProxies);
+	if (SelectedSequencerProxies.Num() > 0)
+	{
+		SelectedProxies.Append(SelectedSequencerProxies);
+	}
+	return SelectedProxies;
+}
+
+void UControlRigDetailPanelControlProxies::ValuesChanged()
+{
+	//need to do all proxies
+	for (UControlRigControlsProxy* P1 : SelectedControlRigProxies)
+	{
+		P1->ValueChanged();
+	}
+	for (UControlRigControlsProxy* P2 : SelectedSequencerProxies)
+	{
+		P2->ValueChanged();
+	}
+}
+
+UControlRigControlsProxy* UControlRigDetailPanelControlProxies::AddProxy(UControlRig* ControlRig, FRigControlElement* ControlElement)
+{
+	//check if forced to be individual
+	bool bIsIndividual = (ControlElement->IsAnimationChannel()) ||
+		(ControlElement->Settings.AnimationType == ERigControlAnimationType::ProxyControl);
+
+	UControlRigControlsProxy* Proxy = FindProxy(ControlRig, ControlElement);
+	if (!Proxy && ControlElement != nullptr)
+	{
+		Proxy = NewProxyFromType(ControlElement->Settings.ControlType, ControlElement->Settings.ControlEnum);
+		//proxy was created so add it
+		if (Proxy)
+		{
+			FNameToProxyMap* ControlRigProxies = ControlRigOnlyProxies.Find(ControlRig);
+			if (ControlRigProxies)
+			{
+				ControlRigProxies->NameToProxy.Add(ControlElement->GetFName(), Proxy);
+			}
+			else
+			{
+				FNameToProxyMap NewControlRigProxies;
+				NewControlRigProxies.NameToProxy.Add(ControlElement->GetFName(), Proxy);
+				ControlRigOnlyProxies.Add(ControlRig, NewControlRigProxies);
+			}
+		}
+		if (Proxy)
+		{
+			Proxy->Type = ControlElement->Settings.ControlType;;
+			Proxy->OwnerControlElement = ControlElement;
+			Proxy->OwnerControlRig = ControlRig;
+			Proxy->AddControlRigControl(ControlRig, ControlElement->GetKey().Name);
+			Proxy->bIsIndividual = bIsIndividual;
+			Proxy->SetFlags(RF_Transactional);
+			Proxy->Modify();
+			Proxy->ValueChanged();
+		}
+	}
+	return Proxy;
+}
+
+UControlRigControlsProxy* UControlRigDetailPanelControlProxies::AddProxy(UObject* InObject, ERigControlType Type, TWeakObjectPtr<UMovieSceneTrack>& Track, TSharedPtr<FTrackInstancePropertyBindings>& Binding)
+{
+	FName PropertyName = Binding->GetPropertyName();
+	UControlRigControlsProxy* Proxy = FindProxy(InObject, PropertyName);
+	if (!Proxy)
+	{
+		TObjectPtr<UEnum> Enum = nullptr;
+		Proxy = NewProxyFromType(Type, Enum);
+		//proxy was created so add it
+		if (Proxy)
+		{
+			FNameToProxyMap* SequencerProxies = SequencerOnlyProxies.Find(InObject);
+			if (SequencerProxies)
+			{
+				SequencerProxies->NameToProxy.Add(PropertyName, Proxy);
+			}
+			else
+			{
+				FNameToProxyMap NewProxies;
+				NewProxies.NameToProxy.Add(PropertyName, Proxy);
+				SequencerOnlyProxies.Add(InObject, NewProxies);
+			}
+		}
+		if (Proxy)
+		{
+			Proxy->Type = Type;
+			Proxy->OwnerObject = InObject;
+			Proxy->OwnerBindingAndTrack.WeakTrack = Track;
+			Proxy->OwnerBindingAndTrack.Binding = Binding;
+			Proxy->AddSequencerProxyItem(InObject, Proxy->OwnerBindingAndTrack.WeakTrack, Binding);
+			Proxy->bIsIndividual = false;
+			Proxy->SetFlags(RF_Transactional);
+			Proxy->Modify();
+			Proxy->ValueChanged();
+		}
+	}
+	return Proxy;
+}
+
+
+UControlRigControlsProxy* UControlRigDetailPanelControlProxies::NewProxyFromType(ERigControlType ControlType, TObjectPtr<UEnum>& EnumPtr)
+{
+	UControlRigControlsProxy* Proxy = nullptr;
+	switch (ControlType)
+	{
+	case ERigControlType::Transform:
+	case ERigControlType::TransformNoScale:
+	case ERigControlType::EulerTransform:
+	{
+		Proxy = NewObject<UAnimDetailControlsProxyTransform>(this, NAME_None);
+		break;
+	}
+	case ERigControlType::Float:
+	case ERigControlType::ScaleFloat:
+	{
+		Proxy = NewObject<UAnimDetailControlsProxyFloat>(this, NAME_None);
+		break;
+	}
+	case ERigControlType::Integer:
+	{
+		if (EnumPtr == nullptr)
+		{
+			Proxy = NewObject<UAnimDetailControlsProxyInteger>(this, NAME_None); //was GetTransientPackage....? whay
+		}
+		else
+		{
+			UAnimDetailControlsProxyEnum* EnumProxy = NewObject<UAnimDetailControlsProxyEnum>(this, NAME_None);
+			EnumProxy->Enum.EnumType = EnumPtr;
+			Proxy = EnumProxy;
+		}
+		break;
+
+	}
+	case ERigControlType::Position:
+	{
+		Proxy = NewObject<UAnimDetailControlsProxyLocation>(this, NAME_None);
+		break;
+	}
+	case ERigControlType::Rotator:
+	{
+		Proxy = NewObject<UAnimDetailControlsProxyRotation>(this, NAME_None);
+		break;
+	}
+	case ERigControlType::Scale:
+	{
+		Proxy = NewObject<UAnimDetailControlsProxyScale>(this, NAME_None);
+		break;
+	}
+	case ERigControlType::Vector2D:
+	{
+		Proxy = NewObject<UAnimDetailControlsProxyVector2D>(this, NAME_None);
+		break;
+	}
+	case ERigControlType::Bool:
+	{
+		Proxy = NewObject<UAnimDetailControlsProxyBool>(this, NAME_None);
+		break;
+	}
+	default:
+		break;
+	}
+
+	return Proxy;
+}
+
+void UControlRigDetailPanelControlProxies::RemoveAllProxies()
+{
+	RemoveControlRigProxies(nullptr);
+	RemoveSequencerProxies(nullptr);
+}
+
+void UControlRigDetailPanelControlProxies::RemoveSequencerProxies(UObject* InObject)
+{
+	//no control rig remove all
+	if (InObject == nullptr)
+	{
+		for (TPair<TObjectPtr<UObject>, FNameToProxyMap>& Proxies : SequencerOnlyProxies)
+		{
+			for (TPair<FName, TObjectPtr<UControlRigControlsProxy> >& Pair : Proxies.Value.NameToProxy)
+			{
+				UControlRigControlsProxy* ExistingProxy = Pair.Value;
+				if (ExistingProxy)
+				{
+					ExistingProxy->Rename(nullptr, GetTransientPackage(), REN_ForceNoResetLoaders);
+					ExistingProxy->MarkAsGarbage();
+				}
+			}
+		}
+		SequencerOnlyProxies.Empty();
+		SelectedSequencerProxies.SetNum(0);
+	}
+	else
+	{
+		FNameToProxyMap* Proxies = SequencerOnlyProxies.Find(InObject);
+		if (Proxies)
+		{
+			for (TPair<FName, TObjectPtr<UControlRigControlsProxy>>& Pair : Proxies->NameToProxy)
+			{
+				UControlRigControlsProxy* ExistingProxy = Pair.Value;
+				if (ExistingProxy)
+				{
+					SelectedControlRigProxies.Remove(ExistingProxy);
+					ExistingProxy->Rename(nullptr, GetTransientPackage(), REN_ForceNoResetLoaders);
+					ExistingProxy->MarkAsGarbage();
+				}
+			}
+			SequencerOnlyProxies.Remove(InObject);
+		}
+	}
+
+}
+
+void UControlRigDetailPanelControlProxies::RemoveControlRigProxies(UControlRig* ControlRig)
+{
+	//no control rig remove all
+	if (ControlRig == nullptr)
+	{
+		for (TPair<TObjectPtr<UControlRig>, FNameToProxyMap>& ControlRigProxies : ControlRigOnlyProxies)
+		{
+			for (TPair<FName, TObjectPtr<UControlRigControlsProxy> >& Pair : ControlRigProxies.Value.NameToProxy)
+			{
+				UControlRigControlsProxy* ExistingProxy = Pair.Value;
+				if (ExistingProxy)
+				{
+					ExistingProxy->Rename(nullptr, GetTransientPackage(), REN_ForceNoResetLoaders);
+					ExistingProxy->MarkAsGarbage();
+				}
+			}
+		}
+		ControlRigOnlyProxies.Empty();
+		SelectedControlRigProxies.SetNum(0);
+	}
+	else
+	{
+		FNameToProxyMap* ControlRigProxies = ControlRigOnlyProxies.Find(ControlRig);
+		if (ControlRigProxies)
+		{
+			for (TPair<FName, TObjectPtr<UControlRigControlsProxy>>& Pair : ControlRigProxies->NameToProxy)
+			{
+				UControlRigControlsProxy* ExistingProxy = Pair.Value;
+				if (ExistingProxy)
+				{
+					SelectedControlRigProxies.Remove(ExistingProxy);
+					ExistingProxy->Rename(nullptr, GetTransientPackage(), REN_ForceNoResetLoaders);
+					ExistingProxy->MarkAsGarbage();
+				}
+			}
+
+			ControlRigOnlyProxies.Remove(ControlRig);
+		}
+	}
+}
+
+void UControlRigDetailPanelControlProxies::RecreateAllProxies(UControlRig* ControlRig)
+{
+	RemoveControlRigProxies(ControlRig);
+	TArray<FRigControlElement*> Controls = ControlRig->AvailableControls();
+	for (FRigControlElement* ControlElement : Controls)
+	{
+		if (ControlElement->Settings.AnimationType != ERigControlAnimationType::VisualCue)
+		{
+			AddProxy(ControlRig, ControlElement);
+		}
+	}
+}
+
+void UControlRigDetailPanelControlProxies::ProxyChanged(UControlRig* ControlRig, FRigControlElement* ControlElement, bool bModify)
+{
+	if (IsInGameThread())
+	{
+		UControlRigControlsProxy* Proxy = FindProxy(ControlRig, ControlElement);
+		if (Proxy)
+		{
+			if (bModify)
+			{
+				Modify();
+				Proxy->Modify();
+			}
+			Proxy->ValueChanged();
+		}
+	}
+}
+
+void UControlRigDetailPanelControlProxies::SelectProxy(UControlRig* ControlRig, FRigControlElement* RigElement, bool bSelected)
+{
+	UControlRigControlsProxy* Proxy = FindProxy(ControlRig, RigElement);
+	if (Proxy)
+	{
+		Modify();
+		if (bSelected)
+		{
+			if (!SelectedControlRigProxies.Contains(Proxy))
+			{
+				SelectedControlRigProxies.Add(Proxy);
+			}
+		}
+		else
+		{
+			SelectedControlRigProxies.Remove(Proxy);
+		}
+		Proxy->SelectionChanged(bSelected);
+	}
+}
+
+bool UControlRigDetailPanelControlProxies::IsSelected(UControlRig* InControlRig, FRigControlElement* RigElement) const
+{
+	if (UControlRigControlsProxy* Proxy = FindProxy(InControlRig, RigElement))
+	{
+		return SelectedControlRigProxies.Contains(Proxy);
+	}
+	return false;
+}
+
+void FControlRigEnumControlProxyValueDetails::CustomizeHeader(TSharedRef<IPropertyHandle> InStructPropertyHandle, FDetailWidgetRow& HeaderRow, IPropertyTypeCustomizationUtils& StructCustomizationUtils)
+{
+	TArray<UObject*> Objects;
+	InStructPropertyHandle->GetOuterObjects(Objects);
+	ensure(Objects.Num() == 1); // This is in here to ensure we are only showing the modifier details in the blueprint editor
+
+	for (UObject* Object : Objects)
+	{
+		if (Object->IsA<UAnimDetailControlsProxyEnum>())
+		{
+			ProxyBeingCustomized = Cast<UAnimDetailControlsProxyEnum>(Object);
+		}
+	}
+
+	check(ProxyBeingCustomized);
+
+	HeaderRow
+		.NameContent()
+		[
+			InStructPropertyHandle->CreatePropertyNameWidget()
+		]
+	.ValueContent()
+		[
+			SNew(SEnumComboBox, ProxyBeingCustomized->Enum.EnumType)
+			.OnEnumSelectionChanged(SEnumComboBox::FOnEnumSelectionChanged::CreateSP(this, &FControlRigEnumControlProxyValueDetails::OnEnumValueChanged, InStructPropertyHandle))
+		.CurrentValue(this, &FControlRigEnumControlProxyValueDetails::GetEnumValue)
+		.Font(FAppStyle::GetFontStyle(TEXT("MenuItem.Font")))
+		];
+}
+
+void FControlRigEnumControlProxyValueDetails::CustomizeChildren(TSharedRef<IPropertyHandle> InStructPropertyHandle, IDetailChildrenBuilder& StructBuilder, IPropertyTypeCustomizationUtils& StructCustomizationUtils)
+{
+}
+
+int32 FControlRigEnumControlProxyValueDetails::GetEnumValue() const
+{
+	if (ProxyBeingCustomized)
+	{
+		return ProxyBeingCustomized->Enum.EnumIndex;
+	}
+	return 0;
+}
+
+void FControlRigEnumControlProxyValueDetails::OnEnumValueChanged(int32 InValue, ESelectInfo::Type InSelectInfo, TSharedRef<IPropertyHandle> InStructHandle)
+{
+	if (ProxyBeingCustomized)
+	{
+		ProxyBeingCustomized->Enum.EnumIndex = InValue;
+		InStructHandle->NotifyPostChange(EPropertyChangeType::ValueSet);
+	}
+}
