New Anim Details. Main code is the new objects in AnimDetailsProxy that handle support for control rig/sequencer, having multiple objects, and shared selection state with sequencer. CR doesn't support

diff --git a/Engine/Plugins/Animation/ControlRig/Source/ControlRigEditor/Private/EditMode/AnimDetailsProxy.h b/Engine/Plugins/Animation/ControlRig/Source/ControlRigEditor/Private/EditMode/AnimDetailsProxy.h
new file mode 100644
--- /dev/null
+++ b/Engine/Plugins/Animation/ControlRig/Source/ControlRigEditor/Private/EditMode/AnimDetailsProxy.h
@@ -0,0 +1,487 @@
+﻿// Copyright Epic Games, Inc. All Rights Reserved.
+
+#pragma once
+
+#include "UObject/Object.h"
+#include "Engine/EngineBaseTypes.h"
+#include "ControlRigControlsProxy.h"
+#include "AnimDetailsProxy.generated.h"
+
+
+/**
+* State for details, including if they are selected(shown in curve editor) or have multiple different values.
+*/
+struct FAnimDetailPropertyState
+{
+	bool bMultiple = false;
+	bool bSelected = false;
+};
+
+struct FAnimDetailVectorState
+{
+	bool bXMultiple = false;
+	bool bYMultiple = false;
+	bool bZMultiple = false;
+
+	bool bXSelected = false;
+	bool bYSelected = false;
+	bool bZSelected = false;
+};
+
+//note control rig uses 'float' controls so we call this float though it's a 
+//double internally so we can use same for non-control rig parameters
+USTRUCT(BlueprintType)
+struct FAnimDetailProxyFloat
+{
+	GENERATED_BODY()
+
+	UPROPERTY(EditAnywhere, Interp, Category = "Float", meta = (SliderExponent = "1.0"))
+	double Float;
+
+};
+
+USTRUCT(BlueprintType)
+struct FAnimDetailProxyBool
+{
+	GENERATED_BODY()
+
+	UPROPERTY(EditAnywhere, Interp, Category = "Bool")
+	bool Bool;
+
+};
+
+USTRUCT(BlueprintType)
+struct FControlRigEnumControlProxyValue
+{
+	GENERATED_USTRUCT_BODY()
+
+	FControlRigEnumControlProxyValue()
+	{
+		EnumType = nullptr;
+		EnumIndex = INDEX_NONE;
+	}
+
+	UPROPERTY()
+	TObjectPtr<UEnum> EnumType;
+
+	UPROPERTY(EditAnywhere, Category = Enum)
+	int32 EnumIndex;
+};
+
+class FControlRigEnumControlProxyValueDetails : public IPropertyTypeCustomization
+{
+public:
+
+	static TSharedRef<IPropertyTypeCustomization> MakeInstance()
+	{
+		return MakeShareable(new FControlRigEnumControlProxyValueDetails);
+	}
+
+	/** IPropertyTypeCustomization interface */
+	virtual void CustomizeHeader(TSharedRef<class IPropertyHandle> InStructPropertyHandle, class FDetailWidgetRow& HeaderRow, IPropertyTypeCustomizationUtils& StructCustomizationUtils) override;
+	virtual void CustomizeChildren(TSharedRef<class IPropertyHandle> InStructPropertyHandle, class IDetailChildrenBuilder& StructBuilder, IPropertyTypeCustomizationUtils& StructCustomizationUtils) override;
+
+private:
+
+	int32 GetEnumValue() const;
+	void OnEnumValueChanged(int32 InValue, ESelectInfo::Type InSelectInfo, TSharedRef<IPropertyHandle> InStructHandle);
+
+	UAnimDetailControlsProxyEnum* ProxyBeingCustomized;
+};
+
+USTRUCT(BlueprintType)
+struct FAnimDetailProxyInteger
+{
+	GENERATED_BODY()
+
+	UPROPERTY(EditAnywhere, Interp, Category = "Integer")
+	int64 Integer;
+
+};
+
+USTRUCT(BlueprintType)
+struct FAnimDetailProxyVector3 
+{
+	GENERATED_BODY()
+
+	UPROPERTY(EditAnywhere, Interp, Category = "Vector")
+	double X;
+
+	UPROPERTY(EditAnywhere, Interp, Category = "Vector")
+	double Y;
+
+	UPROPERTY(EditAnywhere, Interp, Category = "Vector")
+	double Z;
+
+	FVector ToVector() const { return FVector(X, Y, Z); }
+};
+
+USTRUCT(BlueprintType)
+struct  FAnimDetailProxyLocation 
+{
+	GENERATED_BODY()
+
+	UPROPERTY(EditAnywhere, Interp, Category = "Location", meta = (Delta = "0.5", SliderExponent = "1", LinearDeltaSensitivity = "1"))
+	double LX;
+
+	UPROPERTY(EditAnywhere, Interp, Category = "Location", meta = (Delta = "0.5", SliderExponent = "1", LinearDeltaSensitivity = "1"))
+	double LY;
+
+	UPROPERTY(EditAnywhere, Interp, Category = "Location", meta = (Delta = "0.5", SliderExponent = "1", LinearDeltaSensitivity = "1"))
+	double LZ;
+
+	FAnimDetailVectorState State;
+
+	FAnimDetailProxyLocation() : LX(0.0), LY(0.0), LZ(0.0) {};
+	FAnimDetailProxyLocation(const FVector& InVector, const FAnimDetailVectorState &InState) 
+	{
+		LX = InVector.X; LY = InVector.Y; LZ = InVector.Z; State = InState;
+	}
+	FAnimDetailProxyLocation(const FVector3f& InVector, const FAnimDetailVectorState& InState) 
+	{
+		LX = InVector.X; LY = InVector.Y; LZ = InVector.Z; State = InState;
+	}
+	FVector ToVector() const { return FVector(LX, LY, LZ); }
+	FVector3f ToVector3f() const { return FVector3f(LX, LY, LZ); }
+
+};
+
+USTRUCT(BlueprintType)
+struct  FAnimDetailProxyRotation
+{
+	GENERATED_BODY()
+
+	UPROPERTY(EditAnywhere, Interp, Category = "Rotation", meta = (Delta = "0.5", SliderExponent = "1", LinearDeltaSensitivity = "1"))
+	double RX;
+
+	UPROPERTY(EditAnywhere, Interp, Category = "Rotation", meta = (Delta = "0.5", SliderExponent = "1", LinearDeltaSensitivity = "1"))
+	double RY;
+
+	UPROPERTY(EditAnywhere, Interp, Category = "Rotation", meta = (Delta = "0.5", SliderExponent = "1", LinearDeltaSensitivity = "1"))
+	double RZ;
+
+	FAnimDetailVectorState State;
+
+	FAnimDetailProxyRotation() : RX(0.0), RY(0.0), RZ(0.0) {};
+	FAnimDetailProxyRotation(const FRotator& InRotator, const FAnimDetailVectorState& InState) { FromRotator(InRotator); State = InState; }
+	FAnimDetailProxyRotation(const FVector3f& InVector, const FAnimDetailVectorState& InState) 
+	{ 
+		RX = InVector.X; RY = InVector.Y; RZ = InVector.Z; State = InState;
+	}
+
+	FVector ToVector() const { return FVector(RX, RY, RZ); }
+	FVector3f ToVector3f() const { return FVector3f(RX, RY, RZ); }
+	FRotator ToRotator() const { FRotator Rot; Rot = Rot.MakeFromEuler(ToVector()); return Rot; }
+	void FromRotator(const FRotator& InRotator) { FVector Vec(InRotator.Euler()); RX = Vec.X; RY = Vec.Y; RZ = Vec.Z; }
+};
+
+USTRUCT(BlueprintType)
+struct  FAnimDetailProxyScale 
+{
+	GENERATED_BODY()
+
+	UPROPERTY(EditAnywhere, Interp, Category = "Scale", meta = (Delta = "0.5", SliderExponent = "1", LinearDeltaSensitivity = "1"))
+	double SX;
+
+	UPROPERTY(EditAnywhere, Interp, Category = "Scale", meta = (Delta = "0.5", SliderExponent = "1", LinearDeltaSensitivity = "1"))
+	double SY;
+
+	UPROPERTY(EditAnywhere, Interp, Category = "Scale", meta = (Delta = "0.5", SliderExponent = "1", LinearDeltaSensitivity = "1"))
+	double SZ;
+
+	FAnimDetailVectorState State;
+
+	FAnimDetailProxyScale() : SX(0.0), SY(0.0), SZ(0.0) {};
+	FAnimDetailProxyScale(const FVector& InVector, const FAnimDetailVectorState& InState) 
+	{
+		SX = InVector.X; SY = InVector.Y; SZ = InVector.Z; State = InState;
+	}
+	FAnimDetailProxyScale(const FVector3f& InVector, const FAnimDetailVectorState& InState)
+	{
+		SX = InVector.X; SY = InVector.Y; SZ = InVector.Z; State = InState;
+	}
+	FVector ToVector() const { return FVector(SX, SY, SZ); }
+	FVector3f ToVector3f() const { return FVector3f(SX, SX, SZ); }
+};
+
+USTRUCT(BlueprintType)
+struct  FAnimDetailProxyVector2D
+{
+	GENERATED_BODY()
+
+	UPROPERTY(EditAnywhere, Interp, Category = "Vector2D", meta = (Delta = "0.05", SliderExponent = "1", LinearDeltaSensitivity = "1"))
+	double X;
+
+	UPROPERTY(EditAnywhere, Interp, Category = "Vector2D", meta = (Delta = "0.05", SliderExponent = "1", LinearDeltaSensitivity = "1"))
+	double Y;
+
+
+	FAnimDetailVectorState State;
+
+	FAnimDetailProxyVector2D() : X(0.0), Y(0.0) {};
+	FAnimDetailProxyVector2D(const FVector2D& InVector, const FAnimDetailVectorState& InState) 
+	{ 
+		X = InVector.X; Y = InVector.Y; State = InState;
+
+	}
+	FVector2D ToVector2D() const { return FVector2D(X,Y); }
+};
+
+UCLASS()
+class UAnimDetailControlsKeyedProxy : public UControlRigControlsProxy
+{
+	GENERATED_BODY()
+public:
+	virtual void SetKey(TSharedPtr<ISequencer>& Sequencer, const IPropertyHandle& KeyedPropertyHandle) override;
+	virtual EPropertyKeyedStatus GetPropertyKeyedStatus(TSharedPtr<ISequencer>& Sequencer, const IPropertyHandle& PropertyHandle) const override;
+
+#if WITH_EDITOR
+	virtual void PostEditUndo() override;
+#endif
+
+};
+
+UCLASS(EditInlineNew, CollapseCategories)
+class UAnimDetailControlsProxyFloat : public UAnimDetailControlsKeyedProxy
+{
+	GENERATED_BODY()
+public:
+
+	//UControlRigControlsProxy
+	virtual bool PropertyIsOnProxy(FProperty* Property, FProperty* MemberProperty) override;
+	virtual void ValueChanged() override;
+	virtual EControlRigContextChannelToKey GetChannelToKeyFromPropertyName(const FName& InPropertyName) const override;
+	virtual EControlRigContextChannelToKey GetChannelToKeyFromChannelName(const FString& InChannelName) const override;
+	virtual bool IsMultiple(const FName& InPropertyName) const override;
+	virtual void SetControlRigElementValueFromCurrent(UControlRig* ControlRig, FRigControlElement* ControlElement, const FRigControlModifiedContext& Context) override;
+	virtual void GetChannelSelectionState(TWeakPtr<FCurveEditor>& CurveEditor, FAnimDetailVectorSelection& OutLocationSelection, FAnimDetailVectorSelection& OutRotationSelection,
+		FAnimDetailVectorSelection& OutScaleSelection) override;
+	virtual void UpdatePropertyNames(IDetailLayoutBuilder& DetailBuilder);
+	virtual void SetBindingValueFromCurrent(UObject* InObject, TSharedPtr<FTrackInstancePropertyBindings>& Binding, FRigControlModifiedContext& Context) override;
+
+	FAnimDetailPropertyState State;
+
+	UPROPERTY(EditAnywhere, Interp, Category = "Float", meta = (Delta = "0.05", SliderExponent = "1", LinearDeltaSensitivity = "1"))
+	FAnimDetailProxyFloat Float;
+	
+private:
+	void ClearMultipleFlags();
+	void SetMultipleFlags(const float& ValA, const float& ValB);
+};
+
+UCLASS(EditInlineNew, CollapseCategories)
+class UAnimDetailControlsProxyBool : public UAnimDetailControlsKeyedProxy
+{
+	GENERATED_BODY()
+public:
+
+	//UControlRigControlsProxy
+	virtual bool PropertyIsOnProxy(FProperty* Property, FProperty* MemberProperty) override;
+	virtual void ValueChanged() override;
+	virtual EControlRigContextChannelToKey GetChannelToKeyFromPropertyName(const FName& InPropertyName) const override;
+	virtual EControlRigContextChannelToKey GetChannelToKeyFromChannelName(const FString& InChannelName) const override;
+	virtual bool IsMultiple(const FName& InPropertyName) const override;
+	virtual void SetControlRigElementValueFromCurrent(UControlRig* ControlRig, FRigControlElement* ControlElement, const FRigControlModifiedContext& Context) override;
+	virtual void GetChannelSelectionState(TWeakPtr<FCurveEditor>& CurveEditor, FAnimDetailVectorSelection& OutLocationSelection, FAnimDetailVectorSelection& OutRotationSelection,
+		FAnimDetailVectorSelection& OutScaleSelection) override;
+	virtual void UpdatePropertyNames(IDetailLayoutBuilder& DetailBuilder);
+	virtual void SetBindingValueFromCurrent(UObject* InObject, TSharedPtr<FTrackInstancePropertyBindings>& Binding, FRigControlModifiedContext& Context) override;
+
+	
+	FAnimDetailPropertyState State;
+
+	UPROPERTY(EditAnywhere, Interp, Category = "Bool")
+	FAnimDetailProxyBool Bool;
+
+private:
+	void ClearMultipleFlags();
+	void SetMultipleFlags(const bool& ValA, const bool& ValB);
+};
+
+UCLASS(EditInlineNew, CollapseCategories)
+class UAnimDetailControlsProxyInteger : public UAnimDetailControlsKeyedProxy
+{
+	GENERATED_BODY()
+public:
+
+	//UControlRigControlsProxy
+	virtual bool PropertyIsOnProxy(FProperty* Property, FProperty* MemberProperty) override;
+	virtual void ValueChanged() override;
+	virtual EControlRigContextChannelToKey GetChannelToKeyFromPropertyName(const FName& InPropertyName) const override;
+	virtual EControlRigContextChannelToKey GetChannelToKeyFromChannelName(const FString& InChannelName) const override;
+	virtual bool IsMultiple(const FName& InPropertyName) const override;
+	virtual void SetControlRigElementValueFromCurrent(UControlRig* ControlRig, FRigControlElement* ControlElement, const FRigControlModifiedContext& Context) override;
+	virtual void GetChannelSelectionState(TWeakPtr<FCurveEditor>& CurveEditor, FAnimDetailVectorSelection& OutLocationSelection, FAnimDetailVectorSelection& OutRotationSelection,
+		FAnimDetailVectorSelection& OutScaleSelection) override;
+	virtual void UpdatePropertyNames(IDetailLayoutBuilder& DetailBuilder);
+	virtual void SetBindingValueFromCurrent(UObject* InObject, TSharedPtr<FTrackInstancePropertyBindings>& Binding, FRigControlModifiedContext& Context) override;
+
+
+	FAnimDetailPropertyState State;
+
+	UPROPERTY(EditAnywhere, Interp, Category = "Integer")
+	FAnimDetailProxyInteger Integer;
+
+private:
+	void ClearMultipleFlags();
+	void SetMultipleFlags(const int64& ValA, const int64& ValB);
+};
+
+UCLASS(EditInlineNew, CollapseCategories)
+class UAnimDetailControlsProxyEnum : public UAnimDetailControlsKeyedProxy
+{
+	GENERATED_BODY()
+public:
+
+	//UControlRigControlsProxy
+	virtual bool PropertyIsOnProxy(FProperty* Property, FProperty* MemberProperty) override;
+	virtual void ValueChanged() override;
+	virtual EControlRigContextChannelToKey GetChannelToKeyFromPropertyName(const FName& InPropertyName) const override;
+	virtual EControlRigContextChannelToKey GetChannelToKeyFromChannelName(const FString& InChannelName) const override;
+	virtual bool IsMultiple(const FName& InPropertyName) const override;
+	virtual void SetControlRigElementValueFromCurrent(UControlRig* ControlRig, FRigControlElement* ControlElement, const FRigControlModifiedContext& Context) override;
+	virtual void GetChannelSelectionState(TWeakPtr<FCurveEditor>& CurveEditor, FAnimDetailVectorSelection& OutLocationSelection, FAnimDetailVectorSelection& OutRotationSelection,
+		FAnimDetailVectorSelection& OutScaleSelection) override;
+	virtual void UpdatePropertyNames(IDetailLayoutBuilder& DetailBuilder);
+
+
+	FAnimDetailPropertyState State;
+
+	UPROPERTY(EditAnywhere, Interp, Category = "Enum")
+	FControlRigEnumControlProxyValue Enum;
+
+private:
+	void ClearMultipleFlags();
+	void SetMultipleFlags(const int32& ValA, const int32& ValB);
+};
+
+
+UCLASS(EditInlineNew,CollapseCategories)
+class UAnimDetailControlsProxyTransform : public UAnimDetailControlsKeyedProxy
+{
+	GENERATED_BODY()
+public:
+
+	//UControlRigControlsProxy
+	virtual bool PropertyIsOnProxy(FProperty* Property, FProperty* MemberProperty) override;
+	virtual void ValueChanged() override;
+	virtual EControlRigContextChannelToKey GetChannelToKeyFromPropertyName(const FName& InPropertyName) const override;
+	virtual EControlRigContextChannelToKey GetChannelToKeyFromChannelName(const FString& InChannelName) const override;
+	virtual bool IsMultiple(const FName& InPropertyName) const override;
+	virtual void SetControlRigElementValueFromCurrent(UControlRig* ControlRig, FRigControlElement* ControlElement, const FRigControlModifiedContext& Context) override;
+	virtual void SetBindingValueFromCurrent(UObject* InObject, TSharedPtr<FTrackInstancePropertyBindings>& Binding, FRigControlModifiedContext& Context) override;
+	virtual void GetChannelSelectionState(TWeakPtr<FCurveEditor>& CurveEditor, FAnimDetailVectorSelection& OutLocationSelection, FAnimDetailVectorSelection& OutRotationSelection,
+		FAnimDetailVectorSelection& OutScaleSelection) override;
+
+	UPROPERTY(EditAnywhere, Interp, Category = Transform)
+	FAnimDetailProxyLocation Location;
+
+	UPROPERTY(EditAnywhere, Interp, Category = Transform)
+	FAnimDetailProxyRotation Rotation;
+
+	UPROPERTY(EditAnywhere, Interp, Category = Transform)
+	FAnimDetailProxyScale Scale;
+
+private:
+	void ClearMultipleFlags();
+	void SetMultipleFlags(const FEulerTransform& ValA, const FEulerTransform& ValB);
+};
+
+
+UCLASS(EditInlineNew, CollapseCategories)
+class UAnimDetailControlsProxyLocation : public UAnimDetailControlsKeyedProxy
+{
+	GENERATED_BODY()
+public:
+
+	//UControlRigControlsProxy
+	virtual bool PropertyIsOnProxy(FProperty* Property, FProperty* MemberProperty) override;
+	virtual void ValueChanged() override;
+	virtual EControlRigContextChannelToKey GetChannelToKeyFromPropertyName(const FName& InPropertyName) const override;
+	virtual EControlRigContextChannelToKey GetChannelToKeyFromChannelName(const FString& InChannelName) const override;
+	virtual void SetControlRigElementValueFromCurrent(UControlRig* ControlRig, FRigControlElement* ControlElement, const FRigControlModifiedContext& Context) override;
+	virtual bool IsMultiple(const FName& InPropertyName) const override;
+	virtual void GetChannelSelectionState(TWeakPtr<FCurveEditor>& CurveEditor, FAnimDetailVectorSelection& OutLocationSelection, FAnimDetailVectorSelection& OutRotationSelection,
+		FAnimDetailVectorSelection& OutScaleSelection) override;
+
+	UPROPERTY(EditAnywhere, Interp, Category = Location)
+	FAnimDetailProxyLocation Location;
+
+	private:
+	void ClearMultipleFlags();
+	void SetMultipleFlags(const FVector3f& ValA, const FVector3f& ValB);
+};
+
+UCLASS(EditInlineNew, CollapseCategories)
+class UAnimDetailControlsProxyRotation : public UAnimDetailControlsKeyedProxy
+{
+	GENERATED_BODY()
+public:
+
+	//UControlRigControlsProxy
+	virtual bool PropertyIsOnProxy(FProperty* Property, FProperty* MemberProperty) override;
+	virtual void ValueChanged() override;
+	virtual EControlRigContextChannelToKey GetChannelToKeyFromPropertyName(const FName& InPropertyName) const override;
+	virtual EControlRigContextChannelToKey GetChannelToKeyFromChannelName(const FString& InChannelName) const override;
+	virtual void SetControlRigElementValueFromCurrent(UControlRig* ControlRig, FRigControlElement* ControlElement, const FRigControlModifiedContext& Context) override;
+	virtual bool IsMultiple(const FName& InPropertyName) const override;
+	virtual void GetChannelSelectionState(TWeakPtr<FCurveEditor>& CurveEditor, FAnimDetailVectorSelection& OutLocationSelection, FAnimDetailVectorSelection& OutRotationSelection,
+		FAnimDetailVectorSelection& OutScaleSelection) override;
+
+	UPROPERTY(EditAnywhere, Interp, Category = Rotation)
+	FAnimDetailProxyRotation Rotation;
+
+private:
+	void ClearMultipleFlags();
+	void SetMultipleFlags(const FVector3f& ValA, const FVector3f& ValB);
+};
+
+UCLASS(EditInlineNew, CollapseCategories)
+class UAnimDetailControlsProxyScale : public UAnimDetailControlsKeyedProxy
+{
+	GENERATED_BODY()
+public:
+
+	//UControlRigControlsProxy
+	virtual bool PropertyIsOnProxy(FProperty* Property, FProperty* MemberProperty) override;
+	virtual void ValueChanged() override;
+	virtual EControlRigContextChannelToKey GetChannelToKeyFromPropertyName(const FName& InPropertyName) const override;
+	virtual EControlRigContextChannelToKey GetChannelToKeyFromChannelName(const FString& InChannelName) const override;
+	virtual void SetControlRigElementValueFromCurrent(UControlRig* ControlRig, FRigControlElement* ControlElement, const FRigControlModifiedContext& Context) override;
+	virtual bool IsMultiple(const FName& InPropertyName) const override;
+	virtual void GetChannelSelectionState(TWeakPtr<FCurveEditor>& CurveEditor, FAnimDetailVectorSelection& OutLocationSelection, FAnimDetailVectorSelection& OutRotationSelection,
+		FAnimDetailVectorSelection& OutScaleSelection) override;
+
+	UPROPERTY(EditAnywhere, Interp, Category = Scale)
+	FAnimDetailProxyScale Scale;
+
+private:
+	void ClearMultipleFlags();
+	void SetMultipleFlags(const FVector3f& ValA, const FVector3f& ValB);
+};
+
+UCLASS(EditInlineNew, CollapseCategories)
+class UAnimDetailControlsProxyVector2D : public UAnimDetailControlsKeyedProxy
+{
+	GENERATED_BODY()
+public:
+
+	//UControlRigControlsProxy
+	virtual bool PropertyIsOnProxy(FProperty* Property, FProperty* MemberProperty) override;
+	virtual void ValueChanged() override;
+	virtual EControlRigContextChannelToKey GetChannelToKeyFromPropertyName(const FName& InPropertyName) const override;
+	virtual EControlRigContextChannelToKey GetChannelToKeyFromChannelName(const FString& InChannelName) const override;
+	virtual void SetControlRigElementValueFromCurrent(UControlRig* ControlRig, FRigControlElement* ControlElement, const FRigControlModifiedContext& Context) override;
+	virtual bool IsMultiple(const FName& InPropertyName) const override;
+	virtual void GetChannelSelectionState(TWeakPtr<FCurveEditor>& CurveEditor, FAnimDetailVectorSelection& OutLocationSelection, FAnimDetailVectorSelection& OutRotationSelection,
+		FAnimDetailVectorSelection& OutScaleSelection) override;
+
+	UPROPERTY(EditAnywhere, Interp, Category = Vector2D)
+		FAnimDetailProxyVector2D Vector2D;
+
+private:
+	void ClearMultipleFlags();
+	void SetMultipleFlags(const FVector2D& ValA, const FVector2D& ValB);
+};
+
+
diff --git a/Engine/Plugins/Animation/ControlRig/Source/ControlRigEditor/Private/EditMode/ControlRigBaseDockableView.cpp b/Engine/Plugins/Animation/ControlRig/Source/ControlRigEditor/Private/EditMode/ControlRigBaseDockableView.cpp
--- a/Engine/Plugins/Animation/ControlRig/Source/ControlRigEditor/Private/EditMode/ControlRigBaseDockableView.cpp
+++ b/Engine/Plugins/Animation/ControlRig/Source/ControlRigEditor/Private/EditMode/ControlRigBaseDockableView.cpp
@@ -73,6 +73,13 @@ void FControlRigBaseDockableView::HandlElementSelected(UControlRig* ControlRig,
 	}
 }
 
+
+FControlRigEditMode* FControlRigBaseDockableView::GetEditMode() const
+{
+	FControlRigEditMode* EditMode = static_cast<FControlRigEditMode*>(ModeTools->GetActiveMode(FControlRigEditMode::ModeName));
+	return EditMode;
+}
+
 ISequencer* FControlRigBaseDockableView::GetSequencer() const
 {
 	if (FControlRigEditMode* EditMode = static_cast<FControlRigEditMode*>(ModeTools->GetActiveMode(FControlRigEditMode::ModeName)))
diff --git a/Engine/Plugins/Animation/ControlRig/Source/ControlRigEditor/Private/EditMode/ControlRigBaseDockableView.h b/Engine/Plugins/Animation/ControlRig/Source/ControlRigEditor/Private/EditMode/ControlRigBaseDockableView.h
--- a/Engine/Plugins/Animation/ControlRig/Source/ControlRigEditor/Private/EditMode/ControlRigBaseDockableView.h
+++ b/Engine/Plugins/Animation/ControlRig/Source/ControlRigEditor/Private/EditMode/ControlRigBaseDockableView.h
@@ -22,6 +22,7 @@ public:
 	TArray<UControlRig*> GetControlRigs() const;
 
 	virtual void SetEditMode(FControlRigEditMode& InEditMode);
+	FControlRigEditMode* GetEditMode() const;
 
 protected:
 	virtual void HandleControlSelected(UControlRig* Subject, FRigControlElement* InControl, bool bSelected);
diff --git a/Engine/Plugins/Animation/ControlRig/Source/ControlRigEditor/Private/EditMode/ControlRigControlsProxy.cpp b/Engine/Plugins/Animation/ControlRig/Source/ControlRigEditor/Private/EditMode/ControlRigControlsProxy.cpp
--- a/Engine/Plugins/Animation/ControlRig/Source/ControlRigEditor/Private/EditMode/ControlRigControlsProxy.cpp
+++ b/Engine/Plugins/Animation/ControlRig/Source/ControlRigEditor/Private/EditMode/ControlRigControlsProxy.cpp
@@ -5,1044 +5,317 @@
 #include "EditMode/ControlRigEditMode.h"
 #include "ControlRig.h"
 #include "Rigs/RigHierarchy.h"
-
+#include "Sequencer/MovieSceneControlRigParameterSection.h"
+#include "Sequencer/MovieSceneControlRigParameterTrack.h"
 #include "Components/SkeletalMeshComponent.h"
 #include "Subsystems/AssetEditorSubsystem.h"
 #include "MovieSceneCommonHelpers.h"
 #include "PropertyHandle.h"
-
-#include UE_INLINE_GENERATED_CPP_BY_NAME(ControlRigControlsProxy)
-
-#if WITH_EDITOR
-
+#include "Sequencer/ControlRigSequencerHelpers.h"
+#include "ISequencer.h"
+#include "CurveEditor.h"
+#include "CurveModel.h"
+#include "MovieScene.h"
+#include "Tracks/MovieScene3DTransformTrack.h"
+#include "Tracks/MovieScenePropertyTrack.h"
+#include "MovieSceneCommonHelpers.h"
+#include "UnrealEdGlobals.h"
+#include "UnrealEdMisc.h"
+#include "Editor/UnrealEdEngine.h"
+#include "ScopedTransaction.h"
+#include "Channels/MovieSceneBoolChannel.h"
+#include "Channels/MovieSceneChannelProxy.h"
+#include "Channels/MovieSceneDoubleChannel.h"
+#include "Channels/MovieSceneIntegerChannel.h"
+#include "Channels/MovieSceneByteChannel.h"
 #include "DetailLayoutBuilder.h"
 #include "DetailCategoryBuilder.h"
 #include "DetailWidgetRow.h"
 #include "SEnumCombo.h"
 
-void UControlRigControlsProxy::SetIsMultiple(bool bIsVal)
-{ 
-	bIsMultiple = bIsVal; 
-	if (bIsMultiple)
-	{
-		FString DisplayString = TEXT("Multiple");
-		FName DisplayName(*DisplayString);
-		Name = DisplayName;
-	}
-	else
-	{
-		Name = ControlName;
-	}
-}
-
-
-void FControlRigEnumControlProxyValueDetails::CustomizeHeader(TSharedRef<IPropertyHandle> InStructPropertyHandle, FDetailWidgetRow& HeaderRow, IPropertyTypeCustomizationUtils& StructCustomizationUtils)
-{
-	TArray<UObject*> Objects;
-	InStructPropertyHandle->GetOuterObjects(Objects);
-	ensure(Objects.Num() == 1); // This is in here to ensure we are only showing the modifier details in the blueprint editor
-
-	for (UObject* Object : Objects)
-	{
-		if (Object->IsA<UControlRigEnumControlProxy>())
-		{
-			ProxyBeingCustomized = Cast<UControlRigEnumControlProxy>(Object);
-		}
-	}
-
-	check(ProxyBeingCustomized);
-
-	HeaderRow
-	.NameContent()
-	[
-		InStructPropertyHandle->CreatePropertyNameWidget()
-	]
-	.ValueContent()
-	[
-		SNew(SEnumComboBox, ProxyBeingCustomized->Enum.EnumType)
-		.OnEnumSelectionChanged(SEnumComboBox::FOnEnumSelectionChanged::CreateSP(this, &FControlRigEnumControlProxyValueDetails::OnEnumValueChanged, InStructPropertyHandle))
-		.CurrentValue(this, &FControlRigEnumControlProxyValueDetails::GetEnumValue)
-		.Font(FAppStyle::GetFontStyle(TEXT("MenuItem.Font")))
-	];
-}
-
-void FControlRigEnumControlProxyValueDetails::CustomizeChildren(TSharedRef<IPropertyHandle> InStructPropertyHandle, IDetailChildrenBuilder& StructBuilder, IPropertyTypeCustomizationUtils& StructCustomizationUtils)
-{
-}
-
-int32 FControlRigEnumControlProxyValueDetails::GetEnumValue() const
-{
-	if (ProxyBeingCustomized)
-	{
-		return ProxyBeingCustomized->Enum.EnumIndex;
-	}
-	return 0;
-}
-
-void FControlRigEnumControlProxyValueDetails::OnEnumValueChanged(int32 InValue, ESelectInfo::Type InSelectInfo, TSharedRef<IPropertyHandle> InStructHandle){
-	if (ProxyBeingCustomized)
-	{
-		ProxyBeingCustomized->Enum.EnumIndex = InValue;
-		InStructHandle->NotifyPostChange(EPropertyChangeType::ValueSet);
-	}
-}
-
-#endif
-
-void UControlRigControlsProxy::SelectionChanged(bool bInSelected)
-{
-	FRigControlElement* ControlElement = GetControlElement();
-	if (ControlElement)
-	{
-		Modify();
-		const FName PropertyName("bSelected");
-		FTrackInstancePropertyBindings Binding(PropertyName, PropertyName.ToString());
-		Binding.CallFunction<bool>(*this, bInSelected);
-	}
-}
-
-void UControlRigControlsProxy::PostEditChangeChainProperty(struct FPropertyChangedChainEvent& PropertyChangedEvent)
-{
-	Super::PostEditChangeChainProperty(PropertyChangedEvent);
-	if (PropertyChangedEvent.Property && PropertyChangedEvent.Property->GetFName() == GET_MEMBER_NAME_CHECKED(UControlRigControlsProxy, bSelected))
-	{
-		FRigControlElement* ControlElement = GetControlElement();
-		if (ControlElement && ControlRig.IsValid())
-		{
-			FControlRigInteractionScope InteractionScope(ControlRig.Get(), ControlElement->GetKey());
-			ControlRig->SelectControl(ControlName, bSelected);
-			ControlRig->Evaluate_AnyThread();
-		}
-	}
-#if WITH_EDITOR
-	if (PropertyChangedEvent.Property)
-	{
-		FRigControlElement* ControlElement = GetControlElement();
-		if (PropertyChangedEvent.ChangeType == EPropertyChangeType::Interactive)
-		{
-			EControlRigInteractionType Type = EControlRigInteractionType::None;
-			FProperty* Owner = PropertyChangedEvent.Property->GetOwnerProperty();
-			if (FEditPropertyChain::TDoubleLinkedListNode* MemberNode = PropertyChangedEvent.PropertyChain.GetActiveMemberNode())
-			{
-				if (FProperty* MemberProperty = MemberNode->GetValue())
-				{
-					if (MemberProperty->GetFName() == GET_MEMBER_NAME_CHECKED(UControlRigTransformControlProxy, Transform))
-					{
-						if (FEditPropertyChain::TDoubleLinkedListNode* NextNode = MemberNode->GetNextNode())
-						{
-							if (FProperty* SubMemberProperty = NextNode->GetValue())
-							{
-								if (SubMemberProperty->GetFName() == GET_MEMBER_NAME_CHECKED(FEulerTransform, Location))
-								{
-									Type = EControlRigInteractionType::Translate;
-								}
-								else if (SubMemberProperty->GetFName() == GET_MEMBER_NAME_CHECKED(FEulerTransform, Rotation))
-								{
-									Type = EControlRigInteractionType::Rotate;
-								}
-								else if (SubMemberProperty->GetFName() == GET_MEMBER_NAME_CHECKED(FEulerTransform, Scale))
-								{
-									Type = EControlRigInteractionType::Scale;
-								}
-							}
-						}
-					}
-				}
-			}
-			FControlRigInteractionScope* InteractionScope = new FControlRigInteractionScope(ControlRig.Get(), ControlElement->GetKey(), Type);
-			InteractionScopes.Add(InteractionScope);
-		}
-		else
-		{
-			for (FControlRigInteractionScope* Scope : InteractionScopes)
-			{
-				if (Scope)
-				{
-					delete Scope; 
-				}
-			}
-			InteractionScopes.Reset();
-		}
-	}
-#endif
-}
-
-#if WITH_EDITOR
-void UControlRigControlsProxy::PostEditUndo()
-{
-	FRigControlElement* ControlElement = GetControlElement();
-	if (ControlElement && ControlRig.IsValid() && ControlRig->GetHierarchy()->Contains(FRigElementKey(ControlName, ERigElementType::Control)))
-	{
-		ControlRig->SelectControl(ControlName, bSelected);
-	}
-}
-#endif
+#include UE_INLINE_GENERATED_CPP_BY_NAME(ControlRigControlsProxy)
 
-FRigControlElement* UControlRigControlsProxy::GetControlElement() const
+void UControlRigControlsProxy::AddControlRigControl(UControlRig* InControlRig, const FName& InName)
 {
-	if(ControlRig.IsValid())
+	if (InControlRig == nullptr)
 	{
-		return ControlRig->GetHierarchy()->Find<FRigControlElement>(FRigElementKey(ControlName, ERigElementType::Control));
+		return;
 	}
-	return nullptr;
-}
-
-void UControlRigTransformControlProxy::PostEditChangeChainProperty(struct FPropertyChangedChainEvent& PropertyChangedEvent)
-{
-	Super::PostEditChangeChainProperty(PropertyChangedEvent);
-	FProperty* Property = PropertyChangedEvent.Property;
-	FProperty* MemberProperty = nullptr;
-	if (PropertyChangedEvent.PropertyChain.GetActiveMemberNode())
+	FRigControlElement* ControlElement = InControlRig->FindControl(InName);
+	if (ControlElement == nullptr)
 	{
-		MemberProperty = PropertyChangedEvent.PropertyChain.GetActiveMemberNode()->GetValue();
+		return;
 	}
 
-	if ((Property && Property->GetFName() == GET_MEMBER_NAME_CHECKED(UControlRigTransformControlProxy, Transform))
-		|| (MemberProperty && MemberProperty->GetFName() == GET_MEMBER_NAME_CHECKED(UControlRigTransformControlProxy, Transform)))
+	FControlRigProxyItem& Item = ControlRigItems.FindOrAdd(InControlRig);
+	Item.ControlRig = InControlRig;
+	if (Item.ControlElements.Contains(ControlElement) == false)
 	{
-		FRigControlElement* ControlElement = GetControlElement();
-		if (ControlElement && ControlRig.IsValid())
-		{
-			//MUST set through ControlRig
-			const FTransform RealTransform = Transform.ToFTransform(); //Transform is FEulerTransform
-			ControlRig->SetControlValue<FRigControlValue::FTransform_Float>(ControlName, RealTransform, true, EControlRigSetKey::DoNotCare,false);
-			FVector EulerAngle(Transform.Rotation.Roll, Transform.Rotation.Pitch, Transform.Rotation.Yaw);
-			ControlRig->GetHierarchy()->SetControlSpecifiedEulerAngle(ControlElement, EulerAngle);
-			ControlRig->Evaluate_AnyThread();
-
-		}
+		Item.ControlElements.Add(ControlElement);
 	}
-}
-
-void UControlRigTransformControlProxy::ValueChanged()
-{
-	FRigControlElement* ControlElement = GetControlElement();
-	if (ControlElement && ControlRig.IsValid())
+	if (ControlRigItems.Num() > 1 || Item.ControlElements.Num() > 1 || SequencerItems.Num() > 0)
 	{
-		const FName PropertyName("Transform");
-		FTrackInstancePropertyBindings Binding(PropertyName, PropertyName.ToString());
-		const FTransform NewTransform = ControlRig.Get()->GetControlValue(ControlElement, ERigControlValueType::Current).Get<FRigControlValue::FTransform_Float>().ToTransform();
-		FEulerTransform EulerTransform(NewTransform);
-		EulerTransform.Rotation = ControlRig->GetHierarchy()->GetControlPreferredRotator(ControlElement);
-		Binding.CallFunction<FEulerTransform>(*this, EulerTransform);
+		FString DisplayString = TEXT("Multiple");
+		FName DisplayName(*DisplayString);
+		Name = DisplayName;
 	}
-}
-
-#if WITH_EDITOR
-void UControlRigTransformControlProxy::PostEditUndo()
-{
-	FRigControlElement* ControlElement = GetControlElement();
-	if (ControlElement && ControlRig.IsValid() && ControlRig->GetHierarchy()->Contains(FRigElementKey(ControlName, ERigElementType::Control)))
+	else
 	{
-		ControlRig->SelectControl(ControlName, bSelected);
-		const FTransform RealTransform = Transform.ToFTransform(); //Transform is FEulerTransform
-		ControlRig->SetControlValue<FRigControlValue::FTransform_Float>(ControlName, RealTransform, true, EControlRigSetKey::Never,false);
-		FVector EulerAngle(Transform.Rotation.Roll, Transform.Rotation.Pitch, Transform.Rotation.Yaw);
-		ControlRig->GetHierarchy()->SetControlSpecifiedEulerAngle(ControlElement, EulerAngle);
+		Name = ControlElement->GetDisplayName();
 	}
 }
-#endif
 
-void UControlRigTransformControlProxy::SetKey(const IPropertyHandle& KeyedPropertyHandle)
+TArray<FRigControlElement*> UControlRigControlsProxy::GetControlElements() const
 {
-	FRigControlElement* ControlElement = GetControlElement();
-	if (ControlElement)
+	TArray<FRigControlElement*> Elements;
+	for (const TPair<UControlRig*, FControlRigProxyItem>& Items : ControlRigItems)
 	{
-		FRigControlModifiedContext Context;
-		Context.SetKey = EControlRigSetKey::Always;
-		if (KeyedPropertyHandle.GetProperty()->GetFName() == GET_MEMBER_NAME_CHECKED(FEulerTransform, Location))
-		{
-			Context.KeyMask = (uint32)EControlRigContextChannelToKey::Translation;
-		}
-		else if (KeyedPropertyHandle.GetProperty()->GetFName() == GET_MEMBER_NAME_CHECKED(FEulerTransform, Rotation))
+		for (FRigControlElement* Element : Items.Value.ControlElements)
 		{
-			Context.KeyMask = (uint32)EControlRigContextChannelToKey::Rotation;
+			Elements.Add(Element);
 		}
-		else if (KeyedPropertyHandle.GetProperty()->GetFName() == GET_MEMBER_NAME_CHECKED(FEulerTransform, Scale))
-		{
-			Context.KeyMask = (uint32)EControlRigContextChannelToKey::Scale;
-		}
-		const FTransform RealTransform = Transform.ToFTransform(); //Transform is FEulerTransform
-		ControlRig->SetControlValue<FRigControlValue::FTransform_Float>(ControlName, RealTransform, true, Context, false);
-		FVector EulerAngle(Transform.Rotation.Roll, Transform.Rotation.Pitch, Transform.Rotation.Yaw);
-		ControlRig->GetHierarchy()->SetControlSpecifiedEulerAngle(ControlElement, EulerAngle);
 	}
+	return Elements;
 }
 
-void UControlRigTransformNoScaleControlProxy::PostEditChangeChainProperty(struct FPropertyChangedChainEvent& PropertyChangedEvent)
+void UControlRigControlsProxy::ResetControlRigItems()
 {
-	Super::PostEditChangeChainProperty(PropertyChangedEvent);
-	FProperty* Property = PropertyChangedEvent.Property;
-	FProperty* MemberProperty = nullptr;
-	if (PropertyChangedEvent.PropertyChain.GetActiveMemberNode())
+	ChildProxies.Reset();
+	ControlRigItems.Reset();
+	if (OwnerControlRig.IsValid())
 	{
-		MemberProperty = PropertyChangedEvent.PropertyChain.GetActiveMemberNode()->GetValue();
-	}
-
-	if ((Property && Property->GetFName() == GET_MEMBER_NAME_CHECKED(UControlRigTransformNoScaleControlProxy, Transform))
-		|| (MemberProperty && MemberProperty->GetFName() == GET_MEMBER_NAME_CHECKED(UControlRigTransformNoScaleControlProxy, Transform)))
-	{
-		FRigControlElement* ControlElement = GetControlElement();
-		if (ControlElement && ControlRig.IsValid())
-		{
-			//MUST set through ControlRig
-			ControlRig->SetControlValue<FRigControlValue::FTransformNoScale_Float>(ControlName, Transform, true, EControlRigSetKey::DoNotCare,false);
-			ControlRig->Evaluate_AnyThread();
-
-		}
+		AddControlRigControl(OwnerControlRig.Get(), OwnerControlElement->GetFName());
 	}
 }
 
-void UControlRigTransformNoScaleControlProxy::ValueChanged()
+void UControlRigControlsProxy::ResetItems()
 {
-	FRigControlElement* ControlElement = GetControlElement();
-	if (ControlElement)
-	{
-		const FName PropertyName("Transform");
-		FTrackInstancePropertyBindings Binding(PropertyName, PropertyName.ToString());
-		const FTransformNoScale NewTransform = ControlRig.Get()->GetControlValue(ControlElement, ERigControlValueType::Current).Get<FRigControlValue::FTransformNoScale_Float>().ToTransform();
-		Binding.CallFunction<FTransformNoScale>(*this, NewTransform);
-	}
+	ResetControlRigItems();
+	ResetSequencerItems();
 }
 
-#if WITH_EDITOR
-void UControlRigTransformNoScaleControlProxy::PostEditUndo()
+void UControlRigControlsProxy::AddItem(UControlRigControlsProxy* ControlProxy)
 {
-	FRigControlElement* ControlElement = GetControlElement();
-	if (ControlElement && ControlRig.IsValid() && ControlRig->GetHierarchy()->Contains(FRigElementKey(ControlName, ERigElementType::Control)))
+	if (ControlProxy->OwnerControlRig.IsValid())
 	{
-		ControlRig->SelectControl(ControlName, bSelected);
-		ControlRig->SetControlValue<FRigControlValue::FTransformNoScale_Float>(ControlName, Transform, true, EControlRigSetKey::Never,false);
+		AddControlRigControl(ControlProxy->OwnerControlRig.Get(), ControlProxy->OwnerControlElement->GetFName());
 	}
-}
-#endif
-
-void UControlRigTransformNoScaleControlProxy::SetKey(const IPropertyHandle& KeyedPropertyHandle)
-{
-	FRigControlElement* ControlElement = GetControlElement();
-	if (ControlElement)
+	else if(ControlProxy->OwnerObject.IsValid())
 	{
-		FRigControlModifiedContext Context;
-		Context.SetKey = EControlRigSetKey::Always;
-		if (KeyedPropertyHandle.GetProperty()->GetFName() == GET_MEMBER_NAME_CHECKED(FTransformNoScale, Location))
-		{
-			Context.KeyMask = (uint32)EControlRigContextChannelToKey::Translation;
-		}
-		else if (KeyedPropertyHandle.GetProperty()->GetFName() == GET_MEMBER_NAME_CHECKED(FTransformNoScale, Rotation))
-		{
-			Context.KeyMask = (uint32)EControlRigContextChannelToKey::Rotation;
-		}
-		ControlRig->SetControlValue<FRigControlValue::FTransformNoScale_Float>(ControlName, Transform, true, Context, false);
+		AddSequencerProxyItem(ControlProxy->OwnerObject.Get(), ControlProxy->OwnerBindingAndTrack.WeakTrack, ControlProxy->OwnerBindingAndTrack.Binding);
 	}
 }
 
-
-void UControlRigEulerTransformControlProxy::PostEditChangeChainProperty(struct FPropertyChangedChainEvent& PropertyChangedEvent)
+void UControlRigControlsProxy::AddSequencerProxyItem(UObject* InObject, TWeakObjectPtr<UMovieSceneTrack>& InTrack, TSharedPtr<FTrackInstancePropertyBindings>& InBinding)
 {
-	Super::PostEditChangeChainProperty(PropertyChangedEvent);
-
-	FProperty* Property = PropertyChangedEvent.Property;
-	FProperty* MemberProperty = nullptr;
-	if (PropertyChangedEvent.PropertyChain.GetActiveMemberNode())
+	if (InObject == nullptr)
 	{
-		MemberProperty = PropertyChangedEvent.PropertyChain.GetActiveMemberNode()->GetValue();
+		return;
 	}
 
-	if ((Property && Property->GetFName() == GET_MEMBER_NAME_CHECKED(UControlRigEulerTransformControlProxy, Transform))
-		|| (MemberProperty && MemberProperty->GetFName() == GET_MEMBER_NAME_CHECKED(UControlRigEulerTransformControlProxy, Transform)))
+	FSequencerProxyItem& Item = SequencerItems.FindOrAdd(InObject);
+	Item.OwnerObject = InObject;
+	bool bContainsBinding = false;
+	for (FBindingAndTrack& Binding : Item.Bindings)
 	{
-		FRigControlElement* ControlElement = GetControlElement();
-		if (ControlElement && ControlRig.IsValid())
+		if (Binding.WeakTrack == InTrack && Binding.Binding->GetPropertyName() == InBinding->GetPropertyName())
 		{
-			URigHierarchy* Hierarchy = ControlRig->GetHierarchy();
-			FVector EulerAngle(Transform.Rotation.Roll, Transform.Rotation.Pitch, Transform.Rotation.Yaw);
-			FQuat Quat = Hierarchy->GetControlQuaternion(ControlElement, EulerAngle);
-			Hierarchy->SetControlSpecifiedEulerAngle(ControlElement, EulerAngle);
-			FRotator UERotator(Quat);
-			FEulerTransform UETransform = Transform;
-			UETransform.Rotation = UERotator;
-
-			ControlRig->SetControlValue<FRigControlValue::FEulerTransform_Float>(ControlName, UETransform, true, EControlRigSetKey::DoNotCare,false);
-			ControlRig->Evaluate_AnyThread();
-
+			bContainsBinding = true;
 		}
 	}
-}
-
-void UControlRigEulerTransformControlProxy::ValueChanged()
-{
-	FRigControlElement* ControlElement = GetControlElement();
-	if (ControlElement)
-	{
-		const FName PropertyName("Transform");
-		FTrackInstancePropertyBindings Binding(PropertyName, PropertyName.ToString());
-		FEulerTransform NewTransform = ControlRig.Get()->GetControlValue(ControlElement, ERigControlValueType::Current).Get<FRigControlValue::FEulerTransform_Float>().ToTransform();
-		NewTransform.Rotation = ControlRig->GetHierarchy()->GetControlPreferredRotator(ControlElement);
-		Binding.CallFunction<FEulerTransform>(*this, NewTransform);
-
-	}
-}
-
-#if WITH_EDITOR
-void UControlRigEulerTransformControlProxy::PostEditUndo()
-{
-	FRigControlElement* ControlElement = GetControlElement();
-	if (ControlElement && ControlRig.IsValid() && ControlRig->GetHierarchy()->Contains(FRigElementKey(ControlName, ERigElementType::Control)))
+	if (bContainsBinding == false)
 	{
-		ControlRig->SelectControl(ControlName, bSelected);
-		ControlRig->SetControlValue<FRigControlValue::FEulerTransform_Float>(ControlName, Transform, true, EControlRigSetKey::Never,false);
-		FVector EulerAngle(Transform.Rotation.Roll, Transform.Rotation.Pitch, Transform.Rotation.Yaw);
-		ControlRig->GetHierarchy()->SetControlSpecifiedEulerAngle(ControlElement, EulerAngle);
+		FBindingAndTrack BindingAndTrack(InBinding, InTrack.Get());
+		Item.Bindings.Add(BindingAndTrack);
 	}
-}
-#endif
-
-void UControlRigEulerTransformControlProxy::SetKey(const IPropertyHandle& KeyedPropertyHandle)
-{
-	FRigControlElement* ControlElement = GetControlElement();
-	if (ControlElement)
+	if (SequencerItems.Num() > 1 || Item.Bindings.Num() > 1 || ControlRigItems.Num() > 0)
 	{
-		FRigControlModifiedContext Context;
-		Context.SetKey = EControlRigSetKey::Always;
-		if (KeyedPropertyHandle.GetProperty()->GetFName() == GET_MEMBER_NAME_CHECKED(FEulerTransform, Location))
-		{
-			Context.KeyMask = (uint32)EControlRigContextChannelToKey::Translation;
-		}
-		else if (KeyedPropertyHandle.GetProperty()->GetFName() == GET_MEMBER_NAME_CHECKED(FEulerTransform, Rotation))
-		{
-			Context.KeyMask = (uint32)EControlRigContextChannelToKey::Rotation;
-		}
-		else if (KeyedPropertyHandle.GetProperty()->GetFName() == GET_MEMBER_NAME_CHECKED(FEulerTransform, Scale))
-		{
-			Context.KeyMask = (uint32)EControlRigContextChannelToKey::Scale;
-		}
-		ControlRig->SetControlValue<FRigControlValue::FEulerTransform_Float>(ControlName, Transform, true, Context,false);
-		FVector EulerAngle(Transform.Rotation.Roll, Transform.Rotation.Pitch, Transform.Rotation.Yaw);
-		ControlRig->GetHierarchy()->SetControlSpecifiedEulerAngle(ControlElement, EulerAngle);
+		FString DisplayString = TEXT("Multiple");
+		FName DisplayName(*DisplayString);
+		Name = DisplayName;
 	}
-}
-
-void UControlRigFloatControlProxy::PostEditChangeChainProperty(struct FPropertyChangedChainEvent& PropertyChangedEvent)
-{
-	Super::PostEditChangeChainProperty(PropertyChangedEvent);
-	if (PropertyChangedEvent.Property && PropertyChangedEvent.Property->GetFName() == GET_MEMBER_NAME_CHECKED(UControlRigFloatControlProxy, Float))
+	else
 	{
-		FRigControlElement* ControlElement = GetControlElement();
-		if (ControlElement && ControlRig.IsValid())
+		if (AActor* Actor = Cast<AActor>(InObject))
 		{
-			//MUST set through ControlRig
-			ControlRig->SetControlValue<float>(ControlName, Float, true, EControlRigSetKey::DoNotCare,false);
-			ControlRig->Evaluate_AnyThread();
-
+			FName DisplayName(*Actor->GetActorLabel());
+			Name = DisplayName;
 		}
-	}
-}
-
-void UControlRigFloatControlProxy::ValueChanged()
-{
-	FRigControlElement* ControlElement = GetControlElement();
-	if (ControlElement)
-	{
-		const FName PropertyName("Float");
-		FTrackInstancePropertyBindings Binding(PropertyName, PropertyName.ToString());
-		const float Val = ControlRig.Get()->GetControlValue(ControlElement, ERigControlValueType::Current).Get<float>();
-		Binding.CallFunction<float>(*this, Val);
-	}
-}
-
-#if WITH_EDITOR
-void UControlRigFloatControlProxy::PostEditUndo()
-{
-	FRigControlElement* ControlElement = GetControlElement();
-	if (ControlElement && ControlRig.IsValid() && ControlRig->GetHierarchy()->Contains(FRigElementKey(ControlName, ERigElementType::Control)))
-	{
-		ControlRig->SelectControl(ControlName, bSelected);
-		ControlRig->SetControlValue<float>(ControlName, Float, true, EControlRigSetKey::Never,false);
-	}
-}
-#endif
-
-
-void UControlRigFloatControlProxy::SetKey(const IPropertyHandle& KeyedPropertyHandle)
-{
-	FRigControlElement* ControlElement = GetControlElement();
-	if (ControlElement)
-	{
-		ControlRig->SetControlValue<float>(ControlName, Float, true, EControlRigSetKey::Always,false);
-	}
-}
-
-void UControlRigIntegerControlProxy::PostEditChangeChainProperty(struct FPropertyChangedChainEvent& PropertyChangedEvent)
-{
-	Super::PostEditChangeChainProperty(PropertyChangedEvent);
-	if (PropertyChangedEvent.Property && PropertyChangedEvent.Property->GetFName() == GET_MEMBER_NAME_CHECKED(UControlRigIntegerControlProxy, Integer))
-	{
-		FRigControlElement* ControlElement = GetControlElement();
-		if (ControlElement && ControlRig.IsValid())
+		else if (UActorComponent* Component = Cast<UActorComponent>(InObject))
 		{
-			//MUST set through ControlRig
-			ControlRig->SetControlValue<int32>(ControlName, Integer, true, EControlRigSetKey::DoNotCare,false);
-			ControlRig->Evaluate_AnyThread();
+			FName DisplayName(*Component->GetName());
+			Name = DisplayName;
 		}
-	}
-}
-
-void UControlRigIntegerControlProxy::ValueChanged()
-{
-	FRigControlElement* ControlElement = GetControlElement();
-	if (ControlElement)
-	{
-		const FName PropertyName("Integer");
-		FTrackInstancePropertyBindings Binding(PropertyName, PropertyName.ToString());
-		const int32 Val = ControlRig.Get()->GetControlValue(ControlElement, ERigControlValueType::Current).Get<int32>();
-		Binding.CallFunction<int32>(*this, Val);
-	}
-}
-
-#if WITH_EDITOR
-void UControlRigIntegerControlProxy::PostEditUndo()
-{
-	FRigControlElement* ControlElement = GetControlElement();
-	if (ControlElement && ControlRig.IsValid() && ControlRig->GetHierarchy()->Contains(FRigElementKey(ControlName, ERigElementType::Control)))
-	{
-		ControlRig->SelectControl(ControlName, bSelected);
-		ControlRig->SetControlValue<int32>(ControlName, Integer, true, EControlRigSetKey::Never,false);
-	}
-}
-#endif
-
-void UControlRigIntegerControlProxy::SetKey(const IPropertyHandle& KeyedPropertyHandle)
-{
-	FRigControlElement* ControlElement = GetControlElement();
-	if (ControlElement)
-	{
-		ControlRig->SetControlValue<int32>(ControlName, Integer, true, EControlRigSetKey::Always,false);
-	}
-}
-
-void UControlRigEnumControlProxy::PostEditChangeChainProperty(struct FPropertyChangedChainEvent& PropertyChangedEvent)
-{
-	Super::PostEditChangeChainProperty(PropertyChangedEvent);
-	if (PropertyChangedEvent.Property && PropertyChangedEvent.Property->GetFName() == GET_MEMBER_NAME_CHECKED(UControlRigEnumControlProxy, Enum))
-	{
-		FRigControlElement* ControlElement = GetControlElement();
-		if (ControlElement && ControlRig.IsValid())
+		else
 		{
-			//MUST set through ControlRig
-			ControlRig->SetControlValue<int32>(ControlName, Enum.EnumIndex, true, EControlRigSetKey::DoNotCare,false);
-			ControlRig->Evaluate_AnyThread();
-
+			Name = InBinding->GetPropertyName();
 		}
 	}
 }
 
-void UControlRigEnumControlProxy::ValueChanged()
+TArray<FBindingAndTrack> UControlRigControlsProxy::GetSequencerItems() const
 {
-	FRigControlElement* ControlElement = GetControlElement();
-	if (ControlElement)
+	TArray<FBindingAndTrack> Elements;
+	for (const TPair<UObject*, FSequencerProxyItem>& Items : SequencerItems)
 	{
-		const FName PropertyName("Enum");
-		FTrackInstancePropertyBindings Binding(PropertyName, PropertyName.ToString());
-
-		FControlRigEnumControlProxyValue Val;
-		Val.EnumType = ControlElement->Settings.ControlEnum;
-		Val.EnumIndex = ControlRig.Get()->GetControlValue(ControlElement, ERigControlValueType::Current).Get<int32>();
-
-		Binding.CallFunction<FControlRigEnumControlProxyValue>(*this, Val);
-	}
-}
-
-#if WITH_EDITOR
-void UControlRigEnumControlProxy::PostEditUndo()
-{
-	FRigControlElement* ControlElement = GetControlElement();
-	if (ControlElement && ControlRig.IsValid() && ControlRig->GetHierarchy()->Contains(FRigElementKey(ControlName, ERigElementType::Control)))
-	{
-		ControlRig->SelectControl(ControlName, bSelected);
-		ControlRig->SetControlValue<int32>(ControlName, Enum.EnumIndex, true, EControlRigSetKey::Never,false);
-	}
-}
-#endif
-
-void UControlRigEnumControlProxy::SetKey(const IPropertyHandle& KeyedPropertyHandle)
-{
-	FRigControlElement* ControlElement = GetControlElement();
-	if (ControlElement)
-	{
-		ControlRig->SetControlValue<int32>(ControlName, Enum.EnumIndex, true, EControlRigSetKey::Always,false);
-	}
-}
-
-void UControlRigVectorControlProxy::PostEditChangeChainProperty(struct FPropertyChangedChainEvent& PropertyChangedEvent)
-{
-	Super::PostEditChangeChainProperty(PropertyChangedEvent);
-	FProperty* Property = PropertyChangedEvent.Property;
-	FProperty* MemberProperty = nullptr;
-	if (PropertyChangedEvent.PropertyChain.GetActiveMemberNode())
-	{
-		MemberProperty = PropertyChangedEvent.PropertyChain.GetActiveMemberNode()->GetValue();
-	}
-
-	if ((Property && Property->GetFName() == GET_MEMBER_NAME_CHECKED(UControlRigVectorControlProxy, Vector))
-		|| (MemberProperty && MemberProperty->GetFName() == GET_MEMBER_NAME_CHECKED(UControlRigVectorControlProxy, Vector)))
-	{
-		FRigControlElement* ControlElement = GetControlElement();
-		if (ControlElement && ControlRig.IsValid())
+		for (const FBindingAndTrack& Element : Items.Value.Bindings)
 		{
-			//MUST set through ControlRig
-			if (ControlElement->Settings.ControlType == ERigControlType::Rotator)
-			{
-				FVector DVector(Vector.X, Vector.Y, Vector.Z);
-				ControlRig->GetHierarchy()->SetControlSpecifiedEulerAngle(ControlElement, DVector);
-			}
-			ControlRig->SetControlValue<FVector3f>(ControlName, Vector, true, EControlRigSetKey::DoNotCare,false);
-			ControlRig->Evaluate_AnyThread();
-
+			Elements.Add(Element);
 		}
 	}
+	return Elements;
 }
 
-void UControlRigVectorControlProxy::ValueChanged()
+void UControlRigControlsProxy::ResetSequencerItems()
 {
-	FRigControlElement* ControlElement = GetControlElement();
-	if (ControlElement)
+	ChildProxies.Reset();
+	SequencerItems.Reset();
+	if (OwnerObject.IsValid())
 	{
-		const FName PropertyName("Vector");
-		FTrackInstancePropertyBindings Binding(PropertyName, PropertyName.ToString());
-		FVector3f Val = ControlRig.Get()->GetControlValue(ControlElement, ERigControlValueType::Current).Get<FVector3f>();
-		if (ControlElement->Settings.ControlType == ERigControlType::Rotator)
-		{ 
-			FVector DVector = ControlRig->GetControlSpecifiedEulerAngle(ControlElement);
-			Val = FVector3f(DVector.X, DVector.Y, DVector.Z);
-		}
-		Binding.CallFunction<FVector3f>(*this, Val);
-	}
-}
-
-#if WITH_EDITOR
-void UControlRigVectorControlProxy::PostEditUndo()
-{
-	FRigControlElement* ControlElement = GetControlElement();
-	if (ControlElement && ControlRig.IsValid() && ControlRig->GetHierarchy()->Contains(FRigElementKey(ControlName, ERigElementType::Control)))
-	{
-		ControlRig->SelectControl(ControlName, bSelected);
-		if (ControlElement && ControlElement->Settings.ControlType == ERigControlType::Rotator)
-		{
-			FVector DVector(Vector.X, Vector.Y, Vector.Z);
-			ControlRig->GetHierarchy()->SetControlSpecifiedEulerAngle(ControlElement, DVector);
-		}
-		ControlRig->SetControlValue<FVector3f>(ControlName, (FVector3f)Vector, true, EControlRigSetKey::Never,false);
+		AddSequencerProxyItem(OwnerObject.Get(), OwnerBindingAndTrack.WeakTrack, OwnerBindingAndTrack.Binding);
 	}
 }
-#endif
 
-
-void UControlRigVectorControlProxy::SetKey(const IPropertyHandle& KeyedPropertyHandle)
+void UControlRigControlsProxy::AddChildProxy(UControlRigControlsProxy* ControlProxy)
 {
-	FRigControlElement* ControlElement = GetControlElement();
-	if (ControlElement)
+	if (ChildProxies.Contains(ControlProxy) == false)
 	{
-		FRigControlModifiedContext Context;
-		Context.SetKey = EControlRigSetKey::Always;
-		if (KeyedPropertyHandle.GetProperty()->GetFName() == GET_MEMBER_NAME_CHECKED(FVector3f, X))
-		{
-			switch (ControlElement->Settings.ControlType)
-			{
-			case ERigControlType::Position:
-				Context.KeyMask = (uint32)EControlRigContextChannelToKey::TranslationX;
-				break;
-			case ERigControlType::Rotator:	
-				Context.KeyMask = (uint32)EControlRigContextChannelToKey::RotationX;
-				break;
-			case ERigControlType::Scale:
-				Context.KeyMask = (uint32)EControlRigContextChannelToKey::ScaleX;
-				break;
-			}
-		}
-		else if (KeyedPropertyHandle.GetProperty()->GetFName() == GET_MEMBER_NAME_CHECKED(FVector3f, Y))
-		{
-			switch (ControlElement->Settings.ControlType)
-			{
-			case ERigControlType::Position:
-				Context.KeyMask = (uint32)EControlRigContextChannelToKey::TranslationY;
-				break;
-			case ERigControlType::Rotator:
-				Context.KeyMask = (uint32)EControlRigContextChannelToKey::RotationY;
-				break;
-			case ERigControlType::Scale:
-				Context.KeyMask = (uint32)EControlRigContextChannelToKey::ScaleY;
-				break;
-			}
-		}
-		else if (KeyedPropertyHandle.GetProperty()->GetFName() == GET_MEMBER_NAME_CHECKED(FVector3f, Z))
-		{
-			switch (ControlElement->Settings.ControlType)
-			{
-			case ERigControlType::Position:
-				Context.KeyMask = (uint32)EControlRigContextChannelToKey::TranslationZ;
-				break;
-			case ERigControlType::Rotator:
-				Context.KeyMask = (uint32)EControlRigContextChannelToKey::RotationZ;
-				break;
-			case ERigControlType::Scale:
-				Context.KeyMask = (uint32)EControlRigContextChannelToKey::ScaleZ;
-				break;
-			}
-		}
-		if (ControlElement->Settings.ControlType == ERigControlType::Rotator)
-		{
-			FVector DVector(Vector.X, Vector.Y, Vector.Z);
-			ControlRig->GetHierarchy()->SetControlSpecifiedEulerAngle(ControlElement, DVector);
-		}
-		ControlRig->SetControlValue<FVector3f>(ControlName, (FVector3f)Vector, true, Context,false);
+		ChildProxies.Add(ControlProxy);
 	}
 }
 
-void UControlRigVector2DControlProxy::PostEditChangeChainProperty(struct FPropertyChangedChainEvent& PropertyChangedEvent)
-{
-	Super::PostEditChangeChainProperty(PropertyChangedEvent);
-	FProperty* Property = PropertyChangedEvent.Property;
-	FProperty* MemberProperty = nullptr;
-	if (PropertyChangedEvent.PropertyChain.GetActiveMemberNode())
-	{
-		MemberProperty = PropertyChangedEvent.PropertyChain.GetActiveMemberNode()->GetValue();
-	}
-
-	if ((Property && Property->GetFName() == GET_MEMBER_NAME_CHECKED(UControlRigVector2DControlProxy, Vector2D))
-		|| (MemberProperty && MemberProperty->GetFName() == GET_MEMBER_NAME_CHECKED(UControlRigVector2DControlProxy, Vector2D)))
-	{
-		FRigControlElement* ControlElement = GetControlElement();
-		if (ControlElement && ControlRig.IsValid())
-		{
-			//MUST set through ControlRig
-			ControlRig->SetControlValue<FVector3f>(ControlName, FVector3f(Vector2D.X, Vector2D.Y, 0.f), true, EControlRigSetKey::DoNotCare,false);
-			ControlRig->Evaluate_AnyThread();
-
-		}
-	}
-}
-
-void UControlRigVector2DControlProxy::ValueChanged()
+void UControlRigControlsProxy::SelectionChanged(bool bInSelected)
 {
-	FRigControlElement* ControlElement = GetControlElement();
-	if (ControlElement)
+	if (OwnerControlElement)
 	{
-		const FName PropertyName("Vector2D");
+		Modify();
+		const FName PropertyName("bSelected");
 		FTrackInstancePropertyBindings Binding(PropertyName, PropertyName.ToString());
-		const FVector3f TempValue = ControlRig.Get()->GetControlValue(ControlElement, ERigControlValueType::Current).Get<FVector3f>();
-		const FVector2D Val(TempValue.X, TempValue.Y);
-		Binding.CallFunction<FVector2D>(*this, Val);
-	}
-}
-
-#if WITH_EDITOR
-void UControlRigVector2DControlProxy::PostEditUndo()
-{
-	FRigControlElement* ControlElement = GetControlElement();
-	if (ControlElement && ControlRig.IsValid() && ControlRig->GetHierarchy()->Contains(FRigElementKey(ControlName, ERigElementType::Control)))
-	{
-		ControlRig->SelectControl(ControlName, bSelected);
-		ControlRig->SetControlValue<FVector3f>(ControlName, FVector3f(Vector2D.X, Vector2D.Y, 0.f), true, EControlRigSetKey::Never,false);
-	}
-}
-#endif
-
-
-void UControlRigVector2DControlProxy::SetKey(const IPropertyHandle& KeyedPropertyHandle)
-{
-	FRigControlElement* ControlElement = GetControlElement();
-	if (ControlElement)
-	{
-		FRigControlModifiedContext Context;
-		Context.SetKey = EControlRigSetKey::Always;
-		if (KeyedPropertyHandle.GetProperty()->GetFName() == GET_MEMBER_NAME_CHECKED(FVector2D, X))
-		{
-			Context.KeyMask = (uint32)EControlRigContextChannelToKey::TranslationX;
-		}
-		else if (KeyedPropertyHandle.GetProperty()->GetFName() == GET_MEMBER_NAME_CHECKED(FVector2D, Y))
-		{
-			Context.KeyMask = (uint32)EControlRigContextChannelToKey::TranslationY;
-		}
-		ControlRig->SetControlValue<FVector3f>(ControlName, FVector3f(Vector2D.X, Vector2D.Y, 0.f), true, Context,false);
+		Binding.CallFunction<bool>(*this, bInSelected);
 	}
 }
 
-
-void UControlRigBoolControlProxy::PostEditChangeChainProperty(struct FPropertyChangedChainEvent& PropertyChangedEvent)
+void UControlRigControlsProxy::PostEditChangeChainProperty(struct FPropertyChangedChainEvent& PropertyChangedEvent)
 {
 	Super::PostEditChangeChainProperty(PropertyChangedEvent);
-	if (PropertyChangedEvent.Property && PropertyChangedEvent.Property->GetFName() == GET_MEMBER_NAME_CHECKED(UControlRigBoolControlProxy, Bool))
+	if (PropertyChangedEvent.Property && PropertyChangedEvent.Property->GetFName() == GET_MEMBER_NAME_CHECKED(UControlRigControlsProxy, bSelected))
 	{
-		FRigControlElement* ControlElement = GetControlElement();
-		if (ControlElement && ControlRig.IsValid())
+		if (OwnerControlElement && OwnerControlRig.IsValid())
 		{
-			//MUST set through ControlRig
-			ControlRig->SetControlValue<bool>(ControlName, Bool, true, EControlRigSetKey::DoNotCare,false);
-			ControlRig->Evaluate_AnyThread();
+			FControlRigInteractionScope InteractionScope(OwnerControlRig.Get(), OwnerControlElement->GetKey());
+			OwnerControlRig.Get()->SelectControl(OwnerControlElement->GetKey().Name, bSelected);
+			OwnerControlRig.Get()->Evaluate_AnyThread();
 		}
 	}
-}
-
-void UControlRigBoolControlProxy::ValueChanged()
-{
-	FRigControlElement* ControlElement = GetControlElement();
-	if (ControlElement)
-	{
-		const FName PropertyName("Bool");
-		FTrackInstancePropertyBindings Binding(PropertyName, PropertyName.ToString());
-		const bool Val = ControlRig.Get()->GetControlValue(ControlElement, ERigControlValueType::Current).Get<bool>();
-		Binding.CallFunction<bool>(*this, Val);
-	}
-}
-
 #if WITH_EDITOR
-void UControlRigBoolControlProxy::PostEditUndo()
-{
-	FRigControlElement* ControlElement = GetControlElement();
-	if (ControlElement && ControlRig.IsValid() && ControlRig->GetHierarchy()->Contains(FRigElementKey(ControlName, ERigElementType::Control)))
-	{
-		ControlRig->SelectControl(ControlName, bSelected);
-		ControlRig->SetControlValue<bool>(ControlName, Bool, true, EControlRigSetKey::Never,false);
-	}
-}
-#endif
-
-
-void UControlRigBoolControlProxy::SetKey(const IPropertyHandle& KeyedPropertyHandle)
-{
-	FRigControlElement* ControlElement = GetControlElement();
-	if (ControlElement)
-	{
-		ControlRig->SetControlValue<bool>(ControlName, Bool, true, EControlRigSetKey::Always,false);
-	}
-}
-
-
-//////UControlDetailPanelControlProxies////////
-
-UControlRigControlsProxy* UControlRigDetailPanelControlProxies::FindProxy(UControlRig* ControlRig, const FName& Name) const
-{
-	const FControlToProxyMap* ControlRigProxies = AllProxies.Find(ControlRig);
-	if (ControlRigProxies)
-	{
-		TObjectPtr<UControlRigControlsProxy> const* Proxy = ControlRigProxies->ControlToProxy.Find(Name);
-		if (Proxy && Proxy[0])
-		{
-			return Proxy[0];
-		}
-	}
-	return nullptr;
-}
-
-void UControlRigDetailPanelControlProxies::AddProxy(UControlRig* ControlRig, const FName& Name,  FRigControlElement* ControlElement)
-{
-	UControlRigControlsProxy* Proxy = FindProxy(ControlRig,Name);
-	if (!Proxy && ControlElement != nullptr)
+	if (PropertyChangedEvent.Property)
 	{
-		switch(ControlElement->Settings.ControlType)
+		if (PropertyChangedEvent.ChangeType == EPropertyChangeType::Interactive)
 		{
-			case ERigControlType::Transform:
-			{
-				Proxy = NewObject<UControlRigTransformControlProxy>(GetTransientPackage(), NAME_None);
-				break;
-
-			}
-			case ERigControlType::TransformNoScale:
-			{
-				Proxy = NewObject<UControlRigTransformNoScaleControlProxy>(GetTransientPackage(), NAME_None);
-				break;
-
-			}
-			case ERigControlType::EulerTransform:
-			{
-				Proxy = NewObject<UControlRigEulerTransformControlProxy>(GetTransientPackage(), NAME_None);
-				break;
-			}
-			case ERigControlType::Float:
-			case ERigControlType::ScaleFloat:
-			{
-				Proxy = NewObject<UControlRigFloatControlProxy>(GetTransientPackage(), NAME_None);
-				break;
-
-			}
-			case ERigControlType::Integer:
+			EControlRigInteractionType InteractionType = EControlRigInteractionType::None;
+			FProperty* Owner = PropertyChangedEvent.Property->GetOwnerProperty();
+			if (FEditPropertyChain::TDoubleLinkedListNode* MemberNode = PropertyChangedEvent.PropertyChain.GetActiveMemberNode())
 			{
-				if(ControlElement->Settings.ControlEnum == nullptr)
-				{
-					Proxy = NewObject<UControlRigIntegerControlProxy>(GetTransientPackage(), NAME_None);
-				}
-				else
+				if (FProperty* MemberProperty = MemberNode->GetValue())
 				{
-					UControlRigEnumControlProxy* EnumProxy = NewObject<UControlRigEnumControlProxy>(GetTransientPackage(), NAME_None);
-					EnumProxy->Enum.EnumType = ControlElement->Settings.ControlEnum;
-					Proxy = EnumProxy;
+					if (MemberProperty->GetFName() == GET_MEMBER_NAME_CHECKED(FEulerTransform, Location))
+					{
+						InteractionType = EControlRigInteractionType::Translate;
+					}
+					else if (MemberProperty->GetFName() == GET_MEMBER_NAME_CHECKED(FEulerTransform, Rotation))
+					{
+						InteractionType = EControlRigInteractionType::Rotate;
+					}
+					else if (MemberProperty->GetFName() == GET_MEMBER_NAME_CHECKED(FEulerTransform, Scale))
+					{
+						InteractionType = EControlRigInteractionType::Scale;
+					}
 				}
-				break;
-
-			}
-			case ERigControlType::Position:
-			case ERigControlType::Rotator:
-			case ERigControlType::Scale:
-			{
-				Proxy = NewObject<UControlRigVectorControlProxy>(GetTransientPackage(), NAME_None);
-				break;
-
-			}
-			case ERigControlType::Vector2D:
-			{
-				Proxy = NewObject<UControlRigVector2DControlProxy>(GetTransientPackage(), NAME_None);
-				break;
-
 			}
-			case ERigControlType::Bool:
+			for (const TPair<UControlRig*, FControlRigProxyItem>& Items : ControlRigItems)
 			{
-				Proxy = NewObject<UControlRigBoolControlProxy>(GetTransientPackage(), NAME_None);
-				break;
-
+				if (UControlRig* ControlRig = Items.Value.ControlRig.Get())
+				{
+					for (FRigControlElement* ControlElement : Items.Value.ControlElements)
+					{
+						if (ControlElement)
+						{
+							if (InteractionScopes.Contains(ControlElement) == false)
+							{
+								FControlRigInteractionScope* InteractionScope = new FControlRigInteractionScope(ControlRig, ControlElement->GetKey(), InteractionType);
+								InteractionScopes.Add(ControlElement, InteractionScope);
+							}
+						}
+					}
+				}
 			}
-			default:
-				break;
-		}
-		if (Proxy)
-		{
-			Proxy->bIsIndividual =
-				(ControlElement->IsAnimationChannel()) ||
-				(ControlElement->Settings.AnimationType == ERigControlAnimationType::ProxyControl) ;
-			Proxy->SetFlags(RF_Transactional);
-			Proxy->SetName(Name);
-			Proxy->ControlRig = ControlRig;
-			Proxy->Modify();
-			Proxy->ValueChanged();
-
-			FControlToProxyMap* ControlRigProxies = AllProxies.Find(ControlRig);
-			if (ControlRigProxies)
+			//set values
+			FProperty* Property = PropertyChangedEvent.Property;
+			FProperty* MemberProperty = nullptr;
+			if (PropertyChangedEvent.PropertyChain.GetActiveMemberNode())
 			{
-				ControlRigProxies->ControlToProxy.Add(Name, Proxy);
+				MemberProperty = PropertyChangedEvent.PropertyChain.GetActiveMemberNode()->GetValue();
 			}
-			else
+			if (PropertyIsOnProxy(Property, MemberProperty))
 			{
-				FControlToProxyMap NewControlRigProxies;
-				NewControlRigProxies.ControlToProxy.Add(Name, Proxy);
-				AllProxies.Add(ControlRig, NewControlRigProxies);
-			}
-		}
-	}
-}
-
-void UControlRigDetailPanelControlProxies::RemoveProxy(UControlRig* ControlRig, const FName& Name)
-{
-	UControlRigControlsProxy* ExistingProxy = FindProxy(ControlRig,Name);
-	if (ExistingProxy)
-	{
-		ExistingProxy->Rename(nullptr, GetTransientPackage(), REN_ForceNoResetLoaders);
-		ExistingProxy->MarkAsGarbage();
-	}
-	FControlToProxyMap* ControlRigProxies = AllProxies.Find(ControlRig);
-	if (ControlRigProxies)
-	{
-		ControlRigProxies->ControlToProxy.Remove(Name);
-	}
-}
+				FRigControlModifiedContext Context;
+				Context.SetKey = EControlRigSetKey::DoNotCare;
+				Context.KeyMask = (uint32)GetChannelToKeyFromPropertyName(Property->GetFName());
 
-void UControlRigDetailPanelControlProxies::RemoveAllProxies(UControlRig* ControlRig)
-{
-	//no control rig remove all
-	if (ControlRig == nullptr)
-	{
-		for (TPair<TObjectPtr<UControlRig>, FControlToProxyMap>& ControlRigProxies : AllProxies)
-		{
-			for (TPair<FName, TObjectPtr<UControlRigControlsProxy> >& Pair : ControlRigProxies.Value.ControlToProxy)
-			{
-				UControlRigControlsProxy* ExistingProxy = Pair.Value;
-				if (ExistingProxy)
+				for (const TPair<UControlRig*, FControlRigProxyItem>& Items : ControlRigItems)
 				{
-					ExistingProxy->Rename(nullptr, GetTransientPackage(), REN_ForceNoResetLoaders);
-					ExistingProxy->MarkAsGarbage();
+					if (UControlRig* ControlRig = Items.Value.ControlRig.Get())
+					{
+						//we do this backwards so ValueChanged later is set up correctly since that iterates in the other direction
+						for (int32 Index = Items.Value.ControlElements.Num() - 1; Index >= 0; --Index)
+						{
+							FRigControlElement* ControlElement = Items.Value.ControlElements[Index];
+							SetControlRigElementValueFromCurrent(ControlRig, ControlElement, Context);
+						}
+					}
 				}
-			}
-		}
-		AllProxies.Empty();
-		SelectedProxies.SetNum(0);
-	}
-	else
-	{
-		FControlToProxyMap* ControlRigProxies = AllProxies.Find(ControlRig);
-		if (ControlRigProxies)
-		{
-			for (TPair<FName, TObjectPtr<UControlRigControlsProxy>>& Pair : ControlRigProxies->ControlToProxy)
-			{
-				UControlRigControlsProxy* ExistingProxy = Pair.Value;
-				if (ExistingProxy)
+				for (TPair <UObject*, FSequencerProxyItem>& SItems : SequencerItems)
 				{
-					SelectedProxies.Remove(ExistingProxy);
-					ExistingProxy->Rename(nullptr, GetTransientPackage(), REN_ForceNoResetLoaders);
-					ExistingProxy->MarkAsGarbage();
+					//we do this backwards so ValueChanged later is set up correctly since that iterates in the other direction
+					for (int32 Index = SItems.Value.Bindings.Num() - 1; Index >= 0; --Index)
+					{
+						FBindingAndTrack& Binding = SItems.Value.Bindings[Index];
+						SetBindingValueFromCurrent(SItems.Key, Binding.Binding, Context);
+					}
 				}
 			}
-
-			AllProxies.Remove(ControlRig);
-		}
-	}
-}
-
-void UControlRigDetailPanelControlProxies::RecreateAllProxies(UControlRig* ControlRig)
-{
-	RemoveAllProxies(ControlRig);
-	TArray<FRigControlElement*> Controls = ControlRig->AvailableControls();
-	for (FRigControlElement* ControlElement : Controls)
-	{
-		if(ControlElement->Settings.AnimationType != ERigControlAnimationType::VisualCue)
-		{
-			AddProxy(ControlRig,ControlElement->GetFName(), ControlElement);
 		}
-	}
-}
-
-void UControlRigDetailPanelControlProxies::ProxyChanged(UControlRig* ControlRig, const FName& Name, bool bModify )
-{
-	if (IsInGameThread())
-	{
-		UControlRigControlsProxy* Proxy = FindProxy(ControlRig,Name);
-		if (Proxy)
+		else
 		{
-			if (bModify)
+			for (TPair<FRigControlElement*,FControlRigInteractionScope*>& Scope : InteractionScopes)
 			{
-				Modify();
-				Proxy->Modify();
+				if (Scope.Value)
+				{
+					delete Scope.Value; 
+				}
 			}
-			Proxy->ValueChanged();
+			InteractionScopes.Reset();
 		}
 	}
+#endif
 }
 
-void UControlRigDetailPanelControlProxies::SelectProxy(UControlRig* ControlRig,const FName& Name, bool bSelected)
+#if WITH_EDITOR
+void UControlRigControlsProxy::PostEditUndo()
 {
-	UControlRigControlsProxy* Proxy = FindProxy(ControlRig,Name);
-	if (Proxy)
+	for (const TPair<UControlRig*, FControlRigProxyItem>& Items : ControlRigItems)
 	{
-		Modify();
-		if (bSelected)
+		if (UControlRig* ControlRig = Items.Value.ControlRig.Get())
 		{
-			if (!SelectedProxies.Contains(Proxy))
+			for (FRigControlElement* ControlElement : Items.Value.ControlElements)
 			{
-				SelectedProxies.Add(Proxy);
+				if (ControlElement && ControlRig->GetHierarchy()->Contains(FRigElementKey(ControlElement->GetKey().Name, ERigElementType::Control)))
+				{
+					ControlRig->SelectControl(ControlElement->GetKey().Name, bSelected);
+				}
 			}
 		}
-		else
-		{
-			SelectedProxies.Remove(Proxy);
-		}
-		Proxy->SelectionChanged(bSelected);
 	}
 }
-
-bool UControlRigDetailPanelControlProxies::IsSelected(UControlRig* InControlRig, const FName& Name) const
-{
-	if (UControlRigControlsProxy* Proxy = FindProxy(InControlRig, Name))
-	{
-		return SelectedProxies.Contains(Proxy);
-	}
-	return false;
-}
-
+#endif
 
diff --git a/Engine/Plugins/Animation/ControlRig/Source/ControlRigEditor/Private/EditMode/ControlRigControlsProxy.h b/Engine/Plugins/Animation/ControlRig/Source/ControlRigEditor/Private/EditMode/ControlRigControlsProxy.h
--- a/Engine/Plugins/Animation/ControlRig/Source/ControlRigEditor/Private/EditMode/ControlRigControlsProxy.h
+++ b/Engine/Plugins/Animation/ControlRig/Source/ControlRigEditor/Private/EditMode/ControlRigControlsProxy.h
@@ -6,10 +6,12 @@
 #include "UObject/LazyObjectPtr.h"
 #include "TransformNoScale.h"
 #include "EulerTransform.h"
+#include "Math/Rotator.h"
 
-#if WITH_EDITOR
+#include "IDetailKeyframeHandler.h"
 #include "IPropertyTypeCustomization.h"
-#endif
+#include "Rigs/RigHierarchyDefines.h"
+#include "MovieSceneCommonHelpers.h"
 
 #include "ControlRigControlsProxy.generated.h"
 
@@ -17,323 +19,189 @@ struct FRigControlElement;
 class UControlRig;
 class IPropertyHandle;
 class FControlRigInteractionScope;
-
-UCLASS(Abstract)
-class UControlRigControlsProxy : public UObject
+struct FRigControlModifiedContext;
+class ISequencer;
+class FCurveEditor;
+class IDetailLayoutBuilder;
+class UMovieSceneTrack;
+class FProperty;
+
+//channel selection states for selection matching with curves
+enum class EAnimDetailSelectionState : uint8
 {
-	GENERATED_BODY()
-
-public:
-	UControlRigControlsProxy() : bSelected(false) {}
-	virtual void SetName(const FName& InName) { Name = ControlName = InName; }
-	virtual FName GetName() const { return Name; }
-	virtual void ValueChanged() {}
-	virtual void SelectionChanged(bool bInSelected);
-	virtual void SetKey(const IPropertyHandle& KeyedPropertyHandle) {};
-
-	// UObject interface
-	virtual void PostEditChangeChainProperty( struct FPropertyChangedChainEvent& PropertyChangedEvent ) override;
-#if WITH_EDITOR
-	virtual void PostEditUndo() override;
-#endif
-	void SetIsMultiple(bool bIsVal); 
-	void SetIsIndividual(bool bIsVal);
-public:
-
-	FRigControlElement* GetControlElement() const;
-	TWeakObjectPtr<UControlRig> ControlRig;
-
-	UPROPERTY()
-	bool bSelected;
-	UPROPERTY(VisibleAnywhere, Category = "Control")
-	FName ControlName;
-
-	//if individual it will show up independently, this will happen for certain nested controls
-	bool bIsIndividual = false;
-
-protected:
-	bool bIsMultiple = 0;
-	FName Name;
-#if WITH_EDITOR
-	TArray<FControlRigInteractionScope*> InteractionScopes;
-#endif
+	None = 0x0, Partial = 0x1, All = 0x2
 };
 
-UCLASS()
-class UControlRigTransformControlProxy : public UControlRigControlsProxy
+struct FAnimDetailVectorSelection
 {
-	GENERATED_BODY()
-	UControlRigTransformControlProxy() {}
-
-	// UObject interface
-	virtual void PostEditChangeChainProperty(struct FPropertyChangedChainEvent& PropertyChangedEvent) override;
-#if WITH_EDITOR
-	virtual void PostEditUndo() override;
-#endif
-
-	//UControlRigControlsProxy
-	virtual void ValueChanged() override;
-	virtual void SetKey(const IPropertyHandle& KeyedPropertyHandle) override;
-
-public:
-	
-	UPROPERTY(EditAnywhere, Interp, Category = "Control", meta = (SliderExponent = "1.0"))
-	FEulerTransform Transform; //FTransform doesn't work with multiple values for some reason, so for now using eulertransform which does work
+	EAnimDetailSelectionState  XSelected = EAnimDetailSelectionState::None;
+	EAnimDetailSelectionState  YSelected = EAnimDetailSelectionState::None;
+	EAnimDetailSelectionState  ZSelected = EAnimDetailSelectionState::None;
 };
 
-
-UCLASS()
-class UControlRigEulerTransformControlProxy : public UControlRigControlsProxy
+//item to specify control rig, todo move to handle class
+struct FControlRigProxyItem
 {
-	GENERATED_BODY()
-	UControlRigEulerTransformControlProxy() {}
-
-	// UObject interface
-	virtual void PostEditChangeChainProperty(struct FPropertyChangedChainEvent& PropertyChangedEvent) override;
-#if WITH_EDITOR
-	virtual void PostEditUndo() override;
-#endif
-
-	//UControlRigControlsProxy
-	virtual void ValueChanged() override;
-	virtual void SetKey(const IPropertyHandle& KeyedPropertyHandle) override;
-
-public:
-
-	UPROPERTY(EditAnywhere, Interp, Category = "Control", meta = (SliderExponent = "1.0"))
-	FEulerTransform Transform;
+	TWeakObjectPtr<UControlRig> ControlRig;
+	TArray<FRigControlElement*> ControlElements;
 };
 
-
-UCLASS()
-class UControlRigTransformNoScaleControlProxy : public UControlRigControlsProxy
+struct FBindingAndTrack
 {
-	GENERATED_BODY()
-	UControlRigTransformNoScaleControlProxy() {}
-
-	// UObject interface
-	virtual void PostEditChangeChainProperty(struct FPropertyChangedChainEvent& PropertyChangedEvent) override;
-#if WITH_EDITOR
-	virtual void PostEditUndo() override;
-#endif
-
-	//UControlRigControlsProxy
-	virtual void ValueChanged() override;
-	virtual void SetKey(const IPropertyHandle& KeyedPropertyHandle) override;
-
-public:
-
-	UPROPERTY(EditAnywhere, Interp, Category = "Control", meta = (SliderExponent = "1.0"))
-	FTransformNoScale Transform;
+	FBindingAndTrack() : Binding(nullptr), WeakTrack(nullptr) {};
+	FBindingAndTrack(TSharedPtr<FTrackInstancePropertyBindings>& InBinding, UMovieSceneTrack* InWeakTrack) :
+	Binding(InBinding), WeakTrack(InWeakTrack){}
+	TSharedPtr<FTrackInstancePropertyBindings> Binding;
+	TWeakObjectPtr<UMovieSceneTrack> WeakTrack;
 };
-
-UCLASS()
-class UControlRigFloatControlProxy : public UControlRigControlsProxy
+//item to specify sequencer item, todo move to handle class
+struct FSequencerProxyItem
 {
-	GENERATED_BODY()
-	UControlRigFloatControlProxy()  {}
-
-	// UObject interface
-	virtual void PostEditChangeChainProperty(struct FPropertyChangedChainEvent& PropertyChangedEvent) override;
-#if WITH_EDITOR
-	virtual void PostEditUndo() override;
-#endif
-
-	//UControlRigControlsProxy
-	virtual void ValueChanged() override;
-	virtual void SetKey(const IPropertyHandle& KeyedPropertyHandle) override;
-
-public:
-
-	UPROPERTY(EditAnywhere, Interp, Category = "Control", meta = (SliderExponent = "1.0"))
-	float Float;
+	TWeakObjectPtr<UObject> OwnerObject;
+	TArray<FBindingAndTrack> Bindings;
 };
 
-UCLASS()
-class UControlRigIntegerControlProxy : public UControlRigControlsProxy
+UCLASS(Abstract)
+class UControlRigControlsProxy : public UObject
 {
 	GENERATED_BODY()
-	UControlRigIntegerControlProxy() {}
-
-	// UObject interface
-	virtual void PostEditChangeChainProperty(struct FPropertyChangedChainEvent& PropertyChangedEvent) override;
-#if WITH_EDITOR
-	virtual void PostEditUndo() override;
-#endif
-
-	//UControlRigControlsProxy
-	virtual void ValueChanged() override;
-	virtual void SetKey(const IPropertyHandle& KeyedPropertyHandle) override;
-
-public:
-
-	UPROPERTY(EditAnywhere, Interp, Category = "Control")
-	int32 Integer;
-};
-
-USTRUCT(BlueprintType)
-struct FControlRigEnumControlProxyValue
-{
-	GENERATED_USTRUCT_BODY()
-
-	FControlRigEnumControlProxyValue()
-	{
-		EnumType = nullptr;
-		EnumIndex = INDEX_NONE;
-	}
-	
-	UPROPERTY()
-	TObjectPtr<UEnum> EnumType;
-
-	UPROPERTY(EditAnywhere, Category = Enum)
-	int32 EnumIndex;
-};
 
-#if WITH_EDITOR
-
-class FControlRigEnumControlProxyValueDetails : public IPropertyTypeCustomization
-{
 public:
-
-	static TSharedRef<IPropertyTypeCustomization> MakeInstance()
-	{
-		return MakeShareable(new FControlRigEnumControlProxyValueDetails);
-	}
-
-	/** IPropertyTypeCustomization interface */
-	virtual void CustomizeHeader(TSharedRef<class IPropertyHandle> InStructPropertyHandle, class FDetailWidgetRow& HeaderRow, IPropertyTypeCustomizationUtils& StructCustomizationUtils) override;
-	virtual void CustomizeChildren(TSharedRef<class IPropertyHandle> InStructPropertyHandle, class IDetailChildrenBuilder& StructBuilder, IPropertyTypeCustomizationUtils& StructCustomizationUtils) override;
-
-private:
-
-	int32 GetEnumValue() const;
-	void OnEnumValueChanged(int32 InValue, ESelectInfo::Type InSelectInfo, TSharedRef<IPropertyHandle> InStructHandle);
-
-	UControlRigEnumControlProxy* ProxyBeingCustomized;
-};
-#endif
-
-UCLASS()
-class UControlRigEnumControlProxy : public UControlRigControlsProxy
-{
-	GENERATED_BODY()
-	UControlRigEnumControlProxy() {}
+	UControlRigControlsProxy() : bSelected(false) {}
+	//will add controlrig or sequencer item
+	virtual void AddItem(UControlRigControlsProxy* ControlProxy);
+	virtual void AddChildProxy(UControlRigControlsProxy* ControlProxy);
+	virtual FName GetName() const { return Name; }
+	virtual void UpdatePropertyNames(IDetailLayoutBuilder& DetailBuilder) {};
+	virtual void ValueChanged() {}
+	virtual void SelectionChanged(bool bInSelected);
+	virtual void SetKey(TSharedPtr<ISequencer>& Sequencer, const IPropertyHandle& KeyedPropertyHandle) {};
+	virtual EPropertyKeyedStatus GetPropertyKeyedStatus(TSharedPtr<ISequencer>& Sequencer, const IPropertyHandle& PropertyHandle) const { return EPropertyKeyedStatus::NotKeyed; }
+	virtual EControlRigContextChannelToKey GetChannelToKeyFromPropertyName(const FName& PropertyName) const { return EControlRigContextChannelToKey::AllTransform; }
+	virtual EControlRigContextChannelToKey GetChannelToKeyFromChannelName(const FString& InChannelName) const { return EControlRigContextChannelToKey::AllTransform; }
+	virtual bool IsMultiple(const FName& InPropertyName) const { return false; }
+	virtual void SetControlRigElementValueFromCurrent(UControlRig* ControlRig, FRigControlElement* ControlElement, const FRigControlModifiedContext& Context) {};
+	virtual void SetBindingValueFromCurrent(UObject* InObject, TSharedPtr<FTrackInstancePropertyBindings>& Binding, FRigControlModifiedContext& Context) {};
+	virtual void GetChannelSelectionState(TWeakPtr<FCurveEditor>& CurveEditor, FAnimDetailVectorSelection& OutLocationSelection, FAnimDetailVectorSelection& OutRotationSelection,
+		FAnimDetailVectorSelection& OutScaleSelection) {};
+	virtual bool PropertyIsOnProxy(FProperty* Property, FProperty* MemberProperty){return false;}
 
 	// UObject interface
-	virtual void PostEditChangeChainProperty(struct FPropertyChangedChainEvent& PropertyChangedEvent) override;
+	virtual void PostEditChangeChainProperty( struct FPropertyChangedChainEvent& PropertyChangedEvent ) override;
 #if WITH_EDITOR
 	virtual void PostEditUndo() override;
 #endif
 
-	//UControlRigControlsProxy
-	virtual void ValueChanged() override;
-	virtual void SetKey(const IPropertyHandle& KeyedPropertyHandle) override;
+	TArray<FRigControlElement*> GetControlElements() const;
 
-public:
+	TArray<FBindingAndTrack> GetSequencerItems() const;
+	//Default type
+	ERigControlType Type = ERigControlType::Transform;
 
-	UPROPERTY(EditAnywhere, Interp, Category = "Control")
-	FControlRigEnumControlProxyValue Enum;
-};
+	//reset items it owns
+	void ResetItems();
 
-UCLASS()
-class UControlRigVectorControlProxy : public UControlRigControlsProxy
-{
-	GENERATED_BODY()
-	UControlRigVectorControlProxy() {}
+	//add correct item
+	void AddSequencerProxyItem(UObject* InObject, TWeakObjectPtr<UMovieSceneTrack>& InTrack,  TSharedPtr<FTrackInstancePropertyBindings>& Binding);
+	void AddControlRigControl(UControlRig* InControlRig, const FName& InName);
 
-	// UObject interface
-	virtual void PostEditChangeChainProperty(struct FPropertyChangedChainEvent& PropertyChangedEvent) override;
-#if WITH_EDITOR
-	virtual void PostEditUndo() override;
-#endif
+private:
+	void HandlePostEditChangedForControlRigs();
+	void HandlePostEditChangedForSequencer();
 
-	//UControlRigControlsProxy
-	virtual void ValueChanged() override;
-	virtual void SetKey(const IPropertyHandle& KeyedPropertyHandle) override;
+	//reset the  rig controls this also owns
+	void ResetControlRigItems();
+	void ResetSequencerItems();
 
 public:
 
-	UPROPERTY(EditAnywhere, Interp, Category = "Control", meta = (SliderExponent = "1.0"))
-	FVector3f Vector;
-};
-
-UCLASS()
-class UControlRigVector2DControlProxy : public UControlRigControlsProxy
-{
-	GENERATED_BODY()
-	UControlRigVector2DControlProxy() {}
-
-	// UObject interface
-	virtual void PostEditChangeChainProperty(struct FPropertyChangedChainEvent& PropertyChangedEvent) override;
-#if WITH_EDITOR
-	virtual void PostEditUndo() override;
-#endif
+	//if individual it will show up independently, this will happen for certain nested controls
+	bool bIsIndividual = false;
 
-	//UControlRigControlsProxy
-	virtual void ValueChanged() override;
-	virtual void SetKey(const IPropertyHandle& KeyedPropertyHandle) override;
+	UPROPERTY()
+	bool bSelected;
+	UPROPERTY(VisibleAnywhere, Category = "Control")
+	FName Name;
 
-public:
+	//refactor
+	//We can set/get values form multiple control rig elements but only one owns this.
+	TWeakObjectPtr<UControlRig> OwnerControlRig;
+	FRigControlElement* OwnerControlElement = nullptr;
+	TMap<UControlRig*, FControlRigProxyItem> ControlRigItems;
+	
+	TWeakObjectPtr<UObject> OwnerObject;
+	FBindingAndTrack OwnerBindingAndTrack;
+	TMap <UObject*, FSequencerProxyItem> SequencerItems;
 
-	UPROPERTY(EditAnywhere, Interp, Category = "Control", meta = (SliderExponent = "1.0"))
-	FVector2D Vector2D;
-};
+	//list of child/animation channel proxies that we will customize
+	TArray<UControlRigControlsProxy*> ChildProxies;
 
-UCLASS()
-class UControlRigBoolControlProxy : public UControlRigControlsProxy
-{
-	GENERATED_BODY()
-	UControlRigBoolControlProxy() {}
-
-	// UObject interface
-	virtual void PostEditChangeChainProperty(struct FPropertyChangedChainEvent& PropertyChangedEvent) override;
 #if WITH_EDITOR
-	virtual void PostEditUndo() override;
+	TMap<FRigControlElement*, FControlRigInteractionScope*> InteractionScopes;
 #endif
-
-	//UControlRigControlsProxy
-	virtual void ValueChanged() override;
-	virtual void SetKey(const IPropertyHandle& KeyedPropertyHandle) override;
-
-public:
-
-	UPROPERTY(EditAnywhere, Interp, Category = "Control")
-	bool Bool;
 };
 
 USTRUCT()
-struct FControlToProxyMap
+struct FNameToProxyMap
 {
 	GENERATED_BODY();
-
 	UPROPERTY()
-	TMap <FName, TObjectPtr<UControlRigControlsProxy>> ControlToProxy;
+	TMap <FName, TObjectPtr<UControlRigControlsProxy>> NameToProxy;;
+};
+
+struct FSequencerProxyPerType
+{
+	TMap<UObject*, TArray<FBindingAndTrack>> Bindings;
 };
+
 /** Proxy in Details Panel */
 UCLASS()
 class UControlRigDetailPanelControlProxies :public UObject
 {
 	GENERATED_BODY()
 
-	UControlRigDetailPanelControlProxies() {}
+	UControlRigDetailPanelControlProxies();
+	~UControlRigDetailPanelControlProxies();
+
 protected:
 
 	UPROPERTY()
-	TMap<TObjectPtr<UControlRig>, FControlToProxyMap> AllProxies; //proxies themselves contain weakobjectptr to the controlrig
+	TMap<TObjectPtr<UControlRig>, FNameToProxyMap> ControlRigOnlyProxies; //proxies themselves contain weakobjectptr to the controlrig
 
 	UPROPERTY()
-	TArray< TObjectPtr<UControlRigControlsProxy>> SelectedProxies;
+	TArray< TObjectPtr<UControlRigControlsProxy>> SelectedControlRigProxies;
+	
+	UPROPERTY()
+	TMap<TObjectPtr<UObject>, FNameToProxyMap> SequencerOnlyProxies;
 
+	UPROPERTY()
+	TArray< TObjectPtr<UControlRigControlsProxy>> SelectedSequencerProxies;
+	
+	TWeakPtr<ISequencer> Sequencer;
 
 public:
-	void SelectProxy(UControlRig* InControlRig, const FName& Name, bool bSelected);
-	UControlRigControlsProxy* FindProxy(UControlRig* InControlRig, const FName& Name) const;
-	void AddProxy(UControlRig* InControlRig, const FName& Name,  FRigControlElement* ControlElement);
-	void RemoveProxy(UControlRig* InControlRig, const FName& Name );
-	void ProxyChanged(UControlRig* InControlRig, const FName& Name, bool bModify = true);
-	void RemoveAllProxies(UControlRig* InControlRig);
+
+	void SelectProxy(UControlRig* InControlRig, FRigControlElement* RigElement, bool bSelected);
+	UControlRigControlsProxy* AddProxy(UControlRig* InControlRig, FRigControlElement* RigElement);
+	UControlRigControlsProxy* AddProxy(UObject* InObject, ERigControlType Type, TWeakObjectPtr<UMovieSceneTrack>& Track, TSharedPtr<FTrackInstancePropertyBindings>& Binding);
+	void RemoveControlRigProxies(UControlRig* InControlRig);
+	void RemoveSequencerProxies(UObject* InObject);
+	void RemoveAllProxies();
+	void ProxyChanged(UControlRig* InControlRig, FRigControlElement* RigElement, bool bModify = true);
 	void RecreateAllProxies(UControlRig* InControlRig);
-	const TArray<UControlRigControlsProxy*>& GetSelectedProxies() const { return SelectedProxies;}
-	bool IsSelected(UControlRig* InControlRig, const FName& Name) const;
+	const TArray<UControlRigControlsProxy*> GetAllSelectedProxies() const;
+	bool IsSelected(UControlRig* InControlRig, FRigControlElement* RigElement) const;
+	void SetSequencer(TWeakPtr<ISequencer> InSequencer) { Sequencer = InSequencer; }
+	ISequencer* GetSequencer() const { return Sequencer.Pin().Get(); }
+	void ValuesChanged();
+	void ResetSequencerProxies(TMap<ERigControlType, FSequencerProxyPerType>& ProxyPerType);
+
+private:
+	UControlRigControlsProxy* FindProxy(UControlRig* InControlRig, FRigControlElement* RigElement) const;
+	UControlRigControlsProxy* FindProxy(UObject* InObject, FName PropertyName) const;
+
+	UControlRigControlsProxy* NewProxyFromType(ERigControlType Type, TObjectPtr<UEnum>& EnumPtr);
 
+	//delegate for changed objects
+	void OnPostPropertyChanged(UObject* InObject, FPropertyChangedEvent& InPropertyChangedEvent);
 };
