New Anim Details. Main code is the new objects in AnimDetailsProxy that handle support for control rig/sequencer, having multiple objects, and shared selection state with sequencer. CR doesn't support

diff --git a/Engine/Plugins/Animation/ControlRig/Source/ControlRigEditor/Private/EditMode/ControlRigEditMode.cpp b/Engine/Plugins/Animation/ControlRig/Source/ControlRigEditor/Private/EditMode/ControlRigEditMode.cpp
--- a/Engine/Plugins/Animation/ControlRig/Source/ControlRigEditor/Private/EditMode/ControlRigEditMode.cpp
+++ b/Engine/Plugins/Animation/ControlRig/Source/ControlRigEditor/Private/EditMode/ControlRigEditMode.cpp
@@ -36,7 +36,7 @@
 #include "ControlRigGizmoActor.h"
 #include "Framework/MultiBox/MultiBoxBuilder.h"
 #include "SEditorViewport.h"
-#include "EditMode/ControlRigControlsProxy.h"
+#include "EditMode/AnimDetailsProxy.h"
 #include "ScopedTransaction.h"
 #include "RigVMModel/RigVMController.h"
 #include "Rigs/AdditiveControlRig.h"
@@ -47,6 +47,7 @@
 #include "Units/Execution/RigUnit_InteractionExecution.h"
 #include "IPersonaPreviewScene.h"
 #include "PersonaSelectionProxies.h"
+#include "PropertyHandle.h"
 #include "Framework/Application/SlateApplication.h"
 #include "UnrealEdGlobals.h"
 #include "Editor/UnrealEdEngine.h"
@@ -72,6 +73,7 @@
 #include "DragTool_FrustumSelect.h"
 #include "AnimationEditorViewportClient.h"
 #include "EditorInteractiveGizmoManager.h"
+#include "Tools/BakingHelper.h"
 
 #include UE_INLINE_GENERATED_CPP_BY_NAME(ControlRigEditMode)
 
@@ -164,6 +166,7 @@ FControlRigEditMode::FControlRigEditMode()
 {
 	ControlProxy = NewObject<UControlRigDetailPanelControlProxies>(GetTransientPackage(), NAME_None);
 	ControlProxy->SetFlags(RF_Transactional);
+	DetailKeyFrameCache = MakeShared<FDetailKeyFrameCacheAndHandler>();
 
 	UControlRigEditModeSettings* Settings = GetMutableDefault<UControlRigEditModeSettings>();
 	bShowControlsAsOverlay = Settings->bShowControlsAsOverlay;
@@ -213,6 +216,9 @@ bool FControlRigEditMode:: SetSequencer(TWeakPtr<ISequencer> InSequencer)
 	if (InSequencer != WeakSequencer)
 	{
 		WeakSequencer = InSequencer;
+
+		DetailKeyFrameCache->UnsetDelegates();
+
 		DestroyShapesActors(nullptr);
 		TArray<TWeakObjectPtr<UControlRig>> PreviousRuntimeRigs = RuntimeControlRigs;
 		for (int32 PreviousRuntimeRigIndex = 0; PreviousRuntimeRigIndex < PreviousRuntimeRigs.Num(); PreviousRuntimeRigIndex++)
@@ -238,6 +244,8 @@ bool FControlRigEditMode:: SetSequencer(TWeakPtr<ISequencer> InSequencer)
 				}
 			}
 			LastMovieSceneSig = Sequencer->GetFocusedMovieSceneSequence()->GetMovieScene()->GetSignature();
+			DetailKeyFrameCache->SetDelegates(WeakSequencer, this);
+			ControlProxy->SetSequencer(WeakSequencer);
 		}
 		SetObjects_Internal();
 	}
@@ -383,6 +391,10 @@ void FControlRigEditMode::Enter()
 	LastMovieSceneSig = FGuid();
 	if (UsesToolkits())
 	{
+		if (WeakSequencer.IsValid() == false)
+		{
+			SetSequencer(FBakingHelper::GetSequencer());
+		}
 		if (!Toolkit.IsValid())
 		{
 			Toolkit = MakeShareable(new FControlRigEditModeToolkit(*this));
@@ -491,7 +503,7 @@ void FControlRigEditMode::Exit()
 	ModeManager->OnCoordSystemChanged().RemoveAll(this);
 
 	//clear proxies
-	ControlProxy->RemoveAllProxies(nullptr);
+	ControlProxy->RemoveAllProxies();
 
 	//make sure the widget is reset
 	ResetControlShapeSize();
@@ -585,10 +597,9 @@ void FControlRigEditMode::Tick(FEditorViewportClient* ViewportClient, float Delt
 							{
 								if (!ControlRig->IsCurveControl(ControlElement))
 								{
-									ControlProxy->AddProxy(ControlRig, SelectedKey.Name, ControlElement);
+									ControlProxy->AddProxy(ControlRig, ControlElement);
 								}
 							}
-							
 						}
 					}
 				}
@@ -4077,7 +4088,7 @@ void FControlRigEditMode::OnHierarchyModified(ERigHierarchyNotification InNotif,
 				if (FRigControlElement* ControlElement = InHierarchy->Find<FRigControlElement>(InElement->GetKey()))
 				{
 					UControlRig* ControlRig = InHierarchy->GetTypedOuter<UControlRig>();
-					if(ControlProxy->IsSelected(ControlRig, ControlElement->GetFName()))
+					if(ControlProxy->IsSelected(ControlRig, ControlElement))
 					{
 						// reselect the control - to affect the details panel / sequencer
 						if(URigHierarchyController* Controller = InHierarchy->GetController())
@@ -4149,9 +4160,9 @@ void FControlRigEditMode::OnHierarchyModified(ERigHierarchyNotification InNotif,
 						}
 						if (!AreEditingControlRigDirectly())
 						{
-							if (const FRigControlElement* ControlElement = ControlRig->GetHierarchy()->Find<FRigControlElement>(Key))
+							if (FRigControlElement* ControlElement = ControlRig->GetHierarchy()->Find<FRigControlElement>(Key))
 							{
-								ControlProxy->SelectProxy(ControlRig,Key.Name, bSelected);
+								ControlProxy->SelectProxy(ControlRig, ControlElement, bSelected);
 
 								if(ControlElement->CanDriveControls())
 								{
@@ -4160,9 +4171,9 @@ void FControlRigEditMode::OnHierarchyModified(ERigHierarchyNotification InNotif,
 									const TArray<FRigElementKey>& DrivenKeys = ControlElement->Settings.DrivenControls;
 									for(const FRigElementKey& DrivenKey : DrivenKeys)
 									{
-										if (const FRigControlElement* DrivenControl = ControlRig->GetHierarchy()->Find<FRigControlElement>(DrivenKey))
+										if (FRigControlElement* DrivenControl = ControlRig->GetHierarchy()->Find<FRigControlElement>(DrivenKey))
 										{
-											ControlProxy->SelectProxy(ControlRig, DrivenControl->GetFName(), bSelected);
+											ControlProxy->SelectProxy(ControlRig, DrivenControl, bSelected);
 
 											if (AControlRigShapeActor* DrivenShapeActor = GetControlShapeFromControlName(ControlRig,DrivenControl->GetFName()))
 											{
@@ -4189,7 +4200,7 @@ void FControlRigEditMode::OnHierarchyModified(ERigHierarchyNotification InNotif,
 													if(DrivenKeys.Contains(ParentControlElement->GetKey()) ||
 														ParentControlElement->GetKey() == ControlElement->GetKey())
 													{
-														ControlProxy->SelectProxy(ControlRig, AnimationChannelControl->GetFName(), bSelected);
+														ControlProxy->SelectProxy(ControlRig, AnimationChannelControl, bSelected);
 													}
 												}
 											}
@@ -4281,7 +4292,7 @@ void FControlRigEditMode::OnControlModified(UControlRig* Subject, FRigControlEle
 {
 	//this makes sure the details panel ui get's updated, don't remove
 	const bool bModify = Context.SetKey != EControlRigSetKey::Never;
-	ControlProxy->ProxyChanged(Subject,InControlElement->GetFName(), bModify);
+	ControlProxy->ProxyChanged(Subject, InControlElement, bModify);
 
 	/*
 	FScopedTransaction ScopedTransaction(LOCTEXT("ModifyControlTransaction", "Modify Control"),!GIsTransacting && Context.SetKey != EControlRigSetKey::Never);
@@ -5310,5 +5321,257 @@ bool FControlRigEditMode::GetOnlySelectRigControls()const
 	return Settings->bOnlySelectRigControls;
 }
 
+/**
+* FDetailKeyFrameCacheAndHandler
+*/
+
+bool FDetailKeyFrameCacheAndHandler::IsPropertyKeyable(const UClass* InObjectClass, const IPropertyHandle& InPropertyHandle) const
+{
+	TArray<UObject*> OuterObjects;
+	InPropertyHandle.GetOuterObjects(OuterObjects);
+	if (OuterObjects.Num() == 1)
+	{
+		if (UControlRigControlsProxy* Proxy = Cast< UControlRigControlsProxy>(OuterObjects[0]))
+		{
+			for (const TPair<UControlRig*, FControlRigProxyItem>& Items : Proxy->ControlRigItems)
+			{
+				if (UControlRig* ControlRig = Items.Value.ControlRig.Get())
+				{
+					for (FRigControlElement* ControlElement : Items.Value.ControlElements)
+					{
+						if (ControlElement)
+						{
+							if (!ControlRig->GetHierarchy()->IsAnimatable(ControlElement))
+							{
+								return false;
+							}
+						}
+					}
+				}
+			}
+		}
+	}
+
+	if (InObjectClass
+		&& InObjectClass->IsChildOf(UAnimDetailControlsProxyTransform::StaticClass())
+		&& InObjectClass->IsChildOf(UAnimDetailControlsProxyLocation::StaticClass())
+		&& InObjectClass->IsChildOf(UAnimDetailControlsProxyRotation::StaticClass())
+		&& InObjectClass->IsChildOf(UAnimDetailControlsProxyScale::StaticClass())
+		&& InObjectClass->IsChildOf(UAnimDetailControlsProxyVector2D::StaticClass())
+		&& InObjectClass->IsChildOf(UAnimDetailControlsProxyFloat::StaticClass())
+		&& InObjectClass->IsChildOf(UAnimDetailControlsProxyBool::StaticClass())
+		)
+	{
+		return true;
+	}
+
+	if ((InObjectClass && InObjectClass->IsChildOf(UAnimDetailControlsProxyTransform::StaticClass()) && InPropertyHandle.GetProperty())
+		&& (InPropertyHandle.GetProperty()->GetFName() == GET_MEMBER_NAME_CHECKED(UAnimDetailControlsProxyTransform, Location) ||
+		InPropertyHandle.GetProperty()->GetFName() == GET_MEMBER_NAME_CHECKED(UAnimDetailControlsProxyTransform, Rotation) ||
+		InPropertyHandle.GetProperty()->GetFName() == GET_MEMBER_NAME_CHECKED(UAnimDetailControlsProxyTransform, Scale)))
+	{
+		return true;
+	}
+
+
+	if (InObjectClass && InObjectClass->IsChildOf(UAnimDetailControlsProxyLocation::StaticClass()) && InPropertyHandle.GetProperty()
+		&& InPropertyHandle.GetProperty()->GetFName() == GET_MEMBER_NAME_CHECKED(UAnimDetailControlsProxyLocation, Location))
+	{
+		return true;
+	}
+
+	if (InObjectClass && InObjectClass->IsChildOf(UAnimDetailControlsProxyRotation::StaticClass()) && InPropertyHandle.GetProperty()
+		&& InPropertyHandle.GetProperty()->GetFName() == GET_MEMBER_NAME_CHECKED(UAnimDetailControlsProxyRotation, Rotation))
+	{
+		return true;
+	}
+
+	if (InObjectClass && InObjectClass->IsChildOf(UAnimDetailControlsProxyScale::StaticClass()) && InPropertyHandle.GetProperty()
+		&& InPropertyHandle.GetProperty()->GetFName() == GET_MEMBER_NAME_CHECKED(UAnimDetailControlsProxyScale, Scale))
+	{
+		return true;
+	}
+
+	if (InObjectClass && InObjectClass->IsChildOf(UAnimDetailControlsProxyVector2D::StaticClass()) && InPropertyHandle.GetProperty()
+		&& InPropertyHandle.GetProperty()->GetFName() == GET_MEMBER_NAME_CHECKED(UAnimDetailControlsProxyVector2D, Vector2D))
+	{
+		return true;
+	}
+
+	if (InObjectClass && InObjectClass->IsChildOf(UAnimDetailControlsProxyBool::StaticClass()) && InPropertyHandle.GetProperty()
+		&& InPropertyHandle.GetProperty()->GetFName() == GET_MEMBER_NAME_CHECKED(UAnimDetailControlsProxyBool, Bool))
+	{
+		return true;
+	}
+	if (InObjectClass && InObjectClass->IsChildOf(UAnimDetailControlsProxyFloat::StaticClass()) && InPropertyHandle.GetProperty()
+		&& InPropertyHandle.GetProperty()->GetFName() == GET_MEMBER_NAME_CHECKED(UAnimDetailControlsProxyFloat, Float))
+	{
+		return true;
+	}
+	FCanKeyPropertyParams CanKeyPropertyParams(InObjectClass, InPropertyHandle);
+	if (WeakSequencer.IsValid() && WeakSequencer.Pin()->CanKeyProperty(CanKeyPropertyParams))
+	{
+		return true;
+	}
+
+	return false;
+}
+
+bool FDetailKeyFrameCacheAndHandler::IsPropertyKeyingEnabled() const
+{
+	if (WeakSequencer.IsValid() &&  WeakSequencer.Pin()->GetFocusedMovieSceneSequence())
+	{
+		return true;
+	}
+
+	return false;
+}
+
+bool FDetailKeyFrameCacheAndHandler::IsPropertyAnimated(const IPropertyHandle& PropertyHandle, UObject* ParentObject) const
+{
+	if (WeakSequencer.IsValid() && WeakSequencer.Pin()->GetFocusedMovieSceneSequence())
+	{
+		TSharedPtr<ISequencer> Sequencer = WeakSequencer.Pin();
+		constexpr bool bCreateHandleIfMissing = false;
+		FGuid ObjectHandle = Sequencer->GetHandleToObject(ParentObject, bCreateHandleIfMissing);
+		if (ObjectHandle.IsValid())
+		{
+			UMovieScene* MovieScene = Sequencer->GetFocusedMovieSceneSequence()->GetMovieScene();
+			FProperty* Property = PropertyHandle.GetProperty();
+			TSharedRef<FPropertyPath> PropertyPath = FPropertyPath::CreateEmpty();
+			PropertyPath->AddProperty(FPropertyInfo(Property));
+			FName PropertyName(*PropertyPath->ToString(TEXT(".")));
+			TSubclassOf<UMovieSceneTrack> TrackClass; //use empty @todo find way to get the UMovieSceneTrack from the Property type.
+			return MovieScene->FindTrack(TrackClass, ObjectHandle, PropertyName) != nullptr;
+		}
+	}
+	return false;
+}
+
+void FDetailKeyFrameCacheAndHandler::OnKeyPropertyClicked(const IPropertyHandle& KeyedPropertyHandle)
+{
+	if (WeakSequencer.IsValid() && !WeakSequencer.Pin()->IsAllowedToChange())
+	{
+		return;
+	}
+	TSharedPtr<ISequencer> SequencerPtr = WeakSequencer.Pin();
+
+	TArray<UObject*> Objects;
+	KeyedPropertyHandle.GetOuterObjects(Objects);
+	for (UObject* Object : Objects)
+	{
+		UControlRigControlsProxy* Proxy = Cast< UControlRigControlsProxy>(Object);
+		if (Proxy)
+		{
+			Proxy->SetKey(SequencerPtr, KeyedPropertyHandle);
+		}
+	}
+}
+
+EPropertyKeyedStatus FDetailKeyFrameCacheAndHandler::GetPropertyKeyedStatus(const IPropertyHandle& PropertyHandle) const
+{
+	if (WeakSequencer.IsValid() == false)
+	{
+		return EPropertyKeyedStatus::NotKeyed;
+	}		
+
+	if (const EPropertyKeyedStatus* ExistingKeyedStatus = CachedPropertyKeyedStatusMap.Find(&PropertyHandle))
+	{
+		return *ExistingKeyedStatus;
+	}
+
+	TSharedPtr<ISequencer> SequencerPtr = WeakSequencer.Pin();
+	UMovieSceneSequence* Sequence = SequencerPtr->GetFocusedMovieSceneSequence();
+	EPropertyKeyedStatus KeyedStatus = EPropertyKeyedStatus::NotKeyed;
+
+	UMovieScene* MovieScene = Sequence ? Sequence->GetMovieScene() : nullptr;
+	if (!MovieScene)
+	{
+		return KeyedStatus;
+	}
+
+	TArray<UObject*> OuterObjects;
+	PropertyHandle.GetOuterObjects(OuterObjects);
+	if (OuterObjects.IsEmpty())
+	{
+		return EPropertyKeyedStatus::NotKeyed;
+	}
+	
+	for (UObject* Object : OuterObjects)
+	{
+		if (UControlRigControlsProxy* Proxy = Cast< UControlRigControlsProxy>(Object))
+		{
+			KeyedStatus = Proxy->GetPropertyKeyedStatus(SequencerPtr,PropertyHandle);
+		}
+		//else check to see if it's in sequencer
+	}
+	CachedPropertyKeyedStatusMap.Add(&PropertyHandle, KeyedStatus);
+
+	return KeyedStatus;
+}
+
+void FDetailKeyFrameCacheAndHandler::SetDelegates(TWeakPtr<ISequencer>& InWeakSequencer, FControlRigEditMode* InEditMode)
+{
+	WeakSequencer = InWeakSequencer;
+	EditMode = InEditMode;
+	if (TSharedPtr<ISequencer> Sequencer = WeakSequencer.Pin())
+	{
+		Sequencer->OnMovieSceneDataChanged().AddRaw(this, &FDetailKeyFrameCacheAndHandler::OnMovieSceneDataChanged);
+		Sequencer->OnGlobalTimeChanged().AddRaw(this, &FDetailKeyFrameCacheAndHandler::OnGlobalTimeChanged);
+		Sequencer->OnEndScrubbingEvent().AddRaw(this, &FDetailKeyFrameCacheAndHandler::ResetCachedData);
+		Sequencer->OnChannelChanged().AddRaw(this, &FDetailKeyFrameCacheAndHandler::OnChannelChanged);
+		Sequencer->OnStopEvent().AddRaw(this, &FDetailKeyFrameCacheAndHandler::ResetCachedData);
+	}
+}
+
+void FDetailKeyFrameCacheAndHandler::UnsetDelegates()
+{
+	if (TSharedPtr<ISequencer> Sequencer = WeakSequencer.Pin())
+	{
+		Sequencer->OnMovieSceneDataChanged().RemoveAll(this);
+		Sequencer->OnGlobalTimeChanged().RemoveAll(this);
+		Sequencer->OnEndScrubbingEvent().RemoveAll(this);
+		Sequencer->OnChannelChanged().RemoveAll(this);
+		Sequencer->OnStopEvent().RemoveAll(this);
+	}
+}
+
+void FDetailKeyFrameCacheAndHandler::OnGlobalTimeChanged()
+{
+	// Only reset cached data when not playing
+	TSharedPtr<ISequencer> Sequencer = WeakSequencer.Pin();
+	if (Sequencer.IsValid() && Sequencer->GetPlaybackStatus() != EMovieScenePlayerStatus::Playing)
+	{
+		ResetCachedData();
+	}
+}
+
+void FDetailKeyFrameCacheAndHandler::OnMovieSceneDataChanged(EMovieSceneDataChangeType DataChangeType)
+{
+	if (DataChangeType == EMovieSceneDataChangeType::MovieSceneStructureItemAdded
+		|| DataChangeType == EMovieSceneDataChangeType::MovieSceneStructureItemRemoved
+		|| DataChangeType == EMovieSceneDataChangeType::MovieSceneStructureItemsChanged
+		|| DataChangeType == EMovieSceneDataChangeType::ActiveMovieSceneChanged
+		|| DataChangeType == EMovieSceneDataChangeType::RefreshAllImmediately)
+	{
+		ResetCachedData();
+	}
+}
+
+void FDetailKeyFrameCacheAndHandler::OnChannelChanged(const FMovieSceneChannelMetaData*, UMovieSceneSection*)
+{
+	ResetCachedData();
+}
+
+void FDetailKeyFrameCacheAndHandler::ResetCachedData()
+{
+	CachedPropertyKeyedStatusMap.Reset();
+	if (EditMode && EditMode->GetControlProxy())
+	{
+		EditMode->GetControlProxy()->ValuesChanged();
+	}
+}
+
+
 
 #undef LOCTEXT_NAMESPACE
diff --git a/Engine/Plugins/Animation/ControlRig/Source/ControlRigEditor/Private/EditMode/ControlsProxyDetailCustomization.cpp b/Engine/Plugins/Animation/ControlRig/Source/ControlRigEditor/Private/EditMode/ControlsProxyDetailCustomization.cpp
new file mode 100644
--- /dev/null
+++ b/Engine/Plugins/Animation/ControlRig/Source/ControlRigEditor/Private/EditMode/ControlsProxyDetailCustomization.cpp
@@ -0,0 +1,1070 @@
+// Copyright Epic Games, Inc. All Rights Reserved.
+
+#include "ControlsProxyDetailCustomization.h"
+#include "DetailLayoutBuilder.h"
+#include "PropertyHandle.h"
+#include "DetailCategoryBuilder.h"
+#include "DetailLayoutBuilder.h"
+#include "DetailWidgetRow.h"
+#include "IDetailGroup.h"
+#include "IDetailPropertyRow.h"
+#include "EditMode/ControlRigControlsProxy.h"
+#include "Framework/MultiBox/MultiBoxBuilder.h"
+#include "Widgets/Input/SComboButton.h"
+#include "Widgets/Input/SSegmentedControl.h"
+#include "Widgets/Layout/SBox.h"
+#include "Algo/Transform.h"
+#include "SAdvancedTransformInputBox.h"
+#include "PropertyCustomizationHelpers.h"
+#include "PropertyEditorModule.h"
+#include "HAL/PlatformApplicationMisc.h"
+
+#include "Widgets/SBoxPanel.h"
+#include "Widgets/Layout/SBorder.h"
+#include "Widgets/Text/SInlineEditableTextBlock.h"
+#include "Widgets/Text/STextBlock.h"
+#include "Widgets/Input/SButton.h"
+
+#include "MVVM/CurveEditorExtension.h"
+#include "MVVM/ViewModelPtr.h"
+#include "MVVM/Selection/Selection.h"
+#include "MVVM/ViewModels/SequencerEditorViewModel.h"
+#include "MVVM/ViewModels/OutlinerViewModel.h"
+#include "MVVM/ViewModels/ViewModelIterators.h"
+#include "MVVM/Extensions/IOutlinerExtension.h"
+#include "Rigs/RigControlHierarchy.h"
+#include "Tree/SCurveEditorTree.h"
+
+#include "CurveEditor.h"
+#include "CurveModel.h"
+#include "Sequencer/MovieSceneControlRigParameterSection.h"
+#include "Framework/Application/SlateApplication.h"
+
+#define LOCTEXT_NAMESPACE "FControlsProxyDetailCustomization"
+
+void FPropertySelectionCache::OnCurveModelDisplayChanged(FCurveModel* InCurveModel, bool bDisplayed, const FCurveEditor* InCurveEditor)
+{
+	bCacheValid = false;
+}
+
+EAnimDetailSelectionState FPropertySelectionCache::IsPropertySelected(TSharedPtr<FCurveEditor>& InCurveEditor, UControlRigControlsProxy* Proxy, const FName& PropertyName)
+{
+	if (CurveEditor.IsValid() == false || InCurveEditor.Get() != CurveEditor.Pin().Get())
+	{
+		ClearDelegates();
+		CurveEditor = InCurveEditor;
+		if (InCurveEditor->OnCurveArrayChanged.IsBoundToObject(this) == false)
+		{
+			InCurveEditor->OnCurveArrayChanged.AddRaw(this, &FPropertySelectionCache::OnCurveModelDisplayChanged);
+		}
+		bCacheValid = false;
+	}
+	if (bCacheValid == false)
+	{
+		CachePropertySelection(Proxy);
+	}
+	EAnimDetailSelectionState SelectionState = EAnimDetailSelectionState::None;
+
+	if ( PropertyName == GET_MEMBER_NAME_CHECKED(FAnimDetailProxyLocation, LX)
+		|| PropertyName == GET_MEMBER_NAME_CHECKED(FAnimDetailProxyVector2D, X)
+		|| PropertyName == GET_MEMBER_NAME_CHECKED(FAnimDetailProxyFloat, Float)
+		|| PropertyName == GET_MEMBER_NAME_CHECKED(FAnimDetailProxyBool, Bool)
+		|| PropertyName == GET_MEMBER_NAME_CHECKED(FAnimDetailProxyInteger, Integer)
+		|| PropertyName == GET_MEMBER_NAME_CHECKED(FControlRigEnumControlProxyValue, EnumIndex)
+
+		)
+	{
+		SelectionState = LocationSelectionCache.XSelected;
+	}
+	else if (PropertyName == GET_MEMBER_NAME_CHECKED(FAnimDetailProxyLocation, LY)
+		|| PropertyName == GET_MEMBER_NAME_CHECKED(FAnimDetailProxyVector2D, Y))
+	{
+		SelectionState = LocationSelectionCache.YSelected;
+	}
+	else if (PropertyName == GET_MEMBER_NAME_CHECKED(FAnimDetailProxyLocation, LZ))
+	{
+		SelectionState = LocationSelectionCache.ZSelected;
+	}
+	else if (PropertyName == GET_MEMBER_NAME_CHECKED(UAnimDetailControlsProxyVector2D, Vector2D))
+	{
+		if (LocationSelectionCache.XSelected == EAnimDetailSelectionState::All && LocationSelectionCache.YSelected == EAnimDetailSelectionState::All)
+		{
+			SelectionState = EAnimDetailSelectionState::All;
+		}
+		else if (LocationSelectionCache.XSelected == EAnimDetailSelectionState::None && LocationSelectionCache.YSelected == EAnimDetailSelectionState::None)
+		{
+			SelectionState = EAnimDetailSelectionState::None;
+		}
+	}
+	else if (PropertyName == GET_MEMBER_NAME_CHECKED(UAnimDetailControlsProxyTransform, Location))
+	{
+		if (LocationSelectionCache.XSelected == EAnimDetailSelectionState::All && LocationSelectionCache.YSelected == EAnimDetailSelectionState::All
+			&& LocationSelectionCache.ZSelected == EAnimDetailSelectionState::All)
+		{
+			SelectionState = EAnimDetailSelectionState::All;
+		}
+		else if (LocationSelectionCache.XSelected == EAnimDetailSelectionState::None && LocationSelectionCache.YSelected == EAnimDetailSelectionState::None
+			&& LocationSelectionCache.ZSelected == EAnimDetailSelectionState::None)
+		{
+			SelectionState = EAnimDetailSelectionState::None;
+		}
+		else
+		{
+			SelectionState = EAnimDetailSelectionState::Partial;
+		}
+	}
+	else if (PropertyName == GET_MEMBER_NAME_CHECKED(UAnimDetailControlsProxyLocation, Location))
+	{
+		if (LocationSelectionCache.XSelected == EAnimDetailSelectionState::All && LocationSelectionCache.YSelected == EAnimDetailSelectionState::All
+			&& LocationSelectionCache.ZSelected == EAnimDetailSelectionState::All)
+		{
+			SelectionState = EAnimDetailSelectionState::All;
+		}
+		else if (LocationSelectionCache.XSelected == EAnimDetailSelectionState::None && LocationSelectionCache.YSelected == EAnimDetailSelectionState::None
+			&& LocationSelectionCache.ZSelected == EAnimDetailSelectionState::None)
+		{
+			SelectionState = EAnimDetailSelectionState::None;
+		}
+		else
+		{
+			SelectionState = EAnimDetailSelectionState::Partial;
+		}
+	}
+	else if (PropertyName == GET_MEMBER_NAME_CHECKED(FAnimDetailProxyRotation, RX))
+	{
+		SelectionState = RotationSelectionCache.XSelected;
+	}
+	else if (PropertyName == GET_MEMBER_NAME_CHECKED(FAnimDetailProxyRotation, RY))
+	{
+		SelectionState = RotationSelectionCache.YSelected;
+	}
+	else if (PropertyName == GET_MEMBER_NAME_CHECKED(FAnimDetailProxyRotation, RZ))
+	{
+		SelectionState = RotationSelectionCache.ZSelected;
+	}
+	else if (PropertyName == GET_MEMBER_NAME_CHECKED(UAnimDetailControlsProxyTransform, Rotation))
+	{
+		if (RotationSelectionCache.XSelected == EAnimDetailSelectionState::All && RotationSelectionCache.YSelected == EAnimDetailSelectionState::All
+			&& RotationSelectionCache.ZSelected == EAnimDetailSelectionState::All)
+		{
+			SelectionState = EAnimDetailSelectionState::All;
+		}
+		else if (RotationSelectionCache.XSelected == EAnimDetailSelectionState::None && RotationSelectionCache.YSelected == EAnimDetailSelectionState::None
+			&& RotationSelectionCache.ZSelected == EAnimDetailSelectionState::None)
+		{
+			SelectionState = EAnimDetailSelectionState::None;
+		}
+		else
+		{
+			SelectionState = EAnimDetailSelectionState::Partial;
+		}
+	}
+	else if (PropertyName == GET_MEMBER_NAME_CHECKED(FAnimDetailProxyScale, SX))
+	{
+		SelectionState = ScaleSelectionCache.XSelected;
+	}
+	else if (PropertyName == GET_MEMBER_NAME_CHECKED(FAnimDetailProxyScale, SY))
+	{
+		SelectionState = ScaleSelectionCache.YSelected;
+	}
+	else if (PropertyName == GET_MEMBER_NAME_CHECKED(FAnimDetailProxyScale, SZ))
+	{
+		SelectionState = ScaleSelectionCache.ZSelected;
+	}
+	else if (PropertyName == GET_MEMBER_NAME_CHECKED(UAnimDetailControlsProxyTransform, Scale))
+	{
+		if (ScaleSelectionCache.XSelected == EAnimDetailSelectionState::All && ScaleSelectionCache.YSelected == EAnimDetailSelectionState::All
+			&& ScaleSelectionCache.ZSelected == EAnimDetailSelectionState::All)
+		{
+			SelectionState = EAnimDetailSelectionState::All;
+		}
+		else if (ScaleSelectionCache.XSelected == EAnimDetailSelectionState::None && ScaleSelectionCache.YSelected == EAnimDetailSelectionState::None
+			&& ScaleSelectionCache.ZSelected == EAnimDetailSelectionState::None)
+		{
+			SelectionState = EAnimDetailSelectionState::None;
+		}
+		else
+		{
+			SelectionState = EAnimDetailSelectionState::Partial;
+		}
+	}
+
+	return SelectionState;
+}
+
+void FPropertySelectionCache::ClearDelegates()
+{
+	TSharedPtr<FCurveEditor> SharedCurveEditor = CurveEditor.Pin();
+
+	if (SharedCurveEditor)
+	{
+		SharedCurveEditor->OnCurveArrayChanged.RemoveAll(this);
+	}
+}
+
+void FPropertySelectionCache::CachePropertySelection(UControlRigControlsProxy* Proxy)
+{
+	if (CurveEditor.IsValid() == false)
+	{
+		return;
+	}
+	bCacheValid = true;
+	Proxy->GetChannelSelectionState(CurveEditor, LocationSelectionCache, RotationSelectionCache, ScaleSelectionCache);
+}
+
+UControlRigControlsProxy* FAnimDetailValueCustomization::GetProxy(TSharedRef< IPropertyHandle>& PropertyHandle) const
+{
+	TArray<UObject*> OuterObjects;
+	PropertyHandle->GetOuterObjects(OuterObjects);
+	for (UObject* OuterObject : OuterObjects)
+	{
+		UControlRigControlsProxy* Proxy = Cast<UControlRigControlsProxy>(OuterObject);
+		if (Proxy)
+		{
+			return Proxy;
+		}
+	}
+	return nullptr;
+}
+
+UControlRigDetailPanelControlProxies* FAnimDetailValueCustomization::GetProxyOwner(UControlRigControlsProxy* Proxy) const
+{
+	if (Proxy)
+	{
+		UControlRigDetailPanelControlProxies* ProxyOwner = Proxy->GetTypedOuter<UControlRigDetailPanelControlProxies>();
+		if (ProxyOwner == nullptr || ProxyOwner->GetSequencer() == nullptr)
+		{
+			return nullptr;
+		}
+		return ProxyOwner;
+	}
+	return nullptr;
+}
+
+EVisibility FAnimDetailValueCustomization::IsVisible() const 
+{
+	if (DetailPropertyRow == nullptr) //need to wait for it to be constructed
+	{
+		DetailPropertyRow = DetailBuilder->EditPropertyFromRoot(StructPropertyHandlePtr);
+	}
+	return (DetailPropertyRow && DetailPropertyRow->IsExpanded()) ? EVisibility::Collapsed : EVisibility::Visible;
+}
+
+void FAnimDetailValueCustomization::CustomizeChildren(TSharedRef<class IPropertyHandle> StructPropertyHandle, class IDetailChildrenBuilder& StructBuilder, IPropertyTypeCustomizationUtils& StructCustomizationUtils)
+{
+	FOnBooleanValueChanged OnExpansionChanged = FOnBooleanValueChanged::CreateLambda([this](bool bNewValue)
+		{
+			int a = 3;
+		});
+	IDetailCategoryBuilder& CatBuilder = StructBuilder.GetParentCategory();
+	CatBuilder.OnExpansionChanged(OnExpansionChanged);
+	DetailBuilder = &CatBuilder.GetParentLayout();
+	StructPropertyHandlePtr = StructPropertyHandle;
+	FProperty* Property = StructPropertyHandle->GetProperty(); // Is a FProperty*
+	UControlRigControlsProxy* Proxy = GetProxy(StructPropertyHandle);
+
+	for (int32 ChildIndex = 0; ChildIndex < SortedChildHandles.Num(); ++ChildIndex)
+	{
+		TSharedRef<IPropertyHandle> ChildHandle = SortedChildHandles[ChildIndex];
+
+		const FName PropertyName = ChildHandle->GetProperty()->GetFName();
+		const FText PropertyDisplayText = ChildHandle->GetPropertyDisplayName();
+
+		TSharedRef<SWidget> ValueWidget = MakeChildWidget(StructPropertyHandle, ChildHandle);
+		TSharedRef<SWidget> LabelWidget = ChildHandle->CreatePropertyNameWidget();
+
+		// Add the individual properties as children as well so the vector can be expanded for more room
+		
+		StructBuilder.AddProperty(ChildHandle).CustomWidget()
+		.NameContent()
+		.HAlign(HAlign_Fill)
+		[
+			 SNew(SBorder)
+			.BorderImage(FAppStyle::Get().GetBrush("DetailsView.CategoryMiddle"))
+			.BorderBackgroundColor_Lambda([Proxy, PropertyName, this]()
+			{
+				EAnimDetailSelectionState SelectionState = IsPropertySelected(Proxy, PropertyName);
+				if (SelectionState == EAnimDetailSelectionState::All)
+				{
+					return  FStyleColors::Select;
+				}
+				else if (SelectionState == EAnimDetailSelectionState::Partial)
+				{
+					return  FStyleColors::SelectInactive;
+				}
+				return FStyleColors::Transparent;
+			})
+			.OnMouseButtonDown_Lambda([Proxy, PropertyName, this](const FGeometry&, const FPointerEvent& PointerEvent)
+			{
+				TogglePropertySelection(Proxy, PropertyName);
+				return FReply::Handled();
+			})
+			.Content()
+			[
+				SNew(SHorizontalBox)
+				+SHorizontalBox::Slot()
+				.Padding(0,0,0,0)
+				[
+					SNew(STextBlock)
+					//.Font(InCustomizationUtils.GetRegularFont())
+					.Text_Lambda([PropertyDisplayText]()
+					{
+						return PropertyDisplayText;
+					})
+					
+				]
+			]
+		]
+		.ValueContent()
+		.HAlign(HAlign_Fill)
+		[
+			 SNew(SBorder)
+			.BorderImage(FAppStyle::Get().GetBrush("DetailsView.CategoryMiddle"))
+			.BorderBackgroundColor_Lambda([Proxy, PropertyName, this]()
+			{
+				EAnimDetailSelectionState SelectionState = IsPropertySelected(Proxy, PropertyName);
+				if (SelectionState == EAnimDetailSelectionState::All)
+				{
+					return  FStyleColors::Select;
+				}
+				else if (SelectionState == EAnimDetailSelectionState::Partial)
+				{
+					return  FStyleColors::SelectInactive;
+				}
+				return FStyleColors::Transparent;
+			})
+			.OnMouseButtonDown_Lambda([Proxy, PropertyName, this](const FGeometry&, const FPointerEvent& PointerEvent)
+			{
+				TogglePropertySelection(Proxy, PropertyName);
+				return FReply::Handled();
+			})
+			.Content()
+			[
+				ValueWidget
+			]
+		];
+		/*
+		.ValueContent()
+		[
+			ValueWidget
+		];
+		*/
+		
+	}
+}
+
+
+TSharedRef<IPropertyTypeCustomization> FAnimDetailValueCustomization::MakeInstance()
+{
+	return MakeShareable(new FAnimDetailValueCustomization);
+}
+
+FAnimDetailValueCustomization::FAnimDetailValueCustomization() = default;
+
+FAnimDetailValueCustomization::~FAnimDetailValueCustomization() = default;
+
+void FAnimDetailValueCustomization::MakeHeaderRow(TSharedRef<class IPropertyHandle>& StructPropertyHandle, FDetailWidgetRow& Row)
+{
+	if (SortedChildHandles.Num() == 1)
+	{
+		return;
+	}
+
+	FProperty* Property = StructPropertyHandle->GetProperty(); // Is a FProperty*
+	TWeakPtr<IPropertyHandle> StructWeakHandlePtr = StructPropertyHandle;
+	TSharedPtr<SHorizontalBox> HorizontalBox;
+
+	Row.NameContent()
+		[
+			StructPropertyHandle->CreatePropertyNameWidget()
+		]
+	.ValueContent()
+		// Make enough space for each child handle
+		.MinDesiredWidth(125.f * SortedChildHandles.Num())
+		.MaxDesiredWidth(125.f * SortedChildHandles.Num())
+		[
+				SAssignNew(HorizontalBox, SHorizontalBox)
+				.Visibility(this, &FAnimDetailValueCustomization::IsVisible)
+				.IsEnabled(this, &FMathStructCustomization::IsValueEnabled, StructWeakHandlePtr)
+		];
+
+	for (int32 ChildIndex = 0; ChildIndex < SortedChildHandles.Num(); ++ChildIndex)
+	{
+		TSharedRef<IPropertyHandle> ChildHandle = SortedChildHandles[ChildIndex];
+
+		// Propagate metadata to child properties so that it's reflected in the nested, individual spin boxes
+		ChildHandle->SetInstanceMetaData(TEXT("UIMin"), StructPropertyHandle->GetMetaData(TEXT("UIMin")));
+		ChildHandle->SetInstanceMetaData(TEXT("UIMax"), StructPropertyHandle->GetMetaData(TEXT("UIMax")));
+		ChildHandle->SetInstanceMetaData(TEXT("SliderExponent"), StructPropertyHandle->GetMetaData(TEXT("SliderExponent")));
+		ChildHandle->SetInstanceMetaData(TEXT("Delta"), StructPropertyHandle->GetMetaData(TEXT("Delta")));
+		ChildHandle->SetInstanceMetaData(TEXT("LinearDeltaSensitivity"), StructPropertyHandle->GetMetaData(TEXT("LinearDeltaSensitivity")));
+		ChildHandle->SetInstanceMetaData(TEXT("ShiftMouseMovePixelPerDelta"), StructPropertyHandle->GetMetaData(TEXT("ShiftMouseMovePixelPerDelta")));
+		ChildHandle->SetInstanceMetaData(TEXT("SupportDynamicSliderMaxValue"), StructPropertyHandle->GetMetaData(TEXT("SupportDynamicSliderMaxValue")));
+		ChildHandle->SetInstanceMetaData(TEXT("SupportDynamicSliderMinValue"), StructPropertyHandle->GetMetaData(TEXT("SupportDynamicSliderMinValue")));
+		ChildHandle->SetInstanceMetaData(TEXT("ClampMin"), StructPropertyHandle->GetMetaData(TEXT("ClampMin")));
+		ChildHandle->SetInstanceMetaData(TEXT("ClampMax"), StructPropertyHandle->GetMetaData(TEXT("ClampMax")));
+
+		const bool bLastChild = SortedChildHandles.Num() - 1 == ChildIndex;
+
+		TSharedRef<SWidget> ChildWidget = MakeChildWidget(StructPropertyHandle, ChildHandle);
+		if (ChildHandle->GetPropertyClass() == FBoolProperty::StaticClass())
+		{
+			HorizontalBox->AddSlot()
+				.Padding(FMargin(0.f, 2.f, bLastChild ? 0.f : 3.f, 2.f))
+				.AutoWidth()  // keep the check box slots small
+				[
+					ChildWidget
+				];
+		}
+		else
+		{
+			if (ChildHandle->GetPropertyClass() == FDoubleProperty::StaticClass())
+			{
+				NumericEntryBoxWidgetList.Add(ChildWidget);
+			}
+
+			HorizontalBox->AddSlot()
+				.Padding(FMargin(0.f, 2.f, bLastChild ? 0.f : 3.f, 2.f))
+				[
+					ChildWidget
+				];
+		}
+	}
+}
+
+FLinearColor FAnimDetailValueCustomization::GetColorFromProperty(const FName& PropertyName) const
+{
+	FLinearColor Color = FLinearColor::White;
+	if (PropertyName == GET_MEMBER_NAME_CHECKED(FAnimDetailProxyVector3, X) ||
+		PropertyName == GET_MEMBER_NAME_CHECKED(FAnimDetailProxyVector2D, X) ||
+		PropertyName == GET_MEMBER_NAME_CHECKED(FAnimDetailProxyLocation, LX) ||
+		PropertyName == GET_MEMBER_NAME_CHECKED(FAnimDetailProxyRotation, RX) ||
+		PropertyName == GET_MEMBER_NAME_CHECKED(FAnimDetailProxyScale, SX))
+	{
+		Color = SNumericEntryBox<double>::RedLabelBackgroundColor;
+	}
+	else if (PropertyName == GET_MEMBER_NAME_CHECKED(FAnimDetailProxyVector3, Y) ||
+		PropertyName == GET_MEMBER_NAME_CHECKED(FAnimDetailProxyVector2D, Y) ||
+		PropertyName == GET_MEMBER_NAME_CHECKED(FAnimDetailProxyLocation, LY) ||
+		PropertyName == GET_MEMBER_NAME_CHECKED(FAnimDetailProxyRotation, RY) ||
+		PropertyName == GET_MEMBER_NAME_CHECKED(FAnimDetailProxyScale, SY))
+	{
+		Color = SNumericEntryBox<double>::GreenLabelBackgroundColor;
+	}
+	else if (PropertyName == GET_MEMBER_NAME_CHECKED(FAnimDetailProxyVector3, Z) ||
+		PropertyName == GET_MEMBER_NAME_CHECKED(FAnimDetailProxyLocation, LZ) ||
+		PropertyName == GET_MEMBER_NAME_CHECKED(FAnimDetailProxyRotation, RZ) ||
+		PropertyName == GET_MEMBER_NAME_CHECKED(FAnimDetailProxyScale, SZ))
+	{
+		Color = SNumericEntryBox<double>::BlueLabelBackgroundColor;
+	}
+	return Color;
+}
+
+EAnimDetailSelectionState FAnimDetailValueCustomization::IsPropertySelected(UControlRigControlsProxy* Proxy,const FName& PropertyName)
+{
+	using namespace UE::Sequencer;
+
+	UControlRigDetailPanelControlProxies* ProxyOwner = GetProxyOwner(Proxy);
+	if (ProxyOwner == nullptr)
+	{
+		return EAnimDetailSelectionState::None;
+	}
+	const TSharedPtr<FSequencerEditorViewModel> SequencerViewModel = ProxyOwner->GetSequencer()->GetViewModel();
+	const FCurveEditorExtension* CurveEditorExtension = SequencerViewModel->CastDynamic<FCurveEditorExtension>();
+	check(CurveEditorExtension);
+	TSharedPtr<FCurveEditor> CurveEditor = CurveEditorExtension->GetCurveEditor();
+	return SelectionCache.IsPropertySelected(CurveEditor,Proxy, PropertyName);
+	
+}
+
+//may want this to be under each proxy? need to think about this.
+static bool GetChannelNameForCurve(const TArray<FString>& CurveString, const FRigControlElement* ControlElement, FString& OutChannelName)
+{
+	//if single channel expect one item and the name will match
+	if (ControlElement->Settings.ControlType == ERigControlType::ScaleFloat ||
+		ControlElement->Settings.ControlType == ERigControlType::Float ||
+		ControlElement->Settings.ControlType == ERigControlType::Bool ||
+		ControlElement->Settings.ControlType == ERigControlType::Integer)
+	{
+		if (CurveString[0] == ControlElement->GetKey().Name)
+		{
+			if (ControlElement->Settings.ControlType == ERigControlType::ScaleFloat ||
+				ControlElement->Settings.ControlType == ERigControlType::Float)
+			{
+				OutChannelName = FString("Float");
+				return true;
+			}
+			if (ControlElement->Settings.ControlType == ERigControlType::Bool)
+			{
+				OutChannelName = FString("Bool");
+				return true;
+			}
+			if (ControlElement->Settings.ControlType == ERigControlType::Integer)
+			{
+				OutChannelName = FString("Integer");
+				return true;
+			}
+		}
+	}
+	else if (CurveString.Num() > 1)
+	{
+		if (CurveString[0] == ControlElement->GetKey().Name)
+		{
+			if (CurveString.Num() == 3)
+			{
+				OutChannelName = CurveString[1] + "." + CurveString[2];
+				return true;
+			}
+			else if (CurveString.Num() == 2)
+			{
+				OutChannelName = CurveString[1];
+				return true;
+			}
+		}
+	}
+	return false;
+}
+
+void FAnimDetailValueCustomization::TogglePropertySelection(UControlRigControlsProxy* Proxy, const FName& PropertyName) const
+{
+	using namespace UE::Sequencer;
+	UControlRigDetailPanelControlProxies* ProxyOwner = GetProxyOwner(Proxy);
+	if (ProxyOwner == nullptr)
+	{
+		return;
+	}
+
+	const TSharedPtr<FSequencerEditorViewModel> SequencerViewModel = ProxyOwner->GetSequencer()->GetViewModel();
+	const FCurveEditorExtension* CurveEditorExtension = SequencerViewModel->CastDynamic<FCurveEditorExtension>();
+	check(CurveEditorExtension);
+	TSharedPtr<FCurveEditor> CurveEditor = CurveEditorExtension->GetCurveEditor();
+	TSharedPtr<SCurveEditorTree>  CurveEditorTreeView = CurveEditorExtension->GetCurveEditorTreeView();
+	TSharedPtr<FOutlinerViewModel> OutlinerViewModel = SequencerViewModel->GetOutliner();
+
+	const bool bIsShiftDown = FSlateApplication::Get().GetModifierKeys().IsShiftDown();
+	const bool bIsCtrlDown = FSlateApplication::Get().GetModifierKeys().IsControlDown();
+	if (bIsShiftDown == false && bIsCtrlDown == false)
+	{
+		CurveEditorTreeView->ClearSelection();
+	}
+
+	if (Proxy)
+	{
+		for (const TPair<UControlRig*, FControlRigProxyItem>& Items : Proxy->ControlRigItems)
+		{
+			if (UControlRig* ControlRig = Items.Value.ControlRig.Get())
+			{
+				for (FRigControlElement* ControlElement : Items.Value.ControlElements)
+				{
+					if (ControlElement)
+					{
+						EControlRigContextChannelToKey ChannelToKey = Proxy->GetChannelToKeyFromPropertyName(PropertyName);
+						TParentFirstChildIterator<IOutlinerExtension> OutlinerExtenstionIt = OutlinerViewModel->GetRootItem()->GetDescendantsOfType<IOutlinerExtension>();
+						for (; OutlinerExtenstionIt; ++OutlinerExtenstionIt)
+						{
+							FName ID = OutlinerExtenstionIt->GetIdentifier();
+							FString Name = ID.ToString();
+							TArray<FString> StringArray;
+							Name.ParseIntoArray(StringArray, TEXT("."));
+
+							FString ChannelName;
+							if(GetChannelNameForCurve(StringArray, ControlElement, ChannelName))
+							{
+								EControlRigContextChannelToKey ChannelToKeyFromCurve = Proxy->GetChannelToKeyFromChannelName(ChannelName);
+								if (ChannelToKey == ChannelToKeyFromCurve)
+								{
+									if (TViewModelPtr<ICurveEditorTreeItemExtension> CurveEditorItem = OutlinerExtenstionIt.GetCurrentItem().ImplicitCast())
+									{
+										FCurveEditorTreeItemID CurveEditorTreeItem = CurveEditorItem->GetCurveEditorItemID();
+										if (CurveEditorTreeItem != FCurveEditorTreeItemID::Invalid())
+										{
+											const bool bSelected = bIsShiftDown ? true : !CurveEditorTreeView->IsItemSelected(CurveEditorTreeItem);
+											if (bIsCtrlDown || bSelected)
+											{
+												CurveEditorTreeView->SetItemSelection(CurveEditorTreeItem, bSelected);
+											}
+										}
+									}
+								}
+							}
+						}
+					}
+				}
+			}
+		}
+		for (const TPair<UObject*, FSequencerProxyItem>& SItems : Proxy->SequencerItems)
+		{
+			if (UObject* Object = SItems.Key)
+			{
+				for (const FBindingAndTrack& Element : SItems.Value.Bindings)
+				{
+					EControlRigContextChannelToKey ChannelToKey = Proxy->GetChannelToKeyFromPropertyName(PropertyName);
+					TParentFirstChildIterator<IOutlinerExtension> OutlinerExtenstionIt = OutlinerViewModel->GetRootItem()->GetDescendantsOfType<IOutlinerExtension>();
+					for (; OutlinerExtenstionIt; ++OutlinerExtenstionIt)
+					{
+						FName ID = OutlinerExtenstionIt->GetIdentifier();
+						FString Name = ID.ToString();
+						TArray<FString> StringArray;
+						Name.ParseIntoArray(StringArray, TEXT("."));
+
+						FString ChannelName;
+						if (StringArray.Num() == 2)
+						{
+							ChannelName = StringArray[0] + "." + StringArray[1];
+						}
+						else if (StringArray.Num() == 0)
+						{
+							ChannelName = StringArray[0];
+						}
+
+						EControlRigContextChannelToKey ChannelToKeyFromCurve = Proxy->GetChannelToKeyFromChannelName(ChannelName);
+						if (ChannelToKey == ChannelToKeyFromCurve)
+						{
+							if (TViewModelPtr<ICurveEditorTreeItemExtension> CurveEditorItem = OutlinerExtenstionIt.GetCurrentItem().ImplicitCast())
+							{
+								FCurveEditorTreeItemID CurveEditorTreeItem = CurveEditorItem->GetCurveEditorItemID();
+								if (CurveEditorTreeItem != FCurveEditorTreeItemID::Invalid())
+								{
+									const bool bSelected = bIsShiftDown ? true : !CurveEditorTreeView->IsItemSelected(CurveEditorTreeItem);
+									if (bIsCtrlDown || bSelected)
+									{
+										CurveEditorTreeView->SetItemSelection(CurveEditorTreeItem, bSelected);
+									}
+								}
+							}
+						}
+					}
+				}
+			}
+		}
+	}
+}
+
+bool FAnimDetailValueCustomization::IsMultiple(UControlRigControlsProxy* Proxy,const FName& InPropertyName) const
+{
+	if (Proxy)
+	{
+		if (Proxy->IsMultiple(InPropertyName))
+		{
+			return true;
+		}
+	}
+	return false;
+}
+
+TSharedRef<SWidget> FAnimDetailValueCustomization::MakeChildWidget(
+	TSharedRef<IPropertyHandle>& StructurePropertyHandle,
+	TSharedRef<IPropertyHandle>& PropertyHandle)
+{
+	const FFieldClass* PropertyClass = PropertyHandle->GetPropertyClass();
+	FName PropertyName = PropertyHandle->GetProperty()->GetFName();
+	FLinearColor LinearColor = GetColorFromProperty(PropertyName);
+	
+	if (PropertyClass == FDoubleProperty::StaticClass())
+	{
+		return MakeDoubleWidget(StructurePropertyHandle, PropertyHandle, LinearColor);
+	}
+	else if (PropertyClass == FBoolProperty::StaticClass())
+	{
+		return MakeBoolWidget(StructurePropertyHandle, PropertyHandle, LinearColor);
+	}
+	else if (PropertyClass == FInt64Property::StaticClass())
+	{
+		return MakeIntegerWidget(StructurePropertyHandle, PropertyHandle, LinearColor);
+
+	}
+
+	checkf(false, TEXT("Unsupported property class for the Anim Detail Values customization."));
+	return SNullWidget::NullWidget;
+}
+
+
+TSharedRef<SWidget> FAnimDetailValueCustomization::MakeIntegerWidget(
+	TSharedRef<IPropertyHandle>& StructurePropertyHandle,
+	TSharedRef<IPropertyHandle>& PropertyHandle,
+	const FLinearColor& LabelColor)
+{
+
+	UControlRigControlsProxy* Proxy = GetProxy(StructurePropertyHandle);
+
+	TWeakPtr<IPropertyHandle> WeakHandlePtr = PropertyHandle;
+
+	return SNew(SNumericEntryBox<int64>)
+		.IsEnabled(this, &FMathStructCustomization::IsValueEnabled, WeakHandlePtr)
+		.EditableTextBoxStyle(&FCoreStyle::Get().GetWidgetStyle<FEditableTextBoxStyle>("NormalEditableTextBox"))
+		.Value_Lambda([WeakHandlePtr, this, Proxy]()
+			{
+				int64 Value = 0;
+				bool bIsMultiple = IsMultiple(Proxy, WeakHandlePtr.Pin()->GetProperty()->GetFName());
+				return (bIsMultiple == false && WeakHandlePtr.Pin()->GetValue(Value) == FPropertyAccess::Success) ?
+					TOptional<int64>(Value) :
+					TOptional<int64>();  // Value couldn't be accessed, return an unset value
+			})
+		.Font(IDetailLayoutBuilder::GetDetailFont())
+		.UndeterminedString(NSLOCTEXT("PropertyEditor", "MultipleValues", "Multiple Values"))
+		.OnValueCommitted_Lambda([WeakHandlePtr](double Value, ETextCommit::Type)
+			{
+				WeakHandlePtr.Pin()->SetValue(Value, EPropertyValueSetFlags::DefaultFlags);
+			})
+		.OnValueChanged_Lambda([this, WeakHandlePtr](double Value)
+			{
+				if (bIsUsingSlider)
+				{
+					WeakHandlePtr.Pin()->SetValue(Value, EPropertyValueSetFlags::InteractiveChange);
+				}
+			})
+		.OnBeginSliderMovement_Lambda([this]()
+			{
+				bIsUsingSlider = true;
+				GEditor->BeginTransaction(LOCTEXT("SetVectorProperty", "Set Property"));
+			})
+		.OnEndSliderMovement_Lambda([this](double Value)
+			{
+				bIsUsingSlider = false;
+				GEditor->EndTransaction();
+			})
+			// Only allow spin on handles with one object.  Otherwise it is not clear what value to spin
+			.AllowSpin(PropertyHandle->GetNumOuterObjects() < 2)
+			.LabelPadding(FMargin(3))
+			.LabelLocation(SNumericEntryBox<int64>::ELabelLocation::Inside)
+			.Label()
+			[
+				SNumericEntryBox<int64>::BuildNarrowColorLabel(LabelColor)
+			];
+}
+
+TSharedRef<SWidget> FAnimDetailValueCustomization::MakeBoolWidget(
+	TSharedRef<IPropertyHandle>& StructurePropertyHandle,
+	TSharedRef<IPropertyHandle>& PropertyHandle,
+	const FLinearColor& LabelColor)
+{
+	UControlRigControlsProxy* Proxy = GetProxy(StructurePropertyHandle);
+	TWeakPtr<IPropertyHandle> WeakHandlePtr = PropertyHandle;
+	return
+		SNew(SCheckBox)
+		.Type(ESlateCheckBoxType::CheckBox)
+		.IsChecked_Lambda([WeakHandlePtr, this, Proxy]()->ECheckBoxState
+			{
+				bool bIsMultiple = IsMultiple(Proxy, WeakHandlePtr.Pin()->GetProperty()->GetFName());
+				if (bIsMultiple)
+				{
+					return ECheckBoxState::Undetermined;
+				}
+
+				bool bValue = false;
+				WeakHandlePtr.Pin()->GetValue(bValue);
+				return bValue ? ECheckBoxState::Checked : ECheckBoxState::Unchecked;
+
+			})
+		.OnCheckStateChanged_Lambda([WeakHandlePtr](ECheckBoxState CheckBoxState)
+			{
+				WeakHandlePtr.Pin()->SetValue(CheckBoxState == ECheckBoxState::Checked, EPropertyValueSetFlags::DefaultFlags);
+			});
+}
+
+TSharedRef<SWidget> FAnimDetailValueCustomization::MakeDoubleWidget(
+	TSharedRef<IPropertyHandle>& StructurePropertyHandle,
+	TSharedRef<IPropertyHandle>& PropertyHandle,
+	const FLinearColor &LabelColor)
+{
+	TOptional<double> MinValue, MaxValue, SliderMinValue, SliderMaxValue;
+	double SliderExponent, Delta;
+	int32 ShiftMouseMovePixelPerDelta = 1;
+	bool SupportDynamicSliderMaxValue = false;
+	bool SupportDynamicSliderMinValue = false;
+
+	UControlRigControlsProxy* Proxy = GetProxy(StructurePropertyHandle);
+
+	ExtractDoubleMetadata(PropertyHandle, MinValue, MaxValue, SliderMinValue, SliderMaxValue, SliderExponent, Delta, ShiftMouseMovePixelPerDelta, SupportDynamicSliderMaxValue, SupportDynamicSliderMinValue);
+
+	TWeakPtr<IPropertyHandle> WeakHandlePtr = PropertyHandle;
+
+	return SNew(SNumericEntryBox<double>)
+		.IsEnabled(this, &FMathStructCustomization::IsValueEnabled, WeakHandlePtr)
+		.EditableTextBoxStyle(&FCoreStyle::Get().GetWidgetStyle<FEditableTextBoxStyle>("NormalEditableTextBox"))
+		.Value_Lambda([WeakHandlePtr,this,Proxy ]()
+			{
+				bool bIsMultiple = IsMultiple(Proxy, WeakHandlePtr.Pin()->GetProperty()->GetFName());
+				if (bIsMultiple && bIsUsingSlider)
+				{
+					return TOptional<double>(MultipleValue);
+				}
+				double Value = 0.;
+				return (bIsMultiple == false && WeakHandlePtr.Pin()->GetValue(Value) == FPropertyAccess::Success) ?
+					TOptional<double>(Value) :
+					TOptional<double>();  // Value couldn't be accessed, return an unset value
+			})
+		.Font(IDetailLayoutBuilder::GetDetailFont())
+				.UndeterminedString(NSLOCTEXT("PropertyEditor", "MultipleValues", "Multiple Values"))
+				.OnValueCommitted_Lambda([WeakHandlePtr](double Value, ETextCommit::Type)
+					{
+						WeakHandlePtr.Pin()->SetValue(Value, EPropertyValueSetFlags::DefaultFlags);
+					})
+				.OnValueChanged_Lambda([this, WeakHandlePtr](double Value)
+					{
+						if (bIsUsingSlider)
+						{
+							MultipleValue = Value;
+							WeakHandlePtr.Pin()->SetValue(Value, EPropertyValueSetFlags::InteractiveChange);
+						}
+					})
+		.OnBeginSliderMovement_Lambda([WeakHandlePtr, this, Proxy]()
+			{
+				bIsUsingSlider = true;
+				GEditor->BeginTransaction(LOCTEXT("SetVectorProperty", "Set Property"));
+				bool bIsMultiple = IsMultiple(Proxy, WeakHandlePtr.Pin()->GetProperty()->GetFName());
+				if (bIsMultiple)
+				{
+					MultipleValue = 0.0;
+					WeakHandlePtr.Pin()->SetValue(MultipleValue, EPropertyValueSetFlags::DefaultFlags);
+				}
+
+			})
+		.OnEndSliderMovement_Lambda([this](double Value)
+			{
+				bIsUsingSlider = false;
+				GEditor->EndTransaction();
+			})
+				// Only allow spin on handles with one object.  Otherwise it is not clear what value to spin
+				.AllowSpin(PropertyHandle->GetNumOuterObjects() < 2)
+				.ShiftMouseMovePixelPerDelta(ShiftMouseMovePixelPerDelta)
+				.SupportDynamicSliderMaxValue(SupportDynamicSliderMaxValue)
+				.SupportDynamicSliderMinValue(SupportDynamicSliderMinValue)
+				.OnDynamicSliderMaxValueChanged(this, &FAnimDetailValueCustomization::OnDynamicSliderMaxValueChanged)
+				.OnDynamicSliderMinValueChanged(this, &FAnimDetailValueCustomization::OnDynamicSliderMinValueChanged)
+				.MinValue(MinValue)
+				.MaxValue(MaxValue)
+				.MinSliderValue(SliderMinValue)
+				.MaxSliderValue(SliderMaxValue)
+				.SliderExponent(SliderExponent)
+				.LinearDeltaSensitivity(1.0)
+				.Delta(Delta)
+				.LabelPadding(FMargin(3))
+				.LabelLocation(SNumericEntryBox<double>::ELabelLocation::Inside)
+				.Label()
+				[
+					SNumericEntryBox<double>::BuildNarrowColorLabel(LabelColor)
+				];
+}
+
+// The following code is just a plain copy of FMathStructCustomization which
+// would need changes to be able to serve as a base class for this customization.
+void FAnimDetailValueCustomization::ExtractDoubleMetadata(TSharedRef<IPropertyHandle>& PropertyHandle, TOptional<double>& MinValue, TOptional<double>& MaxValue, TOptional<double>& SliderMinValue, TOptional<double>& SliderMaxValue, double& SliderExponent, double& Delta, int32& ShiftMouseMovePixelPerDelta, bool& SupportDynamicSliderMaxValue, bool& SupportDynamicSliderMinValue)
+{
+	FProperty* Property = PropertyHandle->GetProperty();
+
+	const FString& MetaUIMinString = Property->GetMetaData(TEXT("UIMin"));
+	const FString& MetaUIMaxString = Property->GetMetaData(TEXT("UIMax"));
+	const FString& SliderExponentString = Property->GetMetaData(TEXT("SliderExponent"));
+	const FString& DeltaString = Property->GetMetaData(TEXT("Delta"));
+	const FString& ShiftMouseMovePixelPerDeltaString = Property->GetMetaData(TEXT("ShiftMouseMovePixelPerDelta"));
+	const FString& SupportDynamicSliderMaxValueString = Property->GetMetaData(TEXT("SupportDynamicSliderMaxValue"));
+	const FString& SupportDynamicSliderMinValueString = Property->GetMetaData(TEXT("SupportDynamicSliderMinValue"));
+	const FString& ClampMinString = Property->GetMetaData(TEXT("ClampMin"));
+	const FString& ClampMaxString = Property->GetMetaData(TEXT("ClampMax"));
+
+	// If no UIMin/Max was specified then use the clamp string
+	const FString& UIMinString = MetaUIMinString.Len() ? MetaUIMinString : ClampMinString;
+	const FString& UIMaxString = MetaUIMaxString.Len() ? MetaUIMaxString : ClampMaxString;
+
+	double ClampMin = TNumericLimits<double>::Lowest();
+	double ClampMax = TNumericLimits<double>::Max();
+
+	if (!ClampMinString.IsEmpty())
+	{
+		TTypeFromString<double>::FromString(ClampMin, *ClampMinString);
+	}
+
+	if (!ClampMaxString.IsEmpty())
+	{
+		TTypeFromString<double>::FromString(ClampMax, *ClampMaxString);
+	}
+
+	double UIMin = TNumericLimits<double>::Lowest();
+	double UIMax = TNumericLimits<double>::Max();
+	TTypeFromString<double>::FromString(UIMin, *UIMinString);
+	TTypeFromString<double>::FromString(UIMax, *UIMaxString);
+
+	SliderExponent = double(1);
+
+	if (SliderExponentString.Len())
+	{
+		TTypeFromString<double>::FromString(SliderExponent, *SliderExponentString);
+	}
+
+	Delta = double(0);
+
+	if (DeltaString.Len())
+	{
+		TTypeFromString<double>::FromString(Delta, *DeltaString);
+	}
+
+	ShiftMouseMovePixelPerDelta = 1;
+	if (ShiftMouseMovePixelPerDeltaString.Len())
+	{
+		TTypeFromString<int32>::FromString(ShiftMouseMovePixelPerDelta, *ShiftMouseMovePixelPerDeltaString);
+		//The value should be greater or equal to 1
+		// 1 is neutral since it is a multiplier of the mouse drag pixel
+		if (ShiftMouseMovePixelPerDelta < 1)
+		{
+			ShiftMouseMovePixelPerDelta = 1;
+		}
+	}
+
+	const double ActualUIMin = FMath::Max(UIMin, ClampMin);
+	const double ActualUIMax = FMath::Min(UIMax, ClampMax);
+
+	MinValue = ClampMinString.Len() ? ClampMin : TOptional<double>();
+	MaxValue = ClampMaxString.Len() ? ClampMax : TOptional<double>();
+	SliderMinValue = (UIMinString.Len()) ? ActualUIMin : TOptional<double>();
+	SliderMaxValue = (UIMaxString.Len()) ? ActualUIMax : TOptional<double>();
+
+	SupportDynamicSliderMaxValue = SupportDynamicSliderMaxValueString.Len() > 0 && SupportDynamicSliderMaxValueString.ToBool();
+	SupportDynamicSliderMinValue = SupportDynamicSliderMinValueString.Len() > 0 && SupportDynamicSliderMinValueString.ToBool();
+}
+
+void FAnimDetailValueCustomization::OnDynamicSliderMaxValueChanged(double NewMaxSliderValue, TWeakPtr<SWidget> InValueChangedSourceWidget, bool IsOriginator, bool UpdateOnlyIfHigher)
+{
+	for (TWeakPtr<SWidget>& Widget : NumericEntryBoxWidgetList)
+	{
+		TSharedPtr<SNumericEntryBox<double>> NumericBox = StaticCastSharedPtr<SNumericEntryBox<double>>(Widget.Pin());
+
+		if (NumericBox.IsValid())
+		{
+			TSharedPtr<SSpinBox<double>> SpinBox = StaticCastSharedPtr<SSpinBox<double>>(NumericBox->GetSpinBox());
+
+			if (SpinBox.IsValid())
+			{
+				if (SpinBox != InValueChangedSourceWidget)
+				{
+					if ((NewMaxSliderValue > SpinBox->GetMaxSliderValue() && UpdateOnlyIfHigher) || !UpdateOnlyIfHigher)
+					{
+						// Make sure the max slider value is not a getter otherwise we will break the link!
+						verifySlow(!SpinBox->IsMaxSliderValueBound());
+						SpinBox->SetMaxSliderValue(NewMaxSliderValue);
+					}
+				}
+			}
+		}
+	}
+
+	if (IsOriginator)
+	{
+		OnNumericEntryBoxDynamicSliderMaxValueChanged.Broadcast((double)NewMaxSliderValue, InValueChangedSourceWidget, false, UpdateOnlyIfHigher);
+	}
+}
+
+void FAnimDetailValueCustomization::OnDynamicSliderMinValueChanged(double NewMinSliderValue, TWeakPtr<SWidget> InValueChangedSourceWidget, bool IsOriginator, bool UpdateOnlyIfLower)
+{
+	for (TWeakPtr<SWidget>& Widget : NumericEntryBoxWidgetList)
+	{
+		TSharedPtr<SNumericEntryBox<double>> NumericBox = StaticCastSharedPtr<SNumericEntryBox<double>>(Widget.Pin());
+
+		if (NumericBox.IsValid())
+		{
+			TSharedPtr<SSpinBox<double>> SpinBox = StaticCastSharedPtr<SSpinBox<double>>(NumericBox->GetSpinBox());
+
+			if (SpinBox.IsValid())
+			{
+				if (SpinBox != InValueChangedSourceWidget)
+				{
+					if ((NewMinSliderValue < SpinBox->GetMinSliderValue() && UpdateOnlyIfLower) || !UpdateOnlyIfLower)
+					{
+						// Make sure the min slider value is not a getter otherwise we will break the link!
+						verifySlow(!SpinBox->IsMinSliderValueBound());
+						SpinBox->SetMinSliderValue(NewMinSliderValue);
+					}
+				}
+			}
+		}
+	}
+
+	if (IsOriginator)
+	{
+		OnNumericEntryBoxDynamicSliderMinValueChanged.Broadcast((double)NewMinSliderValue, InValueChangedSourceWidget, false, UpdateOnlyIfLower);
+	}
+}
+
+FAnimDetailProxyDetails::FAnimDetailProxyDetails(const FName& InCategoryName)
+{
+	CategoryName = InCategoryName;
+}
+
+TSharedRef<IDetailCustomization> FAnimDetailProxyDetails::MakeInstance(const FName& InCategoryName)
+{
+	return MakeShareable(new FAnimDetailProxyDetails(InCategoryName));
+}
+
+void FAnimDetailProxyDetails::CustomizeDetails(IDetailLayoutBuilder& DetailBuilder)
+{
+	TArray<UControlRigControlsProxy*> ControlProxies;
+
+	TArray<TWeakObjectPtr<UObject> > EditedObjects;
+	DetailBuilder.GetObjectsBeingCustomized(EditedObjects);
+	for (int32 i = 0; i < EditedObjects.Num(); i++)
+	{
+		UControlRigControlsProxy* Proxy = Cast<UControlRigControlsProxy>(EditedObjects[i].Get());
+		if (Proxy)
+		{
+			ControlProxies.Add(Proxy);
+		}
+	}
+	FText Name;
+	if (ControlProxies.Num() == 1)
+	{
+		Name = FText::FromName(ControlProxies[0]->GetName());
+		ControlProxies[0]->UpdatePropertyNames(DetailBuilder);
+	}
+	else
+	{
+		FString DisplayString = TEXT("Multiple");
+		Name = FText::FromString(*DisplayString);
+		
+	}
+	IDetailCategoryBuilder& Category = DetailBuilder.EditCategory(CategoryName, Name, ECategoryPriority::Important);
+		
+	//add custom attributes
+	TArray<UControlRigControlsProxy*> NestedProxies;
+	for (UControlRigControlsProxy* Proxy : ControlProxies)
+	{
+		if (Proxy)
+		{
+			for (UControlRigControlsProxy* ChildProxy : Proxy->ChildProxies)
+			{
+				if (ChildProxy)
+				{
+					NestedProxies.Add(ChildProxy);
+				}
+			}
+		}
+	}
+	IDetailCategoryBuilder& AttributesCategory = DetailBuilder.EditCategory("Attributes");
+
+	if (NestedProxies.Num() > 0)
+	{
+		AttributesCategory.SetCategoryVisibility(true);
+		for (UControlRigControlsProxy* Proxy: NestedProxies)
+		{
+			if (Proxy)
+			{
+				TArray<UObject*> ExternalObjects;
+				ExternalObjects.Add(Proxy);
+
+				FAddPropertyParams Params;
+				Params.CreateCategoryNodes(false);
+				Params.HideRootObjectNode(true);
+				IDetailPropertyRow* NestedRow = AttributesCategory.AddExternalObjects(
+					ExternalObjects,
+					EPropertyLocation::Default,
+					Params);
+			}
+		}
+	}
+	else
+	{
+		AttributesCategory.SetCategoryVisibility(true);
+	}
+}
+
+#undef LOCTEXT_NAMESPACE
diff --git a/Engine/Plugins/Animation/ControlRig/Source/ControlRigEditor/Private/EditMode/ControlsProxyDetailCustomization.h b/Engine/Plugins/Animation/ControlRig/Source/ControlRigEditor/Private/EditMode/ControlsProxyDetailCustomization.h
new file mode 100644
--- /dev/null
+++ b/Engine/Plugins/Animation/ControlRig/Source/ControlRigEditor/Private/EditMode/ControlsProxyDetailCustomization.h
@@ -0,0 +1,104 @@
+// Copyright Epic Games, Inc. All Rights Reserved.
+
+#pragma once
+
+#include "CoreMinimal.h"
+#include "IDetailCustomization.h"
+#include "AnimDetailsProxy.h"
+#include "Customizations/MathStructCustomizations.h"
+
+class IPropertyHandle;
+class FCurveModel;
+class FCurveEditor;
+class IDetailPropertyRow;
+
+struct FPropertySelectionCache
+{
+	~FPropertySelectionCache() { ClearDelegates(); }
+
+	EAnimDetailSelectionState IsPropertySelected(TSharedPtr<FCurveEditor>& InCurveEditor, UControlRigControlsProxy* Proxy, const FName& PropertyName);
+	void OnCurveModelDisplayChanged(FCurveModel* InCurveModel, bool bDisplayed, const FCurveEditor* InCurveEditor);
+	void ClearDelegates();
+	void CachePropertySelection(UControlRigControlsProxy* Proxy);
+	EAnimDetailSelectionState CachePropertySelection(UControlRigControlsProxy* Proxy, const FName& PropertyName);
+
+	FAnimDetailVectorSelection LocationSelectionCache;
+	FAnimDetailVectorSelection RotationSelectionCache;
+	FAnimDetailVectorSelection ScaleSelectionCache;
+
+	bool bCacheValid = false;
+	TWeakPtr<FCurveEditor> CurveEditor;
+};
+
+class FAnimDetailValueCustomization
+		: public FMathStructCustomization
+	{
+	public:
+		static TSharedRef<IPropertyTypeCustomization> MakeInstance();
+
+		FAnimDetailValueCustomization();
+		virtual ~FAnimDetailValueCustomization() override;
+
+	protected:
+		virtual void MakeHeaderRow(TSharedRef<IPropertyHandle>& StructPropertyHandle, FDetailWidgetRow& Row) override;
+		virtual void CustomizeChildren(TSharedRef<IPropertyHandle> StructPropertyHandle, IDetailChildrenBuilder& StructBuilder, IPropertyTypeCustomizationUtils& StructCustomizationUtils) override;
+		virtual TSharedRef<SWidget> MakeChildWidget(TSharedRef<IPropertyHandle>& StructurePropertyHandle, TSharedRef<IPropertyHandle>& PropertyHandle) override;
+
+	private:
+
+		UControlRigControlsProxy* GetProxy(TSharedRef< IPropertyHandle>& PropertyHandle) const;
+		FLinearColor GetColorFromProperty(const FName& PropertyName) const;
+		EAnimDetailSelectionState IsPropertySelected(UControlRigControlsProxy* Proxy, const FName& PropertyName);
+		bool IsMultiple(UControlRigControlsProxy* Proxy, const FName& InPropertyName) const;
+		
+		void TogglePropertySelection(UControlRigControlsProxy* Proxy, const FName& PropertyName) const;
+		static void ExtractDoubleMetadata(
+			TSharedRef<IPropertyHandle>& PropertyHandle,
+			TOptional<double>& MinValue,
+			TOptional<double>& MaxValue,
+			TOptional<double>& SliderMinValue,
+			TOptional<double>& SliderMaxValue,
+			double& SliderExponent,
+			double& Delta,
+			int32& ShiftMouseMovePixelPerDelta,
+			bool& SupportDynamicSliderMaxValue,
+			bool& SupportDynamicSliderMinValue);
+		void OnDynamicSliderMaxValueChanged(double NewMaxSliderValue, TWeakPtr<SWidget> InValueChangedSourceWidget, bool IsOriginator, bool UpdateOnlyIfHigher);
+		void OnDynamicSliderMinValueChanged(double NewMinSliderValue, TWeakPtr<SWidget> InValueChangedSourceWidget, bool IsOriginator, bool UpdateOnlyIfLower);
+
+		//widgets
+		TSharedRef<SWidget> MakeDoubleWidget(TSharedRef<IPropertyHandle>& StructurePropertyHandle, TSharedRef<IPropertyHandle>& PropertyHandle, const FLinearColor& LinearColor);
+		TSharedRef<SWidget> MakeBoolWidget(TSharedRef<IPropertyHandle>& StructurePropertyHandle, TSharedRef<IPropertyHandle>& PropertyHandle, const FLinearColor& LinearColor);
+		TSharedRef<SWidget> MakeIntegerWidget(TSharedRef<IPropertyHandle>& StructurePropertyHandle, TSharedRef<IPropertyHandle>& PropertyHandle, const FLinearColor& LinearColor);
+
+		UControlRigDetailPanelControlProxies* GetProxyOwner(UControlRigControlsProxy* Proxy) const;
+
+		EVisibility IsVisible() const;
+	private:
+		//current selection cache
+		FPropertySelectionCache SelectionCache;
+		//the scrubbed value when doing multiples, is basicually used a as a delta to add onto the multiple values
+		double MultipleValue = 0.0;
+		// use the property row to get expansion state to hide/show widgets, and use the builder/structhandler to get it
+		mutable IDetailPropertyRow* DetailPropertyRow = nullptr;
+		IDetailLayoutBuilder* DetailBuilder = nullptr;
+		TSharedPtr<IPropertyHandle> StructPropertyHandlePtr;
+
+};
+
+class FAnimDetailProxyDetails : public IDetailCustomization
+{
+public:
+	FAnimDetailProxyDetails(const FName& InCategoryName);
+	/** Makes a new instance of this detail layout class for a specific detail view requesting it, with the specified category
+	name*/
+	static TSharedRef<IDetailCustomization> MakeInstance(const FName& InCategoryName);
+
+	// IDetailCustomization interface
+	virtual void CustomizeDetails(IDetailLayoutBuilder& DetailLayout) override;
+
+	FName CategoryName;
+};
+
+
+
diff --git a/Engine/Plugins/Animation/ControlRig/Source/ControlRigEditor/Private/EditMode/SControlRigDetails.cpp b/Engine/Plugins/Animation/ControlRig/Source/ControlRigEditor/Private/EditMode/SControlRigDetails.cpp
--- a/Engine/Plugins/Animation/ControlRig/Source/ControlRigEditor/Private/EditMode/SControlRigDetails.cpp
+++ b/Engine/Plugins/Animation/ControlRig/Source/ControlRigEditor/Private/EditMode/SControlRigDetails.cpp
@@ -22,204 +22,26 @@
 #include "Widgets/Text/STextBlock.h"
 #include "Widgets/Input/SButton.h"
 #include "Widgets/Layout/SScrollBox.h"
-#include "EditMode/ControlRigControlsProxy.h"
+#include "EditMode/AnimDetailsProxy.h"
 #include "EditMode/ControlRigEditModeSettings.h"
 #include "Modules/ModuleManager.h"
 #include "TimerManager.h"
+#include "CurveEditor.h"
+#include "MVVM/CurveEditorExtension.h"
+#include "MVVM/ViewModels/SequencerEditorViewModel.h"
+#include "Tracks/MovieScenePropertyTrack.h"
+#include "Tracks/MovieScene3DTransformTrack.h"
+#include "MovieSceneCommonHelpers.h"
 
 #define LOCTEXT_NAMESPACE "ControlRigDetails"
 
 void FControlRigEditModeGenericDetails::CustomizeDetails(class IDetailLayoutBuilder& DetailLayout)
 {
-	TArray<TWeakObjectPtr<UObject>> ObjectsBeingCustomized;
-	DetailLayout.GetObjectsBeingCustomized(ObjectsBeingCustomized);
-
-	TArray<UControlRigControlsProxy*> ProxiesBeingCustomized;
-	for (TWeakObjectPtr<UObject> ObjectBeingCustomized : ObjectsBeingCustomized)
-	{
-		if (UControlRigControlsProxy* Proxy = Cast< UControlRigControlsProxy>(ObjectBeingCustomized.Get()))
-		{
-			ProxiesBeingCustomized.Add(Proxy);
-		}
-	}
-
-	if (ProxiesBeingCustomized.Num() == 0 || ProxiesBeingCustomized[0]->GetControlElement() == nullptr)
-	{
-		return;
-	}
-	FText ControlText = FText::FromName(ProxiesBeingCustomized[0]->GetName());
-
-	if (ProxiesBeingCustomized.Num() > 1)
-	{
-		if (ProxiesBeingCustomized[0]->GetClass() == UControlRigTransformControlProxy::StaticClass())
-		{
-			ControlText = LOCTEXT("TransformChannels", "Transform Channels");
-		}
-		else if (ProxiesBeingCustomized[0]->GetClass() == UControlRigTransformNoScaleControlProxy::StaticClass())
-		{
-			ControlText = LOCTEXT("TransformNoScaleChannels", "TransformNoScale Channels");
-		}
-		else if (ProxiesBeingCustomized[0]->GetClass() == UControlRigEulerTransformControlProxy::StaticClass())
-		{
-			ControlText = LOCTEXT("EulerTransformChannels", "Euler Transform Channels");
-		}
-		else if (ProxiesBeingCustomized[0]->GetClass() == UControlRigFloatControlProxy::StaticClass())
-		{
-			ControlText = LOCTEXT("FloatChannels", "Float Channels");
-		}
-		else if (ProxiesBeingCustomized[0]->GetClass() == UControlRigVectorControlProxy::StaticClass())
-		{
-			ControlText = LOCTEXT("VectorChannels", "Vector Channels");
-		}
-		else if (ProxiesBeingCustomized[0]->GetClass() == UControlRigVector2DControlProxy::StaticClass())
-		{
-			ControlText = LOCTEXT("Vector2DChannels", "Vector2D Channels");
-		}
-		else if (ProxiesBeingCustomized[0]->GetClass() == UControlRigBoolControlProxy::StaticClass())
-		{
-			ControlText = LOCTEXT("BoolChannels", "Bool Channels");
-		}
-		else if (ProxiesBeingCustomized[0]->GetClass() == UControlRigEnumControlProxy::StaticClass())
-		{
-			ControlText = LOCTEXT("EnumChannels", "Enum Channels");
-		}
-		else if (ProxiesBeingCustomized[0]->GetClass() == UControlRigIntegerControlProxy::StaticClass())
-		{
-			ControlText = LOCTEXT("IntegerChannels", "Integer Channels");
-		}
-	}
-
-	IDetailCategoryBuilder& Category = DetailLayout.EditCategory(TEXT("Control"), ControlText);
-	for (UControlRigControlsProxy* Proxy : ProxiesBeingCustomized)
-	{
-		FRigControlElement* ControlElement = Proxy->GetControlElement();
-		if (ControlElement == nullptr)
-		{
-			continue;
-		}
-
-		FName ValuePropertyName = TEXT("Transform");
-		if (ControlElement->Settings.ControlType == ERigControlType::Float ||
-			ControlElement->Settings.ControlType == ERigControlType::ScaleFloat)
-		{
-			ValuePropertyName = TEXT("Float");
-		}
-		else if (ControlElement->Settings.ControlType == ERigControlType::Integer)
-		{
-			if (ControlElement->Settings.ControlEnum == nullptr)
-			{
-				ValuePropertyName = TEXT("Integer");
-			}
-			else
-			{
-				ValuePropertyName = TEXT("Enum");
-			}
-		}
-		else if (ControlElement->Settings.ControlType == ERigControlType::Bool)
-		{
-			ValuePropertyName = TEXT("Bool");
-		}
-		else if (ControlElement->Settings.ControlType == ERigControlType::Position ||
-			ControlElement->Settings.ControlType == ERigControlType::Scale)
-		{
-			ValuePropertyName = TEXT("Vector");
-		}
-		else if (ControlElement->Settings.ControlType == ERigControlType::Vector2D)
-		{
-			ValuePropertyName = TEXT("Vector2D");
-		}
-		else if (ControlElement->Settings.ControlType == ERigControlType::EulerTransform)
-		{
-			ValuePropertyName = TEXT("EulerTransform");
-		}
-		else if (ControlElement->Settings.ControlType == ERigControlType::TransformNoScale)
-		{
-			ValuePropertyName = TEXT("TransformNoScale");
-		}
-
-		TSharedPtr<IPropertyHandle> ValuePropertyHandle = DetailLayout.GetProperty(ValuePropertyName, Proxy->GetClass());
-		if (ValuePropertyHandle)
-		{
-			ValuePropertyHandle->SetPropertyDisplayName(FText::FromName(Proxy->GetName()));
-		}
-
-		URigHierarchy* Hierarchy = Proxy->ControlRig->GetHierarchy();
-		Hierarchy->ForEach<FRigControlElement>([Hierarchy, Proxy, &Category, this](FRigControlElement* ControlElement) -> bool
-			{
-				FName ParentControlName = NAME_None;
-				FRigControlElement* ParentControlElement = Cast<FRigControlElement>(Hierarchy->GetFirstParent(ControlElement));
-				if (ParentControlElement)
-				{
-					ParentControlName = ParentControlElement->GetFName();
-				}
-
-				if (ParentControlName == ControlElement->GetFName())
-				{
-					if (FControlRigEditMode* EditMode = static_cast<FControlRigEditMode*>(ModeTools->GetActiveMode(FControlRigEditMode::ModeName)))
-					{
-						UControlRig* ControlRig = Proxy->ControlRig.Get();
-						if (UObject* NestedProxy = EditMode->ControlProxy->FindProxy(ControlRig,ControlElement->GetFName()))
-						{
-							FName PropertyName(NAME_None);
-							switch (ControlElement->Settings.ControlType)
-							{
-							case ERigControlType::Bool:
-							{
-								PropertyName = TEXT("Bool");
-								break;
-							}
-							case ERigControlType::Float:
-							case ERigControlType::ScaleFloat:
-							{
-								PropertyName = TEXT("Float");
-								break;
-							}
-							case ERigControlType::Integer:
-							{
-								if (ControlElement->Settings.ControlEnum == nullptr)
-								{
-									PropertyName = TEXT("Integer");
-								}
-								else
-								{
-									PropertyName = TEXT("Enum");
-								}
-								break;
-							}
-							default:
-							{
-								break;
-							}
-							}
-
-							if (PropertyName.IsNone())
-							{
-								return true;
-							}
-
-							TArray<UObject*> NestedProxies;
-							NestedProxies.Add(NestedProxy);
-
-							FAddPropertyParams Params;
-							Params.CreateCategoryNodes(false);
-							IDetailPropertyRow* NestedRow = Category.AddExternalObjectProperty(
-								NestedProxies,
-								PropertyName,
-								EPropertyLocation::Advanced,
-								Params);
-							NestedRow->DisplayName(FText::FromName(ControlElement->Settings.DisplayName));
-
-							Category.SetShowAdvanced(true);
-						}
-					}
-				}
-				return true;
-			});
-	}
 }
-
 void SControlRigDetails::Construct(const FArguments& InArgs, FControlRigEditMode& InEditMode)
 {
+	using namespace UE::Sequencer;
+
 	ModeTools = InEditMode.GetModeManager();
 	FDetailsViewArgs DetailsViewArgs;
 	{
@@ -231,7 +53,7 @@ void SControlRigDetails::Construct(const FArguments& InArgs, FControlRigEditMode
 		DetailsViewArgs.bShowOptions = false;
 		DetailsViewArgs.bShowModifiedPropertiesOption = true;
 		DetailsViewArgs.bCustomNameAreaLocation = true;
-		DetailsViewArgs.bCustomFilterAreaLocation = true;
+		DetailsViewArgs.bCustomFilterAreaLocation = false;
 		DetailsViewArgs.NameAreaSettings = FDetailsViewArgs::HideNameArea;
 		DetailsViewArgs.bAllowMultipleTopLevelObjects = false;
 		DetailsViewArgs.bShowScrollBar = false; // Don't need to show this, as we are putting it in a scroll box
@@ -242,136 +64,33 @@ void SControlRigDetails::Construct(const FArguments& InArgs, FControlRigEditMode
 
 	auto CreateDetailsView = [this](FDetailsViewArgs InDetailsViewArgs) -> TSharedPtr<IDetailsView>
 	{
+		FControlRigEditMode* EditMode = GetEditMode();
 		TSharedPtr<IDetailsView> DetailsView = FModuleManager::GetModuleChecked<FPropertyEditorModule>("PropertyEditor").CreateDetailView(InDetailsViewArgs);
-		DetailsView->SetKeyframeHandler(SharedThis(this));
+		DetailsView->SetKeyframeHandler(EditMode->DetailKeyFrameCache);
 		DetailsView->SetIsPropertyVisibleDelegate(FIsPropertyVisible::CreateSP(this, &SControlRigDetails::ShouldShowPropertyOnDetailCustomization));
 		DetailsView->SetIsPropertyReadOnlyDelegate(FIsPropertyReadOnly::CreateSP(this, &SControlRigDetails::IsReadOnlyPropertyOnDetailCustomization));
 		DetailsView->SetGenericLayoutDetailsDelegate(FOnGetDetailCustomizationInstance::CreateStatic(&FControlRigEditModeGenericDetails::MakeInstance, ModeTools));
 		return DetailsView;
 	};
 
-	ControlEulerTransformDetailsView = CreateDetailsView(DetailsViewArgs);
-	IndividualControlEulerTransformDetailsView = CreateDetailsView(IndividualDetailsViewArgs);
-	ControlTransformDetailsView = CreateDetailsView(DetailsViewArgs);
-	IndividualControlTransformDetailsView = CreateDetailsView(IndividualDetailsViewArgs);
-	ControlTransformNoScaleDetailsView = CreateDetailsView(DetailsViewArgs);
-	IndividualControlTransformNoScaleDetailsView = CreateDetailsView(IndividualDetailsViewArgs);
-	ControlFloatDetailsView = CreateDetailsView(DetailsViewArgs);
-	IndividualControlFloatDetailsView = CreateDetailsView(IndividualDetailsViewArgs);
-	ControlEnumDetailsView = CreateDetailsView(DetailsViewArgs);
-	IndividualControlEnumDetailsView = CreateDetailsView(IndividualDetailsViewArgs);
-	ControlIntegerDetailsView = CreateDetailsView(DetailsViewArgs);
-	IndividualControlIntegerDetailsView = CreateDetailsView(IndividualDetailsViewArgs);
-	ControlBoolDetailsView = CreateDetailsView(DetailsViewArgs);
-	IndividualControlBoolDetailsView = CreateDetailsView(IndividualDetailsViewArgs);
-	ControlVectorDetailsView = CreateDetailsView(DetailsViewArgs);
-	IndividualControlVectorDetailsView = CreateDetailsView(IndividualDetailsViewArgs);
-	ControlVector2DDetailsView = CreateDetailsView(DetailsViewArgs);
-	IndividualControlVector2DDetailsView = CreateDetailsView(IndividualDetailsViewArgs);
+	AllControlsView = CreateDetailsView(IndividualDetailsViewArgs);
 
 	ChildSlot
 		[
 			SNew(SScrollBox)
 			+ SScrollBox::Slot()
 			[
-			SNew(SVerticalBox)
-
-			
-			+ SVerticalBox::Slot()
-				.AutoHeight()
-				[
-					ControlEulerTransformDetailsView.ToSharedRef()
-				]
-			+ SVerticalBox::Slot()
-				.AutoHeight()
-				[
-					ControlTransformDetailsView.ToSharedRef()
-				]
-			+ SVerticalBox::Slot()
-				.AutoHeight()
-				[
-					ControlTransformNoScaleDetailsView.ToSharedRef()
-				]
-			+ SVerticalBox::Slot()
-				.AutoHeight()
-				[
-					ControlBoolDetailsView.ToSharedRef()
-				]
-			+ SVerticalBox::Slot()
-				.AutoHeight()
-				[
-					ControlIntegerDetailsView.ToSharedRef()
-				]
-			+ SVerticalBox::Slot()
-				.AutoHeight()
-				[
-					ControlEnumDetailsView.ToSharedRef()
-				]
-			+ SVerticalBox::Slot()
-				.AutoHeight()
-				[
-					ControlVectorDetailsView.ToSharedRef()
-				]
-			+ SVerticalBox::Slot()
-				.AutoHeight()
-				[
-					ControlVector2DDetailsView.ToSharedRef()
-				]
-			+ SVerticalBox::Slot()
-				.AutoHeight()
-				[
-					ControlFloatDetailsView.ToSharedRef()
-				]
-
-			+ SVerticalBox::Slot()
+				SNew(SVerticalBox)
+				+ SVerticalBox::Slot()
 				.AutoHeight()
 				[
-					IndividualControlEulerTransformDetailsView.ToSharedRef()
-				]
-			+ SVerticalBox::Slot()
-				.AutoHeight()
-				[
-					IndividualControlTransformDetailsView.ToSharedRef()
-				]
-			+ SVerticalBox::Slot()
-				.AutoHeight()
-				[
-					IndividualControlTransformNoScaleDetailsView.ToSharedRef()
-				]
-			+ SVerticalBox::Slot()
-				.AutoHeight()
-				[
-					IndividualControlBoolDetailsView.ToSharedRef()
-				]
-			+ SVerticalBox::Slot()
-				.AutoHeight()
-				[
-					IndividualControlIntegerDetailsView.ToSharedRef()
-				]
-			+ SVerticalBox::Slot()
-				.AutoHeight()
-				[
-					IndividualControlEnumDetailsView.ToSharedRef()
-				]
-			+ SVerticalBox::Slot()
-				.AutoHeight()
-				[
-					IndividualControlVectorDetailsView.ToSharedRef()
-				]
-			+ SVerticalBox::Slot()
-				.AutoHeight()
-				[
-					IndividualControlVector2DDetailsView.ToSharedRef()
-				]
-			+ SVerticalBox::Slot()
-				.AutoHeight()
-				[
-					IndividualControlFloatDetailsView.ToSharedRef()
+					AllControlsView.ToSharedRef()
 				]
 			]
 		];
 
 	SetEditMode(InEditMode);
+	SequencerTracker.SetSequencerAndDetails(InEditMode.GetWeakSequencer(), this);
 }
 
 SControlRigDetails::~SControlRigDetails()
@@ -379,12 +98,176 @@ SControlRigDetails::~SControlRigDetails()
 	//base class handles control rig related cleanup
 }
 
+void SControlRigDetails::SelectedSequencerObjects(const TMap<UObject*, FArrayOfPropertyTracks>& InObjectsTracked)
+{
+	TMap<UObject*, FArrayOfPropertyTracks> SequencerObjects;
+	for (const TPair<UObject*, FArrayOfPropertyTracks>& Pair : InObjectsTracked)
+	{
+		if (AActor* Actor = Cast<AActor>(Pair.Key))
+		{
+			if (Actor->IsSelectedInEditor())
+			{
+				SequencerObjects.Add(Pair);
+			}
+		}
+		else if (UActorComponent* Component = Cast<UActorComponent>(Pair.Key))
+		{
+			if (Component->IsSelectedInEditor())
+			{
+				SequencerObjects.Add(Pair);
+			}
+		}
+	}
+	//make sure the objects that are selected are actually selected in the world
+
+	HandleSequencerObjects(SequencerObjects);
+	UpdateProxies();
+}
 void SControlRigDetails::HandleControlSelected(UControlRig* Subject, FRigControlElement* InControl, bool bSelected)
 {
 	FControlRigBaseDockableView::HandleControlSelected(Subject, InControl, bSelected);
 	UpdateProxies();
 }
 
+static UControlRigControlsProxy* GetParentProxy(UControlRigControlsProxy* ChildProxy, const TArray<UControlRigControlsProxy*>& Proxies)
+{
+	FRigBaseElement* ChildParent = (ChildProxy && ChildProxy->OwnerControlRig.IsValid()) ?
+		ChildProxy->OwnerControlRig->GetHierarchy()->GetFirstParent(ChildProxy->OwnerControlElement) : nullptr;
+	if (ChildParent == nullptr)
+	{
+		return nullptr;
+	}
+	for (UControlRigControlsProxy* Proxy : Proxies)
+	{
+		if (Proxy && Proxy->OwnerControlRig.IsValid() && (ChildParent == Proxy->OwnerControlElement))
+		{
+			return Proxy;
+		}
+	}
+	return nullptr;
+}
+static UControlRigControlsProxy* GetProxyWithSameType(TArray<TWeakObjectPtr<>>& AllProxies, ERigControlType ControlType, bool bIsEnum)
+{
+	for (TWeakObjectPtr<> ExistingProxy : AllProxies)
+	{
+		if (ExistingProxy.IsValid())
+		{
+			switch (ControlType)
+			{
+			case ERigControlType::Transform:
+			case ERigControlType::TransformNoScale:
+			case ERigControlType::EulerTransform:
+			{
+				if (ExistingProxy.Get()->IsA<UAnimDetailControlsProxyTransform>())
+				{
+					return  Cast<UControlRigControlsProxy>(ExistingProxy.Get());
+				}
+				break;
+			}
+			case ERigControlType::Float:
+			case ERigControlType::ScaleFloat:
+			{
+				if (ExistingProxy.Get()->IsA<UAnimDetailControlsProxyFloat>())
+				{
+					return  Cast<UControlRigControlsProxy>(ExistingProxy.Get());
+				}
+				break;
+
+			}
+			case ERigControlType::Integer:
+			{
+				if (bIsEnum == false)
+				{
+					if (ExistingProxy.Get()->IsA<UAnimDetailControlsProxyInteger>())
+					{
+						return  Cast<UAnimDetailControlsProxyInteger>(ExistingProxy.Get());
+					}
+				}
+				else
+				{
+					if (ExistingProxy.Get()->IsA<UAnimDetailControlsProxyEnum>())
+					{
+						return  Cast<UControlRigControlsProxy>(ExistingProxy.Get());
+					}
+				}
+				break;
+
+			}
+			case ERigControlType::Position:
+			{
+				if (ExistingProxy.Get()->IsA<UAnimDetailControlsProxyLocation>())
+				{
+					return  Cast<UControlRigControlsProxy>(ExistingProxy.Get());
+				}
+				break;
+			}
+			case ERigControlType::Rotator:
+				{
+				if (ExistingProxy.Get()->IsA<UAnimDetailControlsProxyRotation>())
+				{
+					return  Cast<UControlRigControlsProxy>(ExistingProxy.Get());
+				}
+				break;
+			}
+			case ERigControlType::Scale:
+			{
+				if (ExistingProxy.Get()->IsA<UAnimDetailControlsProxyScale>())
+				{
+					return  Cast<UControlRigControlsProxy>(ExistingProxy.Get());
+				}
+				break;
+			}
+			case ERigControlType::Vector2D:
+			{
+				if (ExistingProxy.Get()->IsA<UAnimDetailControlsProxyVector2D>())
+				{
+					return  Cast<UControlRigControlsProxy>(ExistingProxy.Get());
+				}
+				break;
+			}
+			case ERigControlType::Bool:
+			{
+				if (ExistingProxy.Get()->IsA<UAnimDetailControlsProxyBool>())
+				{
+					return  Cast<UControlRigControlsProxy>(ExistingProxy.Get());
+				}
+				break;
+			}
+			}
+		}
+	}
+	return nullptr;
+}
+
+void SControlRigDetails::HandleSequencerObjects(TMap<UObject*, FArrayOfPropertyTracks>& SequencerObjects)
+{
+	if (FControlRigEditMode* EditMode = static_cast<FControlRigEditMode*>(ModeTools->GetActiveMode(FControlRigEditMode::ModeName)))
+	{
+		if (UControlRigDetailPanelControlProxies* ControlProxy = EditMode->GetDetailProxies())
+		{
+			TMap<ERigControlType, FSequencerProxyPerType>  ProxyPerType;
+			for (TPair<UObject*, FArrayOfPropertyTracks>& Pair : SequencerObjects)
+			{
+				for (UMovieSceneTrack* Track : Pair.Value.PropertyTracks)
+				{
+					if (UMovieScenePropertyTrack* PropTrack = Cast<UMovieScenePropertyTrack>(Track))
+					{
+						if (PropTrack->IsA<UMovieScene3DTransformTrack>())
+						{
+							FSequencerProxyPerType& Binding = ProxyPerType.FindOrAdd(ERigControlType::Transform);
+							TArray<FBindingAndTrack>& Bindings = Binding.Bindings.FindOrAdd(Pair.Key);
+							TSharedPtr<FTrackInstancePropertyBindings> PropertyBindings = MakeShareable(new FTrackInstancePropertyBindings(PropTrack->GetPropertyName(), PropTrack->GetPropertyPath().ToString()));
+							FBindingAndTrack BindingAndTrack(PropertyBindings, Track);
+							Bindings.Add(BindingAndTrack);
+						}
+					}
+				}
+			}
+			ControlProxy->ResetSequencerProxies(ProxyPerType);
+		}
+	}
+}
+
 void SControlRigDetails::UpdateProxies()
 {
 	TWeakPtr<SWidget> WeakPtr = AsWeak();
@@ -402,346 +285,86 @@ void SControlRigDetails::UpdateProxies()
 			return;
 		}
 		
-		TArray<TWeakObjectPtr<>> Eulers;
-		TArray<TWeakObjectPtr<>> Transforms;
-		TArray<TWeakObjectPtr<>> TransformNoScales;
-		TArray<TWeakObjectPtr<>> Floats;
-		TArray<TWeakObjectPtr<>> Vectors;
-		TArray<TWeakObjectPtr<>> Vector2Ds;
-		TArray<TWeakObjectPtr<>> Bools;
-		TArray<TWeakObjectPtr<>> Integers;
-		TArray<TWeakObjectPtr<>> Enums;
-		TArray<TWeakObjectPtr<>> IndividualEulers;
-		TArray<TWeakObjectPtr<>> IndividualTransforms;
-		TArray<TWeakObjectPtr<>> IndividualTransformNoScales;
-		TArray<TWeakObjectPtr<>> IndividualFloats;
-		TArray<TWeakObjectPtr<>> IndividualVectors;
-		TArray<TWeakObjectPtr<>> IndividualVector2Ds;
-		TArray<TWeakObjectPtr<>> IndividualBools;
-		TArray<TWeakObjectPtr<>> IndividualIntegers;
-		TArray<TWeakObjectPtr<>> IndividualEnums;
-		TArray<UControlRig*> ControlRigs = StrongThis->GetControlRigs();
-	
+		TArray<TWeakObjectPtr<>> AllProxies;
+		TArray<UControlRigControlsProxy*> ChildProxies; //list of 'child' proxies that will show up as custom attributes
 		if (FControlRigEditMode* EditMode = static_cast<FControlRigEditMode*>(StrongThis->ModeTools->GetActiveMode(FControlRigEditMode::ModeName)))
 		{
 			if (UControlRigDetailPanelControlProxies* ControlProxy = EditMode->GetDetailProxies())
 			{
-				const TArray<UControlRigControlsProxy*>& Proxies = ControlProxy->GetSelectedProxies();
+				const TArray<UControlRigControlsProxy*>& Proxies = ControlProxy->GetAllSelectedProxies();
 				for (UControlRigControlsProxy* Proxy : Proxies)
 				{
 					if (Proxy == nullptr)
 					{
 						continue;
 					}
-					if (Proxy->GetClass() == UControlRigTransformControlProxy::StaticClass())
-					{
-						if (Proxy->bIsIndividual)
-						{
-							IndividualTransforms.Add(Proxy);
-						}
-						else
-						{
-							Transforms.Add(Proxy);
-						}
-					}
-					else if (Proxy->GetClass() == UControlRigTransformNoScaleControlProxy::StaticClass())
-					{
-						if (Proxy->bIsIndividual)
-						{
-							IndividualTransformNoScales.Add(Proxy);
-						}
-						else
-						{
-							TransformNoScales.Add(Proxy);
-						}
-					}
-					else if (Proxy->GetClass() == UControlRigEulerTransformControlProxy::StaticClass())
-					{
-						if (Proxy->bIsIndividual)
-						{
-							IndividualEulers.Add(Proxy);
-						}
-						else
-						{
-							Eulers.Add(Proxy);
-						}
-					}
-					else if (Proxy->GetClass() == UControlRigFloatControlProxy::StaticClass())
-					{
-						if (Proxy->bIsIndividual)
-						{
-							IndividualFloats.Add(Proxy);
-						}
-						else
-						{
-							Floats.Add(Proxy);
-						}
-					}
-					else if (Proxy->GetClass() == UControlRigVectorControlProxy::StaticClass())
-					{
-						if (Proxy->bIsIndividual)
-						{
-							IndividualVectors.Add(Proxy);
-						}
-						else
-						{
-							Vectors.Add(Proxy);
-						}
-					}
-					else if (Proxy->GetClass() == UControlRigVector2DControlProxy::StaticClass())
+					Proxy->ResetItems();
+
+					if (Proxy->bIsIndividual)
 					{
-						if (Proxy->bIsIndividual)
-						{
-							IndividualVector2Ds.Add(Proxy);
-						}
-						else
-						{
-							Vector2Ds.Add(Proxy);
-						}
+						ChildProxies.Add(Proxy);
 					}
-					else if (Proxy->GetClass() == UControlRigBoolControlProxy::StaticClass())
+					else
 					{
-						if (Proxy->bIsIndividual)
+						TObjectPtr<UEnum> EnumPtr = Proxy->OwnerControlElement ? Proxy->OwnerControlElement->Settings.ControlEnum : nullptr;
+						if (UControlRigControlsProxy* ExistingProxy = GetProxyWithSameType(AllProxies, Proxy->Type, EnumPtr != nullptr))
 						{
-							IndividualBools.Add(Proxy);
+							ExistingProxy->AddItem(Proxy);
+							ExistingProxy->ValueChanged();
 						}
 						else
 						{
-							Bools.Add(Proxy);
+							AllProxies.Add(Proxy);
 						}
 					}
-					else if (Proxy->GetClass() == UControlRigEnumControlProxy::StaticClass())
+				}
+				//now add child proxies to parents if parents also selected...
+				for (UControlRigControlsProxy* Proxy : ChildProxies)
+				{
+					if (UControlRigControlsProxy* ParentProxy = GetParentProxy(Proxy, Proxies))
 					{
-						if (Proxy->bIsIndividual)
-						{
-							IndividualEnums.Add(Proxy);
-						}
-						else
+						TObjectPtr<UEnum> EnumPtr = ParentProxy->OwnerControlElement ? ParentProxy->OwnerControlElement->Settings.ControlEnum : nullptr;
+						if (UControlRigControlsProxy* ExistingProxy = GetProxyWithSameType(AllProxies, ParentProxy->Type, EnumPtr != nullptr))
 						{
-							Enums.Add(Proxy);
+							ParentProxy->AddChildProxy(Proxy);
 						}
 					}
-					else if (Proxy->GetClass() == UControlRigIntegerControlProxy::StaticClass())
+					else
 					{
-						if (Proxy->bIsIndividual)
-						{
-							IndividualIntegers.Add(Proxy);
-						}
-						else
-						{
-							Integers.Add(Proxy);
-						}
+						AllProxies.Add(Proxy);
 					}
 				}
-			}
-			for (TWeakObjectPtr<>& Object : Transforms)
-			{
-				UControlRigControlsProxy* Proxy = Cast<UControlRigControlsProxy>(Object.Get());
-				if (Proxy)
-				{
-					Proxy->SetIsMultiple(Transforms.Num() > 1);
-				}
-			}
-			for (TWeakObjectPtr<>& Object : Floats)
-			{
-				UControlRigControlsProxy* Proxy = Cast<UControlRigControlsProxy>(Object.Get());
-				if (Proxy)
+				for (UControlRigControlsProxy* Proxy : Proxies)
 				{
-					Proxy->SetIsMultiple(Floats.Num() > 1);
+					Proxy->ValueChanged();
 				}
 			}
 		}
-
-		StrongThis->SetTransformDetailsObjects(Transforms, false);
-		StrongThis->SetTransformNoScaleDetailsObjects(TransformNoScales, false);
-		StrongThis->SetEulerTransformDetailsObjects(Eulers, false);
-		StrongThis->SetVectorDetailsObjects(Vectors, false);
-		StrongThis->SetVector2DDetailsObjects(Vector2Ds, false);
-		StrongThis->SetFloatDetailsObjects(Floats, false);
-		StrongThis->SetBoolDetailsObjects(Bools,false);
-		StrongThis->SetIntegerDetailsObjects(Integers,false);
-		StrongThis->SetEnumDetailsObjects(Enums,false);
-		StrongThis->SetTransformDetailsObjects(IndividualTransforms, true);
-		StrongThis->SetTransformNoScaleDetailsObjects(IndividualTransformNoScales, true);
-		StrongThis->SetEulerTransformDetailsObjects(IndividualEulers, true);
-		StrongThis->SetVectorDetailsObjects(IndividualVectors, true);
-		StrongThis->SetVector2DDetailsObjects(IndividualVector2Ds, true);
-		StrongThis->SetFloatDetailsObjects(IndividualFloats, true);
-		StrongThis->SetBoolDetailsObjects(IndividualBools,true);
-		StrongThis->SetIntegerDetailsObjects(IndividualIntegers,true);
-		StrongThis->SetEnumDetailsObjects(IndividualEnums,true);
+		
+		StrongThis->AllControlsView->SetObjects(AllProxies,true);
+		
 	});
 }
 
-void SControlRigDetails::SetEulerTransformDetailsObjects(const TArray<TWeakObjectPtr<>>& InObjects, bool bIsIndividual)
-{
-	TSharedPtr<IDetailsView>& DetailsView = bIsIndividual ? IndividualControlEulerTransformDetailsView : ControlEulerTransformDetailsView; 
-	if (DetailsView)
-	{
-		DetailsView->SetObjects(InObjects);
-	}
-};
-
-void SControlRigDetails::SetTransformDetailsObjects(const TArray<TWeakObjectPtr<>>& InObjects, bool bIsIndividual)
-{
-	TSharedPtr<IDetailsView>& DetailsView = bIsIndividual ? IndividualControlTransformDetailsView : ControlTransformDetailsView; 
-	if (DetailsView)
-	{
-		DetailsView->SetObjects(InObjects);
-	}
-}
-
-void SControlRigDetails::SetTransformNoScaleDetailsObjects(const TArray<TWeakObjectPtr<>>& InObjects, bool bIsIndividual)
+FReply SControlRigDetails::OnKeyDown(const FGeometry& MyGeometry, const FKeyEvent& InKeyEvent)
 {
-	TSharedPtr<IDetailsView>& DetailsView = bIsIndividual ? IndividualControlTransformNoScaleDetailsView : ControlTransformNoScaleDetailsView; 
-	if (DetailsView)
+	if (FControlRigEditMode* EditMode = static_cast<FControlRigEditMode*>(ModeTools->GetActiveMode(FControlRigEditMode::ModeName)))
 	{
-		DetailsView->SetObjects(InObjects);
-	}
-}
-
-void SControlRigDetails::SetFloatDetailsObjects(const TArray<TWeakObjectPtr<>>& InObjects, bool bIsIndividual)
-{
-	TSharedPtr<IDetailsView>& DetailsView = bIsIndividual ? IndividualControlFloatDetailsView : ControlFloatDetailsView; 
-	if (DetailsView)
-	{
-		DetailsView->SetObjects(InObjects);
-	}
-}
-
-void SControlRigDetails::SetBoolDetailsObjects(const TArray<TWeakObjectPtr<>>& InObjects, bool bIsIndividual)
-{
-	TSharedPtr<IDetailsView>& DetailsView = bIsIndividual ? IndividualControlBoolDetailsView : ControlBoolDetailsView; 
-	if (DetailsView)
-	{
-		DetailsView->SetObjects(InObjects);
-	}
-}
-
-void SControlRigDetails::SetIntegerDetailsObjects(const TArray<TWeakObjectPtr<>>& InObjects, bool bIsIndividual)
-{
-	TSharedPtr<IDetailsView>& DetailsView = bIsIndividual ? IndividualControlIntegerDetailsView : ControlIntegerDetailsView; 
-	if (DetailsView)
-	{
-		DetailsView->SetObjects(InObjects);
-	}
-}
-
-void SControlRigDetails::SetEnumDetailsObjects(const TArray<TWeakObjectPtr<>>& InObjects, bool bIsIndividual)
-{
-	TSharedPtr<IDetailsView>& DetailsView = bIsIndividual ? IndividualControlEnumDetailsView : ControlEnumDetailsView; 
-	if (DetailsView)
-	{
-		DetailsView->SetObjects(InObjects);
-	}
-}
-
-void SControlRigDetails::SetVectorDetailsObjects(const TArray<TWeakObjectPtr<>>& InObjects, bool bIsIndividual)
-{
-	TSharedPtr<IDetailsView>& DetailsView = bIsIndividual ? IndividualControlVectorDetailsView : ControlVectorDetailsView; 
-	if (DetailsView)
-	{
-		DetailsView->SetObjects(InObjects);
-	}
-}
-
-void SControlRigDetails::SetVector2DDetailsObjects(const TArray<TWeakObjectPtr<>>& InObjects, bool bIsIndividual)
-{
-	TSharedPtr<IDetailsView>& DetailsView = bIsIndividual ? IndividualControlVector2DDetailsView : ControlVector2DDetailsView; 
-	if (DetailsView)
-	{
-		DetailsView->SetObjects(InObjects);
-	}
-}
-
-bool SControlRigDetails::IsPropertyKeyable(const UClass* InObjectClass, const IPropertyHandle& InPropertyHandle) const
-{
-	TArray<UObject*> OuterObjects;
-	InPropertyHandle.GetOuterObjects(OuterObjects);
-	if(OuterObjects.Num() == 1)
-	{
-		if (UControlRigControlsProxy* Proxy = Cast< UControlRigControlsProxy>(OuterObjects[0]))
+		TWeakPtr<ISequencer> Sequencer = EditMode->GetWeakSequencer();
+		if (Sequencer.IsValid())
 		{
-			if(UControlRig* ControlRig = Proxy->ControlRig.Get())
+			using namespace UE::Sequencer;
+			const TSharedPtr<FSequencerEditorViewModel> SequencerViewModel = Sequencer.Pin()->GetViewModel();
+			const FCurveEditorExtension* CurveEditorExtension = SequencerViewModel->CastDynamic<FCurveEditorExtension>();
+			check(CurveEditorExtension);
+			TSharedPtr<FCurveEditor> CurveEditor = CurveEditorExtension->GetCurveEditor();
+			if (CurveEditor->GetCommands()->ProcessCommandBindings(InKeyEvent))
 			{
-				const FRigElementKey Key = FRigElementKey(Proxy->ControlName, ERigElementType::Control);
-				if(const FRigControlElement* ControlElement = ControlRig->GetHierarchy()->Find<FRigControlElement>(Key))
-				{
-					if(!ControlRig->GetHierarchy()->IsAnimatable(ControlElement))
-					{
-						return false;
-					}
-				}
+				return FReply::Handled();
 			}
 		}
 	}
-	
-	if (InObjectClass && InObjectClass->IsChildOf(UControlRigTransformNoScaleControlProxy::StaticClass()) && InObjectClass->IsChildOf(UControlRigEulerTransformControlProxy::StaticClass()) && InPropertyHandle.GetProperty()
-		&& InPropertyHandle.GetProperty()->GetFName() == GET_MEMBER_NAME_CHECKED(UControlRigTransformControlProxy, Transform))
-	{
-		return true;
-	}
-
-	FCanKeyPropertyParams CanKeyPropertyParams(InObjectClass, InPropertyHandle);
-	ISequencer* Sequencer = GetSequencer();
-	if (Sequencer && Sequencer->CanKeyProperty(CanKeyPropertyParams))
-	{
-		return true;
-	}
-
-	return false;
-}
-
-bool SControlRigDetails::IsPropertyKeyingEnabled() const
-{
-	ISequencer* Sequencer = GetSequencer();
-	if (Sequencer && Sequencer->GetFocusedMovieSceneSequence())
-	{
-		return true;
-	}
-
-	return false;
-}
-
-bool SControlRigDetails::IsPropertyAnimated(const IPropertyHandle& PropertyHandle, UObject *ParentObject) const
-{
-	ISequencer* Sequencer = GetSequencer();
-	if (Sequencer && Sequencer->GetFocusedMovieSceneSequence())
-	{
-		constexpr bool bCreateHandleIfMissing = false;
-		FGuid ObjectHandle = Sequencer->GetHandleToObject(ParentObject, bCreateHandleIfMissing);
-		if (ObjectHandle.IsValid())
-		{
-			UMovieScene* MovieScene = Sequencer->GetFocusedMovieSceneSequence()->GetMovieScene();
-			FProperty* Property = PropertyHandle.GetProperty();
-			TSharedRef<FPropertyPath> PropertyPath = FPropertyPath::CreateEmpty();
-			PropertyPath->AddProperty(FPropertyInfo(Property));
-			FName PropertyName(*PropertyPath->ToString(TEXT(".")));
-			TSubclassOf<UMovieSceneTrack> TrackClass; //use empty @todo find way to get the UMovieSceneTrack from the Property type.
-			return MovieScene->FindTrack(TrackClass, ObjectHandle, PropertyName) != nullptr;
-		}
-	}
-	return false;
-}
-
-void SControlRigDetails::OnKeyPropertyClicked(const IPropertyHandle& KeyedPropertyHandle)
-{
-	ISequencer* Sequencer = GetSequencer();
-	if (Sequencer && !Sequencer->IsAllowedToChange())
-	{
-		return;
-	}
-
-	TArray<UObject*> Objects;
-	KeyedPropertyHandle.GetOuterObjects(Objects);
-	for (UObject *Object : Objects)
-	{
-		UControlRigControlsProxy* Proxy = Cast< UControlRigControlsProxy>(Object);
-		if (Proxy)
-		{
-			Proxy->SetKey(KeyedPropertyHandle);
-		}
-	}
+	return FReply::Unhandled();
 }
 
 bool SControlRigDetails::ShouldShowPropertyOnDetailCustomization(const FPropertyAndParent& InPropertyAndParent) const
@@ -749,19 +372,6 @@ bool SControlRigDetails::ShouldShowPropertyOnDetailCustomization(const FProperty
 	auto ShouldPropertyBeVisible = [](const FProperty& InProperty)
 	{
 		bool bShow = InProperty.HasAnyPropertyFlags(CPF_Interp) || InProperty.HasMetaData(FRigVMStruct::InputMetaName) || InProperty.HasMetaData(FRigVMStruct::OutputMetaName);
-
-		// Always show settings properties
-		const UClass* OwnerClass = InProperty.GetOwner<UClass>();
-		bShow |= OwnerClass == UControlRigTransformControlProxy::StaticClass();
-		bShow |= OwnerClass == UControlRigTransformNoScaleControlProxy::StaticClass();
-		bShow |= OwnerClass == UControlRigEulerTransformControlProxy::StaticClass();
-		bShow |= OwnerClass == UControlRigFloatControlProxy::StaticClass();
-		bShow |= OwnerClass == UControlRigVectorControlProxy::StaticClass();
-		bShow |= OwnerClass == UControlRigVector2DControlProxy::StaticClass();
-		bShow |= OwnerClass == UControlRigBoolControlProxy::StaticClass();
-		bShow |= OwnerClass == UControlRigEnumControlProxy::StaticClass();
-		bShow |= OwnerClass == UControlRigIntegerControlProxy::StaticClass();
-
 		return bShow;
 	};
 
@@ -787,20 +397,6 @@ bool SControlRigDetails::IsReadOnlyPropertyOnDetailCustomization(const FProperty
 	auto ShouldPropertyBeEnabled = [](const FProperty& InProperty)
 	{
 		bool bShow = InProperty.HasAnyPropertyFlags(CPF_Interp) || InProperty.HasMetaData(FRigVMStruct::InputMetaName);
-
-		// Always show settings properties
-		const UClass* OwnerClass = InProperty.GetOwner<UClass>();
-		bShow |= OwnerClass == UControlRigTransformControlProxy::StaticClass();
-		bShow |= OwnerClass == UControlRigTransformNoScaleControlProxy::StaticClass();
-		bShow |= OwnerClass == UControlRigEulerTransformControlProxy::StaticClass();
-		bShow |= OwnerClass == UControlRigFloatControlProxy::StaticClass();
-		bShow |= OwnerClass == UControlRigVectorControlProxy::StaticClass();
-		bShow |= OwnerClass == UControlRigVector2DControlProxy::StaticClass();
-		bShow |= OwnerClass == UControlRigBoolControlProxy::StaticClass();
-		bShow |= OwnerClass == UControlRigEnumControlProxy::StaticClass();
-		bShow |= OwnerClass == UControlRigIntegerControlProxy::StaticClass();
-
-
 		return bShow;
 	};
 
@@ -822,4 +418,67 @@ bool SControlRigDetails::IsReadOnlyPropertyOnDetailCustomization(const FProperty
 }
 
 
+FSequencerTracker::~FSequencerTracker()
+{
+	TSharedPtr<ISequencer> Sequencer = WeakSequencer.Pin();
+	if (Sequencer)
+	{
+		Sequencer->GetSelectionChangedObjectGuids().Remove(OnSelectionChangedHandle);
+	}
+}
+
+void FSequencerTracker::SetSequencerAndDetails(TWeakPtr<ISequencer> InWeakSequencer, SControlRigDetails* InControlRigDetails)
+{
+	WeakSequencer = InWeakSequencer;
+	ControlRigDetails = InControlRigDetails;
+	if (WeakSequencer.IsValid() == false || InControlRigDetails == nullptr)
+	{
+		return;
+	}
+	TSharedPtr<ISequencer> Sequencer = WeakSequencer.Pin();
+
+	TArray<FGuid> SequencerSelectedObjects;
+	Sequencer->GetSelectedObjects(SequencerSelectedObjects);
+	UpdateSequencerBindings(SequencerSelectedObjects);
+
+	OnSelectionChangedHandle = Sequencer->GetSelectionChangedObjectGuids().AddLambda([this](TArray<FGuid> NewSelection)
+	{
+		UpdateSequencerBindings(NewSelection);
+
+	});
+
+}
+
+void FSequencerTracker::UpdateSequencerBindings(const TArray<FGuid>& SequencerBindings)
+{
+	const FDateTime StartTime = FDateTime::Now();
+
+	TSharedPtr<ISequencer> Sequencer = WeakSequencer.Pin();
+	check(Sequencer);
+	ObjectsTracked.Reset();
+	for (FGuid BindingGuid : SequencerBindings)
+	{
+		FArrayOfPropertyTracks Properties;
+		Properties.PropertyTracks = Sequencer->GetFocusedMovieSceneSequence()->GetMovieScene()->FindTracks(UMovieScenePropertyTrack::StaticClass(), BindingGuid);
+		if (Properties.PropertyTracks.Num() == 0)
+		{
+			continue;
+		}
+		for (TWeakObjectPtr<> BoundObject : Sequencer->FindBoundObjects(BindingGuid, Sequencer->GetFocusedTemplateID()))
+		{
+			if (!BoundObject.IsValid())
+			{
+				continue;
+			}
+			ObjectsTracked.FindOrAdd(BoundObject.Get(), Properties);
+
+		}
+	}
+	if (ControlRigDetails)
+	{
+		ControlRigDetails->SelectedSequencerObjects(ObjectsTracked);
+	}
+}
+
+
 #undef LOCTEXT_NAMESPACE
