New Anim Details. Main code is the new objects in AnimDetailsProxy that handle support for control rig/sequencer, having multiple objects, and shared selection state with sequencer. CR doesn't support

diff --git a/Engine/Plugins/Animation/ControlRig/Source/ControlRigEditor/Private/EditMode/SControlRigDetails.h b/Engine/Plugins/Animation/ControlRig/Source/ControlRigEditor/Private/EditMode/SControlRigDetails.h
--- a/Engine/Plugins/Animation/ControlRig/Source/ControlRigEditor/Private/EditMode/SControlRigDetails.h
+++ b/Engine/Plugins/Animation/ControlRig/Source/ControlRigEditor/Private/EditMode/SControlRigDetails.h
@@ -12,14 +12,37 @@
 #include "Misc/FrameNumber.h"
 #include "IDetailsView.h"
 #include "Rigs/RigHierarchy.h"
-#include "IDetailKeyframeHandler.h"
 #include "IDetailCustomization.h"
 #include "DetailLayoutBuilder.h"
 #include "DetailCategoryBuilder.h"
 #include "DetailWidgetRow.h"
+#include "AnimDetailsProxy.h"
 
 class ISequencer;
 class UControlRig;
+class FUICommandList;
+class UMovieSceneTrack;
+class SControlRigDetails;
+
+struct FArrayOfPropertyTracks
+{
+	TArray<UMovieSceneTrack*> PropertyTracks;
+};
+
+class FSequencerTracker
+{
+public:
+	FSequencerTracker() = default;
+	~FSequencerTracker();
+	void SetSequencerAndDetails(TWeakPtr<ISequencer> InWeakSequencer, SControlRigDetails* InControlRigDetails);
+	TMap<UObject*, FArrayOfPropertyTracks>& GetObjectsTracked() { return ObjectsTracked; }
+private:
+	void UpdateSequencerBindings(const TArray<FGuid>& SequencerBindings);
+	FDelegateHandle OnSelectionChangedHandle;
+	TWeakPtr<ISequencer> WeakSequencer;
+	TMap<UObject*, FArrayOfPropertyTracks> ObjectsTracked;
+	SControlRigDetails* ControlRigDetails = nullptr;
+};
 
 class FControlRigEditModeGenericDetails : public IDetailCustomization
 {
@@ -40,7 +63,7 @@ protected:
 	FEditorModeTools* ModeTools = nullptr;
 };
 
-class SControlRigDetails: public SCompoundWidget, public FControlRigBaseDockableView, public IDetailKeyframeHandler
+class SControlRigDetails: public SCompoundWidget, public FControlRigBaseDockableView
 {
 
 	SLATE_BEGIN_ARGS(SControlRigDetails)
@@ -50,53 +73,28 @@ class SControlRigDetails: public SCompoundWidget, public FControlRigBaseDockable
 
 	void Construct(const FArguments& InArgs, FControlRigEditMode& InEditMode);
 
-	// IDetailKeyframeHandler interface
-	virtual bool IsPropertyKeyable(const UClass* InObjectClass, const class IPropertyHandle& PropertyHandle) const override;
-	virtual bool IsPropertyKeyingEnabled() const override;
-	virtual void OnKeyPropertyClicked(const IPropertyHandle& KeyedPropertyHandle) override;
-	virtual bool IsPropertyAnimated(const class IPropertyHandle& PropertyHandle, UObject *ParentObject) const override;
-
 	/** Display or edit set up for property */
 	bool ShouldShowPropertyOnDetailCustomization(const struct FPropertyAndParent& InPropertyAndParent) const;
 	bool IsReadOnlyPropertyOnDetailCustomization(const struct FPropertyAndParent& InPropertyAndParent) const;
+	void SelectedSequencerObjects(const TMap<UObject*, FArrayOfPropertyTracks>& ObjectsTracked);
 
 private:
 
-	/** Set the objects to be displayed in the details panel */
-	void SetSettingsDetailsObject(const TWeakObjectPtr<>& InObject);
-	void SetEulerTransformDetailsObjects(const TArray<TWeakObjectPtr<>>& InObjects, bool bIsIndividual);
-	void SetTransformDetailsObjects(const TArray<TWeakObjectPtr<>>& InObjects, bool bIsIndividual);
-	void SetTransformNoScaleDetailsObjects(const TArray<TWeakObjectPtr<>>& InObjects, bool bIsIndividual);
-	void SetVectorDetailsObjects(const TArray<TWeakObjectPtr<>>& InObjects, bool bIsIndividual);
-	void SetVector2DDetailsObjects(const TArray<TWeakObjectPtr<>>& InObjects, bool bIsIndividual);
-	void SetFloatDetailsObjects(const TArray<TWeakObjectPtr<>>& InObjects, bool bIsIndividual);
-	void SetBoolDetailsObjects(const TArray<TWeakObjectPtr<>>& InObjects, bool bIsIndividual);
-	void SetIntegerDetailsObjects(const TArray<TWeakObjectPtr<>>& InObjects, bool bIsIndividual);
-	void SetEnumDetailsObjects(const TArray<TWeakObjectPtr<>>& InObjects, bool bIsIndividual);
-
 	void UpdateProxies();
+	void HandleSequencerObjects(TMap<UObject*, FArrayOfPropertyTracks>& InObjectsTracked);
 	virtual void HandleControlSelected(UControlRig* Subject, FRigControlElement* InControl, bool bSelected) override;
 
-	//these views will hold more than one of the same tyhpe
-	TSharedPtr<IDetailsView> ControlEulerTransformDetailsView;
-	TSharedPtr<IDetailsView> ControlTransformDetailsView;
-	TSharedPtr<IDetailsView> ControlTransformNoScaleDetailsView;
-	TSharedPtr<IDetailsView> ControlFloatDetailsView;
-	TSharedPtr<IDetailsView> ControlBoolDetailsView;
-	TSharedPtr<IDetailsView> ControlIntegerDetailsView;
-	TSharedPtr<IDetailsView> ControlEnumDetailsView;
-	TSharedPtr<IDetailsView> ControlVector2DDetailsView;
-	TSharedPtr<IDetailsView> ControlVectorDetailsView;
-	//these will show more than one of the same type, will happen for controls with parents.
-	TSharedPtr<IDetailsView> IndividualControlEulerTransformDetailsView;
-	TSharedPtr<IDetailsView> IndividualControlTransformDetailsView;
-	TSharedPtr<IDetailsView> IndividualControlTransformNoScaleDetailsView;
-	TSharedPtr<IDetailsView> IndividualControlFloatDetailsView;
-	TSharedPtr<IDetailsView> IndividualControlBoolDetailsView;
-	TSharedPtr<IDetailsView> IndividualControlIntegerDetailsView;
-	TSharedPtr<IDetailsView> IndividualControlEnumDetailsView;
-	TSharedPtr<IDetailsView> IndividualControlVector2DDetailsView;
-	TSharedPtr<IDetailsView> IndividualControlVectorDetailsView;
+	TSharedPtr<IDetailsView> AllControlsView;
+
+private:
+	/*~ Keyboard interaction */
+	virtual bool SupportsKeyboardFocus() const override { return true; }
+	virtual FReply OnKeyDown(const FGeometry& MyGeometry, const FKeyEvent& InKeyEvent) override;
+
+	FSequencerTracker SequencerTracker;
 
 };
 
+
+
+
diff --git a/Engine/Plugins/Animation/ControlRig/Source/ControlRigEditor/Private/EditMode/SControlRigEditModeTools.cpp b/Engine/Plugins/Animation/ControlRig/Source/ControlRigEditor/Private/EditMode/SControlRigEditModeTools.cpp
--- a/Engine/Plugins/Animation/ControlRig/Source/ControlRigEditor/Private/EditMode/SControlRigEditModeTools.cpp
+++ b/Engine/Plugins/Animation/ControlRig/Source/ControlRigEditor/Private/EditMode/SControlRigEditModeTools.cpp
@@ -595,11 +595,18 @@ void SControlRigEditModeTools::SetSequencer(TWeakPtr<ISequencer> InSequencer)
 
 bool SControlRigEditModeTools::IsPropertyKeyable(const UClass* InObjectClass, const IPropertyHandle& InPropertyHandle) const
 {
-	if (InObjectClass && InObjectClass->IsChildOf(UControlRigTransformNoScaleControlProxy::StaticClass()) && InObjectClass->IsChildOf(UControlRigEulerTransformControlProxy::StaticClass()) && InPropertyHandle.GetProperty()
-		&& InPropertyHandle.GetProperty()->GetFName() == GET_MEMBER_NAME_CHECKED(UControlRigTransformControlProxy, Transform)) 
+	if (InObjectClass && InObjectClass->IsChildOf(UAnimDetailControlsProxyTransform::StaticClass()))
 	{
 		return true;
 	}
+	if (InObjectClass && InObjectClass->IsChildOf(UAnimDetailControlsProxyTransform::StaticClass()) && InPropertyHandle.GetProperty()
+		&& (InPropertyHandle.GetProperty()->GetFName() == GET_MEMBER_NAME_CHECKED(UAnimDetailControlsProxyTransform, Location) ||
+		InPropertyHandle.GetProperty()->GetFName() == GET_MEMBER_NAME_CHECKED(UAnimDetailControlsProxyTransform, Rotation) ||
+		InPropertyHandle.GetProperty()->GetFName() == GET_MEMBER_NAME_CHECKED(UAnimDetailControlsProxyTransform, Scale)))
+	{
+		return true;
+	}
+
 	FCanKeyPropertyParams CanKeyPropertyParams(InObjectClass, InPropertyHandle);
 	TSharedPtr<ISequencer> Sequencer = WeakSequencer.Pin();
 	if (Sequencer.IsValid() && Sequencer->CanKeyProperty(CanKeyPropertyParams))
@@ -651,12 +658,14 @@ void SControlRigEditModeTools::OnKeyPropertyClicked(const IPropertyHandle& Keyed
 
 	TArray<UObject*> Objects;
 	KeyedPropertyHandle.GetOuterObjects(Objects);
+	TSharedPtr<ISequencer> SequencerPtr = WeakSequencer.Pin();
+
 	for (UObject *Object : Objects)
 	{
 		UControlRigControlsProxy* Proxy = Cast< UControlRigControlsProxy>(Object);
 		if (Proxy)
 	{
-			Proxy->SetKey(KeyedPropertyHandle);
+			Proxy->SetKey(SequencerPtr,KeyedPropertyHandle);
 		}
 	}
 }
@@ -675,16 +684,6 @@ bool SControlRigEditModeTools::ShouldShowPropertyOnDetailCustomization(const FPr
 		// Always show settings properties
 		const UClass* OwnerClass = InProperty.GetOwner<UClass>();
 		bShow |= OwnerClass == UControlRigEditModeSettings::StaticClass();
-		bShow |= OwnerClass == UControlRigTransformControlProxy::StaticClass();		
-		bShow |= OwnerClass == UControlRigTransformNoScaleControlProxy::StaticClass();
-		bShow |= OwnerClass == UControlRigEulerTransformControlProxy::StaticClass();
-		bShow |= OwnerClass == UControlRigFloatControlProxy::StaticClass();
-		bShow |= OwnerClass == UControlRigVectorControlProxy::StaticClass();
-		bShow |= OwnerClass == UControlRigVector2DControlProxy::StaticClass();
-		bShow |= OwnerClass == UControlRigBoolControlProxy::StaticClass();
-		bShow |= OwnerClass == UControlRigEnumControlProxy::StaticClass();
-		bShow |= OwnerClass == UControlRigIntegerControlProxy::StaticClass();
-
 		return bShow;
 	};
 
@@ -714,17 +713,6 @@ bool SControlRigEditModeTools::IsReadOnlyPropertyOnDetailCustomization(const FPr
 		// Always show settings properties
 		const UClass* OwnerClass = InProperty.GetOwner<UClass>();
 		bShow |= OwnerClass == UControlRigEditModeSettings::StaticClass();
-		bShow |= OwnerClass == UControlRigTransformControlProxy::StaticClass();
-		bShow |= OwnerClass == UControlRigTransformNoScaleControlProxy::StaticClass();
-		bShow |= OwnerClass == UControlRigEulerTransformControlProxy::StaticClass();
-		bShow |= OwnerClass == UControlRigFloatControlProxy::StaticClass();
-		bShow |= OwnerClass == UControlRigVectorControlProxy::StaticClass();
-		bShow |= OwnerClass == UControlRigVector2DControlProxy::StaticClass();
-		bShow |= OwnerClass == UControlRigBoolControlProxy::StaticClass();
-		bShow |= OwnerClass == UControlRigEnumControlProxy::StaticClass();
-		bShow |= OwnerClass == UControlRigIntegerControlProxy::StaticClass();
-
-
 		return bShow;
 	};
 
diff --git a/Engine/Plugins/Animation/ControlRig/Source/ControlRigEditor/Public/EditMode/ControlRigEditMode.h b/Engine/Plugins/Animation/ControlRig/Source/ControlRigEditor/Public/EditMode/ControlRigEditMode.h
--- a/Engine/Plugins/Animation/ControlRig/Source/ControlRigEditor/Public/EditMode/ControlRigEditMode.h
+++ b/Engine/Plugins/Animation/ControlRig/Source/ControlRigEditor/Public/EditMode/ControlRigEditMode.h
@@ -14,6 +14,9 @@
 #include "IPersonaEditMode.h"
 #include "Misc/Guid.h"
 #include "EditorDragTools.h"
+#include "Materials/Material.h"
+#include "Materials/MaterialInstanceDynamic.h"
+#include "IDetailKeyframeHandler.h"
 #include "ControlRigEditMode.generated.h"
 
 class FEditorViewportClient;
@@ -40,6 +43,9 @@ class ISequencer;
 enum class EControlRigSetKey : uint8;
 class UToolMenu;
 struct FGizmoState;
+enum class EMovieSceneDataChangeType;
+struct FMovieSceneChannelMetaData;
+class UMovieSceneSection;
 
 DECLARE_DELEGATE_RetVal_ThreeParams(FTransform, FOnGetRigElementTransform, const FRigElementKey& /*RigElementKey*/, bool /*bLocal*/, bool /*bOnDebugInstance*/);
 DECLARE_DELEGATE_ThreeParams(FOnSetRigElementTransform, const FRigElementKey& /*RigElementKey*/, const FTransform& /*Transform*/, bool /*bLocal*/);
@@ -81,6 +87,38 @@ private:
 	FDelegateHandle OnBoneTransformsFinalizedHandle;
 };
 
+
+struct FDetailKeyFrameCacheAndHandler: public IDetailKeyframeHandler
+{
+	FDetailKeyFrameCacheAndHandler() { UnsetDelegates(); }
+
+	/** IDetailKeyframeHandler interface*/
+	virtual bool IsPropertyKeyable(const UClass* InObjectClass, const class IPropertyHandle& PropertyHandle) const override;
+	virtual bool IsPropertyKeyingEnabled() const override;
+	virtual void OnKeyPropertyClicked(const IPropertyHandle& KeyedPropertyHandle) override;
+	virtual bool IsPropertyAnimated(const class IPropertyHandle& PropertyHandle, UObject* ParentObject) const override;
+	virtual EPropertyKeyedStatus GetPropertyKeyedStatus(const IPropertyHandle& PropertyHandle) const override;
+
+	/** Delegates Resetting Cached Data */
+	void OnGlobalTimeChanged();
+	void OnMovieSceneDataChanged(EMovieSceneDataChangeType);
+	void OnChannelChanged(const FMovieSceneChannelMetaData*, UMovieSceneSection*);
+
+	void SetDelegates(TWeakPtr<ISequencer>& InWeakSequencer, FControlRigEditMode* InEditMode);
+	void UnsetDelegates();
+
+	/** Map to the last calculated property keyed status. Resets when Scrubbing, changing Movie Scene Data, etc */
+	mutable TMap<const IPropertyHandle*, EPropertyKeyedStatus> CachedPropertyKeyedStatusMap;
+
+private:
+	TWeakPtr<ISequencer> WeakSequencer;
+	FControlRigEditMode* EditMode = nullptr;
+
+	void ResetCachedData();
+};
+
+
+
 class CONTROLRIGEDITOR_API FControlRigEditMode : public IPersonaEditMode
 {
 public:
@@ -301,6 +339,7 @@ protected:
 	TWeakPtr<ISequencer> WeakSequencer;
 	FGuid LastMovieSceneSig;
 
+
 	/** The scope for the interaction, one per manipulated Control rig */
 	TMap<UControlRig*,FControlRigInteractionScope*> InteractionScopes;
 
@@ -407,6 +446,7 @@ public:
 	/** Request a certain transform widget for the next update */
 	void RequestTransformWidgetMode(UE::Widget::EWidgetMode InWidgetMode);
 	
+	UControlRigDetailPanelControlProxies* GetControlProxy() { return ControlProxy; }
 private:
 	/** Whether or not Pivot Transforms have changed, in which case we need to redraw viewport*/
 	bool HasPivotTransformsChanged() const;
@@ -466,6 +506,9 @@ protected:
 
 	//Get if the hosted component is visible
 	bool IsControlRigSkelMeshVisible(UControlRig* ControlRig) const;
+	
+public:  
+		TSharedPtr<FDetailKeyFrameCacheAndHandler> DetailKeyFrameCache;
 
 private:
 
@@ -532,6 +575,7 @@ private:
 	TArray<uint32> LastHierarchyHash;
 	TArray<uint32> LastShapeLibraryHash;
 
+
 	friend class FControlRigEditorModule;
 	friend class FControlRigEditor;
 	friend class FControlRigEditModeGenericDetails;
diff --git a/Engine/Source/Editor/Sequencer/Public/MVVM/CurveEditorExtension.h b/Engine/Source/Editor/Sequencer/Public/MVVM/CurveEditorExtension.h
--- a/Engine/Source/Editor/Sequencer/Public/MVVM/CurveEditorExtension.h
+++ b/Engine/Source/Editor/Sequencer/Public/MVVM/CurveEditorExtension.h
@@ -48,11 +48,15 @@ public:
 	/** Closes the curve editor */
 	void CloseCurveEditor();
 
+	/** Curve editor tree widget */
+	TSharedPtr<SCurveEditorTree> GetCurveEditorTreeView() const { return CurveEditorTreeView; }
+
 	/**
 	 * Synchronize curve editor selection with sequencer outliner selection on the next update.
 	 */
 	void RequestSyncSelection();
 
+
 public:
 
 	static const FName CurveEditorTabName;
