[NO NON-PII COMMIT MESSAGE]

diff --git a/Engine/Plugins/Animation/ControlRig/Source/ControlRigEditor/Private/EditMode/AnimDetailsProxy.cpp b/Engine/Plugins/Animation/ControlRig/Source/ControlRigEditor/Private/EditMode/AnimDetailsProxy.cpp
--- a/Engine/Plugins/Animation/ControlRig/Source/ControlRigEditor/Private/EditMode/AnimDetailsProxy.cpp
+++ b/Engine/Plugins/Animation/ControlRig/Source/ControlRigEditor/Private/EditMode/AnimDetailsProxy.cpp
@@ -2618,6 +2618,10 @@ void UControlRigDetailPanelControlProxies::ValuesChanged()
 
 UControlRigControlsProxy* UControlRigDetailPanelControlProxies::AddProxy(UControlRig* ControlRig, FRigControlElement* ControlElement)
 {
+	if (ControlRig == nullptr || ControlElement == nullptr)
+	{
+		return nullptr;
+	}
 	//check if forced to be individual
 	bool bIsIndividual = (ControlElement->IsAnimationChannel()) ||
 		(ControlElement->Settings.AnimationType == ERigControlAnimationType::ProxyControl);
diff --git a/Engine/Plugins/Animation/ControlRig/Source/ControlRigEditor/Private/EditMode/ControlRigControlsProxy.h b/Engine/Plugins/Animation/ControlRig/Source/ControlRigEditor/Private/EditMode/ControlRigControlsProxy.h
--- a/Engine/Plugins/Animation/ControlRig/Source/ControlRigEditor/Private/EditMode/ControlRigControlsProxy.h
+++ b/Engine/Plugins/Animation/ControlRig/Source/ControlRigEditor/Private/EditMode/ControlRigControlsProxy.h
@@ -7,7 +7,7 @@
 #include "TransformNoScale.h"
 #include "EulerTransform.h"
 #include "Math/Rotator.h"
-
+#include "MovieSceneTrack.h"
 #include "IDetailKeyframeHandler.h"
 #include "IPropertyTypeCustomization.h"
 #include "Rigs/RigHierarchyDefines.h"
@@ -23,7 +23,6 @@ struct FRigControlModifiedContext;
 class ISequencer;
 class FCurveEditor;
 class IDetailLayoutBuilder;
-class UMovieSceneTrack;
 class FProperty;
 
 //channel selection states for selection matching with curves
