Follow up fix to 30962662

diff --git a/Engine/Plugins/Animation/ControlRig/Source/ControlRigEditor/Private/EditMode/ControlsProxyDetailCustomization.cpp b/Engine/Plugins/Animation/ControlRig/Source/ControlRigEditor/Private/EditMode/ControlsProxyDetailCustomization.cpp
--- a/Engine/Plugins/Animation/ControlRig/Source/ControlRigEditor/Private/EditMode/ControlsProxyDetailCustomization.cpp
+++ b/Engine/Plugins/Animation/ControlRig/Source/ControlRigEditor/Private/EditMode/ControlsProxyDetailCustomization.cpp
@@ -398,7 +398,8 @@ void FAnimDetailValueCustomization::MakeHeaderRow(TSharedRef<class IPropertyHand
 		ChildHandle->SetInstanceMetaData(TEXT("SliderExponent"), StructPropertyHandle->GetMetaData(TEXT("SliderExponent")));
 		ChildHandle->SetInstanceMetaData(TEXT("Delta"), StructPropertyHandle->GetMetaData(TEXT("Delta")));
 		ChildHandle->SetInstanceMetaData(TEXT("LinearDeltaSensitivity"), StructPropertyHandle->GetMetaData(TEXT("LinearDeltaSensitivity")));
-		ChildHandle->SetInstanceMetaData(TEXT("ShiftMouseMovePixelPerDelta"), StructPropertyHandle->GetMetaData(TEXT("ShiftMouseMovePixelPerDelta")));
+		ChildHandle->SetInstanceMetaData(TEXT("ShiftMultiplier"), StructPropertyHandle->GetMetaData(TEXT("ShiftMultiplier")));
+		ChildHandle->SetInstanceMetaData(TEXT("CtrlMultiplier"), StructPropertyHandle->GetMetaData(TEXT("CtrlMultiplier")));
 		ChildHandle->SetInstanceMetaData(TEXT("SupportDynamicSliderMaxValue"), StructPropertyHandle->GetMetaData(TEXT("SupportDynamicSliderMaxValue")));
 		ChildHandle->SetInstanceMetaData(TEXT("SupportDynamicSliderMinValue"), StructPropertyHandle->GetMetaData(TEXT("SupportDynamicSliderMinValue")));
 		ChildHandle->SetInstanceMetaData(TEXT("ClampMin"), StructPropertyHandle->GetMetaData(TEXT("ClampMin")));
@@ -770,13 +771,14 @@ TSharedRef<SWidget> FAnimDetailValueCustomization::MakeDoubleWidget(
 {
 	TOptional<double> MinValue, MaxValue, SliderMinValue, SliderMaxValue;
 	double SliderExponent, Delta;
-	int32 ShiftMouseMovePixelPerDelta = 1;
+	float ShiftMultiplier = 10.f;
+	float CtrlMultiplier = 0.1f;
 	bool SupportDynamicSliderMaxValue = false;
 	bool SupportDynamicSliderMinValue = false;
 
 	UControlRigControlsProxy* Proxy = GetProxy(StructurePropertyHandle);
 
-	ExtractDoubleMetadata(PropertyHandle, MinValue, MaxValue, SliderMinValue, SliderMaxValue, SliderExponent, Delta, ShiftMouseMovePixelPerDelta, SupportDynamicSliderMaxValue, SupportDynamicSliderMinValue);
+	ExtractDoubleMetadata(PropertyHandle, MinValue, MaxValue, SliderMinValue, SliderMaxValue, SliderExponent, Delta, ShiftMultiplier, CtrlMultiplier, SupportDynamicSliderMaxValue, SupportDynamicSliderMinValue);
 
 	TWeakPtr<IPropertyHandle> WeakHandlePtr = PropertyHandle;
 
@@ -828,7 +830,8 @@ TSharedRef<SWidget> FAnimDetailValueCustomization::MakeDoubleWidget(
 			})
 				// Only allow spin on handles with one object.  Otherwise it is not clear what value to spin
 				.AllowSpin(PropertyHandle->GetNumOuterObjects() < 2)
-				.ShiftMouseMovePixelPerDelta(ShiftMouseMovePixelPerDelta)
+				.ShiftMultiplier(ShiftMultiplier)
+				.CtrlMultiplier(CtrlMultiplier)
 				.SupportDynamicSliderMaxValue(SupportDynamicSliderMaxValue)
 				.SupportDynamicSliderMinValue(SupportDynamicSliderMinValue)
 				.OnDynamicSliderMaxValueChanged(this, &FAnimDetailValueCustomization::OnDynamicSliderMaxValueChanged)
@@ -850,7 +853,7 @@ TSharedRef<SWidget> FAnimDetailValueCustomization::MakeDoubleWidget(
 
 // The following code is just a plain copy of FMathStructCustomization which
 // would need changes to be able to serve as a base class for this customization.
-void FAnimDetailValueCustomization::ExtractDoubleMetadata(TSharedRef<IPropertyHandle>& PropertyHandle, TOptional<double>& MinValue, TOptional<double>& MaxValue, TOptional<double>& SliderMinValue, TOptional<double>& SliderMaxValue, double& SliderExponent, double& Delta, int32& ShiftMouseMovePixelPerDelta, bool& SupportDynamicSliderMaxValue, bool& SupportDynamicSliderMinValue)
+void FAnimDetailValueCustomization::ExtractDoubleMetadata(TSharedRef<IPropertyHandle>& PropertyHandle, TOptional<double>& MinValue, TOptional<double>& MaxValue, TOptional<double>& SliderMinValue, TOptional<double>& SliderMaxValue, double& SliderExponent, double& Delta, float& ShiftMultiplier, float& CtrlMultiplier, bool& SupportDynamicSliderMaxValue, bool& SupportDynamicSliderMinValue)
 {
 	FProperty* Property = PropertyHandle->GetProperty();
 
@@ -858,7 +861,8 @@ void FAnimDetailValueCustomization::ExtractDoubleMetadata(TSharedRef<IPropertyHa
 	const FString& MetaUIMaxString = Property->GetMetaData(TEXT("UIMax"));
 	const FString& SliderExponentString = Property->GetMetaData(TEXT("SliderExponent"));
 	const FString& DeltaString = Property->GetMetaData(TEXT("Delta"));
-	const FString& ShiftMouseMovePixelPerDeltaString = Property->GetMetaData(TEXT("ShiftMouseMovePixelPerDelta"));
+	const FString& ShiftMultiplierString = Property->GetMetaData(TEXT("ShiftMultiplier"));
+	const FString& CtrlMultiplierString = Property->GetMetaData(TEXT("CtrlMultiplier"));
 	const FString& SupportDynamicSliderMaxValueString = Property->GetMetaData(TEXT("SupportDynamicSliderMaxValue"));
 	const FString& SupportDynamicSliderMinValueString = Property->GetMetaData(TEXT("SupportDynamicSliderMinValue"));
 	const FString& ClampMinString = Property->GetMetaData(TEXT("ClampMin"));
@@ -900,16 +904,16 @@ void FAnimDetailValueCustomization::ExtractDoubleMetadata(TSharedRef<IPropertyHa
 		TTypeFromString<double>::FromString(Delta, *DeltaString);
 	}
 
-	ShiftMouseMovePixelPerDelta = 1;
-	if (ShiftMouseMovePixelPerDeltaString.Len())
+	ShiftMultiplier = 10.f;
+	if (ShiftMultiplierString.Len())
 	{
-		TTypeFromString<int32>::FromString(ShiftMouseMovePixelPerDelta, *ShiftMouseMovePixelPerDeltaString);
-		//The value should be greater or equal to 1
-		// 1 is neutral since it is a multiplier of the mouse drag pixel
-		if (ShiftMouseMovePixelPerDelta < 1)
-		{
-			ShiftMouseMovePixelPerDelta = 1;
-		}
+		TTypeFromString<float>::FromString(ShiftMultiplier, *ShiftMultiplierString);
+	}
+
+	CtrlMultiplier = 0.1f;
+	if (CtrlMultiplierString.Len())
+	{
+		TTypeFromString<float>::FromString(CtrlMultiplier, *CtrlMultiplierString);
 	}
 
 	const double ActualUIMin = FMath::Max(UIMin, ClampMin);
diff --git a/Engine/Plugins/Animation/ControlRig/Source/ControlRigEditor/Private/EditMode/ControlsProxyDetailCustomization.h b/Engine/Plugins/Animation/ControlRig/Source/ControlRigEditor/Private/EditMode/ControlsProxyDetailCustomization.h
--- a/Engine/Plugins/Animation/ControlRig/Source/ControlRigEditor/Private/EditMode/ControlsProxyDetailCustomization.h
+++ b/Engine/Plugins/Animation/ControlRig/Source/ControlRigEditor/Private/EditMode/ControlsProxyDetailCustomization.h
@@ -60,7 +60,8 @@ class FAnimDetailValueCustomization
 			TOptional<double>& SliderMaxValue,
 			double& SliderExponent,
 			double& Delta,
-			int32& ShiftMouseMovePixelPerDelta,
+			float& ShiftMultiplier,
+			float& CtrlMultiplier,
 			bool& SupportDynamicSliderMaxValue,
 			bool& SupportDynamicSliderMinValue);
 		void OnDynamicSliderMaxValueChanged(double NewMaxSliderValue, TWeakPtr<SWidget> InValueChangedSourceWidget, bool IsOriginator, bool UpdateOnlyIfHigher);
