Control Rig: Fix show/hide schematic viewport not working, and add viewport menu action
#rnx

diff --git a/Engine/Plugins/Animation/ControlRig/Source/ControlRigEditor/Private/Editor/ControlRigEditor.cpp b/Engine/Plugins/Animation/ControlRig/Source/ControlRigEditor/Private/Editor/ControlRigEditor.cpp
--- a/Engine/Plugins/Animation/ControlRig/Source/ControlRigEditor/Private/Editor/ControlRigEditor.cpp
+++ b/Engine/Plugins/Animation/ControlRig/Source/ControlRigEditor/Private/Editor/ControlRigEditor.cpp
@@ -2271,7 +2271,7 @@ void FControlRigEditor::HandleViewportCreated(const TSharedRef<class IPersonaVie
 		{
 			if (Blueprint->IsModularRig())
 			{
-				TSharedRef<SSchematicGraphPanel> SchematicPanel = SNew(SSchematicGraphPanel)
+				SchematicViewport = SNew(SSchematicGraphPanel)
 																.GraphData(&SchematicModel)
 																.IsOverlay(true)
 																.PaddingLeft(30)
@@ -2280,7 +2280,19 @@ void FControlRigEditor::HandleViewportCreated(const TSharedRef<class IPersonaVie
 																.PaddingBottom(60)
 																.PaddingInterNode(5)
 				;
-				InViewport->AddOverlayWidget(SchematicPanel);
+				InViewport->AddOverlayWidget(SchematicViewport.ToSharedRef());
+
+				
+				InViewport->AddToolbarExtender(TEXT("ControlRig"), FMenuExtensionDelegate::CreateLambda([&](FMenuBuilder& InMenuBuilder)
+					{
+						InMenuBuilder.AddMenuSeparator(TEXT("Modular Rig"));
+						InMenuBuilder.BeginSection("ModularRig", LOCTEXT("ModularRig_Label", "Modular Rig"));
+						{
+							InMenuBuilder.AddMenuEntry(FControlRigEditorCommands::Get().ToggleSchematicViewportVisibility);
+						}
+						InMenuBuilder.EndSection();
+					}
+				));
 			}
 		}
 	}
@@ -2348,10 +2360,19 @@ void FControlRigEditor::HandleToggleSchematicViewport()
 	if(SchematicViewport.IsValid())
 	{
 		SchematicModel.UpdateControlRigContent();
-		SchematicViewport.Pin()->ToggleVisibility();
+		SchematicViewport->ToggleVisibility();
 	}
 }
 
+bool FControlRigEditor::IsSchematicViewportActive() const
+{
+	if (SchematicViewport.IsValid())
+	{
+		return SchematicViewport->GetVisibility() != EVisibility::Hidden;
+	}
+	return false;
+}
+
 ECheckBoxState FControlRigEditor::GetToolbarDrawAxesOnSelection() const
 {
 	if (const UControlRigEditModeSettings* Settings = GetDefault<UControlRigEditModeSettings>())
@@ -3339,7 +3360,8 @@ void FControlRigEditor::BindCommands()
 	GetToolkitCommands()->MapAction(
 		FControlRigEditorCommands::Get().ToggleSchematicViewportVisibility,
 		FExecuteAction::CreateSP(this, &FControlRigEditor::HandleToggleSchematicViewport),
-		FCanExecuteAction());
+		FCanExecuteAction(),
+		FIsActionChecked::CreateSP(this, &FControlRigEditor::IsSchematicViewportActive));
 }
 
 void FControlRigEditor::OnHierarchyChanged()
diff --git a/Engine/Plugins/Animation/ControlRig/Source/ControlRigEditor/Private/Editor/ControlRigEditor.h b/Engine/Plugins/Animation/ControlRig/Source/ControlRigEditor/Private/Editor/ControlRigEditor.h
--- a/Engine/Plugins/Animation/ControlRig/Source/ControlRigEditor/Private/Editor/ControlRigEditor.h
+++ b/Engine/Plugins/Animation/ControlRig/Source/ControlRigEditor/Private/Editor/ControlRigEditor.h
@@ -245,6 +245,7 @@ private:
 	TOptional<float> GetToolbarAxesScale() const;
 	void OnToolbarAxesScaleChanged(float InValue);
 	void HandleToggleSchematicViewport();
+	bool IsSchematicViewportActive() const;
 
 		/** Handle switching skeletal meshes */
 	void HandlePreviewMeshChanged(USkeletalMesh* InOldSkeletalMesh, USkeletalMesh* InNewSkeletalMesh);
@@ -352,7 +353,7 @@ protected:
 	TSharedPtr<SComboBox<TSharedPtr<FString>>> DirectManipulationCombo;
 	bool bRefreshDirectionManipulationTargetsRequired;
 	FSimpleMulticastDelegate RequestNavigateToConnectorWarningDelegate;
-	TWeakPtr<SSchematicGraphPanel> SchematicViewport;
+	TSharedPtr<SSchematicGraphPanel> SchematicViewport;
 
 	static const TArray<FName> ForwardsSolveEventQueue;
 	static const TArray<FName> BackwardsSolveEventQueue;
diff --git a/Engine/Plugins/Animation/ControlRig/Source/ControlRigEditor/Private/Editor/ControlRigEditorCommands.cpp b/Engine/Plugins/Animation/ControlRig/Source/ControlRigEditor/Private/Editor/ControlRigEditorCommands.cpp
--- a/Engine/Plugins/Animation/ControlRig/Source/ControlRigEditor/Private/Editor/ControlRigEditorCommands.cpp
+++ b/Engine/Plugins/Animation/ControlRig/Source/ControlRigEditor/Private/Editor/ControlRigEditorCommands.cpp
@@ -13,7 +13,7 @@ void FControlRigEditorCommands::RegisterCommands()
 	UI_COMMAND(RequestDirectManipulationPosition, "Request Direct Manipulation for Position", "Request per node direct manipulation on a position", EUserInterfaceActionType::Button, FInputChord(EKeys::W));
 	UI_COMMAND(RequestDirectManipulationRotation, "Request Direct Manipulation for Rotation", "Request per node direct manipulation on a rotation", EUserInterfaceActionType::Button, FInputChord(EKeys::E));
 	UI_COMMAND(RequestDirectManipulationScale, "Request Direct Manipulation for Scale", "Request per node direct manipulation on a scale", EUserInterfaceActionType::Button, FInputChord(EKeys::R));
-	UI_COMMAND(ToggleSchematicViewportVisibility, "Toggle the viewport schematic", "Toggles the visibility of the viewport schematic", EUserInterfaceActionType::Button, FInputChord(EKeys::Y));
+	UI_COMMAND(ToggleSchematicViewportVisibility, "Show Schematic Viewport", "Toggles the visibility of the viewport schematic", EUserInterfaceActionType::ToggleButton, FInputChord(EKeys::Y));
 }
 
 #undef LOCTEXT_NAMESPACE
