Control Rig: Swap label on StandaloneRig button to represent the action (Swith to Rig Module)
#rnx

diff --git a/Engine/Plugins/Animation/ControlRig/Source/ControlRigEditor/Private/Editor/ControlRigEditor.cpp b/Engine/Plugins/Animation/ControlRig/Source/ControlRigEditor/Private/Editor/ControlRigEditor.cpp
--- a/Engine/Plugins/Animation/ControlRig/Source/ControlRigEditor/Private/Editor/ControlRigEditor.cpp
+++ b/Engine/Plugins/Animation/ControlRig/Source/ControlRigEditor/Private/Editor/ControlRigEditor.cpp
@@ -653,8 +653,8 @@ void FControlRigEditor::FillToolbar(FToolBarBuilder& ToolbarBuilder, bool bEndSe
 				NAME_None,
 				TAttribute<FText>::CreateLambda([WeakBlueprint]()
 				{
-					static const FText StandaloneRig = LOCTEXT("StandaloneRig", "Standalone Rig"); 
-					static const FText RigModule = LOCTEXT("RigModule", "Rig Module");
+					static const FText StandaloneRig = LOCTEXT("SwitchToRigModule", "Switch to Rig Module"); 
+					static const FText RigModule = LOCTEXT("SwitchToStandaloneRig", "Switch to Standalone Rig");
 					if(WeakBlueprint.IsValid())
 					{
 						if(WeakBlueprint->IsControlRigModule())
