Control Rig: Add check to prevent crash
#rnx

diff --git a/Engine/Plugins/Animation/ControlRig/Source/ControlRigEditor/Private/Editor/ControlRigEditor.cpp b/Engine/Plugins/Animation/ControlRig/Source/ControlRigEditor/Private/Editor/ControlRigEditor.cpp
--- a/Engine/Plugins/Animation/ControlRig/Source/ControlRigEditor/Private/Editor/ControlRigEditor.cpp
+++ b/Engine/Plugins/Animation/ControlRig/Source/ControlRigEditor/Private/Editor/ControlRigEditor.cpp
@@ -3184,13 +3184,16 @@ FText FControlRigEditor::GetDirectionManipulationText() const
 {
 	if (UControlRig* DebuggedControlRig = Cast<UControlRig>(GetBlueprintObj()->GetObjectBeingDebugged()))
 	{
-		TArray<FRigControlElement*> TransientControls = DebuggedControlRig->GetHierarchy()->GetTransientControls();
-		for(const FRigControlElement* TransientControl : TransientControls)
+		if (URigHierarchy* Hierarchy = DebuggedControlRig->GetHierarchy())
 		{
-			const FString Target = UControlRig::GetTargetFromTransientControl(TransientControl->GetKey());
-			if(!Target.IsEmpty())
+			TArray<FRigControlElement*> TransientControls = Hierarchy->GetTransientControls();
+			for(const FRigControlElement* TransientControl : TransientControls)
 			{
-				return FText::FromString(Target);
+				const FString Target = UControlRig::GetTargetFromTransientControl(TransientControl->GetKey());
+				if(!Target.IsEmpty())
+				{
+					return FText::FromString(Target);
+				}
 			}
 		}
 	}
