CIS fix for unitialized variables

diff --git a/Engine/Plugins/Animation/ControlRig/Source/ControlRigEditor/Private/EditMode/AnimDetailsProxy.h b/Engine/Plugins/Animation/ControlRig/Source/ControlRigEditor/Private/EditMode/AnimDetailsProxy.h
--- a/Engine/Plugins/Animation/ControlRig/Source/ControlRigEditor/Private/EditMode/AnimDetailsProxy.h
+++ b/Engine/Plugins/Animation/ControlRig/Source/ControlRigEditor/Private/EditMode/AnimDetailsProxy.h
@@ -36,7 +36,7 @@ struct FAnimDetailProxyFloat
 	GENERATED_BODY()
 
 	UPROPERTY(EditAnywhere, Interp, Category = "Float", meta = (SliderExponent = "1.0"))
-	double Float;
+	double Float = 0.0;
 
 };
 
@@ -46,7 +46,7 @@ struct FAnimDetailProxyBool
 	GENERATED_BODY()
 
 	UPROPERTY(EditAnywhere, Interp, Category = "Bool")
-	bool Bool;
+	bool Bool = false;
 
 };
 
@@ -95,7 +95,7 @@ struct FAnimDetailProxyInteger
 	GENERATED_BODY()
 
 	UPROPERTY(EditAnywhere, Interp, Category = "Integer")
-	int64 Integer;
+	int64 Integer = 0;
 
 };
 
@@ -105,13 +105,13 @@ struct FAnimDetailProxyVector3
 	GENERATED_BODY()
 
 	UPROPERTY(EditAnywhere, Interp, Category = "Vector")
-	double X;
+	double X = 0.0;
 
 	UPROPERTY(EditAnywhere, Interp, Category = "Vector")
-	double Y;
+	double Y = 0.0;
 
 	UPROPERTY(EditAnywhere, Interp, Category = "Vector")
-	double Z;
+	double Z = 0.0;
 
 	FVector ToVector() const { return FVector(X, Y, Z); }
 };
@@ -122,13 +122,13 @@ struct  FAnimDetailProxyLocation
 	GENERATED_BODY()
 
 	UPROPERTY(EditAnywhere, Interp, Category = "Location", meta = (Delta = "0.5", SliderExponent = "1", LinearDeltaSensitivity = "1"))
-	double LX;
+	double LX = 0.0;
 
 	UPROPERTY(EditAnywhere, Interp, Category = "Location", meta = (Delta = "0.5", SliderExponent = "1", LinearDeltaSensitivity = "1"))
-	double LY;
+	double LY = 0.0;
 
 	UPROPERTY(EditAnywhere, Interp, Category = "Location", meta = (Delta = "0.5", SliderExponent = "1", LinearDeltaSensitivity = "1"))
-	double LZ;
+	double LZ = 0.0;
 
 	FAnimDetailVectorState State;
 
@@ -152,13 +152,13 @@ struct  FAnimDetailProxyRotation
 	GENERATED_BODY()
 
 	UPROPERTY(EditAnywhere, Interp, Category = "Rotation", meta = (Delta = "0.5", SliderExponent = "1", LinearDeltaSensitivity = "1"))
-	double RX;
+	double RX = 0.0;
 
 	UPROPERTY(EditAnywhere, Interp, Category = "Rotation", meta = (Delta = "0.5", SliderExponent = "1", LinearDeltaSensitivity = "1"))
-	double RY;
+	double RY = 0.0;
 
 	UPROPERTY(EditAnywhere, Interp, Category = "Rotation", meta = (Delta = "0.5", SliderExponent = "1", LinearDeltaSensitivity = "1"))
-	double RZ;
+	double RZ = 0.0;
 
 	FAnimDetailVectorState State;
 
@@ -181,13 +181,13 @@ struct  FAnimDetailProxyScale
 	GENERATED_BODY()
 
 	UPROPERTY(EditAnywhere, Interp, Category = "Scale", meta = (Delta = "0.5", SliderExponent = "1", LinearDeltaSensitivity = "1"))
-	double SX;
+	double SX = 0.0;
 
 	UPROPERTY(EditAnywhere, Interp, Category = "Scale", meta = (Delta = "0.5", SliderExponent = "1", LinearDeltaSensitivity = "1"))
-	double SY;
+	double SY = 0.0;
 
 	UPROPERTY(EditAnywhere, Interp, Category = "Scale", meta = (Delta = "0.5", SliderExponent = "1", LinearDeltaSensitivity = "1"))
-	double SZ;
+	double SZ = 0.0;
 
 	FAnimDetailVectorState State;
 
@@ -210,10 +210,10 @@ struct  FAnimDetailProxyVector2D
 	GENERATED_BODY()
 
 	UPROPERTY(EditAnywhere, Interp, Category = "Vector2D", meta = (Delta = "0.05", SliderExponent = "1", LinearDeltaSensitivity = "1"))
-	double X;
+	double X = 0.0;
 
 	UPROPERTY(EditAnywhere, Interp, Category = "Vector2D", meta = (Delta = "0.05", SliderExponent = "1", LinearDeltaSensitivity = "1"))
-	double Y;
+	double Y = 0.0;
 
 
 	FAnimDetailVectorState State;
@@ -477,7 +477,7 @@ public:
 		FAnimDetailVectorSelection& OutScaleSelection) override;
 
 	UPROPERTY(EditAnywhere, Interp, Category = Vector2D)
-		FAnimDetailProxyVector2D Vector2D;
+	FAnimDetailProxyVector2D Vector2D;
 
 private:
 	void ClearMultipleFlags();
