Sequencer: Anim Details: Move setting value to happen no matter what the post edit change mode,

diff --git a/Engine/Plugins/Animation/ControlRig/Source/ControlRigEditor/Private/EditMode/ControlRigControlsProxy.cpp b/Engine/Plugins/Animation/ControlRig/Source/ControlRigEditor/Private/EditMode/ControlRigControlsProxy.cpp
--- a/Engine/Plugins/Animation/ControlRig/Source/ControlRigEditor/Private/EditMode/ControlRigControlsProxy.cpp
+++ b/Engine/Plugins/Animation/ControlRig/Source/ControlRigEditor/Private/EditMode/ControlRigControlsProxy.cpp
@@ -249,52 +249,53 @@ void UControlRigControlsProxy::PostEditChangeChainProperty(struct FPropertyChang
 					}
 				}
 			}
-			//set values
-			FProperty* Property = PropertyChangedEvent.Property;
-			FProperty* MemberProperty = nullptr;
-			if (PropertyChangedEvent.PropertyChain.GetActiveMemberNode())
-			{
-				MemberProperty = PropertyChangedEvent.PropertyChain.GetActiveMemberNode()->GetValue();
-			}
-			if (PropertyIsOnProxy(Property, MemberProperty))
+		}
+		else
+		{
+			for (TPair<FRigControlElement*, FControlRigInteractionScope*>& Scope : InteractionScopes)
 			{
-				FRigControlModifiedContext Context;
-				Context.SetKey = EControlRigSetKey::DoNotCare;
-				Context.KeyMask = (uint32)GetChannelToKeyFromPropertyName(Property->GetFName());
-
-				for (const TPair<UControlRig*, FControlRigProxyItem>& Items : ControlRigItems)
+				if (Scope.Value)
 				{
-					if (UControlRig* ControlRig = Items.Value.ControlRig.Get())
-					{
-						//we do this backwards so ValueChanged later is set up correctly since that iterates in the other direction
-						for (int32 Index = Items.Value.ControlElements.Num() - 1; Index >= 0; --Index)
-						{
-							FRigControlElement* ControlElement = Items.Value.ControlElements[Index];
-							SetControlRigElementValueFromCurrent(ControlRig, ControlElement, Context);
-						}
-					}
+					delete Scope.Value;
 				}
-				for (TPair <UObject*, FSequencerProxyItem>& SItems : SequencerItems)
+			}
+			InteractionScopes.Reset();
+		}
+
+		//set values
+		FProperty* Property = PropertyChangedEvent.Property;
+		FProperty* MemberProperty = nullptr;
+		if (PropertyChangedEvent.PropertyChain.GetActiveMemberNode())
+		{
+			MemberProperty = PropertyChangedEvent.PropertyChain.GetActiveMemberNode()->GetValue();
+		}
+		if (PropertyIsOnProxy(Property, MemberProperty))
+		{
+			FRigControlModifiedContext Context;
+			Context.SetKey = EControlRigSetKey::DoNotCare;
+			Context.KeyMask = (uint32)GetChannelToKeyFromPropertyName(Property->GetFName());
+
+			for (const TPair<UControlRig*, FControlRigProxyItem>& Items : ControlRigItems)
+			{
+				if (UControlRig* ControlRig = Items.Value.ControlRig.Get())
 				{
 					//we do this backwards so ValueChanged later is set up correctly since that iterates in the other direction
-					for (int32 Index = SItems.Value.Bindings.Num() - 1; Index >= 0; --Index)
+					for (int32 Index = Items.Value.ControlElements.Num() - 1; Index >= 0; --Index)
 					{
-						FBindingAndTrack& Binding = SItems.Value.Bindings[Index];
-						SetBindingValueFromCurrent(SItems.Key, Binding.Binding, Context);
+						FRigControlElement* ControlElement = Items.Value.ControlElements[Index];
+						SetControlRigElementValueFromCurrent(ControlRig, ControlElement, Context);
 					}
 				}
 			}
-		}
-		else
-		{
-			for (TPair<FRigControlElement*,FControlRigInteractionScope*>& Scope : InteractionScopes)
+			for (TPair <UObject*, FSequencerProxyItem>& SItems : SequencerItems)
 			{
-				if (Scope.Value)
+				//we do this backwards so ValueChanged later is set up correctly since that iterates in the other direction
+				for (int32 Index = SItems.Value.Bindings.Num() - 1; Index >= 0; --Index)
 				{
-					delete Scope.Value; 
+					FBindingAndTrack& Binding = SItems.Value.Bindings[Index];
+					SetBindingValueFromCurrent(SItems.Key, Binding.Binding, Context);
 				}
 			}
-			InteractionScopes.Reset();
 		}
 	}
 #endif
