Control Rig: Fx FKControlRig AnimNode not showing controls as inputs
#rnx

diff --git a/Engine/Plugins/Animation/ControlRig/Source/ControlRig/Private/AnimNode_ControlRig.cpp b/Engine/Plugins/Animation/ControlRig/Source/ControlRig/Private/AnimNode_ControlRig.cpp
--- a/Engine/Plugins/Animation/ControlRig/Source/ControlRig/Private/AnimNode_ControlRig.cpp
+++ b/Engine/Plugins/Animation/ControlRig/Source/ControlRig/Private/AnimNode_ControlRig.cpp
@@ -8,6 +8,7 @@
 #include "GameFramework/Actor.h"
 #include "Animation/NodeMappingContainer.h"
 #include "AnimationRuntime.h"
+#include "ControlRigObjectBinding.h"
 #include "Animation/AnimCurveUtils.h"
 #include "Animation/AnimStats.h"
 #include "Units/Execution/RigUnit_PrepareForExecution.h"
@@ -375,6 +376,8 @@ bool FAnimNode_ControlRig::UpdateControlRigIfNeeded(const UAnimInstance* InAnimI
 						
 						// create a new control rig using the new class
 						NewControlRig = NewObject<UControlRig>(InAnimInstance->GetOwningComponent(), ExpectedClass);
+						ControlRig->SetObjectBinding(MakeShared<FControlRigObjectBinding>());
+						ControlRig->GetObjectBinding()->BindToObject(InAnimInstance->GetOwningComponent());
 						NewControlRig->Initialize(true);
 						NewControlRig->RequestInit();
 
@@ -462,6 +465,8 @@ bool FAnimNode_ControlRig::UpdateControlRigIfNeeded(const UAnimInstance* InAnimI
 		if(ControlRig == nullptr)
 		{
 			ControlRig = NewObject<UControlRig>(InAnimInstance->GetOwningComponent(), ExpectedClass);
+			ControlRig->SetObjectBinding(MakeShared<FControlRigObjectBinding>());
+			ControlRig->GetObjectBinding()->BindToObject(InAnimInstance->GetOwningComponent());
 			ControlRig->Initialize(true);
 			ControlRig->RequestInit();
 		}
diff --git a/Engine/Plugins/Animation/ControlRig/Source/ControlRigDeveloper/Private/AnimGraphNode_ControlRig.cpp b/Engine/Plugins/Animation/ControlRig/Source/ControlRigDeveloper/Private/AnimGraphNode_ControlRig.cpp
--- a/Engine/Plugins/Animation/ControlRig/Source/ControlRigDeveloper/Private/AnimGraphNode_ControlRig.cpp
+++ b/Engine/Plugins/Animation/ControlRig/Source/ControlRigDeveloper/Private/AnimGraphNode_ControlRig.cpp
@@ -15,6 +15,7 @@
 #include "AnimationGraphSchema.h"
 #include "RigVMBlueprintGeneratedClass.h"
 #include "ControlRigBlueprint.h"
+#include "ControlRigObjectBinding.h"
 #include "Misc/DefaultValueHelper.h"
 
 #include UE_INLINE_GENERATED_CPP_BY_NAME(AnimGraphNode_ControlRig)
@@ -198,58 +199,45 @@ void UAnimGraphNode_ControlRig::CreateCustomPins(TArray<UEdGraphPin*>* OldPins)
 		}
 	}
 
-	if(const UClass* ControlRigClass = GetTargetClass())
+	TArray<FControlsInfo>& Controls = GetControls();
+	for (FControlsInfo& ControlInfo : Controls)
 	{
-		if(UControlRig* CDO = ControlRigClass->GetDefaultObject<UControlRig>())
+		const FName ControlName = ControlInfo.Name;
+		BeginExposableNames.Remove(ControlName);
+
+		if (CurrentlyExposedNames.Contains(ControlName))
 		{
-			if(const URigHierarchy* Hierarchy = CDO->GetHierarchy())
+			const FEdGraphPinType PinType = ControlInfo.PinType;
+			if (!PinType.PinCategory.IsValid())
 			{
-				Hierarchy->ForEach<FRigControlElement>([&](FRigControlElement* ControlElement) -> bool
-				{
-					if (Hierarchy->IsAnimatable(ControlElement))
-					{
-						const FName ControlName = ControlElement->GetFName();
-						BeginExposableNames.Remove(ControlName);
-
-						if (CurrentlyExposedNames.Contains(ControlName))
-						{
-							const FEdGraphPinType PinType = Hierarchy->GetControlPinType(ControlElement);
-							if (!PinType.PinCategory.IsValid())
-							{
-								return true;
-							}
+				continue;
+			}
 
-							UEdGraphPin* NewPin = CreatePin(EEdGraphPinDirection::EGPD_Input, PinType, ControlName);
-							NewPin->PinFriendlyName = FText::FromName(ControlElement->GetFName());
+			UEdGraphPin* NewPin = CreatePin(EEdGraphPinDirection::EGPD_Input, PinType, ControlName);
+			NewPin->PinFriendlyName = FText::FromName(ControlInfo.Name);
 
-							// Newly created pin does not have an auto generated default value, so we need to generate one here
-							// Missing the auto-gen default would cause UEdGraphPin::MovePersistentDataFromOldPin to override
-							// the actual default value with the empty auto-gen default, causing BP compiler to complain
-							// This is similar to how the following two functions create and initialize new pins, create first 
-							// and then set an auto-gen default
-							// FOptionalPinManager::CreateVisiblePins()
-							// FAnimBlueprintNodeOptionalPinManager::PostInitNewPin()
-							AnimGraphDefaultSchema->SetPinAutogeneratedDefaultValueBasedOnType(NewPin);
+			// Newly created pin does not have an auto generated default value, so we need to generate one here
+			// Missing the auto-gen default would cause UEdGraphPin::MovePersistentDataFromOldPin to override
+			// the actual default value with the empty auto-gen default, causing BP compiler to complain
+			// This is similar to how the following two functions create and initialize new pins, create first 
+			// and then set an auto-gen default
+			// FOptionalPinManager::CreateVisiblePins()
+			// FAnimBlueprintNodeOptionalPinManager::PostInitNewPin()
+			AnimGraphDefaultSchema->SetPinAutogeneratedDefaultValueBasedOnType(NewPin);
 
-							// We cant interrogate CDO here as we may be mid-compile, so we can only really
-							// reset to the autogenerated default.
-							AnimGraphDefaultSchema->ResetPinToAutogeneratedDefaultValue(NewPin, false);
+			// We cant interrogate CDO here as we may be mid-compile, so we can only really
+			// reset to the autogenerated default.
+			AnimGraphDefaultSchema->ResetPinToAutogeneratedDefaultValue(NewPin, false);
 
-							// Extract default values from the Target Control Rig if possible
-							const FString DefaultValue = Hierarchy->GetControlPinDefaultValue(ControlElement, true);
-							if(!DefaultValue.IsEmpty())
-							{
-								AnimGraphDefaultSchema->TrySetDefaultValue(*NewPin, DefaultValue);
-							}
-							
-							// sustain the current set of custom pins - we'll refrain from changing the node until post load is complete 
-							CustomizePinData(NewPin, ControlName, INDEX_NONE);
-						}
-					}
-					
-					return true;
-				});
+			// Extract default values from the Target Control Rig if possible
+			const FString DefaultValue = ControlInfo.DefaultValue;
+			if(!DefaultValue.IsEmpty())
+			{
+				AnimGraphDefaultSchema->TrySetDefaultValue(*NewPin, DefaultValue);
 			}
+			
+			// sustain the current set of custom pins - we'll refrain from changing the node until post load is complete 
+			CustomizePinData(NewPin, ControlName, INDEX_NONE);
 		}
 	}
 
@@ -368,23 +356,11 @@ void UAnimGraphNode_ControlRig::RebuildExposedProperties()
 	}
 
 	// also add all of the controls
-	if(const UClass* TargetClass = GetTargetClass())
+	TArray<FControlsInfo>& Controls = GetControls();
+	for (FControlsInfo& ControlInfo : Controls)
 	{
-		if(UControlRig* CDO = TargetClass->GetDefaultObject<UControlRig>())
-		{
-			if(const URigHierarchy* Hierarchy = CDO->GetHierarchy())
-			{
-				Hierarchy->ForEach<FRigControlElement>([&](const FRigControlElement* ControlElement) -> bool
-				{
-					if (Hierarchy->IsAnimatable(ControlElement))
-					{
-						const FName ControlName = ControlElement->GetFName();
-						CustomPinProperties.Add(MakeOptionalPin(ControlName));
-					}
-					return true;
-				});
-			}
-		}
+		const FName ControlName = ControlInfo.Name;
+		CustomPinProperties.Add(MakeOptionalPin(ControlName));
 	}
 }
 
@@ -396,19 +372,13 @@ bool UAnimGraphNode_ControlRig::IsInputProperty(const FName& PropertyName) const
 	return InputVariables.Contains(PropertyName) || !OutputVariables.Contains(PropertyName);
 }
 
-FRigControlElement* UAnimGraphNode_ControlRig::FindControlElement(const FName& InControlName) const
+UAnimGraphNode_ControlRig::FControlsInfo* UAnimGraphNode_ControlRig::FindControlElement(const FName& InControlName) const
 {
-	if(const UClass* ControlRigClass = GetTargetClass())
+	TArray<FControlsInfo>& Controls = GetControls();
+	return Controls.FindByPredicate([InControlName](const FControlsInfo& Info)
 	{
-		if(UControlRig* CDO = ControlRigClass->GetDefaultObject<UControlRig>())
-		{
-			if(const URigHierarchy* Hierarchy = CDO->GetHierarchy())
-			{
-				return (FRigControlElement*)Hierarchy->Find<FRigControlElement>(FRigElementKey(InControlName, ERigElementType::Control));
-			}
-		}
-	}
-	return nullptr;
+		return Info.Name == InControlName;
+	});
 }
 
 bool UAnimGraphNode_ControlRig::IsAvailableToMapToCurve(const FName& PropertyName, bool bInput) const
@@ -423,9 +393,9 @@ bool UAnimGraphNode_ControlRig::IsAvailableToMapToCurve(const FName& PropertyNam
 		return Variable->TypeName == TEXT("float");
 	}
 
-	if(const FRigControlElement* ControlElement = FindControlElement(PropertyName))
+	if(const FControlsInfo* ControlInfo = FindControlElement(PropertyName))
 	{
-		return (ControlElement->Settings.ControlType == ERigControlType::Float || ControlElement->Settings.ControlType == ERigControlType::ScaleFloat);
+		return (ControlInfo->ControlType == ERigControlType::Float || ControlInfo->ControlType == ERigControlType::ScaleFloat);
 	}
 
 	return ensure(false);
@@ -572,6 +542,54 @@ void UAnimGraphNode_ControlRig::CustomizeDetails(IDetailLayoutBuilder& DetailBui
 	// alpha property blending support END
 }
 
+TArray<UAnimGraphNode_ControlRig::FControlsInfo>& UAnimGraphNode_ControlRig::GetControls() const
+{
+	const UClass* ControlRigClass = GetTargetClass();
+
+	if (ControlsInfoClass != ControlRigClass)
+	{
+		ControlsInfo.Reset();
+		ControlsInfoClass = ControlRigClass;
+		
+		if (ControlRigClass)
+		{
+			if (USkeleton* Skeleton = GetAnimBlueprint()->TargetSkeleton)
+			{
+				UControlRig* TemplateRig = NewObject<UControlRig>(GetTransientPackage(), ControlRigClass, NAME_None, RF_Transient);
+				TemplateRig->SetObjectBinding(MakeShared<FControlRigObjectBinding>());
+				TemplateRig->GetObjectBinding()->BindToObject(Skeleton);
+				TemplateRig->GetDataSourceRegistry()->RegisterDataSource(UControlRig::OwnerComponent, TemplateRig->GetObjectBinding()->GetBoundObject());
+
+				TemplateRig->Initialize();
+				TemplateRig->SetBoneInitialTransformsFromRefSkeleton(Skeleton->GetReferenceSkeleton());
+				TemplateRig->Evaluate_AnyThread();
+
+				if(const URigHierarchy* Hierarchy = TemplateRig->GetHierarchy())
+				{
+					Hierarchy->ForEach<FRigControlElement>([&](FRigControlElement* ControlElement) -> bool
+					{
+						if (Hierarchy->IsAnimatable(ControlElement))
+						{
+							FControlsInfo Info;
+							Info.Name = ControlElement->GetFName();
+							Info.DisplayName = ControlElement->GetName();
+							Info.PinType = TemplateRig->GetHierarchy()->GetControlPinType(ControlElement);
+							Info.ControlType = ControlElement->Settings.ControlType;
+							Info.DefaultValue = TemplateRig->GetHierarchy()->GetControlPinDefaultValue(ControlElement, true);
+							ControlsInfo.Add(Info);
+						}
+						return true;
+					});
+				}
+
+				TemplateRig->MarkAsGarbage();
+			}
+		}
+	}
+
+	return ControlsInfo;
+}
+
 void UAnimGraphNode_ControlRig::GetVariables(bool bInput, TMap<FName, FRigVMExternalVariable>& OutVariables) const
 {
 	DECLARE_SCOPE_HIERARCHICAL_COUNTER_FUNC()
@@ -634,22 +652,10 @@ void UAnimGraphNode_ControlRig::GetAvailableMapping(const FName& PathName, TArra
 		TargetSkeleton->GetCurveMetaDataNames(OutArray);
 		
 		// also add all controls
-		if(const UClass* ControlRigClass = GetTargetClass())
+		TArray<FControlsInfo>& Controls = GetControls();
+		for(FControlsInfo& ControlInfo : Controls)
 		{
-			if(UControlRig* CDO = ControlRigClass->GetDefaultObject<UControlRig>())
-			{
-				if(const URigHierarchy* Hierarchy = CDO->GetHierarchy())
-				{
-					Hierarchy->ForEach<FRigControlElement>([&](FRigControlElement* ControlElement) -> bool
-					{
-						if (Hierarchy->IsAnimatable(ControlElement))
-						{
-							OutArray.Add(ControlElement->GetFName());
-						}
-						return true;
-					});
-				}
-			}
+			OutArray.Add(ControlInfo.Name);
 		}
 
 		// we want to exclude anything that has been mapped already
@@ -708,28 +714,16 @@ void UAnimGraphNode_ControlRig::CreateVariableMapping(const FString& FilteredTex
 	if(bInput)
 	{
 		// add all controls
-		if(const UClass* ControlRigClass = GetTargetClass())
+		TArray<FControlsInfo>& Controls = GetControls();
+		for (FControlsInfo& ControlInfo : Controls)
 		{
-			if(UControlRig* CDO = ControlRigClass->GetDefaultObject<UControlRig>())
-			{
-				if(const URigHierarchy* Hierarchy = CDO->GetHierarchy())
-				{
-					Hierarchy->ForEach<FRigControlElement>([&](FRigControlElement* ControlElement) -> bool
-					{
-						if (Hierarchy->IsAnimatable(ControlElement))
-						{
-							const FName ControlName = ControlElement->GetFName();
-							const FString& DisplayName = ControlElement->GetName();
+			const FName ControlName = ControlInfo.Name;
+			const FString& DisplayName = ControlInfo.DisplayName;
 
-							if (!bDoFiltering || DisplayName.Contains(FilteredText))
-							{
-								TSharedRef<FRigVMVariableMappingInfo> Item = FRigVMVariableMappingInfo::Make(ControlName);
-								OutArray.Add(Item);
-							}
-						}
-						return true;
-					});
-				}
+			if (!bDoFiltering || DisplayName.Contains(FilteredText))
+			{
+				TSharedRef<FRigVMVariableMappingInfo> Item = FRigVMVariableMappingInfo::Make(ControlName);
+				OutArray.Add(Item);
 			}
 		}
 	}
diff --git a/Engine/Plugins/Animation/ControlRig/Source/ControlRigDeveloper/Public/AnimGraphNode_ControlRig.h b/Engine/Plugins/Animation/ControlRig/Source/ControlRigDeveloper/Public/AnimGraphNode_ControlRig.h
--- a/Engine/Plugins/Animation/ControlRig/Source/ControlRigDeveloper/Public/AnimGraphNode_ControlRig.h
+++ b/Engine/Plugins/Animation/ControlRig/Source/ControlRigDeveloper/Public/AnimGraphNode_ControlRig.h
@@ -30,6 +30,18 @@ private:
 	virtual FAnimNode_CustomProperty* GetCustomPropertyNode() override { return &Node; }
 	virtual const FAnimNode_CustomProperty* GetCustomPropertyNode() const override { return &Node; }
 
+	struct FControlsInfo
+	{
+		FName Name;
+		FString DisplayName;
+		FEdGraphPinType PinType;
+		ERigControlType ControlType;
+		FString DefaultValue;
+	};
+	mutable const UClass* ControlsInfoClass = nullptr;
+	mutable TArray<FControlsInfo> ControlsInfo;
+	TArray<FControlsInfo>& GetControls() const;
+
 	// property related things
 	void GetVariables(bool bInput, TMap<FName, FRigVMExternalVariable>& OutParameters) const;
 
@@ -55,6 +67,6 @@ private:
 
 	bool IsAvailableToMapToCurve(const FName& PropertyName, bool bInput) const;
 	bool IsInputProperty(const FName& PropertyName) const;
-	FRigControlElement* FindControlElement(const FName& InControlName) const;
+	UAnimGraphNode_ControlRig::FControlsInfo* FindControlElement(const FName& InControlName) const;
 };
 
