Control Rig: Fix newly added modules not expanded automatically
#rnx

diff --git a/Engine/Plugins/Animation/ControlRig/Source/ControlRig/Private/ModularRigController.cpp b/Engine/Plugins/Animation/ControlRig/Source/ControlRig/Private/ModularRigController.cpp
--- a/Engine/Plugins/Animation/ControlRig/Source/ControlRig/Private/ModularRigController.cpp
+++ b/Engine/Plugins/Animation/ControlRig/Source/ControlRig/Private/ModularRigController.cpp
@@ -1329,7 +1329,7 @@ bool UModularRigController::IsNameAvailable(const FString& InParentModulePath, c
 
 	for (const FRigModuleReference* Child : *Children)
 	{
-		if (FRigName(Child->Name) == DesiredName)
+		if (FRigName(Child->Name).Equals(DesiredName, ESearchCase::IgnoreCase))
 		{
 			if(OutErrorMessage)
 			{
diff --git a/Engine/Plugins/Animation/ControlRig/Source/ControlRig/Public/Rigs/RigName.h b/Engine/Plugins/Animation/ControlRig/Source/ControlRig/Public/Rigs/RigName.h
--- a/Engine/Plugins/Animation/ControlRig/Source/ControlRig/Public/Rigs/RigName.h
+++ b/Engine/Plugins/Animation/ControlRig/Source/ControlRig/Public/Rigs/RigName.h
@@ -61,25 +61,30 @@ public:
 	}
 
 	bool operator ==(const FRigName& InOther) const
+	{
+		return Equals(InOther, ESearchCase::CaseSensitive);
+	}
+
+	bool Equals(const FRigName& InOther, const ESearchCase::Type& InCase) const
 	{
 		if(Name.IsSet() && InOther.Name.IsSet())
 		{
-			return GetFName().IsEqual(InOther.GetFName(), ENameCase::CaseSensitive);
+			return GetFName().IsEqual(InOther.GetFName(), InCase == ESearchCase::CaseSensitive ? ENameCase::CaseSensitive : ENameCase::IgnoreCase);
 		}
 		
 		if(NameString.IsSet() && InOther.NameString.IsSet())
 		{
-			return GetName().Equals(InOther.GetName(), ESearchCase::CaseSensitive);
+			return GetName().Equals(InOther.GetName(), InCase);
 		}
 
 		if(Name.IsSet() && InOther.NameString.IsSet())
 		{
-			return GetName().Equals(InOther.GetName(), ESearchCase::CaseSensitive);
+			return GetName().Equals(InOther.GetName(), InCase);
 		}
 
 		if(NameString.IsSet() && InOther.Name.IsSet())
 		{
-			return GetName().Equals(InOther.GetName(), ESearchCase::CaseSensitive);
+			return GetName().Equals(InOther.GetName(), InCase);
 		}
 
 		return IsNone() == InOther.IsNone();
diff --git a/Engine/Plugins/Animation/ControlRig/Source/ControlRigEditor/Private/Editor/SModularRigTreeView.cpp b/Engine/Plugins/Animation/ControlRig/Source/ControlRigEditor/Private/Editor/SModularRigTreeView.cpp
--- a/Engine/Plugins/Animation/ControlRig/Source/ControlRigEditor/Private/Editor/SModularRigTreeView.cpp
+++ b/Engine/Plugins/Animation/ControlRig/Source/ControlRigEditor/Private/Editor/SModularRigTreeView.cpp
@@ -717,6 +717,8 @@ bool SModularRigTreeView::AddElement(FString InKey, FString InParentKey)
 			RootElements.Add(NewItem);
 		}
 
+		SetItemExpansion(NewItem, true);
+
 		if (const UModularRig* ModularRig = Delegates.GetModularRig())
 		{
 			if (const FRigModuleInstance* Module = ModularRig->FindModule(NewItem->ModulePath))
