Control Rig: Fix RigHierarchy filter text not working

diff --git a/Engine/Plugins/Animation/ControlRig/Source/ControlRigEditor/Private/Editor/SRigHierarchyTreeView.cpp b/Engine/Plugins/Animation/ControlRig/Source/ControlRigEditor/Private/Editor/SRigHierarchyTreeView.cpp
--- a/Engine/Plugins/Animation/ControlRig/Source/ControlRigEditor/Private/Editor/SRigHierarchyTreeView.cpp
+++ b/Engine/Plugins/Animation/ControlRig/Source/ControlRigEditor/Private/Editor/SRigHierarchyTreeView.cpp
@@ -409,7 +409,7 @@ bool SRigHierarchyTreeView::AddElement(FRigElementKey InKey, FRigElementKey InPa
 	bool bAnyFilteredOut = Delegates.OnRigTreeIsItemVisible.IsBound();
 	if (!bAnyFilteredOut)
 	{
-		bAnyFilteredOut = FilteredString.IsEmpty() || !InKey.IsValid();
+		bAnyFilteredOut = !FilteredString.IsEmpty() && InKey.IsValid();
 	}
 
 	if (!bAnyFilteredOut)
