Constraints: enable setting properties when creating a new constraint

diff --git a/Engine/Plugins/Animation/ControlRig/Source/ControlRigEditor/Private/EditMode/SControlRigEditModeTools.cpp b/Engine/Plugins/Animation/ControlRig/Source/ControlRigEditor/Private/EditMode/SControlRigEditModeTools.cpp
--- a/Engine/Plugins/Animation/ControlRig/Source/ControlRigEditor/Private/EditMode/SControlRigEditModeTools.cpp
+++ b/Engine/Plugins/Animation/ControlRig/Source/ControlRigEditor/Private/EditMode/SControlRigEditModeTools.cpp
@@ -37,6 +37,7 @@
 #include "Widgets/Input/SComboButton.h"
 #include "Framework/Notifications/NotificationManager.h"
 #include "ScopedTransaction.h"
+#include "TransformConstraint.h"
 #include "EditMode/ControlRigEditModeToolkit.h"
 #include "EditMode/SControlRigDetails.h"
 #include "Editor/Constraints/SConstraintsWidget.h"
@@ -1055,30 +1056,36 @@ FReply SControlRigEditModeTools::OnBakeControlsToNewSpaceButtonClicked()
 
 FReply SControlRigEditModeTools::HandleAddConstraintClicked()
 {
-	// magic number to auto expand the widget when creating a new constraint. We keep that number below a reasonable
-	// threshold to avoid automatically creating a large number of items (this can be style done by the user) 
-	static constexpr int32 NumAutoExpand = 20;
+	FMenuBuilder MenuBuilder(true, nullptr);
 
-	const TSharedPtr<SConstraintsCreationWidget> Widget =
-		SNew(SConstraintsCreationWidget)
-		.OnConstraintCreated_Lambda( [this]()
+	auto AddConstraintWidget = [&](ETransformConstraintType InConstraintType)
+	{
+		const TSharedRef<SConstraintMenuEntry> Entry =
+			SNew(SConstraintMenuEntry, InConstraintType)
+		.OnConstraintCreated_Lambda([this]()
 		{
-			const int32 NumItems = ConstraintsEditionWidget ? ConstraintsEditionWidget->RefreshConstraintList() : 0;
-			
+			// magic number to auto expand the widget when creating a new constraint. We keep that number below a reasonable
+			// threshold to avoid automatically creating a large number of items (this can be style done by the user) 
+			static constexpr int32 NumAutoExpand = 20;
+			const int32 NumItems = ConstraintsEditionWidget ? ConstraintsEditionWidget->RefreshConstraintList() : 0;	
 			if (ConstraintPickerExpander && NumItems < NumAutoExpand)
 			{
 				ConstraintPickerExpander->SetExpanded(true);
 			}
 		});
+		MenuBuilder.AddWidget(Entry, FText::GetEmpty(), true);
+	};
 	
-	FMenuBuilder MenuBuilder(true, nullptr);
 	MenuBuilder.BeginSection("CreateConstraint", LOCTEXT("CreateConstraintHeader", "Create New..."));
 	{
-		MenuBuilder.AddWidget(Widget.ToSharedRef(), FText::GetEmpty(), true);
+		AddConstraintWidget(ETransformConstraintType::Translation);
+		AddConstraintWidget(ETransformConstraintType::Rotation);
+		AddConstraintWidget(ETransformConstraintType::Scale);
+		AddConstraintWidget(ETransformConstraintType::Parent);
+		AddConstraintWidget(ETransformConstraintType::LookAt);
 	}
 	MenuBuilder.EndSection();
 	
-	
 	FSlateApplication::Get().PushMenu(
 		AsShared(),
 		FWidgetPath(),
diff --git a/Engine/Plugins/Animation/ControlRig/Source/ControlRigEditor/Private/Editor/Constraints/SConstraintsWidget.cpp b/Engine/Plugins/Animation/ControlRig/Source/ControlRigEditor/Private/Editor/Constraints/SConstraintsWidget.cpp
--- a/Engine/Plugins/Animation/ControlRig/Source/ControlRigEditor/Private/Editor/Constraints/SConstraintsWidget.cpp
+++ b/Engine/Plugins/Animation/ControlRig/Source/ControlRigEditor/Private/Editor/Constraints/SConstraintsWidget.cpp
@@ -29,19 +29,12 @@
 #include "Styling/SlateIconFinder.h"
 #include "Widgets/Input/NumericTypeInterface.h"
 #include "FrameNumberDetailsCustomization.h"
+#include "SPrimaryButton.h"
 
 #define LOCTEXT_NAMESPACE "SConstraintsWidget"
 
 BEGIN_SLATE_FUNCTION_BUILD_OPTIMIZATION
 
-TArray< SConstraintsCreationWidget::ItemSharedPtr > SConstraintsCreationWidget::ListItems({
-	FDroppableConstraintItem::Make(ETransformConstraintType::Translation),
-	FDroppableConstraintItem::Make(ETransformConstraintType::Rotation),
-	FDroppableConstraintItem::Make(ETransformConstraintType::Scale),
-	FDroppableConstraintItem::Make(ETransformConstraintType::Parent),
-	FDroppableConstraintItem::Make(ETransformConstraintType::LookAt)
-});
-
 const TArray< const FSlateBrush* >& FConstraintInfo::GetBrushes()
 {
 	static const TArray< const FSlateBrush* > Brushes({
@@ -66,6 +59,17 @@ const TMap< UClass*, ETransformConstraintType >& FConstraintInfo::GetConstraintT
 	return ConstraintToType;
 }
 
+const TArray< UTickableTransformConstraint* >& FConstraintInfo::GetMutableDefaults()
+{
+	static TArray< UTickableTransformConstraint* > MutableDefaults({
+		GetMutableDefault<UTickableTranslationConstraint>(),
+		GetMutableDefault<UTickableRotationConstraint>(),
+		GetMutableDefault<UTickableScaleConstraint>(),
+		GetMutableDefault<UTickableParentConstraint>(),
+		GetMutableDefault<UTickableLookAtConstraint>()});
+	return MutableDefaults;
+}
+
 const FSlateBrush* FConstraintInfo::GetBrush(uint8 InType)
 {
 	static const UEnum* ETransformConstraintTypeEnum = StaticEnum<ETransformConstraintType>();
@@ -86,6 +90,17 @@ int8 FConstraintInfo::GetType(UClass* InClass)
 	return -1;
 }
 
+UTickableTransformConstraint* FConstraintInfo::GetMutable(ETransformConstraintType InType)
+{
+	static const UEnum* ETransformConstraintTypeEnum = StaticEnum<ETransformConstraintType>();
+	const uint8 ConstraintType = static_cast<uint8>(InType);
+	if (ETransformConstraintTypeEnum->IsValidEnumValue(ConstraintType))
+	{
+		return GetMutableDefaults()[ConstraintType];
+	}
+	return nullptr;
+}
+
 namespace
 {
 
@@ -135,117 +150,202 @@ static TWeakPtr<ISequencer> GetSequencerChecked()
  * SConstraintItem
  */
 
-void SDroppableConstraintItem::Construct(
+void SConstraintMenuEntry::Construct(
 		const FArguments& InArgs,
-		const TSharedPtr<const FDroppableConstraintItem>& InItem,
-		TSharedPtr<SConstraintsCreationWidget> InConstraintsWidget)
+		const ETransformConstraintType& InType)
 {
-	ConstraintItem = InItem;
-	ConstraintType = InItem->Type;
-	ConstraintsWidget = InConstraintsWidget;
+	ConstraintType = InType;
+	OnConstraintCreated = InArgs._OnConstraintCreated;
 	
-	const FButtonStyle& ButtonStyle = FAppStyle::GetWidgetStyle<FButtonStyle>( "PlacementBrowser.Asset" );
+	const FButtonStyle& ButtonStyle = FAppStyle::Get().GetWidgetStyle<FButtonStyle>( "Menu.Button" );
 
 	// enum to string
 	const UEnum* ETransformConstraintTypeEnum = StaticEnum<ETransformConstraintType>();
-	// const uint8 ConstraintType = static_cast<uint8>(InItem->Type);
 	const FString TypeStr = ETransformConstraintTypeEnum->GetNameStringByValue((uint8)ConstraintType);
 
 	// tooltip
 	const FString ToolTipStr = FString::Printf(TEXT("Create new %s constraint."), *TypeStr);
 	const TSharedPtr<IToolTip> ToolTip = FSlateApplicationBase::Get().MakeToolTip(FText::FromString(ToolTipStr));
-	
+
+	const float MenuIconSize = FAppStyle::Get().GetFloat("Menu.MenuIconSize");
+
 	ChildSlot
-	.Padding(FMargin(8.f, 2.f, 12.f, 2.f))
+	.HAlign(HAlign_Fill)
+	.VAlign(VAlign_Fill)
 	[
-		SNew(SOverlay)
-		+SOverlay::Slot()
+		SNew( SHorizontalBox )
+
+		+ SHorizontalBox::Slot()
 		[
 			SNew(SBorder)
-			.BorderImage( FAppStyle::Get().GetBrush("PlacementBrowser.Asset.Background"))
-			.Cursor( EMouseCursor::GrabHand )
-			.ToolTip(ToolTip)
-			.Padding(0)
+			.BorderImage_Lambda( [this, &ButtonStyle]
+			{
+				if (bIsPressed) { return &ButtonStyle.Pressed; }
+				if (IsHovered()) { return &ButtonStyle.Hovered; }
+				return &ButtonStyle.Normal;
+			})
+			.ForegroundColor(FLinearColor::White)
+			.ToolTip( ToolTip )
 			[
 				SNew( SHorizontalBox )
 
 				+ SHorizontalBox::Slot()
-				.Padding(8.0f, 4.f)
+				.Padding(14.0f, 0.f, 10.f, 0.0f)
 				.AutoWidth()
 				.HAlign(HAlign_Center)
 				.VAlign(VAlign_Center)
 				[
-					SNew( SBox )
-					.WidthOverride(40)
-					.HeightOverride(40)
+					SNew(SBox)
+					.WidthOverride(MenuIconSize)
+					.HeightOverride(MenuIconSize)
 					[
-					 	SNew(SImage)
-					 	.DesiredSizeOverride(FVector2D(16, 16))
+						SNew(SImage)
 						.Image(FConstraintInfo::GetBrush((uint8)ConstraintType))
-						.ColorAndOpacity(FSlateColor::UseForeground())
+						.ColorAndOpacity(FSlateColor::UseSubduedForeground())
 					]
 				]
 
 				+ SHorizontalBox::Slot()
-				.VAlign(VAlign_Fill)
-				.Padding(0)
+				.FillWidth(1.f)
+				.Padding(1.f, 0.f, 25.f, 0.f)
+				.VAlign(VAlign_Center)
+				.HAlign(HAlign_Left)
 				[
-					SNew(SBorder)
-					.BorderImage(FAppStyle::Get().GetBrush("PlacementBrowser.Asset.LabelBack"))
-					[
-						SNew(SHorizontalBox)
-						+SHorizontalBox::Slot()
-						.Padding(9, 0, 0, 1)
-						.VAlign(VAlign_Center)
-						[
-							SNew( STextBlock )
-							.TextStyle( FAppStyle::Get(), "PlacementBrowser.Asset.Name" )
-							.Text_Lambda( [TypeStr]
-							{
-								return FText::FromString(TypeStr);
-							} )
-						]
-					]
+					SNew( STextBlock )
+					.ColorAndOpacity(FSlateColor::UseForeground())
+					.Text_Lambda( [TypeStr] { return FText::FromString(TypeStr); } )
 				]
 			]
 		]
-
-		+SOverlay::Slot()
+		
+		+ SHorizontalBox::Slot()
+		.VAlign(VAlign_Center)
+		.HAlign(HAlign_Right)
+		.Padding(FMargin(0.f, 1.f, 0.f, 1.f))
+		.AutoWidth()
 		[
-			SNew(SBorder)
-			.BorderImage_Lambda( [this, &ButtonStyle]
-			{
-				if (bIsPressed)
-				{
-					return &ButtonStyle.Pressed;
-				}
-				
-				if (IsHovered())
-				{
-					return &ButtonStyle.Hovered;
-				}
-
-				return &ButtonStyle.Normal;
-			})
-			.Cursor( EMouseCursor::GrabHand )
-			.ToolTip( ToolTip )
+			SNew(SComboButton)
+			.ComboButtonStyle(FAppStyle::Get(), "ConstraintManager.ComboButton")
+			.MenuContent()
+			[
+				GenerateConstraintDefaultWidget()
+			]
 		]
 	];
 }
 
-FReply SDroppableConstraintItem::OnMouseButtonDown(const FGeometry& MyGeometry, const FPointerEvent& MouseEvent)
+FReply SConstraintMenuEntry::OnMouseButtonDown(const FGeometry& MyGeometry, const FPointerEvent& MouseEvent)
 {
 	if (MouseEvent.GetEffectingButton() == EKeys::LeftMouseButton)
 	{
 		bIsPressed = true;
-		// return FReply::Handled().DetectDrag( SharedThis( this ), MouseEvent.GetEffectingButton() );
-		return CreateSelectionPicker();
+		static constexpr bool bUseMutableDefault = false;
+		return CreateSelectionPicker(bUseMutableDefault);
 	}
 
 	return FReply::Unhandled();
 }
 
-FReply SDroppableConstraintItem::OnMouseButtonUp(const FGeometry& MyGeometry, const FPointerEvent& MouseEvent)
+TSharedRef<SWidget> SConstraintMenuEntry::GenerateConstraintDefaultWidget() const
+{
+	// static constexpr bool CloseAfterSelection = true;
+	// FMenuBuilder MenuBuilder(CloseAfterSelection, nullptr);
+
+	static FDetailsViewArgs DetailsViewArgs;
+	{
+		DetailsViewArgs.bAllowSearch = false;
+		DetailsViewArgs.bCustomFilterAreaLocation = true;
+		DetailsViewArgs.bCustomNameAreaLocation = true;
+		DetailsViewArgs.bHideSelectionTip = true;
+		DetailsViewArgs.bLockable = false;
+		DetailsViewArgs.bSearchInitialKeyFocus = true;
+		DetailsViewArgs.bUpdatesFromSelection = false;
+		DetailsViewArgs.bShowOptions = false;
+		DetailsViewArgs.bShowModifiedPropertiesOption = false;
+		DetailsViewArgs.ColumnWidth = 0.45f;
+	}
+	
+	TSharedRef<IDetailsView> DetailsView = FModuleManager::GetModuleChecked<FPropertyEditorModule>("PropertyEditor").CreateDetailView(DetailsViewArgs);
+
+	UTickableTransformConstraint* Constraint = FConstraintInfo::GetMutable(ConstraintType);
+	
+	// manage properties visibility
+	const bool bIsLookAtConstraint = Constraint->GetClass() == UTickableLookAtConstraint::StaticClass();
+	const auto PropertyVisibility = FIsPropertyVisible::CreateLambda(
+		[bIsLookAtConstraint](const FPropertyAndParent& InPropertyAndParent)
+		{
+			const FName PropertyName = InPropertyAndParent.Property.GetFName();
+			
+			// hide active property
+			static const FName ActivePropName = GET_MEMBER_NAME_CHECKED(UTickableTransformConstraint, Active);
+			if (PropertyName == ActivePropName)
+			{
+				return false;
+			}
+
+			// hide offset properties for look at constraints
+			if (bIsLookAtConstraint)
+			{
+				static const FName MaintainOffsetPropName = GET_MEMBER_NAME_CHECKED(UTickableTransformConstraint, bMaintainOffset);
+				static const FName DynamicOffsetPropName = GET_MEMBER_NAME_CHECKED(UTickableTransformConstraint, bDynamicOffset);
+				if (PropertyName == MaintainOffsetPropName || PropertyName == DynamicOffsetPropName)
+				{
+					return false;
+				}
+			}
+
+			return true;
+		});
+	DetailsView->SetIsPropertyVisibleDelegate(PropertyVisibility);
+		
+	DetailsView->SetObject(Constraint);
+
+	return SNew(SBorder)
+		.Visibility(EVisibility::Visible)
+		.BorderImage(FAppStyle::GetBrush("Menu.Background"))
+		[
+			SNew( SVerticalBox )
+
+			+ SVerticalBox::Slot()
+			.AutoHeight()
+			.VAlign( VAlign_Center )
+			[
+				DetailsView
+			]
+			
+			+ SVerticalBox::Slot()
+			.AutoHeight()
+			.Padding( 0.0f, 1.0f )
+			.VAlign( VAlign_Center )
+			[
+				SNew( SHorizontalBox )
+
+				+ SHorizontalBox::Slot()
+				.FillWidth(1.f)
+				.HAlign(HAlign_Fill)
+				[
+					SNew(SSpacer)
+				]
+		
+				+SHorizontalBox::Slot()
+				.Padding(16.f, 1.f)
+				.AutoWidth()
+				[
+					SNew(SButton)
+					.HAlign(HAlign_Center)
+					.VAlign(VAlign_Center)
+					.Text(LOCTEXT("CreateButtonLabel", "Create"))
+					.OnClicked_Lambda([this]()
+					{
+						static constexpr bool bUseMutableDefault = true;
+						return CreateSelectionPicker(bUseMutableDefault);
+					})
+				]
+			]
+		];
+}
+
+FReply SConstraintMenuEntry::OnMouseButtonUp(const FGeometry& MyGeometry, const FPointerEvent& MouseEvent)
 {
 	if ( MouseEvent.GetEffectingButton() == EKeys::LeftMouseButton )
 	{
@@ -255,45 +355,35 @@ FReply SDroppableConstraintItem::OnMouseButtonUp(const FGeometry& MyGeometry, co
 	return FReply::Unhandled();
 }
 
-FReply SDroppableConstraintItem::OnDragDetected(const FGeometry& MyGeometry, const FPointerEvent& MouseEvent)
+FReply SConstraintMenuEntry::OnDragDetected(const FGeometry& MyGeometry, const FPointerEvent& MouseEvent)
 {
 	bIsPressed = false;
-	
-	// if (MouseEvent.IsMouseButtonDown( EKeys::LeftMouseButton ))
-	// {
-	// 	return CreateSelectionPicker();
-	// }
-	
 	return FReply::Handled();
 }
 
-FReply SDroppableConstraintItem::CreateSelectionPicker() const
+FReply SConstraintMenuEntry::CreateSelectionPicker(const bool bUseDefault) const
 {
 	// FIXME temp approach for selecting the parent
 	FSlateApplication::Get().DismissAllMenus();
 	
 	static const FActorPickerModeModule& ActorPickerMode = FModuleManager::Get().GetModuleChecked<FActorPickerModeModule>("ActorPickerMode");
 
-	TSharedPtr<SConstraintsCreationWidget> ConstraintsCreationWidget = this->ConstraintsWidget.Pin();
-	ETransformConstraintType ConstraintTypeCopy = this->ConstraintType;
 	ActorPickerMode.BeginActorPickingMode(
 		FOnGetAllowedClasses(), 
 		FOnShouldFilterActor(), 
-		FOnActorSelected::CreateLambda([ConstraintsCreationWidget, ConstraintTypeCopy](AActor* InActor)
+		FOnActorSelected::CreateLambda([CreationDelegate = OnConstraintCreated, Type = ConstraintType, bUseDefault](AActor* InActor)
 		{
-			const FOnConstraintCreated CreationDelegate = ConstraintsCreationWidget.IsValid() ?
-				ConstraintsCreationWidget->OnConstraintCreated : FOnConstraintCreated();
-			SDroppableConstraintItem::CreateConstraint(InActor, CreationDelegate, ConstraintTypeCopy);
+			CreateConstraint(InActor, CreationDelegate, Type, bUseDefault);
 		}) );
-
 	
 	return FReply::Handled();
 }
 
-void SDroppableConstraintItem::CreateConstraint(
+void SConstraintMenuEntry::CreateConstraint(
 	AActor* InParent,
 	FOnConstraintCreated InCreationDelegate,
-	const ETransformConstraintType InConstraintType)
+	const ETransformConstraintType InConstraintType,
+	const bool bUseDefault)
 {
 	if (!InParent)
 	{
@@ -323,7 +413,7 @@ void SDroppableConstraintItem::CreateConstraint(
 	const TWeakPtr<ISequencer> WeakSequencer = GetSequencerChecked();
 	
 	// create constraints
-	auto CreateConstraint = [InCreationDelegate, InConstraintType, WeakSequencer](
+	auto CreateConstraint = [InCreationDelegate, InConstraintType, WeakSequencer, bUseDefault](
 		const TArray<AActor*>& Selection, UObject* InParent, const FName& InSocketName)
 	{
 		UWorld* World = GetCurrentWorld();
@@ -339,12 +429,16 @@ void SDroppableConstraintItem::CreateConstraint(
 			{
 				FScopedTransaction Transaction(LOCTEXT("CreateConstraintKey", "Create Constraint Key"));
 				UTickableTransformConstraint* Constraint =
-					FTransformConstraintUtils::CreateAndAddFromObjects(World, InParent, InSocketName, Child, NAME_None, InConstraintType);
+					FTransformConstraintUtils::CreateAndAddFromObjects(World, InParent, InSocketName, Child, NAME_None, InConstraintType, true, bUseDefault);
 				if (Constraint)
 				{
 					bCreated = true;
 					if (WeakSequencer.IsValid())
 					{
+						if (bUseDefault)
+						{
+							Constraint->Evaluate();
+						}
 						FMovieSceneConstraintChannelHelper::SmartConstraintKey(WeakSequencer.Pin(), Constraint, TOptional<bool>(), TOptional<FFrameNumber>());
 					}
 					else
@@ -427,44 +521,6 @@ void SDroppableConstraintItem::CreateConstraint(
 	CreateMenu(MenuBuilder.MakeWidget(), SummonLocation);
 }
 
-/**
- * SConstraintCreationWidget
- */
-
-void SConstraintsCreationWidget::Construct(const FArguments& InArgs)
-{
-	OnConstraintCreated = InArgs._OnConstraintCreated;
-	
-	ChildSlot
-	[
-		SNew( SVerticalBox )
-
-		+ SVerticalBox::Slot()
-		.Padding(FMargin(0.0f, 3.f))
-		[
-			SNew(SOverlay)
-
-			+ SOverlay::Slot()
-			[
-				SAssignNew(ListView, ConstraintItemListView)
-				.SelectionMode(ESelectionMode::None)
-				.ListItemsSource( &ListItems )
-				.OnGenerateRow(this, &SConstraintsCreationWidget::OnGenerateWidgetForItem)
-			]
-		]
-	];
-}
-
-TSharedRef<ITableRow> SConstraintsCreationWidget::OnGenerateWidgetForItem(
-	ItemSharedPtr InItem,
-	const TSharedRef<STableViewBase>& OwnerTable)
-{
-	return SNew(STableRow<ItemSharedPtr>, OwnerTable)
-	[
-		SNew(SDroppableConstraintItem, InItem.ToSharedRef(), SharedThis(this))
-	];
-}
-
 /**
  * SEditableConstraintItem
  */
@@ -1248,6 +1304,33 @@ TSharedPtr<SWidget> SConstraintsEditionWidget::CreateContextMenu()
 			return true;
 		});
 	DetailsView->SetIsPropertyVisibleDelegate(PropertyVisibility);
+
+	const UConstraintsManager* Manager = UConstraintsManager::Find(World);
+	const bool bIsStatic = Manager ? Manager->IsStaticConstraint(Constraint) : false;
+	auto IsPropertyReadOnly = [bIsStatic](const FProperty* InProperty)
+	{
+		if (InProperty && InProperty->HasAnyPropertyFlags(CPF_DisableEditOnInstance))
+		{
+			return !bIsStatic;
+		}
+		return false;
+	};
+	
+	FIsPropertyReadOnly PropertyReadonly = FIsPropertyReadOnly::CreateLambda([IsPropertyReadOnly](const FPropertyAndParent& InPropertyAndParent)
+	{
+		if (IsPropertyReadOnly(&InPropertyAndParent.Property))
+		{
+			return true;
+		}
+
+		const bool bIsParentReadOnly = InPropertyAndParent.ParentProperties.ContainsByPredicate([IsPropertyReadOnly](const FProperty* Parent)
+		{
+			return IsPropertyReadOnly(Parent);
+		});
+		
+		return bIsParentReadOnly;
+	});
+	DetailsView->SetIsPropertyReadOnlyDelegate(PropertyReadonly);
 	
 	TArray<TWeakObjectPtr<UObject>> ConstrainsToEdit;
 	ConstrainsToEdit.Add(Constraint);
diff --git a/Engine/Plugins/Animation/ControlRig/Source/ControlRigEditor/Private/Editor/Constraints/SConstraintsWidget.h b/Engine/Plugins/Animation/ControlRig/Source/ControlRigEditor/Private/Editor/Constraints/SConstraintsWidget.h
--- a/Engine/Plugins/Animation/ControlRig/Source/ControlRigEditor/Private/Editor/Constraints/SConstraintsWidget.h
+++ b/Engine/Plugins/Animation/ControlRig/Source/ControlRigEditor/Private/Editor/Constraints/SConstraintsWidget.h
@@ -11,6 +11,7 @@
 #include "Misc/QualifiedFrameTime.h"
 #include "ConstraintsManager.h"
 
+class UTickableTransformConstraint;
 class AActor;
 class SConstraintsCreationWidget;
 class SConstraintsEditionWidget;
@@ -26,9 +27,11 @@ class FConstraintInfo
 public:
 	static const FSlateBrush* GetBrush(uint8 InType);
 	static int8 GetType(UClass* InClass);
+	static UTickableTransformConstraint* GetMutable(ETransformConstraintType InType);
 private:
 	static const TArray< const FSlateBrush* >& GetBrushes();
 	static const TMap< UClass*, ETransformConstraintType >& GetConstraintToType();
+	static const TArray< UTickableTransformConstraint* >& GetMutableDefaults();
 };
 
 /**
@@ -61,20 +64,22 @@ private:
 };
 
 /**
- * SDroppableConstraintItem
+ * SConstraintMenuEntry
  */
 
-class SDroppableConstraintItem : public SCompoundWidget
+class SConstraintMenuEntry : public SCompoundWidget
 {
 public:
-	SLATE_BEGIN_ARGS(SDroppableConstraintItem){}
+	SLATE_BEGIN_ARGS(SConstraintMenuEntry){}
+
+		SLATE_EVENT(FOnConstraintCreated, OnConstraintCreated)
+	
 	SLATE_END_ARGS()
 
 	/** Constructs this widget with InArgs and the actual tree item. */
 	void Construct(
 		const FArguments& InArgs,
-		const TSharedPtr<const FDroppableConstraintItem>& InItem,
-		TSharedPtr<SConstraintsCreationWidget> InConstraintsWidget);
+		const ETransformConstraintType& InType);
 
 	// SWidget interface
 	virtual FReply OnMouseButtonDown(const FGeometry& MyGeometry, const FPointerEvent& MouseEvent) override;
@@ -84,14 +89,21 @@ public:
 	
 private:
 
-	/** todo documentation */
-	FReply CreateSelectionPicker() const;
+	/** Activates the ActorPickerMode to pick a parent in the viewport. */
+	FReply CreateSelectionPicker(const bool bUseDefault) const;
+
+	/** Delegate triggered when a new constraint has been created. */
+	FOnConstraintCreated OnConstraintCreated;
+
+	/** Creates a widget to edit the constraint default properties. */
+	TSharedRef<SWidget> GenerateConstraintDefaultWidget() const;
 	
 	/** Creates the constraint between the current selection and the picked actor. */
 	static void CreateConstraint(
 		AActor* InParent,
 		FOnConstraintCreated InDelegate,
-		const ETransformConstraintType InConstraintType);
+		const ETransformConstraintType InConstraintType,
+		const bool bUseDefault);
 
 	/** TSharedPtr to the tree item. */
 	TSharedPtr<const FDroppableConstraintItem> ConstraintItem;
diff --git a/Engine/Source/Runtime/Experimental/Animation/Constraints/Private/TransformConstraint.cpp b/Engine/Source/Runtime/Experimental/Animation/Constraints/Private/TransformConstraint.cpp
--- a/Engine/Source/Runtime/Experimental/Animation/Constraints/Private/TransformConstraint.cpp
+++ b/Engine/Source/Runtime/Experimental/Animation/Constraints/Private/TransformConstraint.cpp
@@ -327,6 +327,11 @@ void UTickableTransformConstraint::PostEditChangeProperty(FPropertyChangedEvent&
 {
 	Super::PostEditChangeProperty(PropertyChangedEvent);
 
+	if (HasAnyFlags(RF_ClassDefaultObject))
+	{
+		return;
+	}
+	
 	const FName PropertyName = PropertyChangedEvent.GetPropertyName();
 	if (PropertyName == GET_MEMBER_NAME_CHECKED(UTickableTransformConstraint, bMaintainOffset))
 	{
@@ -728,6 +733,11 @@ void UTickableTranslationConstraint::PostEditChangeProperty(FPropertyChangedEven
 {
 	Super::PostEditChangeProperty(PropertyChangedEvent);
 
+	if (HasAnyFlags(RF_ClassDefaultObject))
+	{
+		return;
+	}
+
 	const FName PropertyName = PropertyChangedEvent.GetPropertyName();
 	if (PropertyName == GET_MEMBER_NAME_CHECKED(UTickableTransformConstraint, bDynamicOffset))
 	{
@@ -870,6 +880,11 @@ void UTickableRotationConstraint::PostEditChangeProperty(FPropertyChangedEvent&
 {
 	Super::PostEditChangeProperty(PropertyChangedEvent);
 
+	if (HasAnyFlags(RF_ClassDefaultObject))
+	{
+		return;
+	}
+	
 	const FName PropertyName = PropertyChangedEvent.GetPropertyName();
 	if (PropertyName == GET_MEMBER_NAME_CHECKED(UTickableTransformConstraint, bDynamicOffset))
 	{
@@ -1014,6 +1029,11 @@ void UTickableScaleConstraint::PostEditChangeProperty(FPropertyChangedEvent& Pro
 {
 	Super::PostEditChangeProperty(PropertyChangedEvent);
 
+	if (HasAnyFlags(RF_ClassDefaultObject))
+	{
+		return;
+	}
+	
 	const FName PropertyName = PropertyChangedEvent.GetPropertyName();
 	if (PropertyName == GET_MEMBER_NAME_CHECKED(UTickableTransformConstraint, bDynamicOffset))
 	{
@@ -1307,6 +1327,11 @@ void UTickableParentConstraint::PostEditChangeProperty(FPropertyChangedEvent& Pr
 {
 	Super::PostEditChangeProperty(PropertyChangedEvent);
 
+	if (HasAnyFlags(RF_ClassDefaultObject))
+	{
+		return;
+	}
+	
 	auto UpdateOffset = [&]()
 	{
 		FTransform ChildGlobalTransform = GetChildGlobalTransform();
@@ -1450,6 +1475,11 @@ void UTickableLookAtConstraint::PostEditChangeProperty(FPropertyChangedEvent& Pr
 {
 	Super::PostEditChangeProperty(PropertyChangedEvent);
 
+	if (HasAnyFlags(RF_ClassDefaultObject))
+	{
+		return;
+	}
+
 	const FName PropertyName = PropertyChangedEvent.GetMemberPropertyName();
 	if (PropertyName == GET_MEMBER_NAME_CHECKED(UTickableLookAtConstraint, Axis))
 	{
@@ -1532,7 +1562,8 @@ void FTransformConstraintUtils::GetParentConstraints(
 
 UTickableTransformConstraint* FTransformConstraintUtils::CreateFromType(
 	UWorld* InWorld,
-	const ETransformConstraintType InType)
+	const ETransformConstraintType InType,
+	const bool bUseDefault)
 {
 	if (!InWorld)
 	{
@@ -1557,19 +1588,19 @@ UTickableTransformConstraint* FTransformConstraintUtils::CreateFromType(
 	switch (InType)
 	{
 	case ETransformConstraintType::Translation:
-		Constraint = Controller.AllocateConstraintT<UTickableTranslationConstraint>(BaseName);
+		Constraint = Controller.AllocateConstraintT<UTickableTranslationConstraint>(BaseName, bUseDefault);
 		break;
 	case ETransformConstraintType::Rotation:
-		Constraint = Controller.AllocateConstraintT<UTickableRotationConstraint>(BaseName);
+		Constraint = Controller.AllocateConstraintT<UTickableRotationConstraint>(BaseName, bUseDefault);
 		break;
 	case ETransformConstraintType::Scale:
-		Constraint = Controller.AllocateConstraintT<UTickableScaleConstraint>(BaseName);
+		Constraint = Controller.AllocateConstraintT<UTickableScaleConstraint>(BaseName, bUseDefault);
 		break;
 	case ETransformConstraintType::Parent:
-		Constraint = Controller.AllocateConstraintT<UTickableParentConstraint>(BaseName);
+		Constraint = Controller.AllocateConstraintT<UTickableParentConstraint>(BaseName, bUseDefault);
 		break;
 	case ETransformConstraintType::LookAt:
-		Constraint = Controller.AllocateConstraintT<UTickableLookAtConstraint>(BaseName);
+		Constraint = Controller.AllocateConstraintT<UTickableLookAtConstraint>(BaseName, bUseDefault);
 		break;
 	default:
 		ensure(false);
@@ -1668,7 +1699,8 @@ UTickableTransformConstraint* FTransformConstraintUtils::CreateAndAddFromObjects
 	UObject* InParent, const FName& InParentSocketName,
 	UObject* InChild, const FName& InChildSocketName,
 	const ETransformConstraintType InType,
-	const bool bMaintainOffset)
+	const bool bMaintainOffset,
+	const bool bUseDefault)
 {
 	static const TCHAR* ErrorPrefix = TEXT("FTransformConstraintUtils::CreateAndAddFromActors");
 	
@@ -1699,10 +1731,10 @@ UTickableTransformConstraint* FTransformConstraintUtils::CreateAndAddFromObjects
 		return nullptr;
 	}
 	
-	UTickableTransformConstraint* Constraint = FTransformConstraintUtils::CreateFromType(InWorld, InType);
+	UTickableTransformConstraint* Constraint = CreateFromType(InWorld, InType, bUseDefault);
 	if (Constraint && (ParentHandle->IsValid() && ChildHandle->IsValid()))
 	{
-		if (AddConstraint(InWorld, ParentHandle, ChildHandle, Constraint, bMaintainOffset) == false)
+		if (AddConstraint(InWorld, ParentHandle, ChildHandle, Constraint, bMaintainOffset, bUseDefault) == false)
 		{
 			Constraint->MarkAsGarbage();
 			Constraint = nullptr;
@@ -1716,7 +1748,8 @@ bool FTransformConstraintUtils::AddConstraint(
 	UTransformableHandle* InParentHandle,
 	UTransformableHandle* InChildHandle,
 	UTickableTransformConstraint* Constraint,
-	const bool bMaintainOffset)
+	const bool bMaintainOffset,
+	const bool bUseDefault)
 {
 	const bool bIsValidParent = InParentHandle && InParentHandle->IsValid();
 	const bool bIsValidChild = InChildHandle && InChildHandle->IsValid();
@@ -1745,8 +1778,11 @@ bool FTransformConstraintUtils::AddConstraint(
 
 	Constraint->ParentTRSHandle = InParentHandle;
 	Constraint->ChildTRSHandle = InChildHandle;
-	Constraint->bMaintainOffset = bMaintainOffset;
-	Constraint->Setup();
+	if (!bUseDefault)
+	{
+		Constraint->bMaintainOffset = bMaintainOffset;
+		Constraint->Setup();
+	}
 	Constraint->InitConstraint(InWorld);
 	// add dependencies with the last child constraint
 	const FName NewConstraintName = Constraint->GetFName();
@@ -1793,7 +1829,12 @@ FTransform FTransformConstraintUtils::ComputeRelativeTransform(
 	case ETransformConstraintType::Translation:
 		{
 			FTransform RelativeTransform = InChildLocal;
-			RelativeTransform.SetLocation(InChildWorld.GetLocation() - InSpaceWorld.GetLocation());
+			FVector RelativeTranslation = InChildWorld.GetLocation() - InSpaceWorld.GetLocation();
+			if (const UTickableTranslationConstraint* TranslationConstraint = Cast<UTickableTranslationConstraint>(InConstraint))
+			{
+				TranslationConstraint->AxisFilter.FilterVector(RelativeTranslation, InChildLocal.GetTranslation());
+			}
+			RelativeTransform.SetLocation(RelativeTranslation);
 			return RelativeTransform;
 		}
 	case ETransformConstraintType::Rotation:
@@ -1801,6 +1842,10 @@ FTransform FTransformConstraintUtils::ComputeRelativeTransform(
 			FTransform RelativeTransform = InChildLocal;
 			FQuat RelativeRotation = InSpaceWorld.GetRotation().Inverse() * InChildWorld.GetRotation();
 			RelativeRotation.Normalize();
+			if (const UTickableRotationConstraint* RotationConstraint = Cast<UTickableRotationConstraint>(InConstraint))
+			{
+				RotationConstraint->AxisFilter.FilterQuat(RelativeRotation, InChildLocal.GetRotation());
+			}
 			RelativeTransform.SetRotation(RelativeRotation);
 			return RelativeTransform;
 		}
@@ -1812,6 +1857,10 @@ FTransform FTransformConstraintUtils::ComputeRelativeTransform(
 			RelativeScale[0] = FMath::Abs(SpaceScale[0]) > KINDA_SMALL_NUMBER ? RelativeScale[0] / SpaceScale[0] : 0.f;
 			RelativeScale[1] = FMath::Abs(SpaceScale[1]) > KINDA_SMALL_NUMBER ? RelativeScale[1] / SpaceScale[1] : 0.f;
 			RelativeScale[2] = FMath::Abs(SpaceScale[2]) > KINDA_SMALL_NUMBER ? RelativeScale[2] / SpaceScale[2] : 0.f;
+			if (const UTickableScaleConstraint* ScaleConstraint = Cast<UTickableScaleConstraint>(InConstraint))
+			{
+				ScaleConstraint->AxisFilter.FilterVector(RelativeScale, InChildLocal.GetScale3D());
+			}
 			RelativeTransform.SetScale3D(RelativeScale);
 			return RelativeTransform;
 		}
@@ -1827,6 +1876,27 @@ FTransform FTransformConstraintUtils::ComputeRelativeTransform(
 			}
 
 			FTransform RelativeTransform = ChildTransform.GetRelativeTransform(InSpaceWorld);
+
+			if (ParentConstraint && !ParentConstraint->TransformFilter.TranslationFilter.HasNoEffect())
+			{
+				FVector RelativeLocation = RelativeTransform.GetLocation();
+				ParentConstraint->TransformFilter.TranslationFilter.FilterVector(RelativeLocation, InChildLocal.GetLocation());
+				RelativeTransform.SetLocation(RelativeLocation);
+			}
+
+			if (ParentConstraint && !ParentConstraint->TransformFilter.RotationFilter.HasNoEffect())
+			{
+				FQuat RelativeRotation = RelativeTransform.GetRotation();
+				ParentConstraint->TransformFilter.RotationFilter.FilterQuat(RelativeRotation, InChildLocal.GetRotation());
+				RelativeTransform.SetRotation(RelativeRotation);
+			}
+
+			if (ParentConstraint && !ParentConstraint->TransformFilter.ScaleFilter.HasNoEffect())
+			{
+				FVector RelativeScale = RelativeTransform.GetScale3D();
+				ParentConstraint->TransformFilter.ScaleFilter.FilterVector(RelativeScale, InChildLocal.GetScale3D());
+				RelativeTransform.SetScale3D(RelativeScale);
+			}
 			
 			if (!bScale)
 			{
diff --git a/Engine/Source/Runtime/Experimental/Animation/Constraints/Public/ConstraintsManager.h b/Engine/Source/Runtime/Experimental/Animation/Constraints/Public/ConstraintsManager.h
--- a/Engine/Source/Runtime/Experimental/Animation/Constraints/Public/ConstraintsManager.h
+++ b/Engine/Source/Runtime/Experimental/Animation/Constraints/Public/ConstraintsManager.h
@@ -241,7 +241,7 @@ public:
 
 	/** Allocates a new constraint with the constraints manager as the owner. */
 	template< typename TConstraint >
-	TConstraint* AllocateConstraintT(const FName& InBaseName) const;
+	TConstraint* AllocateConstraintT(const FName& InBaseName, const bool bUseDefault = false) const;
 
 	/** Add this constraint to the manager */
 	CONSTRAINTS_API bool AddConstraint(UTickableConstraint* InConstraint) const;
diff --git a/Engine/Source/Runtime/Experimental/Animation/Constraints/Public/ConstraintsManager.inl b/Engine/Source/Runtime/Experimental/Animation/Constraints/Public/ConstraintsManager.inl
--- a/Engine/Source/Runtime/Experimental/Animation/Constraints/Public/ConstraintsManager.inl
+++ b/Engine/Source/Runtime/Experimental/Animation/Constraints/Public/ConstraintsManager.inl
@@ -7,7 +7,7 @@
 #include "Algo/StableSort.h"
 
 template< typename TConstraint >
-TConstraint* FConstraintsManagerController::AllocateConstraintT(const FName& InBaseName) const
+TConstraint* FConstraintsManagerController::AllocateConstraintT(const FName& InBaseName, const bool bUseDefault) const
 {
 	UConstraintSubsystem* Subsystem =  UConstraintSubsystem::Get();
 	if (!Subsystem)
@@ -17,7 +17,7 @@ TConstraint* FConstraintsManagerController::AllocateConstraintT(const FName& InB
 	// unique name (we may want to use another approach here to manage uniqueness)
 	const FName Name = MakeUniqueObjectName(Subsystem, TConstraint::StaticClass(), InBaseName);
 
-	TConstraint* NewConstraint = NewObject<TConstraint>(Subsystem, Name, RF_Transactional);
+	TConstraint* NewConstraint = NewObject<TConstraint>(Subsystem, Name, RF_Transactional, bUseDefault ? GetMutableDefault<TConstraint>() : nullptr, bUseDefault);
 	NewConstraint->Modify();
 	return NewConstraint;
 }
diff --git a/Engine/Source/Runtime/Experimental/Animation/Constraints/Public/TransformConstraint.h b/Engine/Source/Runtime/Experimental/Animation/Constraints/Public/TransformConstraint.h
--- a/Engine/Source/Runtime/Experimental/Animation/Constraints/Public/TransformConstraint.h
+++ b/Engine/Source/Runtime/Experimental/Animation/Constraints/Public/TransformConstraint.h
@@ -222,7 +222,7 @@ public:
 	FVector OffsetTranslation = FVector::ZeroVector;
 
 	/** Defines which translation axis is constrained. */
-	UPROPERTY(BlueprintReadWrite, Category = "Axis Filter")
+	UPROPERTY(EditDefaultsOnly, BlueprintReadWrite, Category = "Axis Filter")
 	FFilterOptionPerAxis AxisFilter;
 
 #if WITH_EDITOR
@@ -270,7 +270,7 @@ public:
 	FQuat OffsetRotation = FQuat::Identity;
 
 	/** Defines which rotation axis is constrained. */
-	UPROPERTY(BlueprintReadWrite, Category = "Axis Filter")
+	UPROPERTY(EditDefaultsOnly, BlueprintReadWrite, Category = "Axis Filter")
 	FFilterOptionPerAxis AxisFilter;
 
 #if WITH_EDITOR
@@ -317,7 +317,7 @@ public:
 	FVector OffsetScale = FVector::OneVector;
 
 	/** Defines which scale axis is constrained. */
-	UPROPERTY(BlueprintReadWrite, Category = "Axis Filter")
+	UPROPERTY(EditDefaultsOnly, BlueprintReadWrite, Category = "Axis Filter")
 	FFilterOptionPerAxis AxisFilter;
 
 #if WITH_EDITOR
@@ -380,7 +380,7 @@ public:
 	bool bScaling = false;
 
 	/** Defines which translation/rotation/scale axis are constrained. */
-	UPROPERTY(BlueprintReadWrite, Category = "Axis Filter") 
+	UPROPERTY(EditDefaultsOnly, BlueprintReadWrite, Category = "Axis Filter") 
 	FTransformFilter TransformFilter;
 
 #if WITH_EDITOR
@@ -452,7 +452,8 @@ struct FTransformConstraintUtils
 	/** Creates a new transform constraint based on the InType. */
 	static CONSTRAINTS_API UTickableTransformConstraint* CreateFromType(
 		UWorld* InWorld,
-		const ETransformConstraintType InType);
+		const ETransformConstraintType InType,
+		const bool bUseDefault = false);
 
 	/** Creates respective handles and creates a new InType transform constraint. */	
 	static CONSTRAINTS_API UTickableTransformConstraint* CreateAndAddFromObjects(
@@ -460,7 +461,8 @@ struct FTransformConstraintUtils
 		UObject* InParent, const FName& InParentSocketName,
 		UObject* InChild, const FName& InChildSocketName,
 		const ETransformConstraintType InType,
-		const bool bMaintainOffset = true);
+		const bool bMaintainOffset = true,
+		const bool bUseDefault = false);
 
 	/** Registers a new transform constraint within the constraints manager. */	
 	static CONSTRAINTS_API bool AddConstraint(
@@ -468,7 +470,8 @@ struct FTransformConstraintUtils
 		UTransformableHandle* InParentHandle,
 		UTransformableHandle* InChildHandle,
 		UTickableTransformConstraint* Constraint,
-		const bool bMaintainOffset = true);
+		const bool bMaintainOffset = true,
+		const bool bUseDefault = false);
 
 	/** Computes the relative transform between both transform based on the constraint's InType. */
 	static CONSTRAINTS_API FTransform ComputeRelativeTransform(
diff --git a/Engine/Source/Runtime/SlateCore/Private/Styling/StarshipCoreStyle.cpp b/Engine/Source/Runtime/SlateCore/Private/Styling/StarshipCoreStyle.cpp
--- a/Engine/Source/Runtime/SlateCore/Private/Styling/StarshipCoreStyle.cpp
+++ b/Engine/Source/Runtime/SlateCore/Private/Styling/StarshipCoreStyle.cpp
@@ -438,8 +438,25 @@ TSharedRef<ISlateStyle> FStarshipCoreStyle::Create()
 		Style->Set("Icons.Toolbar.Settings", new IMAGE_BRUSH_SVG("Starship/Common/settings", Icon20x20));
 
 		// Constraint Manager Icons
-		Style->Set("Icons.ConstraintManager.LookAt", new IMAGE_BRUSH_SVG("Starship/Common/LookAt", Icon16x16));
-		Style->Set("Icons.ConstraintManager.ParentHierarchy", new IMAGE_BRUSH_SVG("Starship/Common/ParentHierarchy", Icon16x16));
+		{
+			Style->Set("Icons.ConstraintManager.LookAt", new IMAGE_BRUSH_SVG("Starship/Common/LookAt", Icon16x16));
+			Style->Set("Icons.ConstraintManager.ParentHierarchy", new IMAGE_BRUSH_SVG("Starship/Common/ParentHierarchy", Icon16x16));
+			
+			const FButtonStyle ConstraintOptionButton = FButtonStyle()
+				.SetNormal(FSlateRoundedBoxBrush(FStyleColors::Transparent, 4.f))
+				.SetHovered(FSlateRoundedBoxBrush(FStyleColors::Hover, 4.f))
+				.SetHoveredForeground(FLinearColor::White)
+				.SetPressed(FSlateRoundedBoxBrush(FStyleColors::Hover, 4.f))
+				.SetPressedForeground(FLinearColor::White)
+				.SetPressedPadding(FMargin(0.0, 1.0, 0.0, 0.0));
+			
+			FComboButtonStyle ConstraintComboButton = FComboButtonStyle()
+				.SetButtonStyle(ConstraintOptionButton)
+				.SetDownArrowImage(IMAGE_BRUSH_SVG("Starship/Common/ellipsis-vertical-narrow", FVector2f(6.f, 15.f)));
+				ConstraintComboButton.ButtonStyle = ConstraintOptionButton;
+			
+			Style->Set( "ConstraintManager.ComboButton", ConstraintComboButton);
+		}
 
 		Style->Set("Icons.Role", new IMAGE_BRUSH_SVG("Starship/Common/Role", Icon16x16));
 		Style->Set("Icons.Merge", new IMAGE_BRUSH_SVG("Starship/Common/Merge", Icon16x16));
