Sequencer: Anim Details: Random crash, don't use FRigControlElement cached pointers but instead just use their names and get a fresh pointer and use more weak pointers.

diff --git a/Engine/Plugins/Animation/ControlRig/Source/ControlRigEditor/Private/EditMode/AnimDetailsProxy.cpp b/Engine/Plugins/Animation/ControlRig/Source/ControlRigEditor/Private/EditMode/AnimDetailsProxy.cpp
--- a/Engine/Plugins/Animation/ControlRig/Source/ControlRigEditor/Private/EditMode/AnimDetailsProxy.cpp
+++ b/Engine/Plugins/Animation/ControlRig/Source/ControlRigEditor/Private/EditMode/AnimDetailsProxy.cpp
@@ -99,13 +99,13 @@ static EAnimDetailSelectionState CachePropertySelection(TWeakPtr<FCurveEditor>&
 		const TMap<FCurveModelID, TUniquePtr<FCurveModel>>& Curves = CurveEditor->GetCurves();
 		if (Proxy)
 		{
-			for (const TPair<UControlRig*, FControlRigProxyItem>& Items : Proxy->ControlRigItems)
+			for (const TPair<TWeakObjectPtr<UControlRig>, FControlRigProxyItem>& Items : Proxy->ControlRigItems)
 			{
 				if (UControlRig* ControlRig = Items.Value.ControlRig.Get())
 				{
-					for (FRigControlElement* ControlElement : Items.Value.ControlElements)
+					for (const FName& CName : Items.Value.ControlElements)
 					{
-						if (ControlElement)
+						if (FRigControlElement* ControlElement = Items.Value.GetControlElement(CName))
 						{
 							++TotalNum;
 							EControlRigContextChannelToKey ChannelToKey = Proxy->GetChannelToKeyFromPropertyName(PropertyName);
@@ -154,9 +154,9 @@ static EAnimDetailSelectionState CachePropertySelection(TWeakPtr<FCurveEditor>&
 					}
 				}
 			}
-			for (const TPair<UObject*, FSequencerProxyItem>& SItems : Proxy->SequencerItems)
+			for (const TPair<TWeakObjectPtr<UObject>, FSequencerProxyItem>& SItems : Proxy->SequencerItems)
 			{
-				if (UObject* Object = SItems.Key)
+				if (UObject* Object = SItems.Key.Get())
 				{
 					for (const FBindingAndTrack& Element : SItems.Value.Bindings)
 					{
@@ -217,28 +217,31 @@ static EAnimDetailSelectionState CachePropertySelection(TWeakPtr<FCurveEditor>&
 
 void UAnimDetailControlsKeyedProxy::SetKey(TSharedPtr<ISequencer>& Sequencer, const IPropertyHandle& KeyedPropertyHandle)
 {
-	for (const TPair<UControlRig*, FControlRigProxyItem>& Items : ControlRigItems)
+	for (const TPair<TWeakObjectPtr<UControlRig>, FControlRigProxyItem>& Items : ControlRigItems)
 	{
 		if (UControlRig* ControlRig = Items.Value.ControlRig.Get())
 		{
-			for (FRigControlElement* ControlElement : Items.Value.ControlElements)
+			for (const FName& CName : Items.Value.ControlElements)
 			{
-				if (ControlElement && ControlRig->GetHierarchy()->Contains(FRigElementKey(ControlElement->GetKey().Name, ERigElementType::Control)))
+				if (FRigControlElement* ControlElement = Items.Value.GetControlElement(CName))
 				{
-					FRigControlModifiedContext Context;
-					Context.SetKey = EControlRigSetKey::Always;
-					FName PropertyName = KeyedPropertyHandle.GetProperty()->GetFName();
-					Context.KeyMask = (uint32)GetChannelToKeyFromPropertyName(PropertyName);
-					SetControlRigElementValueFromCurrent(ControlRig, ControlElement, Context);
+					if (ControlRig->GetHierarchy()->Contains(FRigElementKey(ControlElement->GetKey().Name, ERigElementType::Control)))
+					{
+						FRigControlModifiedContext Context;
+						Context.SetKey = EControlRigSetKey::Always;
+						FName PropertyName = KeyedPropertyHandle.GetProperty()->GetFName();
+						Context.KeyMask = (uint32)GetChannelToKeyFromPropertyName(PropertyName);
+						SetControlRigElementValueFromCurrent(ControlRig, ControlElement, Context);
+					}
 				}
 			}
 		}
 	}
-	for (const TPair<UObject*, FSequencerProxyItem>& SItems : SequencerItems)
+	for (const TPair<TWeakObjectPtr<UObject>, FSequencerProxyItem>& SItems : SequencerItems)
 	{
-		if (SItems.Key)
+		if (SItems.Key.IsValid())
 		{
-			TArray<UObject*> ObjectsToKey = { SItems.Key };
+			TArray<UObject*> ObjectsToKey = { SItems.Key.Get()};
 			for (const FBindingAndTrack& Element : SItems.Value.Bindings)
 			{
 				FName PropertyName = KeyedPropertyHandle.GetProperty()->GetFName();
@@ -681,7 +684,7 @@ EPropertyKeyedStatus UAnimDetailControlsKeyedProxy::GetPropertyKeyedStatus(TShar
 	const TRange<FFrameNumber> FrameRange = TRange<FFrameNumber>(Sequencer->GetLocalTime().Time.FrameNumber);
 	FName PropertyName = PropertyHandle.GetProperty()->GetFName();
 	EControlRigContextChannelToKey ChannelToKey = GetChannelToKeyFromPropertyName(PropertyName);
-	for (const TPair<UControlRig*, FControlRigProxyItem>& Items : ControlRigItems)
+	for (const TPair<TWeakObjectPtr<UControlRig>, FControlRigProxyItem>& Items : ControlRigItems)
 	{
 		if (UControlRig* ControlRig = Items.Value.ControlRig.Get())
 		{
@@ -690,10 +693,13 @@ EPropertyKeyedStatus UAnimDetailControlsKeyedProxy::GetPropertyKeyedStatus(TShar
 			{
 				continue;
 			}
-			for (FRigControlElement* ControlElement : Items.Value.ControlElements)
+			for (const FName& CName : Items.Value.ControlElements)
 			{
-				EPropertyKeyedStatus NewKeyedStatus = GetKeyedStatusInTrack(ControlRig, ControlElement->GetKey().Name, Track, FrameRange, ChannelToKey);
-				KeyedStatus = FMath::Max(KeyedStatus, NewKeyedStatus);
+				if (FRigControlElement* ControlElement = Items.Value.GetControlElement(CName))
+				{
+					EPropertyKeyedStatus NewKeyedStatus = GetKeyedStatusInTrack(ControlRig, ControlElement->GetKey().Name, Track, FrameRange, ChannelToKey);
+					KeyedStatus = FMath::Max(KeyedStatus, NewKeyedStatus);
+				}
 			}
 		}
 	}
@@ -706,7 +712,7 @@ EPropertyKeyedStatus UAnimDetailControlsKeyedProxy::GetPropertyKeyedStatus(TShar
 	{
 		MaxNumIndices = 2;
 	}
-	for (const TPair<UObject*, FSequencerProxyItem>& Items : SequencerItems)
+	for (const TPair<TWeakObjectPtr<UObject>, FSequencerProxyItem>& Items : SequencerItems)
 	{
 		for (const FBindingAndTrack& Element : Items.Value.Bindings)
 		{
@@ -727,26 +733,32 @@ void UAnimDetailControlsKeyedProxy::PostEditUndo()
 {
 	FRigControlModifiedContext Context;
 	Context.SetKey = EControlRigSetKey::Never;
-	for (const TPair<UControlRig*, FControlRigProxyItem>& Items : ControlRigItems)
+	for (const TPair<TWeakObjectPtr<UControlRig>, FControlRigProxyItem>& Items : ControlRigItems)
 	{
 		if (UControlRig* ControlRig = Items.Value.ControlRig.Get())
 		{
-			for (FRigControlElement* ControlElement : Items.Value.ControlElements)
+			for (const FName& CName : Items.Value.ControlElements)
 			{
-				if (ControlElement && ControlRig->GetHierarchy()->Contains(FRigElementKey(ControlElement->GetKey().Name, ERigElementType::Control)))
+				if (FRigControlElement* ControlElement = Items.Value.GetControlElement(CName))
 				{
-					ControlRig->SelectControl(ControlElement->GetKey().Name, bSelected);
-					SetControlRigElementValueFromCurrent(ControlRig, ControlElement, Context);
+					if (ControlRig->GetHierarchy()->Contains(FRigElementKey(ControlElement->GetKey().Name, ERigElementType::Control)))
+					{
+						ControlRig->SelectControl(ControlElement->GetKey().Name, bSelected);
+						SetControlRigElementValueFromCurrent(ControlRig, ControlElement, Context);
+					}
 				}
 			}
 		}
 	}
-	for (TPair <UObject*, FSequencerProxyItem>& SItems : SequencerItems)
+	for (TPair <TWeakObjectPtr<UObject>, FSequencerProxyItem>& SItems : SequencerItems)
 	{
 		//we do this backwards so ValueChanged later is set up correctly since that iterates in the other direction
-		for (FBindingAndTrack& Binding: SItems.Value.Bindings)
+		if (SItems.Key.IsValid())
 		{
-			SetBindingValueFromCurrent(SItems.Key, Binding.Binding, Context);
+			for (FBindingAndTrack& Binding : SItems.Value.Bindings)
+			{
+				SetBindingValueFromCurrent(SItems.Key.Get(), Binding.Binding, Context);
+			}
 		}
 	}
 }
@@ -1031,34 +1043,41 @@ void UAnimDetailControlsProxyTransform::ValueChanged()
 	TOptional<FEulerTransform> LastEulerTransform;
 	ClearMultipleFlags();
 
-	for (const TPair<UControlRig*, FControlRigProxyItem>& Items : ControlRigItems)
+	for (const TPair<TWeakObjectPtr<UControlRig>, FControlRigProxyItem>& Items : ControlRigItems)
 	{
 		if (UControlRig* ControlRig = Items.Value.ControlRig.Get())
 		{
-			for (FRigControlElement* ControlElement : Items.Value.ControlElements)
+			for (const FName& CName : Items.Value.ControlElements)
 			{
-				FEulerTransform EulerTransform = GetCurrentValue(ControlRig, ControlElement);
-				if (LastEulerTransform.IsSet())
+				if (FRigControlElement* ControlElement = Items.Value.GetControlElement(CName))
 				{
-					const FEulerTransform LastVal = LastEulerTransform.GetValue();
-					SetMultipleFlags(LastVal, EulerTransform);
-				}
-				else
-				{
-					LastEulerTransform = EulerTransform;
+					FEulerTransform EulerTransform = GetCurrentValue(ControlRig, ControlElement);
+					if (LastEulerTransform.IsSet())
+					{
+						const FEulerTransform LastVal = LastEulerTransform.GetValue();
+						SetMultipleFlags(LastVal, EulerTransform);
+					}
+					else
+					{
+						LastEulerTransform = EulerTransform;
+					}
 				}
 			}
 		}
 	}
-	for (TPair <UObject*, FSequencerProxyItem>& SItems : SequencerItems)
+	for (TPair <TWeakObjectPtr<UObject>, FSequencerProxyItem>& SItems : SequencerItems)
 	{
+		if (SItems.Key.IsValid() == false)
+		{
+			continue;
+		}
 		//we do this backwards so ValueChanged later is set up correctly since that iterates in the other direction
 		for (int32 Index = SItems.Value.Bindings.Num() - 1; Index >= 0; --Index)
 		{
 			FBindingAndTrack& Binding = SItems.Value.Bindings[Index];
 			if (Binding.Binding.IsValid())
 			{
-				FEulerTransform EulerTransform = GetCurrentValue(SItems.Key, Binding.Binding);
+				FEulerTransform EulerTransform = GetCurrentValue(SItems.Key.Get(), Binding.Binding);
 				if (LastEulerTransform.IsSet())
 				{
 					const FEulerTransform LastVal = LastEulerTransform.GetValue();
@@ -1314,26 +1333,29 @@ void UAnimDetailControlsProxyLocation::ValueChanged()
 	TOptional<FVector3f> LastValue;
 	ClearMultipleFlags();
 
-	for (const TPair<UControlRig*, FControlRigProxyItem>& Items : ControlRigItems)
+	for (const TPair<TWeakObjectPtr<UControlRig>, FControlRigProxyItem>& Items : ControlRigItems)
 	{
 		if (UControlRig* ControlRig = Items.Value.ControlRig.Get())
 		{
-			for (FRigControlElement* ControlElement : Items.Value.ControlElements)
+			for (const FName& CName : Items.Value.ControlElements)
 			{
-				FVector3f Value = FVector3f::ZeroVector;
-				if (ControlElement->Settings.ControlType == ERigControlType::Position)
-				{
-					FRigControlValue ControlValue = ControlRig->GetControlValue(ControlElement, ERigControlValueType::Current);
-					Value = ControlValue.Get<FVector3f>();
-				}
-				if (LastValue.IsSet())
+				if (FRigControlElement* ControlElement = Items.Value.GetControlElement(CName))
 				{
-					const FVector3f LastVal = LastValue.GetValue();
-					SetMultipleFlags(LastVal, Value);
-				}
-				else
-				{
-					LastValue = Value;
+					FVector3f Value = FVector3f::ZeroVector;
+					if (ControlElement->Settings.ControlType == ERigControlType::Position)
+					{
+						FRigControlValue ControlValue = ControlRig->GetControlValue(ControlElement, ERigControlValueType::Current);
+						Value = ControlValue.Get<FVector3f>();
+					}
+					if (LastValue.IsSet())
+					{
+						const FVector3f LastVal = LastValue.GetValue();
+						SetMultipleFlags(LastVal, Value);
+					}
+					else
+					{
+						LastValue = Value;
+					}
 				}
 			}
 		}
@@ -1478,27 +1500,29 @@ void UAnimDetailControlsProxyRotation::ValueChanged()
 {
 	TOptional<FVector3f> LastValue;
 	ClearMultipleFlags();
-
-	for (const TPair<UControlRig*, FControlRigProxyItem>& Items : ControlRigItems)
+	for (const TPair<TWeakObjectPtr<UControlRig>, FControlRigProxyItem>& Items : ControlRigItems)
 	{
 		if (UControlRig* ControlRig = Items.Value.ControlRig.Get())
 		{
-			for (FRigControlElement* ControlElement : Items.Value.ControlElements)
+			for (const FName& CName : Items.Value.ControlElements)
 			{
-				FVector3f Value = FVector3f::ZeroVector;
-				if (ControlElement->Settings.ControlType == ERigControlType::Rotator)
-				{
-					FRigControlValue ControlValue = ControlRig->GetControlValue(ControlElement, ERigControlValueType::Current);
-					Value = ControlValue.Get<FVector3f>();
-				}
-				if (LastValue.IsSet())
+				if (FRigControlElement* ControlElement = Items.Value.GetControlElement(CName))
 				{
-					const FVector3f LastVal = LastValue.GetValue();
-					SetMultipleFlags(LastVal, Value);
-				}
-				else
-				{
-					LastValue = Value;
+					FVector3f Value = FVector3f::ZeroVector;
+					if (ControlElement->Settings.ControlType == ERigControlType::Rotator)
+					{
+						FRigControlValue ControlValue = ControlRig->GetControlValue(ControlElement, ERigControlValueType::Current);
+						Value = ControlValue.Get<FVector3f>();
+					}
+					if (LastValue.IsSet())
+					{
+						const FVector3f LastVal = LastValue.GetValue();
+						SetMultipleFlags(LastVal, Value);
+					}
+					else
+					{
+						LastValue = Value;
+					}
 				}
 			}
 		}
@@ -1643,26 +1667,29 @@ void UAnimDetailControlsProxyScale::ValueChanged()
 	TOptional<FVector3f> LastValue;
 	ClearMultipleFlags();
 
-	for (const TPair<UControlRig*, FControlRigProxyItem>& Items : ControlRigItems)
+	for (const TPair<TWeakObjectPtr<UControlRig>, FControlRigProxyItem>& Items : ControlRigItems)
 	{
 		if (UControlRig* ControlRig = Items.Value.ControlRig.Get())
 		{
-			for (FRigControlElement* ControlElement : Items.Value.ControlElements)
+			for (const FName& CName : Items.Value.ControlElements)
 			{
-				FVector3f Value = FVector3f::ZeroVector;
-				if (ControlElement->Settings.ControlType == ERigControlType::Scale)
-				{
-					FRigControlValue ControlValue = ControlRig->GetControlValue(ControlElement, ERigControlValueType::Current);
-					Value = ControlValue.Get<FVector3f>();
-				}
-				if (LastValue.IsSet())
+				if (FRigControlElement* ControlElement = Items.Value.GetControlElement(CName))
 				{
-					const FVector3f LastVal = LastValue.GetValue();
-					SetMultipleFlags(LastVal, Value);
-				}
-				else
-				{
-					LastValue = Value;
+					FVector3f Value = FVector3f::ZeroVector;
+					if (ControlElement->Settings.ControlType == ERigControlType::Scale)
+					{
+						FRigControlValue ControlValue = ControlRig->GetControlValue(ControlElement, ERigControlValueType::Current);
+						Value = ControlValue.Get<FVector3f>();
+					}
+					if (LastValue.IsSet())
+					{
+						const FVector3f LastVal = LastValue.GetValue();
+						SetMultipleFlags(LastVal, Value);
+					}
+					else
+					{
+						LastValue = Value;
+					}
 				}
 			}
 		}
@@ -1796,28 +1823,31 @@ void UAnimDetailControlsProxyVector2D::ValueChanged()
 	TOptional<FVector2D> LastValue;
 	ClearMultipleFlags();
 
-	for (const TPair<UControlRig*, FControlRigProxyItem>& Items : ControlRigItems)
+	for (const TPair<TWeakObjectPtr<UControlRig>, FControlRigProxyItem>& Items : ControlRigItems)
 	{
 		if (UControlRig* ControlRig = Items.Value.ControlRig.Get())
 		{
-			for (FRigControlElement* ControlElement : Items.Value.ControlElements)
+			for (const FName& CName : Items.Value.ControlElements)
 			{
-				FVector2D Value = FVector2D::ZeroVector;
-				if (ControlElement->Settings.ControlType == ERigControlType::Vector2D)
-				{
-					FRigControlValue ControlValue = ControlRig->GetControlValue(ControlElement, ERigControlValueType::Current);
-					//FVector2D version deleted for some reason so need to convert
-					FVector3f Val = ControlValue.Get<FVector3f>();
-					Value = FVector2D(Val.X, Val.Y);
-				}
-				if (LastValue.IsSet())
+				if (FRigControlElement* ControlElement = Items.Value.GetControlElement(CName))
 				{
-					const FVector2D LastVal = LastValue.GetValue();
-					SetMultipleFlags(LastVal, Value);
-				}
-				else
-				{
-					LastValue = Value;
+					FVector2D Value = FVector2D::ZeroVector;
+					if (ControlElement->Settings.ControlType == ERigControlType::Vector2D)
+					{
+						FRigControlValue ControlValue = ControlRig->GetControlValue(ControlElement, ERigControlValueType::Current);
+						//FVector2D version deleted for some reason so need to convert
+						FVector3f Val = ControlValue.Get<FVector3f>();
+						Value = FVector2D(Val.X, Val.Y);
+					}
+					if (LastValue.IsSet())
+					{
+						const FVector2D LastVal = LastValue.GetValue();
+						SetMultipleFlags(LastVal, Value);
+					}
+					else
+					{
+						LastValue = Value;
+					}
 				}
 			}
 		}
@@ -1937,33 +1967,40 @@ void UAnimDetailControlsProxyFloat::ValueChanged()
 	TOptional<float> LastValue;
 	ClearMultipleFlags();
 
-	for (const TPair<UControlRig*, FControlRigProxyItem>& Items : ControlRigItems)
+	for (const TPair<TWeakObjectPtr<UControlRig>, FControlRigProxyItem>& Items : ControlRigItems)
 	{
 		if (UControlRig* ControlRig = Items.Value.ControlRig.Get())
 		{
-			for (FRigControlElement* ControlElement : Items.Value.ControlElements)
+			for (const FName& CName : Items.Value.ControlElements)
 			{
-				float Value = 0.0f;
-				if (ControlElement->Settings.ControlType == ERigControlType::Float ||
-					(ControlElement->Settings.ControlType == ERigControlType::ScaleFloat))
-				{
-					FRigControlValue ControlValue = ControlRig->GetControlValue(ControlElement, ERigControlValueType::Current);
-					Value = ControlValue.Get<float>();
-				}
-				if (LastValue.IsSet())
+				if (FRigControlElement* ControlElement = Items.Value.GetControlElement(CName))
 				{
-					const float LastVal = LastValue.GetValue();
-					SetMultipleFlags(LastVal, Value);
-				}
-				else
-				{
-					LastValue = Value;
+					float Value = 0.0f;
+					if (ControlElement->Settings.ControlType == ERigControlType::Float ||
+						(ControlElement->Settings.ControlType == ERigControlType::ScaleFloat))
+					{
+						FRigControlValue ControlValue = ControlRig->GetControlValue(ControlElement, ERigControlValueType::Current);
+						Value = ControlValue.Get<float>();
+					}
+					if (LastValue.IsSet())
+					{
+						const float LastVal = LastValue.GetValue();
+						SetMultipleFlags(LastVal, Value);
+					}
+					else
+					{
+						LastValue = Value;
+					}
 				}
 			}
 		}
 	}
-	for (TPair <UObject*, FSequencerProxyItem>& SItems : SequencerItems)
+	for (TPair <TWeakObjectPtr<UObject>, FSequencerProxyItem>& SItems : SequencerItems)
 	{
+		if (SItems.Key.IsValid() == false)
+		{
+			continue;
+		}
 		for (FBindingAndTrack& Binding: SItems.Value.Bindings)
 		{
 			if (Binding.Binding.IsValid())
@@ -2102,32 +2139,39 @@ void UAnimDetailControlsProxyBool::ValueChanged()
 	TOptional<bool> LastValue;
 	ClearMultipleFlags();
 
-	for (const TPair<UControlRig*, FControlRigProxyItem>& Items : ControlRigItems)
+	for (const TPair<TWeakObjectPtr<UControlRig>, FControlRigProxyItem>& Items : ControlRigItems)
 	{
 		if (UControlRig* ControlRig = Items.Value.ControlRig.Get())
 		{
-			for (FRigControlElement* ControlElement : Items.Value.ControlElements)
+			for (const FName& CName : Items.Value.ControlElements)
 			{
-				bool Value = false;
-				if (ControlElement->Settings.ControlType == ERigControlType::Bool)
-				{
-					FRigControlValue ControlValue = ControlRig->GetControlValue(ControlElement, ERigControlValueType::Current);
-					Value = ControlValue.Get<bool>();
-				}
-				if (LastValue.IsSet())
+				if (FRigControlElement* ControlElement = Items.Value.GetControlElement(CName))
 				{
-					const bool LastVal = LastValue.GetValue();
-					SetMultipleFlags(LastVal, Value);
-				}
-				else
-				{
-					LastValue = Value;
+					bool Value = false;
+					if (ControlElement->Settings.ControlType == ERigControlType::Bool)
+					{
+						FRigControlValue ControlValue = ControlRig->GetControlValue(ControlElement, ERigControlValueType::Current);
+						Value = ControlValue.Get<bool>();
+					}
+					if (LastValue.IsSet())
+					{
+						const bool LastVal = LastValue.GetValue();
+						SetMultipleFlags(LastVal, Value);
+					}
+					else
+					{
+						LastValue = Value;
+					}
 				}
 			}
 		}
 	}
-	for (TPair <UObject*, FSequencerProxyItem>& SItems : SequencerItems)
+	for (TPair <TWeakObjectPtr<UObject>, FSequencerProxyItem>& SItems : SequencerItems)
 	{
+		if (SItems.Key.IsValid() == false)
+		{
+			continue;
+		}
 		for (FBindingAndTrack& Binding : SItems.Value.Bindings)
 		{
 			if (Binding.Binding.IsValid())
@@ -2268,31 +2312,38 @@ void UAnimDetailControlsProxyInteger::ValueChanged()
 	TOptional<int64> LastValue;
 	ClearMultipleFlags();
 
-	for (const TPair<UControlRig*, FControlRigProxyItem>& Items : ControlRigItems)
+	for (const TPair<TWeakObjectPtr<UControlRig>, FControlRigProxyItem>& Items : ControlRigItems)
 	{
 		if (UControlRig* ControlRig = Items.Value.ControlRig.Get())
 		{
-			for (FRigControlElement* ControlElement : Items.Value.ControlElements)
+			for (const FName& CName : Items.Value.ControlElements)
 			{
-				int64 Value = 0;
-				if (ControlElement->Settings.ControlType == ERigControlType::Integer)
-				{
-					FRigControlValue ControlValue = ControlRig->GetControlValue(ControlElement, ERigControlValueType::Current);
-					Value = ControlValue.Get<int32>();
-				}
-				if (LastValue.IsSet())
-				{
-					const int64 LastVal = LastValue.GetValue();
-					SetMultipleFlags(LastVal, Value);
-				}
-				else
+				if (FRigControlElement* ControlElement = Items.Value.GetControlElement(CName))
 				{
-					LastValue = Value;
+					int64 Value = 0;
+					if (ControlElement->Settings.ControlType == ERigControlType::Integer)
+					{
+						FRigControlValue ControlValue = ControlRig->GetControlValue(ControlElement, ERigControlValueType::Current);
+						Value = ControlValue.Get<int32>();
+					}
+					if (LastValue.IsSet())
+					{
+						const int64 LastVal = LastValue.GetValue();
+						SetMultipleFlags(LastVal, Value);
+					}
+					else
+					{
+						LastValue = Value;
+					}
 				}
 			}
 		}
-		for (TPair <UObject*, FSequencerProxyItem>& SItems : SequencerItems)
+		for (TPair <TWeakObjectPtr<UObject>, FSequencerProxyItem>& SItems : SequencerItems)
 		{
+			if (SItems.Key.IsValid() == false)
+			{
+				continue;
+			}
 			for (FBindingAndTrack& Binding : SItems.Value.Bindings)
 			{
 				if (Binding.Binding.IsValid())
@@ -2431,29 +2482,32 @@ void UAnimDetailControlsProxyEnum::ValueChanged()
 	TOptional<int32> LastValue;
 	ClearMultipleFlags();
 	TObjectPtr<UEnum> EnumType = nullptr;
-	for (const TPair<UControlRig*, FControlRigProxyItem>& Items : ControlRigItems)
+	for (const TPair<TWeakObjectPtr<UControlRig>, FControlRigProxyItem>& Items : ControlRigItems)
 	{
 		if (UControlRig* ControlRig = Items.Value.ControlRig.Get())
 		{
-			for (FRigControlElement* ControlElement : Items.Value.ControlElements)
+			for (const FName& CName : Items.Value.ControlElements)
 			{
-				int32 Value = 0;
-				if (ControlElement->Settings.ControlType == ERigControlType::Integer &&
-					ControlElement->Settings.ControlEnum)
+				if (FRigControlElement* ControlElement = Items.Value.GetControlElement(CName))
 				{
-					EnumType = ControlElement->Settings.ControlEnum;
-					FRigControlValue ControlValue = ControlRig->GetControlValue(ControlElement, ERigControlValueType::Current);
-					Value = ControlValue.Get<int32>();
-				}
+					int32 Value = 0;
+					if (ControlElement->Settings.ControlType == ERigControlType::Integer &&
+						ControlElement->Settings.ControlEnum)
+					{
+						EnumType = ControlElement->Settings.ControlEnum;
+						FRigControlValue ControlValue = ControlRig->GetControlValue(ControlElement, ERigControlValueType::Current);
+						Value = ControlValue.Get<int32>();
+					}
 
-				if (LastValue.IsSet())
-				{
-					const int32 LastVal = LastValue.GetValue();
-					SetMultipleFlags(LastVal, Value);
-				}
-				else
-				{
-					LastValue = Value;
+					if (LastValue.IsSet())
+					{
+						const int32 LastVal = LastValue.GetValue();
+						SetMultipleFlags(LastVal, Value);
+					}
+					else
+					{
+						LastValue = Value;
+					}
 				}
 			}
 		}
diff --git a/Engine/Plugins/Animation/ControlRig/Source/ControlRigEditor/Private/EditMode/ControlRigControlsProxy.cpp b/Engine/Plugins/Animation/ControlRig/Source/ControlRigEditor/Private/EditMode/ControlRigControlsProxy.cpp
--- a/Engine/Plugins/Animation/ControlRig/Source/ControlRigEditor/Private/EditMode/ControlRigControlsProxy.cpp
+++ b/Engine/Plugins/Animation/ControlRig/Source/ControlRigEditor/Private/EditMode/ControlRigControlsProxy.cpp
@@ -35,6 +35,16 @@
 
 #include UE_INLINE_GENERATED_CPP_BY_NAME(ControlRigControlsProxy)
 
+FRigControlElement* FControlRigProxyItem::GetControlElement(const FName& InName) const
+{
+	FRigControlElement* Element = nullptr;
+	if (ControlRig.IsValid())
+	{
+		Element = ControlRig->FindControl(InName);
+	}
+	return Element;
+}
+
 void UControlRigControlsProxy::AddControlRigControl(UControlRig* InControlRig, const FName& InName)
 {
 	if (InControlRig == nullptr)
@@ -49,9 +59,9 @@ void UControlRigControlsProxy::AddControlRigControl(UControlRig* InControlRig, c
 
 	FControlRigProxyItem& Item = ControlRigItems.FindOrAdd(InControlRig);
 	Item.ControlRig = InControlRig;
-	if (Item.ControlElements.Contains(ControlElement) == false)
+	if (Item.ControlElements.Contains(InName) == false)
 	{
-		Item.ControlElements.Add(ControlElement);
+		Item.ControlElements.Add(InName);
 	}
 	if (ControlRigItems.Num() > 1 || Item.ControlElements.Num() > 1 || SequencerItems.Num() > 0)
 	{
@@ -68,11 +78,18 @@ void UControlRigControlsProxy::AddControlRigControl(UControlRig* InControlRig, c
 TArray<FRigControlElement*> UControlRigControlsProxy::GetControlElements() const
 {
 	TArray<FRigControlElement*> Elements;
-	for (const TPair<UControlRig*, FControlRigProxyItem>& Items : ControlRigItems)
+	for (const TPair<TWeakObjectPtr<UControlRig>, FControlRigProxyItem>& Items : ControlRigItems)
 	{
-		for (FRigControlElement* Element : Items.Value.ControlElements)
+		if (Items.Key.IsValid() == false)
 		{
-			Elements.Add(Element);
+			continue;
+		}
+		for (const FName& CName : Items.Value.ControlElements)
+		{
+			if (FRigControlElement* Element = Items.Value.GetControlElement(CName))
+			{
+				Elements.Add(Element);
+			}
 		}
 	}
 	return Elements;
@@ -156,8 +173,12 @@ void UControlRigControlsProxy::AddSequencerProxyItem(UObject* InObject, TWeakObj
 TArray<FBindingAndTrack> UControlRigControlsProxy::GetSequencerItems() const
 {
 	TArray<FBindingAndTrack> Elements;
-	for (const TPair<UObject*, FSequencerProxyItem>& Items : SequencerItems)
+	for (const TPair<TWeakObjectPtr<UObject>, FSequencerProxyItem>& Items : SequencerItems)
 	{
+		if (Items.Key.IsValid() == false)
+		{
+			continue;
+		}
 		for (const FBindingAndTrack& Element : Items.Value.Bindings)
 		{
 			Elements.Add(Element);
@@ -232,13 +253,13 @@ void UControlRigControlsProxy::PostEditChangeChainProperty(struct FPropertyChang
 					}
 				}
 			}
-			for (const TPair<UControlRig*, FControlRigProxyItem>& Items : ControlRigItems)
+			for (const TPair<TWeakObjectPtr<UControlRig>, FControlRigProxyItem>& Items : ControlRigItems)
 			{
 				if (UControlRig* ControlRig = Items.Value.ControlRig.Get())
 				{
-					for (FRigControlElement* ControlElement : Items.Value.ControlElements)
+					for (const FName& CName : Items.Value.ControlElements)
 					{
-						if (ControlElement)
+						if (FRigControlElement * ControlElement = Items.Value.GetControlElement(CName))
 						{
 							if (InteractionScopes.Contains(ControlElement) == false)
 							{
@@ -275,25 +296,31 @@ void UControlRigControlsProxy::PostEditChangeChainProperty(struct FPropertyChang
 			Context.SetKey = EControlRigSetKey::DoNotCare;
 			Context.KeyMask = (uint32)GetChannelToKeyFromPropertyName(Property->GetFName());
 
-			for (const TPair<UControlRig*, FControlRigProxyItem>& Items : ControlRigItems)
+			for (const TPair<TWeakObjectPtr<UControlRig>, FControlRigProxyItem>& Items : ControlRigItems)
 			{
 				if (UControlRig* ControlRig = Items.Value.ControlRig.Get())
 				{
 					//we do this backwards so ValueChanged later is set up correctly since that iterates in the other direction
 					for (int32 Index = Items.Value.ControlElements.Num() - 1; Index >= 0; --Index)
 					{
-						FRigControlElement* ControlElement = Items.Value.ControlElements[Index];
-						SetControlRigElementValueFromCurrent(ControlRig, ControlElement, Context);
+						if (FRigControlElement* ControlElement = Items.Value.GetControlElement(Items.Value.ControlElements[Index]))
+						{
+							SetControlRigElementValueFromCurrent(ControlRig, ControlElement, Context);
+						}
 					}
 				}
 			}
-			for (TPair <UObject*, FSequencerProxyItem>& SItems : SequencerItems)
+			for (TPair <TWeakObjectPtr<UObject>, FSequencerProxyItem>& SItems : SequencerItems)
 			{
+				if (SItems.Key.IsValid() == false)
+				{
+					continue;
+				}
 				//we do this backwards so ValueChanged later is set up correctly since that iterates in the other direction
 				for (int32 Index = SItems.Value.Bindings.Num() - 1; Index >= 0; --Index)
 				{
 					FBindingAndTrack& Binding = SItems.Value.Bindings[Index];
-					SetBindingValueFromCurrent(SItems.Key, Binding.Binding, Context);
+					SetBindingValueFromCurrent(SItems.Key.Get(), Binding.Binding, Context);
 				}
 			}
 		}
@@ -304,15 +331,18 @@ void UControlRigControlsProxy::PostEditChangeChainProperty(struct FPropertyChang
 #if WITH_EDITOR
 void UControlRigControlsProxy::PostEditUndo()
 {
-	for (const TPair<UControlRig*, FControlRigProxyItem>& Items : ControlRigItems)
+	for (const TPair<TWeakObjectPtr<UControlRig>, FControlRigProxyItem>& Items : ControlRigItems)
 	{
 		if (UControlRig* ControlRig = Items.Value.ControlRig.Get())
 		{
-			for (FRigControlElement* ControlElement : Items.Value.ControlElements)
+			for (const FName& CName : Items.Value.ControlElements)
 			{
-				if (ControlElement && ControlRig->GetHierarchy()->Contains(FRigElementKey(ControlElement->GetKey().Name, ERigElementType::Control)))
+				if (FRigControlElement* ControlElement = Items.Value.GetControlElement(CName))
 				{
-					ControlRig->SelectControl(ControlElement->GetKey().Name, bSelected);
+					if (ControlRig->GetHierarchy()->Contains(FRigElementKey(ControlElement->GetKey().Name, ERigElementType::Control)))
+					{
+						ControlRig->SelectControl(ControlElement->GetKey().Name, bSelected);
+					}
 				}
 			}
 		}
diff --git a/Engine/Plugins/Animation/ControlRig/Source/ControlRigEditor/Private/EditMode/ControlRigControlsProxy.h b/Engine/Plugins/Animation/ControlRig/Source/ControlRigEditor/Private/EditMode/ControlRigControlsProxy.h
--- a/Engine/Plugins/Animation/ControlRig/Source/ControlRigEditor/Private/EditMode/ControlRigControlsProxy.h
+++ b/Engine/Plugins/Animation/ControlRig/Source/ControlRigEditor/Private/EditMode/ControlRigControlsProxy.h
@@ -42,7 +42,8 @@ struct FAnimDetailVectorSelection
 struct FControlRigProxyItem
 {
 	TWeakObjectPtr<UControlRig> ControlRig;
-	TArray<FRigControlElement*> ControlElements;
+	TArray<FName> ControlElements;
+	FRigControlElement* GetControlElement(const FName& InName) const;
 };
 
 struct FBindingAndTrack
@@ -126,11 +127,11 @@ public:
 	//We can set/get values form multiple control rig elements but only one owns this.
 	TWeakObjectPtr<UControlRig> OwnerControlRig;
 	FRigControlElement* OwnerControlElement = nullptr;
-	TMap<UControlRig*, FControlRigProxyItem> ControlRigItems;
+	TMap<TWeakObjectPtr<UControlRig>, FControlRigProxyItem> ControlRigItems;
 	
 	TWeakObjectPtr<UObject> OwnerObject;
 	FBindingAndTrack OwnerBindingAndTrack;
-	TMap <UObject*, FSequencerProxyItem> SequencerItems;
+	TMap <TWeakObjectPtr<UObject>, FSequencerProxyItem> SequencerItems;
 
 	//list of child/animation channel proxies that we will customize
 	TArray<UControlRigControlsProxy*> ChildProxies;
diff --git a/Engine/Plugins/Animation/ControlRig/Source/ControlRigEditor/Private/EditMode/ControlRigEditMode.cpp b/Engine/Plugins/Animation/ControlRig/Source/ControlRigEditor/Private/EditMode/ControlRigEditMode.cpp
--- a/Engine/Plugins/Animation/ControlRig/Source/ControlRigEditor/Private/EditMode/ControlRigEditMode.cpp
+++ b/Engine/Plugins/Animation/ControlRig/Source/ControlRigEditor/Private/EditMode/ControlRigEditMode.cpp
@@ -5333,13 +5333,13 @@ bool FDetailKeyFrameCacheAndHandler::IsPropertyKeyable(const UClass* InObjectCla
 	{
 		if (UControlRigControlsProxy* Proxy = Cast< UControlRigControlsProxy>(OuterObjects[0]))
 		{
-			for (const TPair<UControlRig*, FControlRigProxyItem>& Items : Proxy->ControlRigItems)
+			for (const TPair<TWeakObjectPtr<UControlRig>, FControlRigProxyItem>& Items : Proxy->ControlRigItems)
 			{
 				if (UControlRig* ControlRig = Items.Value.ControlRig.Get())
 				{
-					for (FRigControlElement* ControlElement : Items.Value.ControlElements)
+					for (const FName& CName : Items.Value.ControlElements)
 					{
-						if (ControlElement)
+						if (FRigControlElement* ControlElement = Items.Value.GetControlElement(CName))
 						{
 							if (!ControlRig->GetHierarchy()->IsAnimatable(ControlElement))
 							{
diff --git a/Engine/Plugins/Animation/ControlRig/Source/ControlRigEditor/Private/EditMode/ControlsProxyDetailCustomization.cpp b/Engine/Plugins/Animation/ControlRig/Source/ControlRigEditor/Private/EditMode/ControlsProxyDetailCustomization.cpp
--- a/Engine/Plugins/Animation/ControlRig/Source/ControlRigEditor/Private/EditMode/ControlsProxyDetailCustomization.cpp
+++ b/Engine/Plugins/Animation/ControlRig/Source/ControlRigEditor/Private/EditMode/ControlsProxyDetailCustomization.cpp
@@ -552,13 +552,13 @@ void FAnimDetailValueCustomization::TogglePropertySelection(UControlRigControlsP
 
 	if (Proxy)
 	{
-		for (const TPair<UControlRig*, FControlRigProxyItem>& Items : Proxy->ControlRigItems)
+		for (const TPair<TWeakObjectPtr<UControlRig>, FControlRigProxyItem>& Items : Proxy->ControlRigItems)
 		{
 			if (UControlRig* ControlRig = Items.Value.ControlRig.Get())
 			{
-				for (FRigControlElement* ControlElement : Items.Value.ControlElements)
+				for (const FName& CName : Items.Value.ControlElements)
 				{
-					if (ControlElement)
+					if (FRigControlElement* ControlElement = Items.Value.GetControlElement(CName))
 					{
 						EControlRigContextChannelToKey ChannelToKey = Proxy->GetChannelToKeyFromPropertyName(PropertyName);
 						TParentFirstChildIterator<IOutlinerExtension> OutlinerExtenstionIt = OutlinerViewModel->GetRootItem()->GetDescendantsOfType<IOutlinerExtension>();
@@ -594,9 +594,9 @@ void FAnimDetailValueCustomization::TogglePropertySelection(UControlRigControlsP
 				}
 			}
 		}
-		for (const TPair<UObject*, FSequencerProxyItem>& SItems : Proxy->SequencerItems)
+		for (const TPair<TWeakObjectPtr<UObject>, FSequencerProxyItem>& SItems : Proxy->SequencerItems)
 		{
-			if (UObject* Object = SItems.Key)
+			if (UObject* Object = SItems.Key.Get())
 			{
 				for (const FBindingAndTrack& Element : SItems.Value.Bindings)
 				{
