Sequencer:  Change default behavior so CR Section to Keys are set based upon what section is selectied in the outliner

diff --git a/Engine/Plugins/Animation/ControlRig/Source/ControlRigEditor/Private/Sequencer/ControlRigParameterTrackEditor.cpp b/Engine/Plugins/Animation/ControlRig/Source/ControlRigEditor/Private/Sequencer/ControlRigParameterTrackEditor.cpp
--- a/Engine/Plugins/Animation/ControlRig/Source/ControlRigEditor/Private/Sequencer/ControlRigParameterTrackEditor.cpp
+++ b/Engine/Plugins/Animation/ControlRig/Source/ControlRigEditor/Private/Sequencer/ControlRigParameterTrackEditor.cpp
@@ -1,8 +1,15 @@
 // Copyright Epic Games, Inc. All Rights Reserved.
 
 #include "Sequencer/ControlRigParameterTrackEditor.h"
-#include "MVVM/Selection/Selection.h"
+#include "MVVM/Extensions/ITrackExtension.h"
+#include "MVVM/ViewModels/OutlinerViewModel.h"
+#include "MVVM/ViewModels/TrackModel.h"
+#include "MVVM/ViewModels/SectionModel.h"
+#include "MVVM/ViewModels/TrackRowModel.h"
 #include "MVVM/ViewModels/SequencerEditorViewModel.h"
+#include "MVVM/ViewModels/ViewModelIterators.h"
+#include "MVVM/ViewModels/OutlinerViewModel.h"
+#include "MVVM/Selection/Selection.h"
 #include "Animation/AnimMontage.h"
 #include "Sequencer/MovieSceneControlRigParameterTrack.h"
 #include "Sequencer/MovieSceneControlRigParameterSection.h"
@@ -93,7 +100,7 @@
 #include "Widgets/Layout/SSpacer.h"
 #include "ControlRigSequencerEditorLibrary.h"
 #include "LevelSequence.h"
-#include "MVVM/Extensions/ITrackExtension.h"
+
 
 #define LOCTEXT_NAMESPACE "FControlRigParameterTrackEditor"
 
@@ -104,7 +111,7 @@ TAutoConsoleVariable<bool> CVarAutoGenerateControlRigTrack(TEXT("ControlRig.Sequ
 
 TAutoConsoleVariable<bool> CVarSelectedKeysSelectControls(TEXT("ControlRig.Sequencer.SelectedKeysSelectControls"), false, TEXT("When true when we select a key in Sequencer it will select the Control, by default false."));
 
-TAutoConsoleVariable<bool> CVarSelectedSectionSetsSectionToKey(TEXT("ControlRig.Sequencer.SelectedSectionSetsSectionToKey"), false, TEXT("When true when we select a channel in a section, if it's the only section selected we set it as the Section To Key, by default false."));
+TAutoConsoleVariable<bool> CVarSelectedSectionSetsSectionToKey(TEXT("ControlRig.Sequencer.SelectedSectionSetsSectionToKey"), true, TEXT("When true when we select a channel in a section, if it's the only section selected we set it as the Section To Key, by default false."));
 
 TAutoConsoleVariable<bool> CVarEnableAdditiveControlRigs(TEXT("ControlRig.Sequencer.EnableAdditiveControlRigs"), true, TEXT("When true it is possible to add an additive control rig to a skeletal mesh component."));
 
@@ -2175,6 +2182,38 @@ void FControlRigParameterTrackEditor::OnSelectionChanged(TArray<UMovieSceneTrack
 				}
 			}
 		}
+		const bool bSelectedSectionSetsSectionToKey = CVarSelectedSectionSetsSectionToKey.GetValueOnGameThread();
+		if (bSelectedSectionSetsSectionToKey)
+		{
+			TMap<UMovieSceneTrack*, TSet<UMovieSceneControlRigParameterSection*>> TracksAndSections;
+			using namespace UE::Sequencer;
+			for (FViewModelPtr ViewModel : GetSequencer()->GetViewModel()->GetSelection()->Outliner)
+			{
+				if (TViewModelPtr<FTrackRowModel> TrackRowModel = ViewModel.ImplicitCast())
+				{
+					for (UMovieSceneSection* Section : TrackRowModel->GetSections())
+					{
+						if (UMovieSceneControlRigParameterSection* CRSection = Cast<UMovieSceneControlRigParameterSection>(Section))
+						{
+							if (UMovieSceneTrack* Track = Section->GetTypedOuter<UMovieSceneTrack>())
+							{
+								TracksAndSections.FindOrAdd(Track).Add(CRSection);
+							}
+						}
+					}
+				}
+			}
+
+			//if we have only one  selected section per track and the track has more than one section we set that to the section to key
+			for (TPair<UMovieSceneTrack*, TSet<UMovieSceneControlRigParameterSection*>>& TrackPair : TracksAndSections)
+			{
+				if (TrackPair.Key->GetAllSections().Num() > 0 && TrackPair.Value.Num() == 1)
+				{
+					TrackPair.Key->SetSectionToKey(TrackPair.Value.Array()[0]);
+				}
+			}
+
+		}
 		return;
 	}
 	SelectRigsAndControls(ControlRig, KeyAreas);
