[Backout] - CL30999475
Original CL Desc
-----------------------------------------------------------------
Control Rig: Fx FKControlRig AnimNode not showing controls as inputs
#rnx

diff --git a/Engine/Plugins/Animation/ControlRig/Source/ControlRig/Private/AnimNode_ControlRig.cpp b/Engine/Plugins/Animation/ControlRig/Source/ControlRig/Private/AnimNode_ControlRig.cpp
--- a/Engine/Plugins/Animation/ControlRig/Source/ControlRig/Private/AnimNode_ControlRig.cpp
+++ b/Engine/Plugins/Animation/ControlRig/Source/ControlRig/Private/AnimNode_ControlRig.cpp
@@ -8,7 +8,6 @@
 #include "GameFramework/Actor.h"
 #include "Animation/NodeMappingContainer.h"
 #include "AnimationRuntime.h"
-#include "ControlRigObjectBinding.h"
 #include "Animation/AnimCurveUtils.h"
 #include "Animation/AnimStats.h"
 #include "Units/Execution/RigUnit_PrepareForExecution.h"
@@ -376,8 +375,6 @@ bool FAnimNode_ControlRig::UpdateControlRigIfNeeded(const UAnimInstance* InAnimI
 						
 						// create a new control rig using the new class
 						NewControlRig = NewObject<UControlRig>(InAnimInstance->GetOwningComponent(), ExpectedClass);
-						ControlRig->SetObjectBinding(MakeShared<FControlRigObjectBinding>());
-						ControlRig->GetObjectBinding()->BindToObject(InAnimInstance->GetOwningComponent());
 						NewControlRig->Initialize(true);
 						NewControlRig->RequestInit();
 
@@ -465,8 +462,6 @@ bool FAnimNode_ControlRig::UpdateControlRigIfNeeded(const UAnimInstance* InAnimI
 		if(ControlRig == nullptr)
 		{
 			ControlRig = NewObject<UControlRig>(InAnimInstance->GetOwningComponent(), ExpectedClass);
-			ControlRig->SetObjectBinding(MakeShared<FControlRigObjectBinding>());
-			ControlRig->GetObjectBinding()->BindToObject(InAnimInstance->GetOwningComponent());
 			ControlRig->Initialize(true);
 			ControlRig->RequestInit();
 		}
diff --git a/Engine/Plugins/Animation/ControlRig/Source/ControlRigDeveloper/Private/AnimGraphNode_ControlRig.cpp b/Engine/Plugins/Animation/ControlRig/Source/ControlRigDeveloper/Private/AnimGraphNode_ControlRig.cpp
--- a/Engine/Plugins/Animation/ControlRig/Source/ControlRigDeveloper/Private/AnimGraphNode_ControlRig.cpp
+++ b/Engine/Plugins/Animation/ControlRig/Source/ControlRigDeveloper/Private/AnimGraphNode_ControlRig.cpp
@@ -15,7 +15,6 @@
 #include "AnimationGraphSchema.h"
 #include "RigVMBlueprintGeneratedClass.h"
 #include "ControlRigBlueprint.h"
-#include "ControlRigObjectBinding.h"
 #include "Misc/DefaultValueHelper.h"
 
 #include UE_INLINE_GENERATED_CPP_BY_NAME(AnimGraphNode_ControlRig)
@@ -199,45 +198,58 @@ void UAnimGraphNode_ControlRig::CreateCustomPins(TArray<UEdGraphPin*>* OldPins)
 		}
 	}
 
-	TArray<FControlsInfo>& Controls = GetControls();
-	for (FControlsInfo& ControlInfo : Controls)
+	if(const UClass* ControlRigClass = GetTargetClass())
 	{
-		const FName ControlName = ControlInfo.Name;
-		BeginExposableNames.Remove(ControlName);
-
-		if (CurrentlyExposedNames.Contains(ControlName))
+		if(UControlRig* CDO = ControlRigClass->GetDefaultObject<UControlRig>())
 		{
-			const FEdGraphPinType PinType = ControlInfo.PinType;
-			if (!PinType.PinCategory.IsValid())
+			if(const URigHierarchy* Hierarchy = CDO->GetHierarchy())
 			{
-				continue;
-			}
+				Hierarchy->ForEach<FRigControlElement>([&](FRigControlElement* ControlElement) -> bool
+				{
+					if (Hierarchy->IsAnimatable(ControlElement))
+					{
+						const FName ControlName = ControlElement->GetFName();
+						BeginExposableNames.Remove(ControlName);
 
-			UEdGraphPin* NewPin = CreatePin(EEdGraphPinDirection::EGPD_Input, PinType, ControlName);
-			NewPin->PinFriendlyName = FText::FromName(ControlInfo.Name);
+						if (CurrentlyExposedNames.Contains(ControlName))
+						{
+							const FEdGraphPinType PinType = Hierarchy->GetControlPinType(ControlElement);
+							if (!PinType.PinCategory.IsValid())
+							{
+								return true;
+							}
 
-			// Newly created pin does not have an auto generated default value, so we need to generate one here
-			// Missing the auto-gen default would cause UEdGraphPin::MovePersistentDataFromOldPin to override
-			// the actual default value with the empty auto-gen default, causing BP compiler to complain
-			// This is similar to how the following two functions create and initialize new pins, create first 
-			// and then set an auto-gen default
-			// FOptionalPinManager::CreateVisiblePins()
-			// FAnimBlueprintNodeOptionalPinManager::PostInitNewPin()
-			AnimGraphDefaultSchema->SetPinAutogeneratedDefaultValueBasedOnType(NewPin);
+							UEdGraphPin* NewPin = CreatePin(EEdGraphPinDirection::EGPD_Input, PinType, ControlName);
+							NewPin->PinFriendlyName = FText::FromName(ControlElement->GetFName());
 
-			// We cant interrogate CDO here as we may be mid-compile, so we can only really
-			// reset to the autogenerated default.
-			AnimGraphDefaultSchema->ResetPinToAutogeneratedDefaultValue(NewPin, false);
+							// Newly created pin does not have an auto generated default value, so we need to generate one here
+							// Missing the auto-gen default would cause UEdGraphPin::MovePersistentDataFromOldPin to override
+							// the actual default value with the empty auto-gen default, causing BP compiler to complain
+							// This is similar to how the following two functions create and initialize new pins, create first 
+							// and then set an auto-gen default
+							// FOptionalPinManager::CreateVisiblePins()
+							// FAnimBlueprintNodeOptionalPinManager::PostInitNewPin()
+							AnimGraphDefaultSchema->SetPinAutogeneratedDefaultValueBasedOnType(NewPin);
 
-			// Extract default values from the Target Control Rig if possible
-			const FString DefaultValue = ControlInfo.DefaultValue;
-			if(!DefaultValue.IsEmpty())
-			{
-				AnimGraphDefaultSchema->TrySetDefaultValue(*NewPin, DefaultValue);
+							// We cant interrogate CDO here as we may be mid-compile, so we can only really
+							// reset to the autogenerated default.
+							AnimGraphDefaultSchema->ResetPinToAutogeneratedDefaultValue(NewPin, false);
+
+							// Extract default values from the Target Control Rig if possible
+							const FString DefaultValue = Hierarchy->GetControlPinDefaultValue(ControlElement, true);
+							if(!DefaultValue.IsEmpty())
+							{
+								AnimGraphDefaultSchema->TrySetDefaultValue(*NewPin, DefaultValue);
+							}
+							
+							// sustain the current set of custom pins - we'll refrain from changing the node until post load is complete 
+							CustomizePinData(NewPin, ControlName, INDEX_NONE);
+						}
+					}
+					
+					return true;
+				});
 			}
-			
-			// sustain the current set of custom pins - we'll refrain from changing the node until post load is complete 
-			CustomizePinData(NewPin, ControlName, INDEX_NONE);
 		}
 	}
 
@@ -356,11 +368,23 @@ void UAnimGraphNode_ControlRig::RebuildExposedProperties()
 	}
 
 	// also add all of the controls
-	TArray<FControlsInfo>& Controls = GetControls();
-	for (FControlsInfo& ControlInfo : Controls)
+	if(const UClass* TargetClass = GetTargetClass())
 	{
-		const FName ControlName = ControlInfo.Name;
-		CustomPinProperties.Add(MakeOptionalPin(ControlName));
+		if(UControlRig* CDO = TargetClass->GetDefaultObject<UControlRig>())
+		{
+			if(const URigHierarchy* Hierarchy = CDO->GetHierarchy())
+			{
+				Hierarchy->ForEach<FRigControlElement>([&](const FRigControlElement* ControlElement) -> bool
+				{
+					if (Hierarchy->IsAnimatable(ControlElement))
+					{
+						const FName ControlName = ControlElement->GetFName();
+						CustomPinProperties.Add(MakeOptionalPin(ControlName));
+					}
+					return true;
+				});
+			}
+		}
 	}
 }
 
@@ -372,13 +396,19 @@ bool UAnimGraphNode_ControlRig::IsInputProperty(const FName& PropertyName) const
 	return InputVariables.Contains(PropertyName) || !OutputVariables.Contains(PropertyName);
 }
 
-UAnimGraphNode_ControlRig::FControlsInfo* UAnimGraphNode_ControlRig::FindControlElement(const FName& InControlName) const
+FRigControlElement* UAnimGraphNode_ControlRig::FindControlElement(const FName& InControlName) const
 {
-	TArray<FControlsInfo>& Controls = GetControls();
-	return Controls.FindByPredicate([InControlName](const FControlsInfo& Info)
+	if(const UClass* ControlRigClass = GetTargetClass())
 	{
-		return Info.Name == InControlName;
-	});
+		if(UControlRig* CDO = ControlRigClass->GetDefaultObject<UControlRig>())
+		{
+			if(const URigHierarchy* Hierarchy = CDO->GetHierarchy())
+			{
+				return (FRigControlElement*)Hierarchy->Find<FRigControlElement>(FRigElementKey(InControlName, ERigElementType::Control));
+			}
+		}
+	}
+	return nullptr;
 }
 
 bool UAnimGraphNode_ControlRig::IsAvailableToMapToCurve(const FName& PropertyName, bool bInput) const
@@ -393,9 +423,9 @@ bool UAnimGraphNode_ControlRig::IsAvailableToMapToCurve(const FName& PropertyNam
 		return Variable->TypeName == TEXT("float");
 	}
 
-	if(const FControlsInfo* ControlInfo = FindControlElement(PropertyName))
+	if(const FRigControlElement* ControlElement = FindControlElement(PropertyName))
 	{
-		return (ControlInfo->ControlType == ERigControlType::Float || ControlInfo->ControlType == ERigControlType::ScaleFloat);
+		return (ControlElement->Settings.ControlType == ERigControlType::Float || ControlElement->Settings.ControlType == ERigControlType::ScaleFloat);
 	}
 
 	return ensure(false);
@@ -542,54 +572,6 @@ void UAnimGraphNode_ControlRig::CustomizeDetails(IDetailLayoutBuilder& DetailBui
 	// alpha property blending support END
 }
 
-TArray<UAnimGraphNode_ControlRig::FControlsInfo>& UAnimGraphNode_ControlRig::GetControls() const
-{
-	const UClass* ControlRigClass = GetTargetClass();
-
-	if (ControlsInfoClass != ControlRigClass)
-	{
-		ControlsInfo.Reset();
-		ControlsInfoClass = ControlRigClass;
-		
-		if (ControlRigClass)
-		{
-			if (USkeleton* Skeleton = GetAnimBlueprint()->TargetSkeleton)
-			{
-				UControlRig* TemplateRig = NewObject<UControlRig>(GetTransientPackage(), ControlRigClass, NAME_None, RF_Transient);
-				TemplateRig->SetObjectBinding(MakeShared<FControlRigObjectBinding>());
-				TemplateRig->GetObjectBinding()->BindToObject(Skeleton);
-				TemplateRig->GetDataSourceRegistry()->RegisterDataSource(UControlRig::OwnerComponent, TemplateRig->GetObjectBinding()->GetBoundObject());
-
-				TemplateRig->Initialize();
-				TemplateRig->SetBoneInitialTransformsFromRefSkeleton(Skeleton->GetReferenceSkeleton());
-				TemplateRig->Evaluate_AnyThread();
-
-				if(const URigHierarchy* Hierarchy = TemplateRig->GetHierarchy())
-				{
-					Hierarchy->ForEach<FRigControlElement>([&](FRigControlElement* ControlElement) -> bool
-					{
-						if (Hierarchy->IsAnimatable(ControlElement))
-						{
-							FControlsInfo Info;
-							Info.Name = ControlElement->GetFName();
-							Info.DisplayName = ControlElement->GetName();
-							Info.PinType = TemplateRig->GetHierarchy()->GetControlPinType(ControlElement);
-							Info.ControlType = ControlElement->Settings.ControlType;
-							Info.DefaultValue = TemplateRig->GetHierarchy()->GetControlPinDefaultValue(ControlElement, true);
-							ControlsInfo.Add(Info);
-						}
-						return true;
-					});
-				}
-
-				TemplateRig->MarkAsGarbage();
-			}
-		}
-	}
-
-	return ControlsInfo;
-}
-
 void UAnimGraphNode_ControlRig::GetVariables(bool bInput, TMap<FName, FRigVMExternalVariable>& OutVariables) const
 {
 	DECLARE_SCOPE_HIERARCHICAL_COUNTER_FUNC()
@@ -652,10 +634,22 @@ void UAnimGraphNode_ControlRig::GetAvailableMapping(const FName& PathName, TArra
 		TargetSkeleton->GetCurveMetaDataNames(OutArray);
 		
 		// also add all controls
-		TArray<FControlsInfo>& Controls = GetControls();
-		for(FControlsInfo& ControlInfo : Controls)
+		if(const UClass* ControlRigClass = GetTargetClass())
 		{
-			OutArray.Add(ControlInfo.Name);
+			if(UControlRig* CDO = ControlRigClass->GetDefaultObject<UControlRig>())
+			{
+				if(const URigHierarchy* Hierarchy = CDO->GetHierarchy())
+				{
+					Hierarchy->ForEach<FRigControlElement>([&](FRigControlElement* ControlElement) -> bool
+					{
+						if (Hierarchy->IsAnimatable(ControlElement))
+						{
+							OutArray.Add(ControlElement->GetFName());
+						}
+						return true;
+					});
+				}
+			}
 		}
 
 		// we want to exclude anything that has been mapped already
@@ -714,16 +708,28 @@ void UAnimGraphNode_ControlRig::CreateVariableMapping(const FString& FilteredTex
 	if(bInput)
 	{
 		// add all controls
-		TArray<FControlsInfo>& Controls = GetControls();
-		for (FControlsInfo& ControlInfo : Controls)
+		if(const UClass* ControlRigClass = GetTargetClass())
 		{
-			const FName ControlName = ControlInfo.Name;
-			const FString& DisplayName = ControlInfo.DisplayName;
-
-			if (!bDoFiltering || DisplayName.Contains(FilteredText))
+			if(UControlRig* CDO = ControlRigClass->GetDefaultObject<UControlRig>())
 			{
-				TSharedRef<FRigVMVariableMappingInfo> Item = FRigVMVariableMappingInfo::Make(ControlName);
-				OutArray.Add(Item);
+				if(const URigHierarchy* Hierarchy = CDO->GetHierarchy())
+				{
+					Hierarchy->ForEach<FRigControlElement>([&](FRigControlElement* ControlElement) -> bool
+					{
+						if (Hierarchy->IsAnimatable(ControlElement))
+						{
+							const FName ControlName = ControlElement->GetFName();
+							const FString& DisplayName = ControlElement->GetName();
+
+							if (!bDoFiltering || DisplayName.Contains(FilteredText))
+							{
+								TSharedRef<FRigVMVariableMappingInfo> Item = FRigVMVariableMappingInfo::Make(ControlName);
+								OutArray.Add(Item);
+							}
+						}
+						return true;
+					});
+				}
 			}
 		}
 	}
diff --git a/Engine/Plugins/Animation/ControlRig/Source/ControlRigDeveloper/Public/AnimGraphNode_ControlRig.h b/Engine/Plugins/Animation/ControlRig/Source/ControlRigDeveloper/Public/AnimGraphNode_ControlRig.h
--- a/Engine/Plugins/Animation/ControlRig/Source/ControlRigDeveloper/Public/AnimGraphNode_ControlRig.h
+++ b/Engine/Plugins/Animation/ControlRig/Source/ControlRigDeveloper/Public/AnimGraphNode_ControlRig.h
@@ -30,18 +30,6 @@ private:
 	virtual FAnimNode_CustomProperty* GetCustomPropertyNode() override { return &Node; }
 	virtual const FAnimNode_CustomProperty* GetCustomPropertyNode() const override { return &Node; }
 
-	struct FControlsInfo
-	{
-		FName Name;
-		FString DisplayName;
-		FEdGraphPinType PinType;
-		ERigControlType ControlType;
-		FString DefaultValue;
-	};
-	mutable const UClass* ControlsInfoClass = nullptr;
-	mutable TArray<FControlsInfo> ControlsInfo;
-	TArray<FControlsInfo>& GetControls() const;
-
 	// property related things
 	void GetVariables(bool bInput, TMap<FName, FRigVMExternalVariable>& OutParameters) const;
 
@@ -67,6 +55,6 @@ private:
 
 	bool IsAvailableToMapToCurve(const FName& PropertyName, bool bInput) const;
 	bool IsInputProperty(const FName& PropertyName) const;
-	UAnimGraphNode_ControlRig::FControlsInfo* FindControlElement(const FName& InControlName) const;
+	FRigControlElement* FindControlElement(const FName& InControlName) const;
 };
 
