Control Rig: Update namespace display for modular rigs

#rnx

diff --git a/Engine/Plugins/Animation/ControlRig/Source/ControlRig/Private/ControlRig.cpp b/Engine/Plugins/Animation/ControlRig/Source/ControlRig/Private/ControlRig.cpp
--- a/Engine/Plugins/Animation/ControlRig/Source/ControlRig/Private/ControlRig.cpp
+++ b/Engine/Plugins/Animation/ControlRig/Source/ControlRig/Private/ControlRig.cpp
@@ -1026,11 +1026,7 @@ bool UControlRig::Execute(const FName& InEventName)
 					
 				} // destroy FTransientControlScope
 				
-				if (PostConstructionEvent.IsBound())
-				{
-					FControlRigBracketScope BracketScope(PostConstructionBracket);
-					PostConstructionEvent.Broadcast(this, FRigUnit_PrepareForExecution::EventName);
-				}
+				RunPostConstructionEvent();
 
 				// Reset the input pose after construction
 				if (CurrentPose.Num() > 0)
@@ -2273,8 +2269,17 @@ bool UControlRig::IsControlSelected(const FName& InControlName)const
 	return false;
 }
 
+void UControlRig::RunPostConstructionEvent()
+{
+	if (PostConstructionEvent.IsBound())
+	{
+		FControlRigBracketScope BracketScope(PostConstructionBracket);
+		PostConstructionEvent.Broadcast(this, FRigUnit_PrepareForExecution::EventName);
+	}
+}
+
 void UControlRig::HandleHierarchyModified(ERigHierarchyNotification InNotification, URigHierarchy* InHierarchy,
-    const FRigBaseElement* InElement)
+                                          const FRigBaseElement* InElement)
 {
 	switch(InNotification)
 	{
diff --git a/Engine/Plugins/Animation/ControlRig/Source/ControlRig/Private/ModularRig.cpp b/Engine/Plugins/Animation/ControlRig/Source/ControlRig/Private/ModularRig.cpp
--- a/Engine/Plugins/Animation/ControlRig/Source/ControlRig/Private/ModularRig.cpp
+++ b/Engine/Plugins/Animation/ControlRig/Source/ControlRig/Private/ModularRig.cpp
@@ -205,6 +205,7 @@ void UModularRig::PostLoad()
 {
 	Super::PostLoad();
 	ModularRigModel.UpdateCachedChildren();
+	ResetShortestDisplayPathCache();
 }
 
 void UModularRig::InitializeVMs(bool bRequestInit)
@@ -574,6 +575,7 @@ void UModularRig::ResetModules(bool bDestroyModuleRigs)
 	RootModules.Reset();
 	Modules.Reset();
 	SupportedEvents.Reset();
+	ResetShortestDisplayPathCache();
 }
 
 const FModularRigModel& UModularRig::GetModularRigModel() const
@@ -629,8 +631,14 @@ void UModularRig::UpdateSupportedEvents()
 	});
 }
 
+void UModularRig::RunPostConstructionEvent()
+{
+	RecomputeShortestDisplayPathCache();
+	Super::RunPostConstructionEvent();
+}
+
 FRigModuleInstance* UModularRig::AddModuleInstance(const FName& InModuleName, TSubclassOf<UControlRig> InModuleClass, const FRigModuleInstance* InParent,
-	const TMap<FRigElementKey, FRigElementKey>& InConnectionMap, const TMap<FName, FString>& InVariableDefaultValues) 
+                                                   const TMap<FRigElementKey, FRigElementKey>& InConnectionMap, const TMap<FName, FString>& InVariableDefaultValues) 
 {
 	// Make sure there are no name clashes
 	if (InParent)
@@ -696,6 +704,7 @@ FRigModuleInstance* UModularRig::AddModuleInstance(const FName& InModuleName, TS
 	NewModule.SetRig(NewModuleRig);
 
 	UpdateCachedChildren();
+	ResetShortestDisplayPathCache();
 	for (const FName& EventName : NewModule.GetRig()->GetSupportedEvents())
 	{
 		SupportedEvents.AddUnique(EventName);
@@ -719,7 +728,7 @@ FRigModuleInstance* UModularRig::AddModuleInstance(const FName& InModuleName, TS
 			NewModule.GetRig()->SetVariableFromString(Variable.Key, Variable.Value);
 		}
 	}
-	
+
 	return &NewModule;
 }
 
@@ -772,6 +781,78 @@ void UModularRig::DiscardModuleRig(UControlRig* InControlRig)
 	}
 }
 
+void UModularRig::ResetShortestDisplayPathCache() const
+{
+	ElementKeyToShortestDisplayPath.Reset();
+}
+
+void UModularRig::RecomputeShortestDisplayPathCache() const
+{
+	ResetShortestDisplayPathCache();
+	
+	const URigHierarchy* Hierarchy = GetHierarchy();
+	check(Hierarchy);
+
+	const TArray<FRigElementKey> AllKeys = Hierarchy->GetAllKeys();
+
+	auto GetNameForElement = [Hierarchy](const FRigElementKey& InElementKey)
+	{
+		const FName DesiredName = Hierarchy->GetNameMetadata(InElementKey, URigHierarchy::DesiredNameMetadataName, NAME_None);
+		if(!DesiredName.IsNone())
+		{
+			return DesiredName;
+		}
+		return InElementKey.Name;
+	};
+	
+	TMap<FName, bool> IsNameUniqueInHierarchy;
+	for(const FRigElementKey& Key : AllKeys)
+	{
+		auto UpdateNameUniqueness = [](const FName& InName, TMap<FName, bool>& UniqueMap)
+		{
+			if(bool* Unique = UniqueMap.Find(InName))
+			{
+				*Unique = false;
+			}
+			else
+			{
+				UniqueMap.Add(InName, true);
+			}
+		};
+		UpdateNameUniqueness(GetNameForElement(Key), IsNameUniqueInHierarchy);
+	}
+
+	const FModularRigModel& Model = GetModularRigModel();
+	for(const FRigElementKey& Key : AllKeys)
+	{
+		const FName Name = GetNameForElement(Key);
+		const FName ModulePath = Hierarchy->GetModulePathFName(Key);
+		
+		if(!ModulePath.IsNone())
+		{
+			if(const FRigModuleReference* Module = Model.FindModule(ModulePath.ToString()))
+			{
+				const FString NameString = Name.ToString();
+				const FString ModuleShortName = Module->GetShortName();
+				const FString NameSpacedName = URigHierarchy::JoinNameSpace(ModuleShortName, NameString);
+				ElementKeyToShortestDisplayPath.Add(Key, { NameSpacedName, IsNameUniqueInHierarchy.FindChecked(Name) ? NameString : NameSpacedName });
+				continue;
+			}
+		}
+
+		if(IsNameUniqueInHierarchy.FindChecked(Name))
+		{
+			const FString NameString = Name.ToString();
+			ElementKeyToShortestDisplayPath.Add(Key, { NameString, NameString});
+		}
+		else
+		{
+			const FString NameString = Key.ToString();
+			ElementKeyToShortestDisplayPath.Add(Key, { Name.ToString(), Name.ToString()});
+		}
+	}
+}
+
 const FRigModuleInstance* UModularRig::FindModule(const FString& InPath) const
 {
 	if(InPath.EndsWith(NamespaceSeparator))
@@ -832,6 +913,15 @@ FString UModularRig::GetParentPath(const FString& InPath) const
 	return FString();
 }
 
+FString UModularRig::GetShortestDisplayPathForElement(const FRigElementKey& InElementKey, bool bAlwaysShowNameSpace) const
+{
+	if(const TTuple<FString, FString>* ShortestDisplayPaths = ElementKeyToShortestDisplayPath.Find(InElementKey))
+	{
+		return bAlwaysShowNameSpace ? ShortestDisplayPaths->Get<0>() : ShortestDisplayPaths->Get<1>();
+	}
+	return FString();
+}
+
 void UModularRig::ForEachModule(TFunctionRef<bool(FRigModuleInstance*)> PerModuleFunction)
 {
 	TArray<FRigModuleInstance*> ModuleInstances = RootModules;
diff --git a/Engine/Plugins/Animation/ControlRig/Source/ControlRig/Private/Rigs/RigHierarchy.cpp b/Engine/Plugins/Animation/ControlRig/Source/ControlRig/Private/Rigs/RigHierarchy.cpp
--- a/Engine/Plugins/Animation/ControlRig/Source/ControlRig/Private/Rigs/RigHierarchy.cpp
+++ b/Engine/Plugins/Animation/ControlRig/Source/ControlRig/Private/Rigs/RigHierarchy.cpp
@@ -1636,6 +1636,15 @@ FText URigHierarchy::GetDisplayNameForUI(const FRigBaseElement* InElement, bool
 {
 	check(InElement);
 
+	if(const UModularRig* ModularRig = GetTypedOuter<UModularRig>())
+	{
+		const FString ShortestPath = ModularRig->GetShortestDisplayPathForElement(InElement->GetKey(), false);
+		if(!ShortestPath.IsEmpty())
+		{
+			return FText::FromString(ShortestPath);
+		}
+	}
+
 	const FName& DisplayName = InElement->GetDisplayName();
 	FString DisplayNameString = DisplayName.ToString();
 	(void)SplitNameSpace(DisplayNameString, nullptr, &DisplayNameString);
diff --git a/Engine/Plugins/Animation/ControlRig/Source/ControlRig/Public/ControlRig.h b/Engine/Plugins/Animation/ControlRig/Source/ControlRig/Public/ControlRig.h
--- a/Engine/Plugins/Animation/ControlRig/Source/ControlRig/Public/ControlRig.h
+++ b/Engine/Plugins/Animation/ControlRig/Source/ControlRig/Public/ControlRig.h
@@ -484,6 +484,10 @@ private:
 	/** Handle changes within the hierarchy */
 	void HandleHierarchyModified(ERigHierarchyNotification InNotification, URigHierarchy* InHierarchy, const FRigBaseElement* InElement);
 
+protected:
+	
+	virtual void RunPostConstructionEvent();
+
 private:
 #if WITH_EDITOR
 	/** Add a transient / temporary control used to interact with a node */
diff --git a/Engine/Plugins/Animation/ControlRig/Source/ControlRig/Public/ModularRig.h b/Engine/Plugins/Animation/ControlRig/Source/ControlRig/Public/ModularRig.h
--- a/Engine/Plugins/Animation/ControlRig/Source/ControlRig/Public/ModularRig.h
+++ b/Engine/Plugins/Animation/ControlRig/Source/ControlRig/Public/ModularRig.h
@@ -181,6 +181,7 @@ public:
 	const FRigModuleInstance* FindModule(const FRigBaseElement* InElement) const;
 	const FRigModuleInstance* FindModule(const FRigElementKey& InElementKey) const;
 	FString GetParentPath(const FString& InPath) const;
+	FString GetShortestDisplayPathForElement(const FRigElementKey& InElementKey, bool bAlwaysShowNameSpace) const;
 
 	void ForEachModule(TFunctionRef<bool(FRigModuleInstance*)> PerModuleFunction);
 	void ForEachModule(TFunctionRef<bool(const FRigModuleInstance*)> PerModuleFunction) const;
@@ -203,6 +204,10 @@ public:
 
 	static const FString NamespaceSeparator;
 
+protected:
+
+	virtual void RunPostConstructionEvent() override;
+
 private:
 
 	/** Adds a module to the rig*/
@@ -216,5 +221,9 @@ private:
 
 	TMap<FString, UControlRig*> PreviousModuleRigs;
 
+	void ResetShortestDisplayPathCache() const;
+	void RecomputeShortestDisplayPathCache() const;
+	mutable TMap<FRigElementKey, TTuple<FString, FString>> ElementKeyToShortestDisplayPath; 
+
 	friend struct FRigModuleInstance;
 };
diff --git a/Engine/Plugins/Animation/ControlRig/Source/ControlRig/Public/Rigs/RigHierarchyDefines.h b/Engine/Plugins/Animation/ControlRig/Source/ControlRig/Public/Rigs/RigHierarchyDefines.h
--- a/Engine/Plugins/Animation/ControlRig/Source/ControlRig/Public/Rigs/RigHierarchyDefines.h
+++ b/Engine/Plugins/Animation/ControlRig/Source/ControlRig/Public/Rigs/RigHierarchyDefines.h
@@ -30,6 +30,9 @@ enum class ERigElementType : uint8
 	Reference = 0x020,
 	Connector = 0x040,
 	Socket = 0x080,
+	
+	First = Bone,
+	Last = Socket,
 	All = Bone | Null | Control | Curve | RigidBody | Reference | Connector | Socket,
 	ToResetAfterConstructionEvent = Bone | Control | Curve | Socket UMETA(Hidden),
 };
diff --git a/Engine/Plugins/Animation/ControlRig/Source/ControlRigEditor/Private/Editor/ControlRigSchematicModel.cpp b/Engine/Plugins/Animation/ControlRig/Source/ControlRigEditor/Private/Editor/ControlRigSchematicModel.cpp
--- a/Engine/Plugins/Animation/ControlRig/Source/ControlRigEditor/Private/Editor/ControlRigSchematicModel.cpp
+++ b/Engine/Plugins/Animation/ControlRig/Source/ControlRigEditor/Private/Editor/ControlRigSchematicModel.cpp
@@ -31,6 +31,53 @@ bool FControlRigSchematicRigElementKeyNode::IsDragSupported() const
 	return Key.Type == ERigElementType::Connector;
 }
 
+const FText& FControlRigSchematicRigElementKeyNode::GetLabel() const
+{
+	FText LabelText = FText::FromString(Key.ToString());
+	if(const FControlRigSchematicModel* ControlRigModel = Cast<FControlRigSchematicModel>(Model))
+	{
+		if(const UModularRig* ModularRig = Cast<UModularRig>(ControlRigModel->ControlRigBeingDebuggedPtr.Get()))
+		{
+			if (const URigHierarchy* Hierarchy = ModularRig->GetHierarchy())
+			{
+				LabelText = Hierarchy->GetDisplayNameForUI(Key);
+
+				switch(Key.Type)
+				{
+					case ERigElementType::Socket:
+					{
+						if(const FRigSocketElement* Socket = Hierarchy->Find<FRigSocketElement>(Key))
+						{
+							const FString Description = Socket->GetDescription(Hierarchy);
+							if(!Description.IsEmpty())
+							{
+								static const FText NodeLabelSocketDescriptionFormat = LOCTEXT("NodeLabelSocketDescriptionFormat", "{0}\n{1}");
+								LabelText = FText::Format(NodeLabelSocketDescriptionFormat, LabelText,  FText::FromString(Description));
+							}
+						}
+						break;
+					}
+					default:
+					{
+						break;
+					}
+				}
+
+				const FModularRigModel& ModularRigModel = ModularRig->GetModularRigModel();
+				const TArray<FRigElementKey>& Connectors = ModularRigModel.Connections.FindConnectorsFromTarget(Key);
+				for(const FRigElementKey& Connector : Connectors)
+				{
+					static const FText NodeLabelConnectionFormat = LOCTEXT("NodeLabelConnectionFormat", "{0}\nConnection: {1}");
+					const FText ConnectorShortestPath = Hierarchy->GetDisplayNameForUI(Connector, false);
+					LabelText = FText::Format(NodeLabelConnectionFormat, LabelText, ConnectorShortestPath);
+				}
+			}
+		}
+	}
+	const_cast<FControlRigSchematicRigElementKeyNode*>(this)->Label = LabelText;
+	return Label;
+}
+
 FControlRigSchematicWarningTag::FControlRigSchematicWarningTag()
 : FSchematicGraphTag()
 {
@@ -175,36 +222,6 @@ FControlRigSchematicRigElementKeyNode* FControlRigSchematicModel::AddElementKeyN
 void FControlRigSchematicModel::ConfigureElementKeyNode(FControlRigSchematicRigElementKeyNode* InNode, const FRigElementKey& InKey)
 {
 	InNode->Key = InKey;
-	InNode->SetLabel(FText::FromString(InNode->GetKey().ToString())); 
-	switch(InNode->GetKey().Type)
-	{
-		case ERigElementType::Socket:
-		{
-			if (ControlRigBeingDebuggedPtr.IsValid())
-			{
-				if (const URigHierarchy* Hierarchy = ControlRigBeingDebuggedPtr->GetHierarchy())
-				{
-					if(const FRigSocketElement* Socket = Hierarchy->Find<FRigSocketElement>(InNode->GetKey()))
-					{
-						const FString Description = Socket->GetDescription(Hierarchy);
-						if(!Description.IsEmpty())
-						{
-							InNode->SetLabel(
-								FText::Format(LOCTEXT("NodeLabelFormat", "{0}\n{1}"), 
-								InNode->GetLabel(), 
-								FText::FromString(Description))
-							);
-						}
-					}
-				}
-			}
-			break;
-		}
-		default:
-		{
-			break;
-		}
-	}
 }
 
 const FControlRigSchematicRigElementKeyNode* FControlRigSchematicModel::FindElementKeyNode(const FRigElementKey& InKey) const
@@ -1105,8 +1122,7 @@ bool FControlRigSchematicModel::GetContextMenuForNode(const FSchematicGraphNode*
 	{
 		if(const FControlRigSchematicRigElementKeyNode* ElementKeyNode = Cast<FControlRigSchematicRigElementKeyNode>(InNode))
 		{
-			const URigHierarchy* Hierarchy = ControlRigBlueprint->GetDebuggedControlRig()->GetHierarchy();
-			check(Hierarchy);
+			const UModularRig* ModularRig = CastChecked<UModularRig>(ControlRigBlueprint->GetDebuggedControlRig());
 			const FModularRigConnections& Connections = ControlRigBlueprint->ModularRigModel.Connections;
 			const TArray<FRigElementKey>& Connectors = Connections.FindConnectorsFromTarget(ElementKeyNode->Key);
 			if(!Connectors.IsEmpty())
@@ -1116,34 +1132,8 @@ bool FControlRigSchematicModel::GetContextMenuForNode(const FSchematicGraphNode*
 				// note: this is a copy on purpose since it is passed into the lambda
 				for(const FRigElementKey Connector : Connectors)
 				{
-					FString Label = *Connector.Name.ToString();
-
-					const FName ModulePath = Hierarchy->GetModulePathFName(Connector);
-					if(!ModulePath.IsNone())
-					{
-						if(const FRigModuleReference* Module = ControlRigBlueprint->ModularRigModel.FindModule(ModulePath.ToString()))
-						{
-							const FName DesiredName = Hierarchy->GetNameMetadata(Connector, URigHierarchy::DesiredNameMetadataName, NAME_None);
-							if(!DesiredName.IsNone())
-							{
-								Label = DesiredName.ToString();
-
-								if(const FRigConnectorElement* PrimaryConnector = Module->FindPrimaryConnector(Hierarchy))
-								{
-									if(PrimaryConnector->GetKey() == Connector)
-									{
-										Label = Module->GetShortName();
-									}
-									else
-									{
-										Label = URigHierarchy::JoinNameSpace(Module->GetShortName(), Label);
-									}
-								}
-							}
-						}
-					}
-					
-					const FText Description = FText::FromString(FString::Printf(TEXT("Disconnect %s"), *Label));
+					FText Label = ModularRig->GetHierarchy()->GetDisplayNameForUI(Connector, false);
+					const FText Description = FText::FromString(FString::Printf(TEXT("Disconnect %s"), *Label.ToString()));
 					
 					OutMenu.AddMenuEntry(Description, Description, FSlateIcon(), FUIAction(
 						FExecuteAction::CreateLambda([this, Connector]()
diff --git a/Engine/Plugins/Animation/ControlRig/Source/ControlRigEditor/Private/Editor/ControlRigSchematicModel.h b/Engine/Plugins/Animation/ControlRig/Source/ControlRigEditor/Private/Editor/ControlRigSchematicModel.h
--- a/Engine/Plugins/Animation/ControlRig/Source/ControlRigEditor/Private/Editor/ControlRigSchematicModel.h
+++ b/Engine/Plugins/Animation/ControlRig/Source/ControlRigEditor/Private/Editor/ControlRigSchematicModel.h
@@ -26,6 +26,7 @@ public:
 	virtual FString GetDragDropDecoratorLabel() const override;
 	virtual bool IsAutoScaleEnabled() const override { return true; }
 	virtual bool IsDragSupported() const override;
+	virtual const FText& GetLabel() const override;
 	
 protected:
 
@@ -139,4 +140,5 @@ private:
 	bool bUpdateElementKeyLinks = true;
 
 	friend class FControlRigEditor;
+	friend class FControlRigSchematicRigElementKeyNode;
 };
\ No newline at end of file
diff --git a/Engine/Source/Editor/AnimationEditorWidgets/Private/SchematicGraphPanel/SSchematicGraphPanel.cpp b/Engine/Source/Editor/AnimationEditorWidgets/Private/SchematicGraphPanel/SSchematicGraphPanel.cpp
--- a/Engine/Source/Editor/AnimationEditorWidgets/Private/SchematicGraphPanel/SSchematicGraphPanel.cpp
+++ b/Engine/Source/Editor/AnimationEditorWidgets/Private/SchematicGraphPanel/SSchematicGraphPanel.cpp
@@ -253,8 +253,8 @@ int32 SSchematicGraphNode::OnPaint(const FPaintArgs& Args, const FGeometry& Allo
 		}
 	}
 
-	const FText& NodeLabel = NodeData->GetLabel();
-	bool bDrawLabel = !NodeLabel.IsEmpty();  
+	FText NodeLabel;
+	bool bDrawLabel = true;
 	if(bDrawLabel)
 	{
 		const FVector2D MouseCursorLocation = FSlateApplication::Get().GetCursorPos();
@@ -274,6 +274,14 @@ int32 SSchematicGraphNode::OnPaint(const FPaintArgs& Args, const FGeometry& Allo
 		}
 	}
 	if(bDrawLabel)
+	{
+		NodeLabel = NodeData->GetLabel();
+		if(NodeLabel.IsEmpty())
+		{
+			bDrawLabel = false;
+		}
+	}
+	if(bDrawLabel)
 	{
 		const FString NodeLabelString = NodeLabel.ToString();
 		const TSharedRef< FSlateFontMeasure > FontMeasureService = FSlateApplication::Get().GetRenderer()->GetFontMeasureService();
