Control Rig: Add category, keywords and description to the rig module settings
#rnx

diff --git a/Engine/Plugins/Animation/ControlRig/Source/ControlRig/Public/Rigs/RigModuleDefines.h b/Engine/Plugins/Animation/ControlRig/Source/ControlRig/Public/Rigs/RigModuleDefines.h
--- a/Engine/Plugins/Animation/ControlRig/Source/ControlRig/Public/Rigs/RigModuleDefines.h
+++ b/Engine/Plugins/Animation/ControlRig/Source/ControlRig/Public/Rigs/RigModuleDefines.h
@@ -89,6 +89,18 @@ struct CONTROLRIG_API FRigModuleSettings
 	UPROPERTY(BlueprintReadWrite, EditAnywhere, Category = Module,  meta = (AllowedClasses = "/Script/Engine.Texture2D"))
 	FSoftObjectPath Icon;
 
+	// The category of the module
+	UPROPERTY(BlueprintReadWrite, EditAnywhere, Category = Module)
+	FString Category;
+
+	// The keywords of the module
+	UPROPERTY(BlueprintReadWrite, EditAnywhere, Category = Module)
+	FString Keywords;
+
+	// The description of the module
+	UPROPERTY(BlueprintReadWrite, EditAnywhere, Category = Module, meta = (MultiLine = true))
+	FString Description;
+
 	UPROPERTY(BlueprintReadOnly, Category = Module)
 	TArray<FRigModuleConnector> ExposedConnectors;
 };
diff --git a/Engine/Plugins/Animation/ControlRig/Source/ControlRigEditor/Private/Editor/SRigModuleAssetBrowser.cpp b/Engine/Plugins/Animation/ControlRig/Source/ControlRigEditor/Private/Editor/SRigModuleAssetBrowser.cpp
--- a/Engine/Plugins/Animation/ControlRig/Source/ControlRigEditor/Private/Editor/SRigModuleAssetBrowser.cpp
+++ b/Engine/Plugins/Animation/ControlRig/Source/ControlRigEditor/Private/Editor/SRigModuleAssetBrowser.cpp
@@ -52,6 +52,7 @@ void SRigModuleAssetBrowser::RefreshView()
 	AssetPickerConfig.bForceShowPluginContent = true;
 	AssetPickerConfig.bForceShowEngineContent = true;
 	AssetPickerConfig.InitialThumbnailSize = EThumbnailSize::Small;
+	AssetPickerConfig.OnGetCustomAssetToolTip = FOnGetCustomAssetToolTip::CreateSP(this, &SRigModuleAssetBrowser::CreateCustomAssetToolTip);
 
 	// hide all asset registry columns by default (we only really want the name and path)
 	UObject* DefaultControlRigBlueprint = UControlRigBlueprint::StaticClass()->GetDefaultObject();
@@ -140,4 +141,162 @@ void SRigModuleAssetBrowser::OnAssetDoubleClicked(const FAssetData& AssetData)
 	}
 }
 
+TSharedRef<SToolTip> SRigModuleAssetBrowser::CreateCustomAssetToolTip(FAssetData& AssetData)
+{
+	// Make a list of tags to show
+	TArray<UObject::FAssetRegistryTag> Tags;
+	UClass* AssetClass = FindObject<UClass>(AssetData.AssetClassPath);
+	check(AssetClass);
+	UObject* DefaultObject = AssetClass->GetDefaultObject();
+	FAssetRegistryTagsContextData TagsContext(DefaultObject, EAssetRegistryTagsCaller::Uncategorized);
+	DefaultObject->GetAssetRegistryTags(TagsContext);
+
+	TArray<FName> TagsToShow;
+	static const FName ModulePath(TEXT("Path"));
+	static const FName ModuleSettings(TEXT("RigModuleSettings"));
+	for (const TPair<FName, UObject::FAssetRegistryTag>& TagPair : TagsContext.Tags)
+	{
+		if(TagPair.Key == ModulePath ||
+			TagPair.Key == ModuleSettings)
+		{
+			TagsToShow.Add(TagPair.Key);
+		}
+	}
+
+	TMap<FName, FText> TagsAndValuesToShow;
+
+	// Add asset registry tags to a text list; except skeleton as that is implied in Persona
+	TSharedRef<SVerticalBox> DescriptionBox = SNew(SVerticalBox);
+	for(TPair<FName, FAssetTagValueRef> TagPair : AssetData.TagsAndValues)
+	{
+		if(TagsToShow.Contains(TagPair.Key))
+		{
+			// Check for DisplayName metadata
+			FName DisplayName;
+			if (FProperty* Field = FindFProperty<FProperty>(AssetClass, TagPair.Key))
+			{
+				DisplayName = *Field->GetDisplayNameText().ToString();
+			}
+			else
+			{
+				DisplayName = TagPair.Key;
+			}
+
+			if (TagPair.Key == ModuleSettings)
+			{
+				FRigModuleSettings Settings;
+				FRigVMPinDefaultValueImportErrorContext ErrorPipe;
+				FRigModuleSettings::StaticStruct()->ImportText(*TagPair.Value.GetValue(), &Settings, nullptr, PPF_None, &ErrorPipe, FString());
+				if (ErrorPipe.NumErrors == 0)
+				{
+					TagsAndValuesToShow.Add(TEXT("Default Name"), FText::FromString(Settings.Identifier.Name));
+					TagsAndValuesToShow.Add(TEXT("Category"), FText::FromString(Settings.Category));
+					TagsAndValuesToShow.Add(TEXT("Keywords"), FText::FromString(Settings.Keywords));
+					TagsAndValuesToShow.Add(TEXT("Description"), FText::FromString(Settings.Description));
+				}
+			}
+			else
+			{
+				TagsAndValuesToShow.Add(DisplayName, TagPair.Value.AsText());
+			}
+		}
+	}
+
+	for (const TPair<FName, FText>& TagPair : TagsAndValuesToShow)
+	{
+		DescriptionBox->AddSlot()
+		.AutoHeight()
+		.Padding(0,0,5,0)
+		[
+			SNew(SHorizontalBox)
+			+ SHorizontalBox::Slot()
+			.AutoWidth()
+			[
+				SNew(STextBlock)
+				.Text(FText::Format(LOCTEXT("AssetTagKey", "{0}: "), FText::FromName(TagPair.Key)))
+				.ColorAndOpacity(FSlateColor::UseSubduedForeground())
+			]
+
+			+ SHorizontalBox::Slot()
+			.AutoWidth()
+			[
+				SNew(STextBlock)
+				.Text(TagPair.Value)
+				.ColorAndOpacity(FSlateColor::UseForeground())
+			]
+		];
+	}
+
+	DescriptionBox->AddSlot()
+		.AutoHeight()
+		.Padding(0,0,5,0)
+		[
+			SNew(SHorizontalBox)
+			+ SHorizontalBox::Slot()
+			.AutoWidth()
+			[
+				SNew(STextBlock)
+				.Text(LOCTEXT("AssetBrowser_FolderPathLabel", "Folder :"))
+				.ColorAndOpacity(FSlateColor::UseSubduedForeground())
+			]
+
+			+ SHorizontalBox::Slot()
+			.AutoWidth()
+			[
+				SNew(STextBlock)
+				.Text(FText::FromName(AssetData.PackagePath))
+				.ColorAndOpacity(FSlateColor::UseForeground())
+				.WrapTextAt(300.f)
+			]
+		];
+
+	TSharedPtr<SHorizontalBox> ContentBox = nullptr;
+	TSharedRef<SToolTip> ToolTipWidget = SNew(SToolTip)
+	.TextMargin(1.f)
+	.BorderImage(FAppStyle::GetBrush("ContentBrowser.TileViewTooltip.ToolTipBorder"))
+	[
+		SNew(SBorder)
+		.Padding(6.f)
+		.BorderImage(FAppStyle::GetBrush("ContentBrowser.TileViewTooltip.NonContentBorder"))
+		[
+			SNew(SVerticalBox)
+			+SVerticalBox::Slot()
+			.AutoHeight()
+			.Padding(0,0,0,4)
+			[
+				SNew(SBorder)
+				.Padding(6.f)
+				.BorderImage(FAppStyle::GetBrush("ContentBrowser.TileViewTooltip.ContentBorder"))
+				[
+					SNew(SBox)
+					.HAlign(HAlign_Left)
+					[
+						SNew(STextBlock)
+						.Text(FText::FromName(AssetData.AssetName))
+						.Font(FAppStyle::GetFontStyle("ContentBrowser.TileViewTooltip.NameFont"))
+					]
+				]
+			]
+		
+			+ SVerticalBox::Slot()
+			[
+				SAssignNew(ContentBox, SHorizontalBox)
+				+SHorizontalBox::Slot()
+				.AutoWidth()
+				[
+					SNew(SBorder)
+					.Padding(6.f)
+					.BorderImage(FAppStyle::GetBrush("ContentBrowser.TileViewTooltip.ContentBorder"))
+					[
+						DescriptionBox
+					]
+				]
+			]
+		]
+	];
+	return ToolTipWidget;
+}
+
+
+
 #undef LOCTEXT_NAMESPACE
diff --git a/Engine/Plugins/Animation/ControlRig/Source/ControlRigEditor/Private/Editor/SRigModuleAssetBrowser.h b/Engine/Plugins/Animation/ControlRig/Source/ControlRigEditor/Private/Editor/SRigModuleAssetBrowser.h
--- a/Engine/Plugins/Animation/ControlRig/Source/ControlRigEditor/Private/Editor/SRigModuleAssetBrowser.h
+++ b/Engine/Plugins/Animation/ControlRig/Source/ControlRigEditor/Private/Editor/SRigModuleAssetBrowser.h
@@ -27,6 +27,7 @@ private:
 	bool OnShouldFilterAsset(const struct FAssetData& AssetData);
 	TSharedPtr<SWidget> OnGetAssetContextMenu(const TArray<FAssetData>& SelectedAssets) const;
 	void OnAssetDoubleClicked(const FAssetData& AssetData);
+	TSharedRef<SToolTip> CreateCustomAssetToolTip(FAssetData& AssetData);
 
 	/** Used to get the currently selected assets */
 	FGetCurrentSelectionDelegate GetCurrentSelectionDelegate;
