Control Rig: Make sure the connector target selection dropdown has a max height, and scrollbar is visible
#rnx

diff --git a/Engine/Plugins/Animation/ControlRig/Source/ControlRigEditor/Private/Editor/SModularRigTreeView.cpp b/Engine/Plugins/Animation/ControlRig/Source/ControlRigEditor/Private/Editor/SModularRigTreeView.cpp
--- a/Engine/Plugins/Animation/ControlRig/Source/ControlRigEditor/Private/Editor/SModularRigTreeView.cpp
+++ b/Engine/Plugins/Animation/ControlRig/Source/ControlRigEditor/Private/Editor/SModularRigTreeView.cpp
@@ -399,6 +399,7 @@ TSharedRef<SWidget> SModularRigModelItem::GenerateWidgetForColumn(const FName& C
 				[
 					SAssignNew(ConnectorComboBox, SSearchableRigHierarchyTreeView)
 						.RigTreeDelegates(TreeDelegates)
+						.MaxHeight(300)
 				]
 			]
 		];
diff --git a/Engine/Plugins/Animation/ControlRig/Source/ControlRigEditor/Private/Editor/SRigHierarchyTreeView.cpp b/Engine/Plugins/Animation/ControlRig/Source/ControlRigEditor/Private/Editor/SRigHierarchyTreeView.cpp
--- a/Engine/Plugins/Animation/ControlRig/Source/ControlRigEditor/Private/Editor/SRigHierarchyTreeView.cpp
+++ b/Engine/Plugins/Animation/ControlRig/Source/ControlRigEditor/Private/Editor/SRigHierarchyTreeView.cpp
@@ -1301,11 +1301,14 @@ void SSearchableRigHierarchyTreeView::Construct(const FArguments& InArgs)
 	FRigTreeDelegates TreeDelegates = InArgs._RigTreeDelegates;
 	SuperGetRigTreeDisplaySettings = TreeDelegates.OnGetDisplaySettings;
 
+	MaxHeight = InArgs._MaxHeight;
+
 	TreeDelegates.OnGetDisplaySettings.BindSP(this, &SSearchableRigHierarchyTreeView::GetDisplaySettings);
-	
+
+	TSharedPtr<SVerticalBox> VerticalBox;
 	ChildSlot
 	[
-		SNew(SVerticalBox)
+		SAssignNew(VerticalBox, SVerticalBox)
 		+SVerticalBox::Slot()
 		.AutoHeight()
 		.VAlign(VAlign_Top)
@@ -1318,7 +1321,6 @@ void SSearchableRigHierarchyTreeView::Construct(const FArguments& InArgs)
 		]
 
 		+SVerticalBox::Slot()
-		.AutoHeight()
 		.VAlign(VAlign_Top)
 		.HAlign(HAlign_Fill)
 		.Padding(0.0f, 0.0f)
@@ -1336,6 +1338,15 @@ void SSearchableRigHierarchyTreeView::Construct(const FArguments& InArgs)
 			]
 		]
 	];
+
+	if (MaxHeight > 0)
+	{
+		VerticalBox->GetSlot(1).SetMaxHeight(MaxHeight);
+	}
+	else
+	{
+		VerticalBox->GetSlot(1).SetAutoHeight();
+	}
 }
 
 const FRigTreeDisplaySettings& SSearchableRigHierarchyTreeView::GetDisplaySettings()
diff --git a/Engine/Plugins/Animation/ControlRig/Source/ControlRigEditor/Private/Editor/SRigHierarchyTreeView.h b/Engine/Plugins/Animation/ControlRig/Source/ControlRigEditor/Private/Editor/SRigHierarchyTreeView.h
--- a/Engine/Plugins/Animation/ControlRig/Source/ControlRigEditor/Private/Editor/SRigHierarchyTreeView.h
+++ b/Engine/Plugins/Animation/ControlRig/Source/ControlRigEditor/Private/Editor/SRigHierarchyTreeView.h
@@ -383,6 +383,7 @@ public:
 	SLATE_BEGIN_ARGS(SSearchableRigHierarchyTreeView) {}
 		SLATE_ARGUMENT(FRigTreeDelegates, RigTreeDelegates)
 		SLATE_ARGUMENT(FText, InitialFilterText)
+		SLATE_ARGUMENT(float, MaxHeight)
 	SLATE_END_ARGS()
 
 	void Construct(const FArguments& InArgs);
@@ -400,4 +401,5 @@ private:
 	FRigTreeDisplaySettings Settings;
 	TSharedPtr<SSearchBox> SearchBox;
 	TSharedPtr<SRigHierarchyTreeView> TreeView;
+	float MaxHeight = 0.f;
 };
