Control Rig: Represent selection state correctly in schematic

#rnx

diff --git a/Engine/Plugins/Animation/ControlRig/Source/ControlRigEditor/Private/Editor/ControlRigEditor.cpp b/Engine/Plugins/Animation/ControlRig/Source/ControlRigEditor/Private/Editor/ControlRigEditor.cpp
--- a/Engine/Plugins/Animation/ControlRig/Source/ControlRigEditor/Private/Editor/ControlRigEditor.cpp
+++ b/Engine/Plugins/Animation/ControlRig/Source/ControlRigEditor/Private/Editor/ControlRigEditor.cpp
@@ -3623,6 +3623,11 @@ void FControlRigEditor::OnHierarchyModified_AnyThread(ERigHierarchyNotification
 	{
 		return;
 	}
+
+	if(SchematicViewport)
+	{
+		SchematicModel.OnHierarchyModified(InNotif, InHierarchy, InElement);
+	}
 	
 	FRigElementKey Key;
 	if(InElement)
diff --git a/Engine/Plugins/Animation/ControlRig/Source/ControlRigEditor/Private/Editor/ControlRigSchematicModel.cpp b/Engine/Plugins/Animation/ControlRig/Source/ControlRigEditor/Private/Editor/ControlRigSchematicModel.cpp
--- a/Engine/Plugins/Animation/ControlRig/Source/ControlRigEditor/Private/Editor/ControlRigSchematicModel.cpp
+++ b/Engine/Plugins/Animation/ControlRig/Source/ControlRigEditor/Private/Editor/ControlRigSchematicModel.cpp
@@ -33,6 +33,12 @@ bool FControlRigSchematicRigElementKeyNode::IsDragSupported() const
 
 const FText& FControlRigSchematicRigElementKeyNode::GetLabel() const
 {
+	if(IsExpanded())
+	{
+		static const FText& EmptyText = FText();
+		return EmptyText;
+	}
+	
 	FText LabelText = FText::FromString(Key.ToString());
 	if(const FControlRigSchematicModel* ControlRigModel = Cast<FControlRigSchematicModel>(Model))
 	{
@@ -225,6 +231,11 @@ void FControlRigSchematicModel::ConfigureElementKeyNode(FControlRigSchematicRigE
 }
 
 const FControlRigSchematicRigElementKeyNode* FControlRigSchematicModel::FindElementKeyNode(const FRigElementKey& InKey) const
+{
+	return const_cast<FControlRigSchematicModel*>(this)->FindElementKeyNode(InKey);
+}
+
+FControlRigSchematicRigElementKeyNode* FControlRigSchematicModel::FindElementKeyNode(const FRigElementKey& InKey)
 {
 	if(const FGuid* FoundGuid = RigElementKeyToGuid.Find(InKey))
 	{
@@ -864,18 +875,32 @@ const FSlateBrush* FControlRigSchematicModel::GetBrushForNode(const FSchematicGr
 
 FLinearColor FControlRigSchematicModel::GetColorForNode(const FSchematicGraphNode* InNode, int32 InLayerIndex) const
 {
+	if (!ControlRigBeingDebuggedPtr.IsValid())
+	{
+		return Super::GetColorForNode(InNode, InLayerIndex);
+	}
+
+	const URigHierarchy* Hierarchy = ControlRigBeingDebuggedPtr->GetHierarchy();
+	if(Hierarchy == nullptr)
+	{
+		return Super::GetColorForNode(InNode, InLayerIndex);
+	}
+
+	static const FLinearColor SelectionColor = FLinearColor(FColor::FromHex(TEXT("#EBA30A")));
+			
 	if(const FControlRigSchematicRigElementKeyNode* ElementKeyNode = Cast<FControlRigSchematicRigElementKeyNode>(InNode))
 	{
+		if(Hierarchy->IsSelected(ElementKeyNode->GetKey()))
+		{
+			return SelectionColor;
+		}
+
 		if(InLayerIndex == 0) // background
 		{
 			return FLinearColor(0, 0, 0, 0.75);
 		}
 		if(InLayerIndex == 1) // outline
 		{
-			if(ControlRigBlueprint->GetHierarchy()->IsSelected(ElementKeyNode->GetKey()))
-			{
-				return FLinearColor(FColor(0, 112, 224));
-			}
 			return FLinearColor::White;
 		}
 		
@@ -891,28 +916,17 @@ FLinearColor FControlRigSchematicModel::GetColorForNode(const FSchematicGraphNod
 			}
 			case ERigElementType::Control:
 			{
-				if (ControlRigBeingDebuggedPtr.IsValid())
+				if(const FRigControlElement* Control = Hierarchy->Find<FRigControlElement>(ElementKeyNode->GetKey()))
 				{
-					if (const URigHierarchy* Hierarchy = ControlRigBeingDebuggedPtr->GetHierarchy())
-					{
-						if(const FRigControlElement* Control = Hierarchy->Find<FRigControlElement>(ElementKeyNode->GetKey()))
-						{
-							return Control->Settings.ShapeColor;
-						}
-					}
+					return Control->Settings.ShapeColor;
 				}
+				break;
 			}
 			case ERigElementType::Socket:
 			{
-				if (ControlRigBeingDebuggedPtr.IsValid())
+				if(const FRigSocketElement* Socket = Hierarchy->Find<FRigSocketElement>(ElementKeyNode->GetKey()))
 				{
-					if (const URigHierarchy* Hierarchy = ControlRigBeingDebuggedPtr->GetHierarchy())
-					{
-						if(const FRigSocketElement* Socket = Hierarchy->Find<FRigSocketElement>(ElementKeyNode->GetKey()))
-						{
-							return Socket->GetColor(Hierarchy);
-						}
-					}
+					return Socket->GetColor(Hierarchy);
 				}
 				break;
 			}
@@ -927,6 +941,19 @@ FLinearColor FControlRigSchematicModel::GetColorForNode(const FSchematicGraphNod
 			}
 		}
 	}
+	else if(const FSchematicGraphAutoGroupNode* AutoGroupNode = Cast<FSchematicGraphAutoGroupNode>(InNode))
+	{
+		for(int32 ChildIndex = 0; ChildIndex < AutoGroupNode->GetNumChildNodes(); ChildIndex++)
+		{
+			if(const FControlRigSchematicRigElementKeyNode* ChildElementKeyNode = Cast<FControlRigSchematicRigElementKeyNode>(AutoGroupNode->GetChildNode(ChildIndex)))
+			{
+				if(Hierarchy->IsSelected(ChildElementKeyNode->GetKey()))
+				{
+					return SelectionColor;
+				}
+			}
+		}
+	}
 	return Super::GetColorForNode(InNode, InLayerIndex);
 }
 
@@ -957,7 +984,10 @@ ESchematicGraphVisibility::Type FControlRigSchematicModel::GetVisibilityForNode(
 									}
 								}
 							}
-							return ESchematicGraphVisibility::Hidden;
+							else
+							{
+								return ESchematicGraphVisibility::Hidden;
+							}
 						}
 					}
 				}
@@ -1729,4 +1759,34 @@ void FControlRigSchematicModel::OnHideCandidatesForConnector()
 	UpdateElementKeyLinks();
 }
 
+void FControlRigSchematicModel::OnHierarchyModified(ERigHierarchyNotification InNotif, URigHierarchy* InHierarchy, const FRigBaseElement* InElement)
+{
+	switch(InNotif)
+	{
+		case ERigHierarchyNotification::ElementSelected:
+		{
+			if(FControlRigSchematicRigElementKeyNode* Node = FindElementKeyNode(InElement->GetKey()))
+			{
+				if(FSchematicGraphGroupNode* GroupNode = Node->GetGroupNode())
+				{
+					GroupNode->SetExpanded(true);
+				}
+			}
+			break;
+		}
+		case ERigHierarchyNotification::ElementDeselected:
+		{
+			if(FControlRigSchematicRigElementKeyNode* Node = FindElementKeyNode(InElement->GetKey()))
+			{
+				Node->SetExpanded(false);
+			}
+			break;
+		}
+		default:
+		{
+			break;
+		}
+	}
+}
+
 #undef LOCTEXT_NAMESPACE
diff --git a/Engine/Plugins/Animation/ControlRig/Source/ControlRigEditor/Private/Editor/ControlRigSchematicModel.h b/Engine/Plugins/Animation/ControlRig/Source/ControlRigEditor/Private/Editor/ControlRigSchematicModel.h
--- a/Engine/Plugins/Animation/ControlRig/Source/ControlRigEditor/Private/Editor/ControlRigSchematicModel.h
+++ b/Engine/Plugins/Animation/ControlRig/Source/ControlRigEditor/Private/Editor/ControlRigSchematicModel.h
@@ -82,6 +82,7 @@ public:
 	virtual void Tick(float InDeltaTime) override;
 	FControlRigSchematicRigElementKeyNode* AddElementKeyNode(const FRigElementKey& InKey, bool bNotify = true);
 	const FControlRigSchematicRigElementKeyNode* FindElementKeyNode(const FRigElementKey& InKey) const;
+	FControlRigSchematicRigElementKeyNode* FindElementKeyNode(const FRigElementKey& InKey);
 	bool ContainsElementKeyNode(const FRigElementKey& InKey) const;
 	virtual bool RemoveNode(const FGuid& InGuid) override;
 	bool RemoveElementKeyNode(const FRigElementKey& InKey);
@@ -128,6 +129,7 @@ private:
 	void OnShowCandidatesForConnector(const FRigModuleConnector* InModuleConnector);
 	void OnShowCandidatesForMatches(const FModularRigResolveResult& InMatches);
 	void OnHideCandidatesForConnector();
+	void OnHierarchyModified(ERigHierarchyNotification InNotif, URigHierarchy* InHierarchy, const FRigBaseElement* InElement);
 
 	TWeakPtr<FControlRigEditor> ControlRigEditor;
 	TWeakObjectPtr<UControlRigBlueprint> ControlRigBlueprint;
diff --git a/Engine/Source/Editor/AnimationEditorWidgets/Private/SchematicGraphPanel/SchematicGraphNode.cpp b/Engine/Source/Editor/AnimationEditorWidgets/Private/SchematicGraphPanel/SchematicGraphNode.cpp
--- a/Engine/Source/Editor/AnimationEditorWidgets/Private/SchematicGraphPanel/SchematicGraphNode.cpp
+++ b/Engine/Source/Editor/AnimationEditorWidgets/Private/SchematicGraphPanel/SchematicGraphNode.cpp
@@ -18,7 +18,7 @@ FSchematicGraphNode::FSchematicGraphNode()
 	Colors = { FLinearColor::White * 0.4f, FLinearColor::White, FLinearColor::Blue };
 }
 
-const FSchematicGraphNode* FSchematicGraphNode::GetParentNode() const
+FSchematicGraphNode* FSchematicGraphNode::GetParentNode()
 {
 	if(Model && ParentNodeGuid.IsValid())
 	{
@@ -27,6 +27,11 @@ const FSchematicGraphNode* FSchematicGraphNode::GetParentNode() const
 	return nullptr;
 }
 
+const FSchematicGraphNode* FSchematicGraphNode::GetParentNode() const
+{
+	return const_cast<FSchematicGraphNode*>(this)->GetParentNode();
+}
+
 const FGuid& FSchematicGraphNode::GetRootNodeGuid() const
 {
 	if(const FSchematicGraphNode* ParentNode = GetParentNode())
@@ -36,7 +41,7 @@ const FGuid& FSchematicGraphNode::GetRootNodeGuid() const
 	return GetGuid();
 }
 
-const FSchematicGraphNode* FSchematicGraphNode::GetRootNode() const
+FSchematicGraphNode* FSchematicGraphNode::GetRootNode()
 {
 	if(Model)
 	{
@@ -45,12 +50,27 @@ const FSchematicGraphNode* FSchematicGraphNode::GetRootNode() const
 	return nullptr;
 }
 
-const FSchematicGraphGroupNode* FSchematicGraphNode::GetGroupNode() const
+const FSchematicGraphNode* FSchematicGraphNode::GetRootNode() const
+{
+	return const_cast<FSchematicGraphNode*>(this)->GetRootNode();
+}
+
+FSchematicGraphGroupNode* FSchematicGraphNode::GetGroupNode() 
 {
 	return Cast<FSchematicGraphGroupNode>(GetParentNode());
 }
 
+const FSchematicGraphGroupNode* FSchematicGraphNode::GetGroupNode() const
+{
+	return const_cast<FSchematicGraphNode*>(this)->GetGroupNode();
+}
+
 const FSchematicGraphNode* FSchematicGraphNode::GetChildNode(int32 InChildNodeIndex) const
+{
+	return const_cast<FSchematicGraphNode*>(this)->GetChildNode(InChildNodeIndex);
+}
+
+FSchematicGraphNode* FSchematicGraphNode::GetChildNode(int32 InChildNodeIndex)
 {
 	check(ChildNodeGuids.IsValidIndex(InChildNodeIndex));
 
@@ -330,6 +350,18 @@ void FSchematicGraphGroupNode::SetExpanded(bool InExpanded, bool bAutoCloseParen
 			Model->SetLastExpandedNode(GroupNode);
 		}
 	}
+
+	// also collapse all child nodes
+	if(!InExpanded)
+	{
+		for(int32 ChildIndex = 0; ChildIndex < GetNumChildNodes(); ChildIndex++)
+		{
+			if(FSchematicGraphGroupNode* ChildNode = Cast<FSchematicGraphGroupNode>(GetChildNode(ChildIndex)))
+			{
+				ChildNode->SetExpanded(false);
+			}
+		}
+	}
 }
 
 TOptional<ESchematicGraphVisibility::Type> FSchematicGraphGroupNode::GetVisibilityForChildNode(const FSchematicGraphNode* InChildNode) const
diff --git a/Engine/Source/Editor/AnimationEditorWidgets/Public/SchematicGraphPanel/SchematicGraphNode.h b/Engine/Source/Editor/AnimationEditorWidgets/Public/SchematicGraphPanel/SchematicGraphNode.h
--- a/Engine/Source/Editor/AnimationEditorWidgets/Public/SchematicGraphPanel/SchematicGraphNode.h
+++ b/Engine/Source/Editor/AnimationEditorWidgets/Public/SchematicGraphPanel/SchematicGraphNode.h
@@ -26,13 +26,17 @@ public:
 	bool HasParentNode() const { return ParentNodeGuid.IsValid(); }
 	const FGuid& GetParentNodeGuid() const { return ParentNodeGuid; }
 	const FSchematicGraphNode* GetParentNode() const;
+	FSchematicGraphNode* GetParentNode();
 	const FGuid& GetRootNodeGuid() const;
 	const FSchematicGraphNode* GetRootNode() const;
+	FSchematicGraphNode* GetRootNode();
 	bool IsRootNode() const { return !HasParentNode(); }
+	FSchematicGraphGroupNode* GetGroupNode();
 	const FSchematicGraphGroupNode* GetGroupNode() const;
 	int32 GetNumChildNodes() const { return ChildNodeGuids.Num(); }
 	const TArray<FGuid>& GetChildNodeGuids() const { return ChildNodeGuids; }
 	const FSchematicGraphNode* GetChildNode(int32 InChildNodeIndex) const;
+	FSchematicGraphNode* GetChildNode(int32 InChildNodeIndex);
 	TOptional<ESchematicGraphVisibility::Type> GetVisibilityForChildNode(const FGuid& InChildGuid) const;
 	virtual TOptional<ESchematicGraphVisibility::Type> GetVisibilityForChildNode(const FSchematicGraphNode* InChildNode) const { return  TOptional<ESchematicGraphVisibility::Type>(); }
 	TOptional<FVector2d> GetPositionForChildNode(const FGuid& InChildGuid) const;
