Control Rig: Add option to reset controls when baking to control rig
#rnx

diff --git a/Engine/Plugins/Animation/ControlRig/Source/ControlRig/Private/Sequencer/MovieSceneControlRigParameterSection.cpp b/Engine/Plugins/Animation/ControlRig/Source/ControlRig/Private/Sequencer/MovieSceneControlRigParameterSection.cpp
--- a/Engine/Plugins/Animation/ControlRig/Source/ControlRig/Private/Sequencer/MovieSceneControlRigParameterSection.cpp
+++ b/Engine/Plugins/Animation/ControlRig/Source/ControlRig/Private/Sequencer/MovieSceneControlRigParameterSection.cpp
@@ -3028,7 +3028,7 @@ void UMovieSceneControlRigParameterSection::RecordControlRigKey(FFrameNumber Fra
 	}
 }
 
-bool UMovieSceneControlRigParameterSection::LoadAnimSequenceIntoThisSection(UAnimSequence* AnimSequence, UMovieScene* MovieScene, UObject* BoundObject, bool bKeyReduce, float Tolerance, FFrameNumber InStartFrame, EMovieSceneKeyInterpolation InInterpolation)
+bool UMovieSceneControlRigParameterSection::LoadAnimSequenceIntoThisSection(UAnimSequence* AnimSequence, UMovieScene* MovieScene, UObject* BoundObject, bool bKeyReduce, float Tolerance, bool bResetControls, FFrameNumber InStartFrame, EMovieSceneKeyInterpolation InInterpolation)
 {
 	USkeletalMeshComponent* SkelMeshComp = Cast<USkeletalMeshComponent>(BoundObject);
 	
@@ -3103,11 +3103,14 @@ bool UMovieSceneControlRigParameterSection::LoadAnimSequenceIntoThisSection(UAni
 
 	// copy the hierarchy from the CDO into the target control rig.
 	// this ensures that the topology version matches in case of a dynamic hierarchy
-	if(!ControlRig->GetClass()->IsNative())
+	if (bResetControls)
 	{
-		if (UControlRig* CDO = Cast<UControlRig>(ControlRig->GetClass()->GetDefaultObject()))
+		if(!ControlRig->GetClass()->IsNative())
 		{
-			ControlRig->GetHierarchy()->CopyHierarchy(CDO->GetHierarchy());
+			if (UControlRig* CDO = Cast<UControlRig>(ControlRig->GetClass()->GetDefaultObject()))
+			{
+				SourceHierarchy->CopyHierarchy(CDO->GetHierarchy());
+			}
 		}
 	}
 
@@ -3120,16 +3123,22 @@ bool UMovieSceneControlRigParameterSection::LoadAnimSequenceIntoThisSection(UAni
 	{
 		ControlRig->SetBoneInitialTransformsFromRefSkeleton(Skeleton->GetReferenceSkeleton());
 	}
-	ControlRig->RequestConstruction();
-	ControlRig->Evaluate_AnyThread();
+	if (bResetControls)
+	{
+		ControlRig->RequestConstruction();
+		ControlRig->Evaluate_AnyThread();
+	}
 
 	for (int32 Index = 0; Index < NumberOfKeys; ++Index)
 	{
 		const float SequenceSecond = AnimSequence->GetTimeAtFrame(Index);
 		const FFrameNumber FrameNumber = StartFrame + (FMath::Max(FrameRateInFrameNumber.Value, 1) * Index);
 
-		ControlRig->GetHierarchy()->ResetPoseToInitial();
-		ControlRig->GetHierarchy()->ResetCurveValues();
+		if (bResetControls)
+		{
+			SourceHierarchy->ResetPoseToInitial();
+			SourceHierarchy->ResetCurveValues();
+		}
 
 		for (const FFloatCurve& Curve : CurveData.FloatCurves)
 		{
diff --git a/Engine/Plugins/Animation/ControlRig/Source/ControlRig/Public/Sequencer/MovieSceneControlRigParameterSection.h b/Engine/Plugins/Animation/ControlRig/Source/ControlRig/Public/Sequencer/MovieSceneControlRigParameterSection.h
--- a/Engine/Plugins/Animation/ControlRig/Source/ControlRig/Public/Sequencer/MovieSceneControlRigParameterSection.h
+++ b/Engine/Plugins/Animation/ControlRig/Source/ControlRig/Public/Sequencer/MovieSceneControlRigParameterSection.h
@@ -315,7 +315,7 @@ public:
 
 	//Function to load an Anim Sequence into this section. It will automatically resize to the section size.
 	//Will return false if fails or is canceled
-	virtual bool LoadAnimSequenceIntoThisSection(UAnimSequence* Sequence, UMovieScene* MovieScene, UObject* BoundObject, bool bKeyReduce, float Tolerance, FFrameNumber InStartFrame , EMovieSceneKeyInterpolation InInterolation);
+	virtual bool LoadAnimSequenceIntoThisSection(UAnimSequence* Sequence, UMovieScene* MovieScene, UObject* BoundObject, bool bKeyReduce, float Tolerance, bool bResetControls, FFrameNumber InStartFrame , EMovieSceneKeyInterpolation InInterolation);
 #endif
 	const TArray<bool>& GetControlsMask() 
 	{
diff --git a/Engine/Plugins/Animation/ControlRig/Source/ControlRigEditor/Private/BakeToControlRigSettings.h b/Engine/Plugins/Animation/ControlRig/Source/ControlRigEditor/Private/BakeToControlRigSettings.h
--- a/Engine/Plugins/Animation/ControlRig/Source/ControlRigEditor/Private/BakeToControlRigSettings.h
+++ b/Engine/Plugins/Animation/ControlRig/Source/ControlRigEditor/Private/BakeToControlRigSettings.h
@@ -21,6 +21,10 @@ public:
 	UPROPERTY(EditAnywhere, Category = "Reduce Keys")
 	float Tolerance = 0.001f;
 
+	/** Reset controls to initial value on every frame */
+	UPROPERTY(EditAnywhere, Category = "Reset Controls")
+	bool bResetControls = true;
+
 	/** Resets the default properties. */
 	void Reset();
 };
diff --git a/Engine/Plugins/Animation/ControlRig/Source/ControlRigEditor/Private/ControlRigEditorModule.cpp b/Engine/Plugins/Animation/ControlRig/Source/ControlRigEditor/Private/ControlRigEditorModule.cpp
--- a/Engine/Plugins/Animation/ControlRig/Source/ControlRigEditor/Private/ControlRigEditorModule.cpp
+++ b/Engine/Plugins/Animation/ControlRig/Source/ControlRigEditor/Private/ControlRigEditorModule.cpp
@@ -1218,13 +1218,13 @@ void FControlRigEditorModule::BakeToControlRig(UClass* ControlRigClass, UAnimSeq
 			
 				FBakeToControlDelegate BakeCallback = FBakeToControlDelegate::CreateLambda([this, WeakSequencer, LevelSequence, 
 					AnimSequence, MovieScene, ControlRig, ParamSection,ActorTrackGuid, SkelMeshComp]
-				(bool bKeyReduce, float KeyReduceTolerance)
+				(bool bKeyReduce, float KeyReduceTolerance, bool bResetControls)
 				{
 					if (ParamSection)
 					{
 						EMovieSceneKeyInterpolation DefaultInterpolation = WeakSequencer.Pin()->GetKeyInterpolation();
 						ParamSection->LoadAnimSequenceIntoThisSection(AnimSequence, MovieScene, SkelMeshComp, bKeyReduce,
-							KeyReduceTolerance, FFrameNumber(0), DefaultInterpolation);
+							KeyReduceTolerance, bResetControls, FFrameNumber(0), DefaultInterpolation);
 					}
 					WeakSequencer.Pin()->EmptySelection();
 					WeakSequencer.Pin()->SelectSection(ParamSection);
diff --git a/Engine/Plugins/Animation/ControlRig/Source/ControlRigEditor/Private/ControlRigSequencerEditorLibrary.cpp b/Engine/Plugins/Animation/ControlRig/Source/ControlRigEditor/Private/ControlRigSequencerEditorLibrary.cpp
--- a/Engine/Plugins/Animation/ControlRig/Source/ControlRigEditor/Private/ControlRigSequencerEditorLibrary.cpp
+++ b/Engine/Plugins/Animation/ControlRig/Source/ControlRigEditor/Private/ControlRigSequencerEditorLibrary.cpp
@@ -1104,7 +1104,7 @@ void UControlRigSequencerEditorLibrary::SetControlRigWorldTransforms(ULevelSeque
 }
 
 bool UControlRigSequencerEditorLibrary::BakeToControlRig(UWorld* World, ULevelSequence* LevelSequence, UClass* InClass, UAnimSeqExportOption* ExportOptions, bool bReduceKeys, float Tolerance,
-	const FMovieSceneBindingProxy& Binding)
+	const FMovieSceneBindingProxy& Binding, bool bResetControls)
 {
 	UMovieScene* MovieScene = LevelSequence->GetMovieScene();
 	if (Binding.Sequence != LevelSequence)
@@ -1263,7 +1263,7 @@ bool UControlRigSequencerEditorLibrary::BakeToControlRig(UWorld* World, ULevelSe
 					}
 
 					ParamSection->LoadAnimSequenceIntoThisSection(TempAnimSequence, MovieScene, SkeletalMeshComp,
-						bReduceKeys, Tolerance, FFrameNumber(0), DefaultInterpolation);
+						bReduceKeys, Tolerance, bResetControls, FFrameNumber(0), DefaultInterpolation);
 
 					//Turn Off Any Skeletal Animation Tracks
 					UMovieSceneSkeletalAnimationTrack* SkelTrack = Cast<UMovieSceneSkeletalAnimationTrack>(MovieScene->FindTrack(UMovieSceneSkeletalAnimationTrack::StaticClass(), Binding.BindingID, NAME_None));
@@ -1311,7 +1311,7 @@ bool UControlRigSequencerEditorLibrary::BakeToControlRig(UWorld* World, ULevelSe
 }
 
 bool UControlRigSequencerEditorLibrary::LoadAnimSequenceIntoControlRigSection(UMovieSceneSection* MovieSceneSection, UAnimSequence* AnimSequence, USkeletalMeshComponent* SkelMeshComp,
-	FFrameNumber InStartFrame, EMovieSceneTimeUnit TimeUnit,bool bKeyReduce, float Tolerance, EMovieSceneKeyInterpolation Interpolation)
+	FFrameNumber InStartFrame, EMovieSceneTimeUnit TimeUnit,bool bKeyReduce, float Tolerance, EMovieSceneKeyInterpolation Interpolation, bool bResetControls)
 {
 	if (MovieSceneSection == nullptr || AnimSequence == nullptr || SkelMeshComp == nullptr)
 	{
@@ -1328,7 +1328,7 @@ bool UControlRigSequencerEditorLibrary::LoadAnimSequenceIntoControlRigSection(UM
 		{
 			InStartFrame = FFrameRate::TransformTime(FFrameTime(InStartFrame, 0), MovieScene->GetDisplayRate(),MovieScene->GetTickResolution()).RoundToFrame();
 		}
-		return Section->LoadAnimSequenceIntoThisSection(AnimSequence, MovieScene, SkelMeshComp, bKeyReduce, Tolerance, InStartFrame, Interpolation);
+		return Section->LoadAnimSequenceIntoThisSection(AnimSequence, MovieScene, SkelMeshComp, bKeyReduce, Tolerance, bResetControls, InStartFrame, Interpolation);
 	}
 	return false;
 }
diff --git a/Engine/Plugins/Animation/ControlRig/Source/ControlRigEditor/Private/SBakeToControlRigDialog.cpp b/Engine/Plugins/Animation/ControlRig/Source/ControlRigEditor/Private/SBakeToControlRigDialog.cpp
--- a/Engine/Plugins/Animation/ControlRig/Source/ControlRigEditor/Private/SBakeToControlRigDialog.cpp
+++ b/Engine/Plugins/Animation/ControlRig/Source/ControlRigEditor/Private/SBakeToControlRigDialog.cpp
@@ -74,7 +74,7 @@ private:
 		TSharedPtr<SWindow> Window = FSlateApplication::Get().FindWidgetWindow(AsShared());
 		if (BakeSettings && Delegate.IsBound())
 		{
-			Delegate.Execute(BakeSettings->bReduceKeys, BakeSettings->Tolerance);
+			Delegate.Execute(BakeSettings->bReduceKeys, BakeSettings->Tolerance, BakeSettings->bResetControls);
 		}
 		if (Window.IsValid())
 		{
diff --git a/Engine/Plugins/Animation/ControlRig/Source/ControlRigEditor/Private/SBakeToControlRigDialog.h b/Engine/Plugins/Animation/ControlRig/Source/ControlRigEditor/Private/SBakeToControlRigDialog.h
--- a/Engine/Plugins/Animation/ControlRig/Source/ControlRigEditor/Private/SBakeToControlRigDialog.h
+++ b/Engine/Plugins/Animation/ControlRig/Source/ControlRigEditor/Private/SBakeToControlRigDialog.h
@@ -5,7 +5,7 @@
 #include "Delegates/Delegate.h"
 #include "Widgets/SWindow.h"
 
-DECLARE_DELEGATE_TwoParams(FBakeToControlDelegate, bool, float);
+DECLARE_DELEGATE_ThreeParams(FBakeToControlDelegate, bool, float, bool);
 
 struct BakeToControlRigDialog
 {
diff --git a/Engine/Plugins/Animation/ControlRig/Source/ControlRigEditor/Private/Sequencer/ControlRigParameterTrackEditor.cpp b/Engine/Plugins/Animation/ControlRig/Source/ControlRigEditor/Private/Sequencer/ControlRigParameterTrackEditor.cpp
--- a/Engine/Plugins/Animation/ControlRig/Source/ControlRigEditor/Private/Sequencer/ControlRigParameterTrackEditor.cpp
+++ b/Engine/Plugins/Animation/ControlRig/Source/ControlRigEditor/Private/Sequencer/ControlRigParameterTrackEditor.cpp
@@ -1089,7 +1089,7 @@ void FControlRigParameterTrackEditor::BakeToControlRig(UClass* InClass, FGuid Ob
 					GetSequencer()->NotifyMovieSceneDataChanged(EMovieSceneDataChangeType::MovieSceneStructureItemAdded);
 					EMovieSceneKeyInterpolation DefaultInterpolation = SequencerParent->GetKeyInterpolation();
 					ParamSection->LoadAnimSequenceIntoThisSection(TempAnimSequence, OwnerMovieScene, SkelMeshComp,
-						BakeSettings->bReduceKeys, BakeSettings->Tolerance, FFrameNumber(0), DefaultInterpolation);
+						BakeSettings->bReduceKeys, BakeSettings->Tolerance, BakeSettings->bResetControls, FFrameNumber(0), DefaultInterpolation);
 
 					//Turn Off Any Skeletal Animation Tracks
 					TArray<UMovieSceneSkeletalAnimationTrack*> SkelAnimationTracks;
@@ -5277,7 +5277,7 @@ void FControlRigParameterSection::OnAnimationAssetSelectedForFK(const FAssetData
 			UMovieScene* MovieScene = SequencerPtr->GetFocusedMovieSceneSequence()->GetMovieScene();
 			FFrameNumber StartFrame = SequencerPtr->GetLocalTime().Time.GetFrame();
 			EMovieSceneKeyInterpolation DefaultInterpolation = SequencerPtr->GetKeyInterpolation();
-			if (!Section->LoadAnimSequenceIntoThisSection(AnimSequence, MovieScene, SkelMeshComp, false, 0.1f, StartFrame, DefaultInterpolation))
+			if (!Section->LoadAnimSequenceIntoThisSection(AnimSequence, MovieScene, SkelMeshComp, false, 0.1f, true, StartFrame, DefaultInterpolation))
 			{
 				Transaction.Cancel();
 			}
diff --git a/Engine/Plugins/Animation/ControlRig/Source/ControlRigEditor/Public/ControlRigSequencerEditorLibrary.h b/Engine/Plugins/Animation/ControlRig/Source/ControlRigEditor/Public/ControlRigSequencerEditorLibrary.h
--- a/Engine/Plugins/Animation/ControlRig/Source/ControlRigEditor/Public/ControlRigSequencerEditorLibrary.h
+++ b/Engine/Plugins/Animation/ControlRig/Source/ControlRigEditor/Public/ControlRigSequencerEditorLibrary.h
@@ -122,12 +122,13 @@ public:
 	* @param bKeyReduce If true do key reduction based upon Tolerance, if false don't
 	* @param Tolerance If reducing keys, tolerance about which keys will be removed, smaller tolerance, more keys usually.
 	* @param Interpolation The key interpolation type to set the keys, defaults to EMovieSceneKeyInterpolation::SmartAuto
+	* @param bResetControls If true will reset all controls to initial value on every frame
 	* @return returns True if successful, False otherwise
 	*/
 	UFUNCTION(BlueprintCallable, Category = "Editor Scripting | Sequencer Tools | Control Rig")
 	static bool LoadAnimSequenceIntoControlRigSection(UMovieSceneSection* MovieSceneSection, UAnimSequence* AnimSequence, USkeletalMeshComponent* SkelMeshComp,
 		FFrameNumber InStartFrame, EMovieSceneTimeUnit TimeUnit = EMovieSceneTimeUnit::DisplayRate, bool bKeyReduce = false, float Tolerance = 0.001f,
-		EMovieSceneKeyInterpolation Interpolation = EMovieSceneKeyInterpolation::SmartAuto);
+		EMovieSceneKeyInterpolation Interpolation = EMovieSceneKeyInterpolation::SmartAuto, bool bResetControls = true);
 
 	/**
 	* Bake the current animation in the binding to a Control Rig track
@@ -138,11 +139,12 @@ public:
 	* @param bKeyReduce If true do key reduction based upon Tolerance, if false don't
 	* @param Tolerance If reducing keys, tolerance about which keys will be removed, smaller tolerance, more keys usually.
 	* @param Binding The binding upon which to bake
+	* @param bResetControls If true will reset all controls to initial value on every frame
 	* @return returns True if successful, False otherwise
 	*/
 	UFUNCTION(BlueprintCallable, Category = "Editor Scripting | Sequencer Tools | Control Rig")
 	static bool BakeToControlRig(UWorld* World, ULevelSequence* LevelSequence, UClass* ControlRigClass, UAnimSeqExportOption* ExportOptions, bool bReduceKeys, float Tolerance,
-			const FMovieSceneBindingProxy& Binding);
+			const FMovieSceneBindingProxy& Binding, bool bResetControls = true);
 
 	/**
 	* Bake the constraint to keys based on the passed in frames. This will use the open sequencer to bake. See ConstraintsScriptingLibrary to get the list of available constraints
