Sequencer: Anim detail fixes:  Added support for bool,ints/float/doubles for Sequencer, Scale reset to 1.0 properly.

diff --git a/Engine/Plugins/Animation/ControlRig/Source/ControlRigEditor/Private/EditMode/AnimDetailsProxy.cpp b/Engine/Plugins/Animation/ControlRig/Source/ControlRigEditor/Private/EditMode/AnimDetailsProxy.cpp
--- a/Engine/Plugins/Animation/ControlRig/Source/ControlRigEditor/Private/EditMode/AnimDetailsProxy.cpp
+++ b/Engine/Plugins/Animation/ControlRig/Source/ControlRigEditor/Private/EditMode/AnimDetailsProxy.cpp
@@ -602,6 +602,7 @@ static EPropertyKeyedStatus GetKeyedStatusInSection(const UMovieSceneSection* Se
 	for (const FMovieSceneChannelEntry& ChannelEntry : ChannelProxy.GetAllEntries())
 	{
 		if (ChannelEntry.GetChannelTypeName() != FMovieSceneDoubleChannel::StaticStruct()->GetFName() &&
+			ChannelEntry.GetChannelTypeName() != FMovieSceneFloatChannel::StaticStruct()->GetFName() &&
 			ChannelEntry.GetChannelTypeName() != FMovieSceneBoolChannel::StaticStruct()->GetFName() &&
 			ChannelEntry.GetChannelTypeName() != FMovieSceneIntegerChannel::StaticStruct()->GetFName() &&
 			ChannelEntry.GetChannelTypeName() != FMovieSceneByteChannel::StaticStruct()->GetFName())
@@ -903,26 +904,113 @@ void UAnimDetailControlsProxyTransform::SetBindingValueFromCurrent(UObject* InOb
 		const FStructProperty* TransformProperty = Binding.IsValid() ? CastField<FStructProperty>(Binding->GetProperty(*InObject)) : nullptr;
 		if (TransformProperty)
 		{
-			if (TransformProperty->Struct == TBaseStructure<FTransform>::Get())
-			{
-			}
-			else if (TransformProperty->Struct == TBaseStructure<FEulerTransform>::Get())
+			if (TransformProperty->Struct == TBaseStructure<FEulerTransform>::Get())
 			{
 				EulerTransform = FEulerTransform(TLocation, TRotation, TScale);
 				Binding->SetCurrentValue<FTransform>(*InObject, RealTransform);
 			}
 		}
-
 		AActor* ActorThatChanged = nullptr;
 		USceneComponent* SceneComponentThatChanged = nullptr;
 		GetActorAndSceneComponentFromObject(InObject, ActorThatChanged, SceneComponentThatChanged);
 		if (SceneComponentThatChanged)
 		{
-			SceneComponentThatChanged->SetRelativeTransform(RealTransform, false, nullptr, ETeleportType::None);
+			FProperty* ValueProperty = nullptr;
+			FProperty* AxisProperty = nullptr;
+			if (Context.SetKey != EControlRigSetKey::Never)
+			{
+				EControlRigContextChannelToKey ChannelsToKey = (EControlRigContextChannelToKey)Context.KeyMask;
+				if (EnumHasAnyFlags(ChannelsToKey, EControlRigContextChannelToKey::TranslationX) )
+				{
+					ValueProperty = FindFProperty<FProperty>(USceneComponent::StaticClass(), USceneComponent::GetRelativeLocationPropertyName());
+					AxisProperty = FindFProperty<FDoubleProperty>(TBaseStructure<FVector>::Get(), GET_MEMBER_NAME_CHECKED(FVector, X));
+				}
+				if (EnumHasAnyFlags(ChannelsToKey, EControlRigContextChannelToKey::TranslationY))
+				{
+					ValueProperty = FindFProperty<FProperty>(USceneComponent::StaticClass(), USceneComponent::GetRelativeLocationPropertyName());
+					AxisProperty = FindFProperty<FDoubleProperty>(TBaseStructure<FVector>::Get(), GET_MEMBER_NAME_CHECKED(FVector, Y));
+				}
+				if (EnumHasAnyFlags(ChannelsToKey, EControlRigContextChannelToKey::TranslationZ))
+				{
+					ValueProperty = FindFProperty<FProperty>(USceneComponent::StaticClass(), USceneComponent::GetRelativeLocationPropertyName());
+					AxisProperty = FindFProperty<FDoubleProperty>(TBaseStructure<FVector>::Get(), GET_MEMBER_NAME_CHECKED(FVector, Z));
+				}
 
+				if (EnumHasAnyFlags(ChannelsToKey, EControlRigContextChannelToKey::RotationX))
+				{
+					ValueProperty = FindFProperty<FProperty>(USceneComponent::StaticClass(), USceneComponent::GetRelativeRotationPropertyName());
+					AxisProperty = FindFProperty<FDoubleProperty>(TBaseStructure<FRotator>::Get(), GET_MEMBER_NAME_CHECKED(FRotator, Roll));
+				}
+				if (EnumHasAnyFlags(ChannelsToKey, EControlRigContextChannelToKey::RotationY))
+				{
+					ValueProperty = FindFProperty<FProperty>(USceneComponent::StaticClass(), USceneComponent::GetRelativeRotationPropertyName());
+					AxisProperty = FindFProperty<FDoubleProperty>(TBaseStructure<FRotator>::Get(), GET_MEMBER_NAME_CHECKED(FRotator, Pitch));
+				}
+				if (EnumHasAnyFlags(ChannelsToKey, EControlRigContextChannelToKey::RotationZ))
+				{
+					ValueProperty = FindFProperty<FProperty>(USceneComponent::StaticClass(), USceneComponent::GetRelativeRotationPropertyName());
+					AxisProperty = FindFProperty<FDoubleProperty>(TBaseStructure<FRotator>::Get(), GET_MEMBER_NAME_CHECKED(FRotator, Yaw));
+				}
+				if (EnumHasAnyFlags(ChannelsToKey, EControlRigContextChannelToKey::ScaleX))
+				{
+					ValueProperty = FindFProperty<FProperty>(USceneComponent::StaticClass(), USceneComponent::GetRelativeScale3DPropertyName());
+					AxisProperty = FindFProperty<FDoubleProperty>(TBaseStructure<FVector>::Get(), GET_MEMBER_NAME_CHECKED(FVector, X));
+				}
+				if (EnumHasAnyFlags(ChannelsToKey, EControlRigContextChannelToKey::ScaleY))
+				{
+					ValueProperty = FindFProperty<FProperty>(USceneComponent::StaticClass(), USceneComponent::GetRelativeScale3DPropertyName());
+					AxisProperty = FindFProperty<FDoubleProperty>(TBaseStructure<FVector>::Get(), GET_MEMBER_NAME_CHECKED(FVector, Y));
+				}
+				if (EnumHasAnyFlags(ChannelsToKey, EControlRigContextChannelToKey::ScaleZ))
+				{
+					ValueProperty = FindFProperty<FProperty>(USceneComponent::StaticClass(), USceneComponent::GetRelativeScale3DPropertyName());
+					AxisProperty = FindFProperty<FDoubleProperty>(TBaseStructure<FVector>::Get(), GET_MEMBER_NAME_CHECKED(FVector, Z));
+				}
+			}
+			// Have to downcast here because of function overloading and inheritance not playing nicely
+			if (ValueProperty)
+			{
+				TArray<UObject*> ModifiedObjects = { InObject };
+				FPropertyChangedEvent PropertyChangedEvent(ValueProperty, EPropertyChangeType::ValueSet, MakeArrayView(ModifiedObjects));
+				FEditPropertyChain PropertyChain;
+
+				if (AxisProperty)
+				{
+					PropertyChain.AddHead(AxisProperty);
+				}
+				PropertyChain.AddHead(ValueProperty);
+				((UObject*)SceneComponentThatChanged)->PreEditChange(PropertyChain);
+
+				if (ActorThatChanged && ActorThatChanged->GetRootComponent() == SceneComponentThatChanged)
+				{
+					((UObject*)ActorThatChanged)->PreEditChange(PropertyChain);
+				}
+			}
+			SceneComponentThatChanged->SetRelativeTransform(RealTransform, false, nullptr, ETeleportType::None);
 			// Force the location and rotation values to avoid Rot->Quat->Rot conversions
 			SceneComponentThatChanged->SetRelativeLocation_Direct(TLocation);
 			SceneComponentThatChanged->SetRelativeRotation_Direct(TRotation);
+
+			if (ValueProperty)
+			{
+				TArray<UObject*> ModifiedObjects = { InObject };
+				FPropertyChangedEvent PropertyChangedEvent(ValueProperty, EPropertyChangeType::ValueSet, MakeArrayView(ModifiedObjects));
+				FEditPropertyChain PropertyChain;
+
+				if (AxisProperty)
+				{
+					PropertyChain.AddHead(AxisProperty);
+				}
+				PropertyChain.AddHead(ValueProperty);
+				FPropertyChangedChainEvent PropertyChangedChainEvent(PropertyChain, PropertyChangedEvent);
+
+				if (ActorThatChanged && ActorThatChanged->GetRootComponent() == SceneComponentThatChanged)
+				{
+					((UObject*)ActorThatChanged)->PostEditChangeChainProperty(PropertyChangedChainEvent);
+				}
+				((UObject*)SceneComponentThatChanged)->PostEditChangeChainProperty(PropertyChangedChainEvent);
+
+			}
 			GUnrealEd->UpdatePivotLocationForSelection();
 		}
 	}
@@ -1919,14 +2007,25 @@ void UAnimDetailControlsProxyFloat::UpdatePropertyNames(IDetailLayoutBuilder& De
 	if (ValuePropertyHandle)
 	{
 		ValuePropertyHandle->SetPropertyDisplayName(FText::FromName(GetName()));
-	}
-	if (bIsIndividual)
-	{
-		const FString PropertyPath = FString::Printf(TEXT("%s.%s"), GET_MEMBER_NAME_STRING_CHECKED(UAnimDetailControlsProxyFloat, Float), GET_MEMBER_NAME_STRING_CHECKED(FAnimDetailProxyFloat, Float));
-		ValuePropertyHandle = DetailBuilder.GetProperty(FName(*PropertyPath), GetClass());
-		if (ValuePropertyHandle)
+
+		if (GetControlElements().Num() == 0 && GetSequencerItems().Num() == 1)
 		{
-			ValuePropertyHandle->SetPropertyDisplayName(FText::FromName(GetName()));
+			TArray<FBindingAndTrack> SItems =  GetSequencerItems();
+			const FString PropertyPath = FString::Printf(TEXT("%s.%s"), GET_MEMBER_NAME_STRING_CHECKED(UAnimDetailControlsProxyFloat, Float), GET_MEMBER_NAME_STRING_CHECKED(FAnimDetailProxyFloat, Float));
+			ValuePropertyHandle = DetailBuilder.GetProperty(FName(*PropertyPath), GetClass());
+			if (ValuePropertyHandle)
+			{
+				ValuePropertyHandle->SetPropertyDisplayName(FText::FromName(SItems[0].Binding->GetPropertyName()));
+			}
+		}
+		else if (bIsIndividual)
+		{
+			const FString PropertyPath = FString::Printf(TEXT("%s.%s"), GET_MEMBER_NAME_STRING_CHECKED(UAnimDetailControlsProxyFloat, Float), GET_MEMBER_NAME_STRING_CHECKED(FAnimDetailProxyFloat, Float));
+			ValuePropertyHandle = DetailBuilder.GetProperty(FName(*PropertyPath), GetClass());
+			if (ValuePropertyHandle)
+			{
+				ValuePropertyHandle->SetPropertyDisplayName(FText::FromName(GetName()));
+			}
 		}
 	}
 }
@@ -2005,7 +2104,23 @@ void UAnimDetailControlsProxyFloat::ValueChanged()
 		{
 			if (Binding.Binding.IsValid())
 			{
-				if (TOptional<double> Value = Binding.Binding->GetOptionalValue<double>(*SItems.Key))
+				TOptional<double> Value;
+				if (FProperty* Property = Binding.Binding->GetProperty((*(SItems.Value.OwnerObject.Get()))))
+				{
+					if (Property->IsA(FDoubleProperty::StaticClass()))
+					{
+						Value = Binding.Binding->GetOptionalValue<double>(*SItems.Key);
+					}
+					else if (Property->IsA(FFloatProperty::StaticClass()))
+					{
+						TOptional<float> FVal = Binding.Binding->GetOptionalValue<float>(*SItems.Key);
+						if (FVal.IsSet())
+						{
+							Value = (double)FVal.GetValue();
+						}
+					}
+				}
+				if (Value)
 				{
 					if (LastValue.IsSet())
 					{
@@ -2081,7 +2196,18 @@ void UAnimDetailControlsProxyFloat::SetBindingValueFromCurrent(UObject* InObject
 {
 	if (InObject && Binding.IsValid())
 	{
-		Binding->SetCurrentValue<double>(*InObject, Float.Float);
+		if (FProperty* Property = Binding->GetProperty((*(InObject))))
+		{
+			if (Property->IsA(FDoubleProperty::StaticClass()))
+			{
+				Binding->SetCurrentValue<double>(*InObject, Float.Float);
+			}
+			else if (Property->IsA(FFloatProperty::StaticClass()))
+			{
+				float FVal = (float)Float.Float;
+				Binding->SetCurrentValue<float>(*InObject, FVal);
+			}
+		}
 	}
 }
 //////UAnimDetailControlsProxyBool////////
@@ -2093,7 +2219,17 @@ void UAnimDetailControlsProxyBool::UpdatePropertyNames(IDetailLayoutBuilder& Det
 	{
 		ValuePropertyHandle->SetPropertyDisplayName(FText::FromName(GetName()));
 	}
-	if (bIsIndividual)
+	if (GetControlElements().Num() == 0 && GetSequencerItems().Num() == 1)
+	{
+		TArray<FBindingAndTrack> SItems = GetSequencerItems();
+		const FString PropertyPath = FString::Printf(TEXT("%s.%s"), GET_MEMBER_NAME_STRING_CHECKED(UAnimDetailControlsProxyBool, Bool), GET_MEMBER_NAME_STRING_CHECKED(FAnimDetailProxyBool, Bool));
+		ValuePropertyHandle = DetailBuilder.GetProperty(FName(*PropertyPath), GetClass());
+		if (ValuePropertyHandle)
+		{
+			ValuePropertyHandle->SetPropertyDisplayName(FText::FromName(SItems[0].Binding->GetPropertyName()));
+		}
+	}
+	else if (bIsIndividual)
 	{
 		const FString PropertyPath = FString::Printf(TEXT("%s.%s"), GET_MEMBER_NAME_STRING_CHECKED(UAnimDetailControlsProxyBool, Bool), GET_MEMBER_NAME_STRING_CHECKED(FAnimDetailProxyBool, Bool));
 		ValuePropertyHandle = DetailBuilder.GetProperty(FName(*PropertyPath), GetClass());
@@ -2265,7 +2401,17 @@ void UAnimDetailControlsProxyInteger::UpdatePropertyNames(IDetailLayoutBuilder&
 	{
 		ValuePropertyHandle->SetPropertyDisplayName(FText::FromName(GetName()));
 	}
-	if (bIsIndividual)
+	if (GetControlElements().Num() == 0 && GetSequencerItems().Num() == 1)
+	{
+		TArray<FBindingAndTrack> SItems = GetSequencerItems();
+		const FString PropertyPath = FString::Printf(TEXT("%s.%s"), GET_MEMBER_NAME_STRING_CHECKED(UAnimDetailControlsProxyInteger, Integer), GET_MEMBER_NAME_STRING_CHECKED(FAnimDetailProxyInteger, Integer));
+		ValuePropertyHandle = DetailBuilder.GetProperty(FName(*PropertyPath), GetClass());
+		if (ValuePropertyHandle)
+		{
+			ValuePropertyHandle->SetPropertyDisplayName(FText::FromName(SItems[0].Binding->GetPropertyName()));
+		}
+	}
+	else if (bIsIndividual)
 	{
 		const FString PropertyPath = FString::Printf(TEXT("%s.%s"), GET_MEMBER_NAME_STRING_CHECKED(UAnimDetailControlsProxyInteger, Integer), GET_MEMBER_NAME_STRING_CHECKED(FAnimDetailProxyInteger, Integer));
 		ValuePropertyHandle = DetailBuilder.GetProperty(FName(*PropertyPath), GetClass());
@@ -2435,7 +2581,17 @@ void UAnimDetailControlsProxyEnum::UpdatePropertyNames(IDetailLayoutBuilder& Det
 	{
 		ValuePropertyHandle->SetPropertyDisplayName(FText::FromName(GetName()));
 	}
-	if (bIsIndividual)
+	if (GetControlElements().Num() == 0 && GetSequencerItems().Num() == 1)
+	{
+		TArray<FBindingAndTrack> SItems = GetSequencerItems();
+		const FString PropertyPath = FString::Printf(TEXT("%s.%s"), GET_MEMBER_NAME_STRING_CHECKED(UAnimDetailControlsProxyEnum, Enum), GET_MEMBER_NAME_STRING_CHECKED(FControlRigEnumControlProxyValue, EnumIndex));
+		ValuePropertyHandle = DetailBuilder.GetProperty(FName(*PropertyPath), GetClass());
+		if (ValuePropertyHandle)
+		{
+			ValuePropertyHandle->SetPropertyDisplayName(FText::FromName(SItems[0].Binding->GetPropertyName()));
+		}
+	}
+	else if (bIsIndividual)
 	{
 		const FString PropertyPath = FString::Printf(TEXT("%s.%s"), GET_MEMBER_NAME_STRING_CHECKED(UAnimDetailControlsProxyEnum, Enum), GET_MEMBER_NAME_STRING_CHECKED(FControlRigEnumControlProxyValue, EnumIndex));
 		ValuePropertyHandle = DetailBuilder.GetProperty(FName(*PropertyPath), GetClass());
diff --git a/Engine/Plugins/Animation/ControlRig/Source/ControlRigEditor/Private/EditMode/AnimDetailsProxy.h b/Engine/Plugins/Animation/ControlRig/Source/ControlRigEditor/Private/EditMode/AnimDetailsProxy.h
--- a/Engine/Plugins/Animation/ControlRig/Source/ControlRigEditor/Private/EditMode/AnimDetailsProxy.h
+++ b/Engine/Plugins/Animation/ControlRig/Source/ControlRigEditor/Private/EditMode/AnimDetailsProxy.h
@@ -181,17 +181,17 @@ struct  FAnimDetailProxyScale
 	GENERATED_BODY()
 
 	UPROPERTY(EditAnywhere, Interp, Category = "Scale", meta = (Delta = "0.5", SliderExponent = "1", LinearDeltaSensitivity = "1"))
-	double SX = 0.0;
+	double SX = 1.0;
 
 	UPROPERTY(EditAnywhere, Interp, Category = "Scale", meta = (Delta = "0.5", SliderExponent = "1", LinearDeltaSensitivity = "1"))
-	double SY = 0.0;
+	double SY = 1.0;
 
 	UPROPERTY(EditAnywhere, Interp, Category = "Scale", meta = (Delta = "0.5", SliderExponent = "1", LinearDeltaSensitivity = "1"))
-	double SZ = 0.0;
+	double SZ = 1.0;
 
 	FAnimDetailVectorState State;
 
-	FAnimDetailProxyScale() : SX(0.0), SY(0.0), SZ(0.0) {};
+	FAnimDetailProxyScale() : SX(1.0), SY(1.0), SZ(1.0) {};
 	FAnimDetailProxyScale(const FVector& InVector, const FAnimDetailVectorState& InState) 
 	{
 		SX = InVector.X; SY = InVector.Y; SZ = InVector.Z; State = InState;
diff --git a/Engine/Plugins/Animation/ControlRig/Source/ControlRigEditor/Private/EditMode/ControlsProxyDetailCustomization.cpp b/Engine/Plugins/Animation/ControlRig/Source/ControlRigEditor/Private/EditMode/ControlsProxyDetailCustomization.cpp
--- a/Engine/Plugins/Animation/ControlRig/Source/ControlRigEditor/Private/EditMode/ControlsProxyDetailCustomization.cpp
+++ b/Engine/Plugins/Animation/ControlRig/Source/ControlRigEditor/Private/EditMode/ControlsProxyDetailCustomization.cpp
@@ -271,7 +271,6 @@ void FAnimDetailValueCustomization::CustomizeChildren(TSharedRef<class IProperty
 		const FText PropertyDisplayText = ChildHandle->GetPropertyDisplayName();
 
 		TSharedRef<SWidget> ValueWidget = MakeChildWidget(StructPropertyHandle, ChildHandle);
-		TSharedRef<SWidget> LabelWidget = ChildHandle->CreatePropertyNameWidget();
 
 		// Add the individual properties as children as well so the vector can be expanded for more room
 		
@@ -373,7 +372,6 @@ void FAnimDetailValueCustomization::MakeHeaderRow(TSharedRef<class IPropertyHand
 	FProperty* Property = StructPropertyHandle->GetProperty(); // Is a FProperty*
 	TWeakPtr<IPropertyHandle> StructWeakHandlePtr = StructPropertyHandle;
 	TSharedPtr<SHorizontalBox> HorizontalBox;
-
 	Row.NameContent()
 		[
 			StructPropertyHandle->CreatePropertyNameWidget()
diff --git a/Engine/Plugins/Animation/ControlRig/Source/ControlRigEditor/Private/EditMode/SControlRigDetails.cpp b/Engine/Plugins/Animation/ControlRig/Source/ControlRigEditor/Private/EditMode/SControlRigDetails.cpp
--- a/Engine/Plugins/Animation/ControlRig/Source/ControlRigEditor/Private/EditMode/SControlRigDetails.cpp
+++ b/Engine/Plugins/Animation/ControlRig/Source/ControlRigEditor/Private/EditMode/SControlRigDetails.cpp
@@ -31,6 +31,10 @@
 #include "MVVM/ViewModels/SequencerEditorViewModel.h"
 #include "Tracks/MovieScenePropertyTrack.h"
 #include "Tracks/MovieScene3DTransformTrack.h"
+#include "Tracks/MovieSceneIntegerTrack.h"
+#include "Tracks/MovieSceneDoubleTrack.h"
+#include "Tracks/MovieSceneFloatTrack.h"
+#include "Tracks/MovieSceneBoolTrack.h"
 #include "MovieSceneCommonHelpers.h"
 
 #define LOCTEXT_NAMESPACE "ControlRigDetails"
@@ -252,13 +256,33 @@ void SControlRigDetails::HandleSequencerObjects(TMap<UObject*, FArrayOfPropertyT
 				{
 					if (UMovieScenePropertyTrack* PropTrack = Cast<UMovieScenePropertyTrack>(Track))
 					{
-						if (PropTrack->IsA<UMovieScene3DTransformTrack>())
+						auto AddBinding = [this, PropTrack](UObject* InObject,FSequencerProxyPerType& Binding)
 						{
-							FSequencerProxyPerType& Binding = ProxyPerType.FindOrAdd(ERigControlType::Transform);
-							TArray<FBindingAndTrack>& Bindings = Binding.Bindings.FindOrAdd(Pair.Key);
+							TArray<FBindingAndTrack>& Bindings = Binding.Bindings.FindOrAdd(InObject);
 							TSharedPtr<FTrackInstancePropertyBindings> PropertyBindings = MakeShareable(new FTrackInstancePropertyBindings(PropTrack->GetPropertyName(), PropTrack->GetPropertyPath().ToString()));
-							FBindingAndTrack BindingAndTrack(PropertyBindings, Track);
+							FBindingAndTrack BindingAndTrack(PropertyBindings, PropTrack);
 							Bindings.Add(BindingAndTrack);
+						};
+						if (PropTrack->IsA<UMovieScene3DTransformTrack>())
+						{
+							FSequencerProxyPerType& Binding = ProxyPerType.FindOrAdd(ERigControlType::Transform);
+							AddBinding(Pair.Key, Binding);
+						}
+						else if (PropTrack->IsA<UMovieSceneBoolTrack>())
+						{
+							FSequencerProxyPerType& Binding = ProxyPerType.FindOrAdd(ERigControlType::Bool);
+							AddBinding(Pair.Key, Binding);
+						}
+						else if (PropTrack->IsA<UMovieSceneIntegerTrack>())
+						{
+							FSequencerProxyPerType& Binding = ProxyPerType.FindOrAdd(ERigControlType::Integer);
+							AddBinding(Pair.Key, Binding);
+						}
+						else if (PropTrack->IsA<UMovieSceneDoubleTrack>() ||
+							PropTrack->IsA<UMovieSceneFloatTrack>())
+						{
+							FSequencerProxyPerType& Binding = ProxyPerType.FindOrAdd(ERigControlType::Float);
+							AddBinding(Pair.Key, Binding);
 						}
 					}
 				}
