Sequencer: Anim Details: Add missing integer details

diff --git a/Engine/Plugins/Animation/ControlRig/Source/ControlRigEditor/Private/ControlRigEditorModule.cpp b/Engine/Plugins/Animation/ControlRig/Source/ControlRigEditor/Private/ControlRigEditorModule.cpp
--- a/Engine/Plugins/Animation/ControlRig/Source/ControlRigEditor/Private/ControlRigEditorModule.cpp
+++ b/Engine/Plugins/Animation/ControlRig/Source/ControlRigEditor/Private/ControlRigEditorModule.cpp
@@ -145,6 +145,17 @@ void FControlRigEditorModule::StartupModule()
 		Section->AddCategory("Attributes");
 	}
 
+	ClassesToUnregisterOnShutdown.Add(UAnimDetailControlsProxyInteger::StaticClass()->GetFName());
+	PropertyEditorModule.RegisterCustomClassLayout(UAnimDetailControlsProxyInteger::StaticClass()->GetFName(), FOnGetDetailCustomizationInstance::CreateLambda([]() {
+		FName CategoryName = FName(TEXT("Integer"));
+		return FAnimDetailProxyDetails::MakeInstance(CategoryName);
+		}));
+
+	{
+		TSharedRef<FPropertySection> Section = RegisterPropertySection(PropertyModule, "AnimDetailControlsProxyInteger", "General", LOCTEXT("General", "General"));
+		Section->AddCategory("Attributes");
+	}
+
 	ClassesToUnregisterOnShutdown.Add(UAnimDetailControlsProxyBool::StaticClass()->GetFName());
 	PropertyEditorModule.RegisterCustomClassLayout(UAnimDetailControlsProxyBool::StaticClass()->GetFName(), FOnGetDetailCustomizationInstance::CreateLambda([]() {
 		FName CategoryName = FName(TEXT("Bool"));
