Control Rig: Fix mirror module not mirroring default config values
#rnx

diff --git a/Engine/Plugins/Animation/ControlRig/Source/ControlRig/Private/ModularRigController.cpp b/Engine/Plugins/Animation/ControlRig/Source/ControlRig/Private/ModularRigController.cpp
--- a/Engine/Plugins/Animation/ControlRig/Source/ControlRig/Private/ModularRigController.cpp
+++ b/Engine/Plugins/Animation/ControlRig/Source/ControlRig/Private/ModularRigController.cpp
@@ -1150,16 +1150,27 @@ FString UModularRigController::MirrorModule(const FString& InModulePath, const F
 		BindModuleVariable(NewModulePath, Pair.Key, NewSourcePath, bSetupUndo);
 	}
 
-	for (const TPair<FName, FString>& Pair : OriginalModule->ConfigValues)
-	{
+	TSet<FName> ConfigValueSet;
 #if WITH_EDITOR
-		const FProperty* Property = NewModule->Class->FindPropertyByName(Pair.Key);
-		if (!Property)
+	for (TFieldIterator<FProperty> PropertyIt(OriginalModule->Class.Get()); PropertyIt; ++PropertyIt)
+	{
+		const FProperty* Property = *PropertyIt;
+		
+		// skip advanced properties for now
+		if (Property->HasAnyPropertyFlags(CPF_AdvancedDisplay))
+		{
+			continue;
+		}
+
+		// skip non-public properties for now
+		const bool bIsPublic = Property->HasAnyPropertyFlags(CPF_Edit | CPF_EditConst);
+		const bool bIsInstanceEditable = !Property->HasAnyPropertyFlags(CPF_DisableEditOnInstance);
+		if(!bIsPublic || !bIsInstanceEditable)
 		{
 			continue;
 		}
 
-		FString CPPType = Property->GetCPPType();
+		const FString CPPType = Property->GetCPPType();
 		bool bIsVector;
 		if (CPPType == TEXT("FVector"))
 		{
@@ -1171,29 +1182,58 @@ FString UModularRigController::MirrorModule(const FString& InModulePath, const F
 		}
 		else
 		{
-			SetConfigValueInModule(NewModulePath, Pair.Key, Pair.Value, bSetupUndo);
 			continue;
 		}
 
-		FString NewValue;
-		if (bIsVector)
+		FString NewValueStr;
+		if (const FString* OriginalValue = OriginalModule->ConfigValues.Find(Property->GetFName()))
 		{
-			FVector Value;
-			FBlueprintEditorUtils::PropertyValueFromString_Direct(Property, Pair.Value, (uint8*)&Value);
-			Value = InSettings.MirrorVector(Value);
-			FBlueprintEditorUtils::PropertyValueToString_Direct(Property, (uint8*)&Value, NewValue, nullptr);
+			if (bIsVector)
+			{
+				FVector Value;
+				FBlueprintEditorUtils::PropertyValueFromString_Direct(Property, *OriginalValue, (uint8*)&Value);
+				Value = InSettings.MirrorVector(Value);
+				FBlueprintEditorUtils::PropertyValueToString_Direct(Property, (uint8*)&Value, NewValueStr, nullptr);
+			}
+			else
+			{
+				FTransform Value;
+				FBlueprintEditorUtils::PropertyValueFromString_Direct(Property, *OriginalValue, (uint8*)&Value);
+				Value = InSettings.MirrorTransform(Value);
+				FBlueprintEditorUtils::PropertyValueToString_Direct(Property, (uint8*)&Value, NewValueStr, nullptr);
+			}
 		}
 		else
 		{
-			FTransform Value;
-			FBlueprintEditorUtils::PropertyValueFromString_Direct(Property, Pair.Value, (uint8*)&Value);
-			Value = InSettings.MirrorTransform(Value);
-			FBlueprintEditorUtils::PropertyValueToString_Direct(Property, (uint8*)&Value, NewValue, nullptr);
+			if (UControlRig* CDO = OriginalModule->Class->GetDefaultObject<UControlRig>())
+			{
+				if (bIsVector)
+				{
+					FVector NewVector = *Property->ContainerPtrToValuePtr<FVector>(CDO);
+					NewVector = InSettings.MirrorVector(NewVector);
+					FBlueprintEditorUtils::PropertyValueToString_Direct(Property, (uint8*)&NewVector, NewValueStr, nullptr);
+				}
+				else
+				{
+					FTransform NewTransform = *Property->ContainerPtrToValuePtr<FTransform>(CDO);
+					NewTransform = InSettings.MirrorTransform(NewTransform);
+					FBlueprintEditorUtils::PropertyValueToString_Direct(Property, (uint8*)&NewTransform, NewValueStr, nullptr);
+				}
+			}
 		}
-		SetConfigValueInModule(NewModulePath, Pair.Key, NewValue, bSetupUndo);
-#else
-		SetConfigValueInModule(NewModulePath, Pair.Key, Pair.Value, bSetupUndo);
+
+		ConfigValueSet.Add(Property->GetFName());
+		SetConfigValueInModule(NewModulePath, Property->GetFName(), NewValueStr, bSetupUndo);
+	}
 #endif
+
+	// Add any other config value that was set in the original module, but was not mirrored
+	for (const TPair<FName, FString>& Pair : OriginalModule->ConfigValues)
+	{
+		if (!ConfigValueSet.Contains(Pair.Key))
+		{
+			SetConfigValueInModule(NewModulePath, Pair.Key, Pair.Value, bSetupUndo);
+		}
 	}
 	
 	return NewModulePath;
