Control Rig: Introduce namespace type enum for metadata nodes

#rnx

diff --git a/Engine/Plugins/Animation/ControlRig/Source/ControlRig/Private/Rigs/RigHierarchy.cpp b/Engine/Plugins/Animation/ControlRig/Source/ControlRig/Private/Rigs/RigHierarchy.cpp
--- a/Engine/Plugins/Animation/ControlRig/Source/ControlRig/Private/Rigs/RigHierarchy.cpp
+++ b/Engine/Plugins/Animation/ControlRig/Source/ControlRig/Private/Rigs/RigHierarchy.cpp
@@ -647,8 +647,10 @@ void URigHierarchy::CopyHierarchy(URigHierarchy* InHierarchy)
 
 		// Increment the topology version to invalidate our cached children.
 		IncrementTopologyVersion();
-	
-		MetadataVersion = InHierarchy->GetMetadataVersion();
+
+		// Keep incrementing the metadata version so that the UI can refresh.
+		MetadataVersion += InHierarchy->GetMetadataVersion();
+		MetadataTagVersion += InHierarchy->GetMetadataTagVersion();
 	}
 
 	if (MetadataChangedDelegate.IsBound())
@@ -5617,7 +5619,7 @@ FRigBaseMetadata* URigHierarchy::GetMetadataForElement(FRigBaseElement* InElemen
 	FMetadataStorage& Storage = ElementMetadata[InElement->MetadataStorageIndex];
 
 	// If repeatedly accessing the same element, store it here for faster access to avoid map lookups.
-	if (Storage.LastAccessName == InName && Storage.LastAccessMetadata->GetType() == InType)
+	if (Storage.LastAccessMetadata && Storage.LastAccessName == InName && Storage.LastAccessMetadata->GetType() == InType)
 	{
 		return Storage.LastAccessMetadata;
 	}
@@ -5646,6 +5648,11 @@ FRigBaseMetadata* URigHierarchy::GetMetadataForElement(FRigBaseElement* InElemen
 		// No metadata with that name existed on the element, create one from scratch.
 		Metadata = FRigBaseMetadata::MakeMetadata(InName, InType);
 		Storage.MetadataMap.Add(InName, Metadata);
+
+		if (bInNotify)
+		{
+			OnMetadataChanged(InElement->Key, InName);
+		}
 	}
 	
 	Storage.LastAccessName = InName;
diff --git a/Engine/Plugins/Animation/ControlRig/Source/ControlRig/Private/Units/Hierarchy/RigUnit_Metadata.cpp b/Engine/Plugins/Animation/ControlRig/Source/ControlRig/Private/Units/Hierarchy/RigUnit_Metadata.cpp
--- a/Engine/Plugins/Animation/ControlRig/Source/ControlRig/Private/Units/Hierarchy/RigUnit_Metadata.cpp
+++ b/Engine/Plugins/Animation/ControlRig/Source/ControlRig/Private/Units/Hierarchy/RigUnit_Metadata.cpp
@@ -13,7 +13,7 @@
 
 FName FRigDispatch_MetadataBase::ItemArgName = TEXT("Item");
 FName FRigDispatch_MetadataBase::NameArgName = TEXT("Name");
-FName FRigDispatch_MetadataBase::UseNameSpaceArgName = TEXT("UseNameSpace");
+FName FRigDispatch_MetadataBase::NameSpaceArgName = TEXT("NameSpace");
 FName FRigDispatch_MetadataBase::CacheArgName = TEXT("Cache");
 FName FRigDispatch_MetadataBase::DefaultArgName = TEXT("Default");
 FName FRigDispatch_MetadataBase::ValueArgName = TEXT("Value");
@@ -72,7 +72,7 @@ const TArray<FRigVMTemplateArgumentInfo>& FRigDispatch_MetadataBase::GetArgument
 	{
 		ItemArgIndex = Infos.Emplace(ItemArgName, ERigVMPinDirection::Input, FRigVMRegistry::Get().GetTypeIndex<FRigElementKey>());
 		NameArgIndex = Infos.Emplace(NameArgName, ERigVMPinDirection::Input, RigVMTypeUtils::TypeIndex::FName);
-		UseNameSpaceArgIndex = Infos.Emplace(UseNameSpaceArgName, ERigVMPinDirection::Input, RigVMTypeUtils::TypeIndex::Bool);
+		NameSpaceArgIndex = Infos.Emplace(NameSpaceArgName, ERigVMPinDirection::Input, FRigVMRegistry::Get().GetTypeIndex<ERigMetaDataNameSpace>());
 		CacheArgIndex = Infos.Emplace(CacheArgName, ERigVMPinDirection::Hidden, FRigVMRegistry::Get().GetTypeIndex<FCachedRigElement>());
 	};
 	return Infos;
@@ -90,9 +90,9 @@ FText FRigDispatch_MetadataBase::GetArgumentTooltip(const FName& InArgumentName,
 	{
 		return NSLOCTEXT("FRigDispatch_MetadataBase", "NameArgTooltip", "The name of the metadata");
 	}
-	if(InArgumentName == UseNameSpaceArgName)
+	if(InArgumentName == NameSpaceArgName)
 	{
-		return NSLOCTEXT("FRigDispatch_MetadataBase", "UseNameSpaceArgTooltip", "If checked the metadata will be available only in this name space / module");
+		return NSLOCTEXT("FRigDispatch_MetadataBase", "NameSpaceArgTooltip", "Defines in which namespace the metadata will be looked up");
 	}
 	if(InArgumentName == DefaultArgName)
 	{
@@ -116,9 +116,10 @@ FText FRigDispatch_MetadataBase::GetArgumentTooltip(const FName& InArgumentName,
 FString FRigDispatch_MetadataBase::GetArgumentDefaultValue(const FName& InArgumentName,
 	TRigVMTypeIndex InTypeIndex) const
 {
-	if(InArgumentName == UseNameSpaceArgName)
+	if(InArgumentName == NameSpaceArgName)
 	{
-		return TrueString;
+		static const FString SelfString = StaticEnum<ERigMetaDataNameSpace>()->GetDisplayNameTextByValue((int64)ERigMetaDataNameSpace::Self).ToString();
+		return SelfString;
 	}
 	return FRigDispatchFactory::GetArgumentDefaultValue(InArgumentName, InTypeIndex);
 }
@@ -185,7 +186,7 @@ const TArray<FRigVMTemplateArgumentInfo>& FRigDispatch_GetMetadata::GetArgumentI
 
 FRigBaseMetadata* FRigDispatch_GetMetadata::FindMetadata(const FRigVMExtendedExecuteContext& InContext,
                                                          const FRigElementKey& InKey, const FName& InName,
-                                                         ERigMetadataType InType, bool bUseNameSpace, FCachedRigElement& Cache)
+                                                         ERigMetadataType InType, ERigMetaDataNameSpace InNameSpace, FCachedRigElement& Cache)
 {
 	const FControlRigExecuteContext& ExecuteContext = InContext.GetPublicData<FControlRigExecuteContext>();
 	if(Cache.UpdateCache(InKey, ExecuteContext.Hierarchy))
@@ -193,7 +194,7 @@ FRigBaseMetadata* FRigDispatch_GetMetadata::FindMetadata(const FRigVMExtendedExe
 		if(FRigBaseElement* Element = ExecuteContext.Hierarchy->Get(Cache.GetIndex()))
 		{
 			// first try to find the metadata in the namespace
-			const FName Name = ExecuteContext.AdaptMetadataName(bUseNameSpace, InName);
+			const FName Name = ExecuteContext.AdaptMetadataName(InNameSpace, InName);
 			return ExecuteContext.Hierarchy->FindMetadataForElement(Element, Name, InType);
 		};
 	}
@@ -312,13 +313,13 @@ const TArray<FRigVMExecuteArgument>& FRigDispatch_SetMetadata::GetExecuteArgumen
 
 FRigBaseMetadata* FRigDispatch_SetMetadata::FindOrAddMetadata(const FControlRigExecuteContext& InContext,
                                                               const FRigElementKey& InKey, const FName& InName, ERigMetadataType InType,
-                                                              bool bUseNameSpace, FCachedRigElement& Cache)
+                                                              ERigMetaDataNameSpace NameSpace, FCachedRigElement& Cache)
 {
 	if(Cache.UpdateCache(InKey, InContext.Hierarchy))
 	{
 		if(FRigBaseElement* Element = InContext.Hierarchy->Get(Cache.GetIndex()))
 		{
-			const FName Name = InContext.AdaptMetadataName(bUseNameSpace, InName);
+			const FName Name = InContext.AdaptMetadataName(NameSpace, InName);
 			constexpr bool bNotify = true;
 			return InContext.Hierarchy->GetMetadataForElement(Element, Name, InType, bNotify);
 		}
@@ -427,9 +428,33 @@ FRigUnit_RemoveMetadata_Execute()
 
 	if(CachedIndex.UpdateCache(Item, Hierarchy))
 	{
+		if(NameSpace == ERigMetaDataNameSpace::All)
+		{
+			Removed = true;
+			
+			// run remove for all relevant namespace types.
+			for(int32 NameSpaceTypeIndex = 0;
+				NameSpaceTypeIndex < (int32)ERigMetaDataNameSpace::Last;
+				NameSpaceTypeIndex++)
+			{
+				const ERigMetaDataNameSpace NameSpaceFromIndex = (ERigMetaDataNameSpace)NameSpaceTypeIndex;
+				if(NameSpaceFromIndex == ERigMetaDataNameSpace::All)
+				{
+					continue;
+				}
+					
+				bool bResult = false;
+				FRigUnit_RemoveMetadata::StaticExecute(ExecuteContext, Item, Name, NameSpaceFromIndex, bResult, CachedIndex);
+				if(!bResult)
+				{
+					Removed = false;
+				}
+			}
+			return;
+		}
 		if(FRigBaseElement* Element = Hierarchy->Get(CachedIndex))
 		{
-			const FName LocalName = ExecuteContext.AdaptMetadataName(UseNameSpace, Name);
+			const FName LocalName = ExecuteContext.AdaptMetadataName(NameSpace, Name);
 			Removed = Element->RemoveMetadata(LocalName);
 			if(!Removed)
 			{
@@ -451,29 +476,51 @@ FRigUnit_RemoveAllMetadata_Execute()
 
 	if(CachedIndex.UpdateCache(Item, Hierarchy))
 	{
+		if(NameSpace == ERigMetaDataNameSpace::All)
+		{
+			Removed = true;
+			
+			// run remove for all relevant namespace types.
+			for(int32 NameSpaceTypeIndex = 0;
+				NameSpaceTypeIndex < (int32)ERigMetaDataNameSpace::Last;
+				NameSpaceTypeIndex++)
+			{
+				const ERigMetaDataNameSpace NameSpaceFromIndex = (ERigMetaDataNameSpace)NameSpaceTypeIndex;
+				if(NameSpaceFromIndex == ERigMetaDataNameSpace::All)
+				{
+					continue;
+				}
+					
+				bool bResult = false;
+				FRigUnit_RemoveAllMetadata::StaticExecute(ExecuteContext, Item, NameSpaceFromIndex, bResult, CachedIndex);
+				if(!bResult)
+				{
+					Removed = false;
+				}
+			}
+			return;
+		}
+		
 		if(FRigBaseElement* Element = Hierarchy->Get(CachedIndex))
 		{
-			if(UseNameSpace)
+			if(NameSpace != ERigMetaDataNameSpace::None)
 			{
 				// only remove the metadata within this module / with this namespace
-				if(ExecuteContext.IsRigModule())
-				{
-					const TArray<FName> MetadataNames = Hierarchy->GetMetadataNames(Element->GetKey());
+				const TArray<FName> MetadataNames = Hierarchy->GetMetadataNames(Element->GetKey());
 
-					Removed = false;
-					const FString NameSpace = ExecuteContext.GetRigModuleNameSpace();
-					for(const FName& MetadataName : MetadataNames)
+				Removed = false;
+				const FString NameSpacePrefix = ExecuteContext.GetElementNameSpace(NameSpace);
+				for(const FName& MetadataName : MetadataNames)
+				{
+					if(MetadataName.ToString().StartsWith(NameSpacePrefix, ESearchCase::CaseSensitive))
 					{
-						if(MetadataName.ToString().StartsWith(NameSpace, ESearchCase::CaseSensitive))
+						if(Element->RemoveMetadata(MetadataName))
 						{
-							if(Element->RemoveMetadata(MetadataName))
-							{
-								Removed = true;
-							}
+							Removed = true;
 						}
 					}
-					return;
 				}
+				return;
 			}
 			Removed = Element->RemoveAllMetadata();
 		}
@@ -494,7 +541,7 @@ FRigUnit_HasMetadata_Execute()
 	{
 		if(const FRigBaseElement* Element = Hierarchy->Get(CachedIndex))
 		{
-			const FName LocalName = ExecuteContext.AdaptMetadataName(UseNameSpace, Name);
+			const FName LocalName = ExecuteContext.AdaptMetadataName(NameSpace, Name);
 			Found = Element->GetMetadata(LocalName, Type) != nullptr;
 			if(!Found)
 			{
@@ -514,7 +561,7 @@ FRigUnit_FindItemsWithMetadata_Execute()
 		return;
 	}
 
-	const FName LocalName = ExecuteContext.AdaptMetadataName(UseNameSpace, Name);
+	const FName LocalName = ExecuteContext.AdaptMetadataName(NameSpace, Name);
 	Hierarchy->Traverse([&Items, Name, LocalName, Type](const FRigBaseElement* Element, bool& bContinue)
 	{
 		if(Element->GetMetadata(LocalName, Type) != nullptr)
@@ -562,7 +609,7 @@ FRigUnit_SetMetadataTag_Execute()
 			if(FRigNameArrayMetadata* Md = Cast<FRigNameArrayMetadata>(Element->SetupValidMetadata(URigHierarchy::TagMetadataName, ERigMetadataType::NameArray)))
 			{
 				const int32 LastIndex = Md->GetValue().Num(); 
-				const FName LocalTag = ExecuteContext.AdaptMetadataName(UseNameSpace, Tag);
+				const FName LocalTag = ExecuteContext.AdaptMetadataName(NameSpace, Tag);
 				if(Md->GetValue().AddUnique(LocalTag) == LastIndex)
 				{
 					Element->NotifyMetadataTagChanged(LocalTag, true);
@@ -589,7 +636,7 @@ FRigUnit_SetMetadataTagArray_Execute()
 				for(const FName& Tag : Tags)
 				{
 					const int32 LastIndex = Md->GetValue().Num(); 
-					const FName LocalTag = ExecuteContext.AdaptMetadataName(UseNameSpace, Tag);
+					const FName LocalTag = ExecuteContext.AdaptMetadataName(NameSpace, Tag);
 					if(Md->GetValue().AddUnique(LocalTag) == LastIndex)
 					{
 						Element->NotifyMetadataTagChanged(LocalTag, true);
@@ -615,7 +662,7 @@ FRigUnit_RemoveMetadataTag_Execute()
 		{
 			if(FRigNameArrayMetadata* Md = Cast<FRigNameArrayMetadata>(Element->GetMetadata(URigHierarchy::TagMetadataName, ERigMetadataType::NameArray)))
 			{
-				const FName LocalTag = ExecuteContext.AdaptMetadataName(UseNameSpace, Tag);
+				const FName LocalTag = ExecuteContext.AdaptMetadataName(NameSpace, Tag);
 				Removed = Md->GetValue().Remove(LocalTag) > 0;
 				if(Removed)
 				{
@@ -642,7 +689,7 @@ FRigUnit_HasMetadataTag_Execute()
 		{
 			if(const FRigNameArrayMetadata* Md = Cast<FRigNameArrayMetadata>(Element->GetMetadata(URigHierarchy::TagMetadataName, ERigMetadataType::NameArray)))
 			{
-				const FName LocalTag = ExecuteContext.AdaptMetadataName(UseNameSpace, Tag);
+				const FName LocalTag = ExecuteContext.AdaptMetadataName(NameSpace, Tag);
 				Found = Md->GetValue().Contains(LocalTag);
 			}
 		}
@@ -668,7 +715,7 @@ FRigUnit_HasMetadataTagArray_Execute()
 				Found = true;
 				for(const FName& Tag : Tags)
 				{
-					const FName LocalTag = ExecuteContext.AdaptMetadataName(UseNameSpace, Tag);
+					const FName LocalTag = ExecuteContext.AdaptMetadataName(NameSpace, Tag);
 					if(!Md->GetValue().Contains(LocalTag))
 					{
 						Found = false;
@@ -690,7 +737,7 @@ FRigUnit_FindItemsWithMetadataTag_Execute()
 		return;
 	}
 
-	const FName LocalTag = ExecuteContext.AdaptMetadataName(UseNameSpace, Tag);
+	const FName LocalTag = ExecuteContext.AdaptMetadataName(NameSpace, Tag);
 	Hierarchy->Traverse([&Items, LocalTag](const FRigBaseElement* Element, bool& bContinue)
 	{
 		if(const FRigNameArrayMetadata* Md = Cast<FRigNameArrayMetadata>(Element->GetMetadata(URigHierarchy::TagMetadataName, ERigMetadataType::NameArray)))
@@ -716,12 +763,14 @@ FRigUnit_FindItemsWithMetadataTagArray_Execute()
 
 	TArrayView<const FName> LocalTags(Tags);
 	TArray<FName> AdaptedTags;
-	if(UseNameSpace && ExecuteContext.IsRigModule())
+
+	const bool bUseNameSpace = (NameSpace != ERigMetaDataNameSpace::None) && (NameSpace != ERigMetaDataNameSpace::All);
+	if(bUseNameSpace && ExecuteContext.IsRigModule())
 	{
 		AdaptedTags.Reserve(Tags.Num());
 		for(const FName& Tag : Tags)
 		{
-			const FName LocalTag = ExecuteContext.AdaptMetadataName(UseNameSpace, Tag);
+			const FName LocalTag = ExecuteContext.AdaptMetadataName(NameSpace, Tag);
 			AdaptedTags.Add(LocalTag);
 		}
 		LocalTags = TArrayView<const FName>(AdaptedTags);
@@ -768,12 +817,13 @@ FRigUnit_FilterItemsByMetadataTags_Execute()
 
 	TArrayView<const FName> LocalTags(Tags);
 	TArray<FName> AdaptedTags;
-	if(UseNameSpace && ExecuteContext.IsRigModule())
+	const bool bUseNameSpace = (NameSpace != ERigMetaDataNameSpace::None) && (NameSpace != ERigMetaDataNameSpace::All);
+	if(bUseNameSpace && ExecuteContext.IsRigModule())
 	{
 		AdaptedTags.Reserve(Tags.Num());
 		for(const FName& Tag : Tags)
 		{
-			const FName LocalTag = ExecuteContext.AdaptMetadataName(UseNameSpace, Tag);
+			const FName LocalTag = ExecuteContext.AdaptMetadataName(NameSpace, Tag);
 			AdaptedTags.Add(LocalTag);
 		}
 		LocalTags = TArrayView<const FName>(AdaptedTags);
diff --git a/Engine/Plugins/Animation/ControlRig/Source/ControlRig/Private/Units/RigUnitContext.cpp b/Engine/Plugins/Animation/ControlRig/Source/ControlRig/Private/Units/RigUnitContext.cpp
--- a/Engine/Plugins/Animation/ControlRig/Source/ControlRig/Private/Units/RigUnitContext.cpp
+++ b/Engine/Plugins/Animation/ControlRig/Source/ControlRig/Private/Units/RigUnitContext.cpp
@@ -2,6 +2,7 @@
 
 #include "Units/RigUnitContext.h"
 #include "ControlRig.h"
+#include "ModularRig.h"
 
 #include UE_INLINE_GENERATED_CPP_BY_NAME(RigUnitContext)
 
@@ -48,10 +49,82 @@ FString FControlRigExecuteContext::RemoveRigModuleNameSpace(const FString& InNam
 	return InName;
 }
 
-FName FControlRigExecuteContext::AdaptMetadataName(bool bUseNameSpace, const FName& InMetadataName) const
+FString FControlRigExecuteContext::GetElementNameSpace(ERigMetaDataNameSpace InNameSpaceType) const
+{
+	if(IsRigModule())
+	{
+		// prefix the meta data name with the namespace to allow modules to store their
+		// metadata in a way that doesn't collide with other modules' metadata.
+		switch(InNameSpaceType)
+		{
+			case ERigMetaDataNameSpace::Self:
+			{
+				return GetRigModuleNameSpace();
+			}
+			case ERigMetaDataNameSpace::Parent:
+			{
+				check(GetRigModuleNameSpace().EndsWith(UModularRig::NamespaceSeparator));
+				FString ParentNameSpace;
+				URigHierarchy::SplitNameSpace(GetRigModuleNameSpace().LeftChop(1), &ParentNameSpace, nullptr, true);
+				if(ParentNameSpace.IsEmpty())
+				{
+					return GetRigModuleNameSpace();
+				}
+				return ParentNameSpace + UModularRig::NamespaceSeparator;
+			}
+			case ERigMetaDataNameSpace::Root:
+			{
+				FString RootNameSpace;
+				URigHierarchy::SplitNameSpace(GetRigModuleNameSpace(), &RootNameSpace, nullptr, false);
+				if(RootNameSpace.IsEmpty())
+				{
+					return GetRigModuleNameSpace();
+				}
+				return RootNameSpace + UModularRig::NamespaceSeparator;
+			}
+			default:
+			{
+				break;
+			}
+		}
+	}
+	else
+	{
+		// prefix the meta data with some mockup namespaces
+		// so we can test this even without a module present.
+		switch(InNameSpaceType)
+		{
+			case ERigMetaDataNameSpace::Self:
+			{
+				// if we are storing on self and this is not a modular
+				// rig let's just not use a namespace.
+				break;
+			}
+			case ERigMetaDataNameSpace::Parent:
+			{
+				static const FString ParentNameSpace = TEXT("Parent:");
+				return ParentNameSpace;
+			}
+			case ERigMetaDataNameSpace::Root:
+			{
+				static const FString RootNameSpace = TEXT("Root:");
+				return RootNameSpace;
+			}
+			default:
+			{
+				break;
+			}
+		}
+	}
+
+	return FString();
+}
+
+FName FControlRigExecuteContext::AdaptMetadataName(ERigMetaDataNameSpace InNameSpaceType, const FName& InMetadataName) const
 {
 	// only if we are within a rig module let's adapt the meta data name
-	if(bUseNameSpace && IsRigModule() && !InMetadataName.IsNone())
+	const bool bUseNameSpace = (InNameSpaceType != ERigMetaDataNameSpace::None) && (InNameSpaceType != ERigMetaDataNameSpace::All);
+	if(bUseNameSpace && !InMetadataName.IsNone())
 	{
 		// if the metadata name already contains a namespace - we are just going
 		// to use it as is. this means that modules have access to other module's metadata,
@@ -64,10 +137,47 @@ FName FControlRigExecuteContext::AdaptMetadataName(bool bUseNameSpace, const FNa
 			return InMetadataName;
 		}
 
-		// prefix the meta data name with the namespace to allow modules to store their
-		// metadata in a way that doesn't collide with other modules' metadata.
-		const FString JoinedMetadataName = GetRigModuleNameSpace() + MetadataNameString;
-		return *JoinedMetadataName;
+		if(IsRigModule())
+		{
+			// prefix the meta data name with the namespace to allow modules to store their
+			// metadata in a way that doesn't collide with other modules' metadata.
+			switch(InNameSpaceType)
+			{
+				case ERigMetaDataNameSpace::Self:
+				case ERigMetaDataNameSpace::Parent:
+				case ERigMetaDataNameSpace::Root:
+				{
+					return *URigHierarchy::JoinNameSpace(GetElementNameSpace(InNameSpaceType), MetadataNameString);
+				}
+				default:
+				{
+					break;
+				}
+			}
+		}
+		else
+		{
+			// prefix the meta data with some mockup namespaces
+			// so we can test this even without a module present.
+			switch(InNameSpaceType)
+			{
+				case ERigMetaDataNameSpace::Self:
+				{
+					// if we are storing on self and this is not a modular
+					// rig let's just not use a namespace.
+					break;
+				}
+				case ERigMetaDataNameSpace::Parent:
+				case ERigMetaDataNameSpace::Root:
+				{
+					return *URigHierarchy::JoinNameSpace(GetElementNameSpace(InNameSpaceType), MetadataNameString);
+				}
+				default:
+				{
+					break;
+				}
+			}
+		}
 	}
 	return InMetadataName;
 }
diff --git a/Engine/Plugins/Animation/ControlRig/Source/ControlRig/Public/Units/Hierarchy/RigUnit_Metadata.h b/Engine/Plugins/Animation/ControlRig/Source/ControlRig/Public/Units/Hierarchy/RigUnit_Metadata.h
--- a/Engine/Plugins/Animation/ControlRig/Source/ControlRig/Public/Units/Hierarchy/RigUnit_Metadata.h
+++ b/Engine/Plugins/Animation/ControlRig/Source/ControlRig/Public/Units/Hierarchy/RigUnit_Metadata.h
@@ -16,6 +16,29 @@ struct CONTROLRIG_API FRigDispatch_MetadataBase : public FRigDispatchFactory
 		FactoryScriptStruct = StaticStruct();
 	}
 
+	template<typename T>
+	static bool Equals(const T& A, const T& B)
+	{
+		return A.Equals(B);
+	}
+
+	template<typename T>
+	static bool ArrayEquals(const TArray<T>& A, const TArray<T>& B)
+	{
+		if(A.Num() != B.Num())
+		{
+			return false;
+		}
+		for(int32 Index = 0; Index < A.Num(); Index++)
+		{
+			if(!Equals<T>(A[Index], B[Index]))
+			{
+				return false;
+			}
+		}
+		return true;
+	}
+
 #if WITH_EDITOR
 	virtual FString GetNodeTitle(const FRigVMTemplateTypeMap& InTypes) const override;;
 #endif
@@ -36,7 +59,7 @@ protected:
 	
 	mutable int32 ItemArgIndex = INDEX_NONE;
 	mutable int32 NameArgIndex = INDEX_NONE;
-	mutable int32 UseNameSpaceArgIndex = INDEX_NONE;
+	mutable int32 NameSpaceArgIndex = INDEX_NONE;
 	mutable int32 CacheArgIndex = INDEX_NONE;
 	mutable int32 DefaultArgIndex = INDEX_NONE;
 	mutable int32 ValueArgIndex = INDEX_NONE;
@@ -44,7 +67,7 @@ protected:
 	mutable int32 SuccessArgIndex = INDEX_NONE;
 	static FName ItemArgName;
 	static FName NameArgName;
-	static FName UseNameSpaceArgName;
+	static FName NameSpaceArgName;
 	static FName CacheArgName;
 	static FName DefaultArgName;
 	static FName ValueArgName;
@@ -52,6 +75,108 @@ protected:
 	static FName SuccessArgName;
 };
 
+template<>
+inline bool FRigDispatch_MetadataBase::Equals(const bool& A, const bool& B)
+{
+	return A == B;
+}
+
+template<>
+inline bool FRigDispatch_MetadataBase::Equals(const int32& A, const int32& B)
+{
+	return A == B;
+}
+
+template<>
+inline bool FRigDispatch_MetadataBase::Equals(const float& A, const float& B)
+{
+	return FMath::IsNearlyEqual(A, B);
+}
+
+template<>
+inline bool FRigDispatch_MetadataBase::Equals(const double& A, const double& B)
+{
+	return FMath::IsNearlyEqual(A, B);
+}
+
+template<>
+inline bool FRigDispatch_MetadataBase::Equals(const FName& A, const FName& B)
+{
+	return A.IsEqual(B, ENameCase::CaseSensitive);
+}
+
+template<>
+inline bool FRigDispatch_MetadataBase::Equals(const FRigElementKey& A, const FRigElementKey& B)
+{
+	return A == B;
+}
+
+template<>
+inline bool FRigDispatch_MetadataBase::Equals(const TArray<bool>& A, const TArray<bool>& B)
+{
+	return ArrayEquals<bool>(A, B);
+}
+
+template<>
+inline bool FRigDispatch_MetadataBase::Equals(const TArray<int32>& A, const TArray<int32>& B)
+{
+	return ArrayEquals<int32>(A, B);
+}
+
+template<>
+inline bool FRigDispatch_MetadataBase::Equals(const TArray<float>& A, const TArray<float>& B)
+{
+	return ArrayEquals<float>(A, B);
+}
+
+template<>
+inline bool FRigDispatch_MetadataBase::Equals(const TArray<double>& A, const TArray<double>& B)
+{
+	return ArrayEquals<double>(A, B);
+}
+
+template<>
+inline bool FRigDispatch_MetadataBase::Equals(const TArray<FName>& A, const TArray<FName>& B)
+{
+	return ArrayEquals<FName>(A, B);
+}
+
+template<>
+inline bool FRigDispatch_MetadataBase::Equals(const TArray<FRigElementKey>& A, const TArray<FRigElementKey>& B)
+{
+	return ArrayEquals<FRigElementKey>(A, B);
+}
+
+template<>
+inline bool FRigDispatch_MetadataBase::Equals(const TArray<FVector>& A, const TArray<FVector>& B)
+{
+	return ArrayEquals<FVector>(A, B);
+}
+
+template<>
+inline bool FRigDispatch_MetadataBase::Equals(const TArray<FRotator>& A, const TArray<FRotator>& B)
+{
+	return ArrayEquals<FRotator>(A, B);
+}
+
+template<>
+inline bool FRigDispatch_MetadataBase::Equals(const TArray<FQuat>& A, const TArray<FQuat>& B)
+{
+	return ArrayEquals<FQuat>(A, B);
+}
+
+template<>
+inline bool FRigDispatch_MetadataBase::Equals(const TArray<FTransform>& A, const TArray<FTransform>& B)
+{
+	return ArrayEquals<FTransform>(A, B);
+}
+
+template<>
+inline bool FRigDispatch_MetadataBase::Equals(const TArray<FLinearColor>& A, const TArray<FLinearColor>& B)
+{
+	return ArrayEquals<FLinearColor>(A, B);
+}
+
 /*
  * Sets some metadata for the provided item
  */
@@ -64,7 +189,7 @@ struct CONTROLRIG_API FRigDispatch_GetMetadata : public FRigDispatch_MetadataBas
 
 protected:
 
-	static FRigBaseMetadata* FindMetadata(const FRigVMExtendedExecuteContext& InContext, const FRigElementKey& InKey, const FName& InName, ERigMetadataType InType, bool bUseNameSpace, FCachedRigElement& Cache);
+	static FRigBaseMetadata* FindMetadata(const FRigVMExtendedExecuteContext& InContext, const FRigElementKey& InKey, const FName& InName, ERigMetadataType InType, ERigMetaDataNameSpace InNameSpace, FCachedRigElement& Cache);
 	virtual FRigVMFunctionPtr GetDispatchFunctionImpl(const FRigVMTemplateTypeMap& InTypes) const override;
 
 #if WITH_EDITOR
@@ -95,14 +220,14 @@ protected:
 		// unpack the memory
 		const FRigElementKey& Item = *reinterpret_cast<const FRigElementKey*>(Handles[Factory->ItemArgIndex].GetData());
 		const FName& Name = *reinterpret_cast<const FName*>(Handles[Factory->NameArgIndex].GetData());
-		const bool bUseNameSpace = *reinterpret_cast<const bool*>(Handles[Factory->UseNameSpaceArgIndex].GetData());
+		const ERigMetaDataNameSpace NameSpace = *reinterpret_cast<const ERigMetaDataNameSpace*>(Handles[Factory->NameSpaceArgIndex].GetData());
 		FCachedRigElement& Cache = *reinterpret_cast<FCachedRigElement*>(Handles[Factory->CacheArgIndex].GetData(false, InContext.GetSlice().GetIndex()));
 		const ValueType& Default = *reinterpret_cast<const ValueType*>(Handles[Factory->DefaultArgIndex].GetData());
 		ValueType& Value = *reinterpret_cast<ValueType*>(Handles[Factory->ValueArgIndex].GetData());
 		bool& Found = *reinterpret_cast<bool*>(Handles[Factory->FoundArgIndex].GetData());
 
 		// extract the metadata
-		if (const MetadataType* Md = Cast<MetadataType>(FindMetadata(InContext, Item, Name, EnumValue, bUseNameSpace, Cache)))
+		if (const MetadataType* Md = Cast<MetadataType>(FindMetadata(InContext, Item, Name, EnumValue, NameSpace, Cache)))
 		{
 			Value = Md->GetValue();
 			Found = true;
@@ -129,7 +254,7 @@ struct CONTROLRIG_API FRigDispatch_SetMetadata : public FRigDispatch_MetadataBas
 
 protected:
 
-	static FRigBaseMetadata* FindOrAddMetadata(const FControlRigExecuteContext& InContext, const FRigElementKey& InKey, const FName& InName, ERigMetadataType InType, bool bUseNameSpace, FCachedRigElement& Cache);
+	static FRigBaseMetadata* FindOrAddMetadata(const FControlRigExecuteContext& InContext, const FRigElementKey& InKey, const FName& InName, ERigMetadataType InType, ERigMetaDataNameSpace InNameSpace, FCachedRigElement& Cache);
 	virtual FRigVMFunctionPtr GetDispatchFunctionImpl(const FRigVMTemplateTypeMap& InTypes) const override;
 
 #if WITH_EDITOR
@@ -160,15 +285,19 @@ protected:
 		FControlRigExecuteContext& ExecuteContext = InContext.GetPublicData<FControlRigExecuteContext>();
 		const FRigElementKey& Item = *(const FRigElementKey*)Handles[Factory->ItemArgIndex].GetData();
 		const FName& Name = *(const FName*)Handles[Factory->NameArgIndex].GetData();
-		const bool bUseNameSpace = *(const bool*)Handles[Factory->UseNameSpaceArgIndex].GetData();
+		const ERigMetaDataNameSpace NameSpace = *(const ERigMetaDataNameSpace*)Handles[Factory->NameSpaceArgIndex].GetData();
 		FCachedRigElement& Cache = *(FCachedRigElement*)Handles[Factory->CacheArgIndex].GetData(false, InContext.GetSlice().GetIndex());
 		ValueType& Value = *(ValueType*)Handles[Factory->ValueArgIndex].GetData();
 		bool& Success = *(bool*)Handles[Factory->SuccessArgIndex].GetData();
 
 		// extract the metadata
-		if (MetadataType* Md = Cast<MetadataType>(FindOrAddMetadata(ExecuteContext, Item, Name, EnumValue, bUseNameSpace, Cache)))
+		if (MetadataType* Md = Cast<MetadataType>(FindOrAddMetadata(ExecuteContext, Item, Name, EnumValue, NameSpace, Cache)))
 		{
-			Md->GetValue() = Value;
+			if(!Equals<ValueType>(Md->GetValue(), Value))
+			{
+				Md->GetValue() = Value;
+				ExecuteContext.Hierarchy->MetadataVersion++;
+			}
 			Success = true;
 		}
 		else
@@ -189,7 +318,7 @@ struct CONTROLRIG_API FRigUnit_RemoveMetadata : public FRigUnitMutable
 	FRigUnit_RemoveMetadata()
 		: Item(NAME_None, ERigElementType::Bone)
 		, Name(NAME_None)
-		, UseNameSpace(true)
+		, NameSpace(ERigMetaDataNameSpace::Self)
 		, Removed(false)
 		, CachedIndex()
 	{}
@@ -210,10 +339,10 @@ struct CONTROLRIG_API FRigUnit_RemoveMetadata : public FRigUnitMutable
 	FName Name;
 
 	/**
-	 * If checked the metadata will be available only in this name space / module
+	 * Defines in which namespace the metadata will be looked up
 	 */
 	UPROPERTY(meta = (Input))
-	bool UseNameSpace;
+	ERigMetaDataNameSpace NameSpace;
 
 	// True if the metadata has been removed
 	UPROPERTY(meta=(Output))
@@ -234,7 +363,7 @@ struct CONTROLRIG_API FRigUnit_RemoveAllMetadata : public FRigUnitMutable
 
 	FRigUnit_RemoveAllMetadata()
 		: Item(NAME_None, ERigElementType::Bone)
-		, UseNameSpace(true)
+		, NameSpace(ERigMetaDataNameSpace::Self)
 		, Removed(false)
 		, CachedIndex()
 	{}
@@ -249,10 +378,10 @@ struct CONTROLRIG_API FRigUnit_RemoveAllMetadata : public FRigUnitMutable
 	FRigElementKey Item;
 
 	/**
-	 * If checked the metadata will be available only in this name space / module
+	 * Defines in which namespace the metadata will be looked up
 	 */
 	UPROPERTY(meta = (Input))
-	bool UseNameSpace;
+	ERigMetaDataNameSpace NameSpace;
 
 	// True if any metadata has been removed
 	UPROPERTY(meta=(Output))
@@ -275,7 +404,7 @@ struct CONTROLRIG_API FRigUnit_HasMetadata : public FRigUnit
 		: Item(NAME_None, ERigElementType::Bone)
 		, Name(NAME_None)
 		, Type(ERigMetadataType::Float)
-		, UseNameSpace(true)
+		, NameSpace(ERigMetaDataNameSpace::Self)
 		, Found(false)
 		, CachedIndex()
 	{}
@@ -302,10 +431,10 @@ struct CONTROLRIG_API FRigUnit_HasMetadata : public FRigUnit
 	ERigMetadataType Type;
 
 	/**
-	 * If checked the metadata will be available only in this name space / module
+	 * Defines in which namespace the metadata will be looked up
 	 */
 	UPROPERTY(meta = (Input))
-	bool UseNameSpace;
+	ERigMetaDataNameSpace NameSpace;
 
 	// True if the item has the metadata
 	UPROPERTY(meta=(Output))
@@ -327,7 +456,7 @@ struct CONTROLRIG_API FRigUnit_FindItemsWithMetadata : public FRigUnit
 	FRigUnit_FindItemsWithMetadata()
 		: Name(NAME_None)
 		, Type(ERigMetadataType::Float)
-		, UseNameSpace(true)
+		, NameSpace(ERigMetaDataNameSpace::Self)
 		, Items()
 	{}
 
@@ -347,10 +476,10 @@ struct CONTROLRIG_API FRigUnit_FindItemsWithMetadata : public FRigUnit
 	ERigMetadataType Type;
 
 	/**
-	 * If checked the metadata will be available only in this name space / module
+	 * Defines in which namespace the metadata will be looked up
 	 */
 	UPROPERTY(meta = (Input))
-	bool UseNameSpace;
+	ERigMetaDataNameSpace NameSpace;
 
 	// The items containing the metadata
 	UPROPERTY(meta=(Output))
@@ -402,7 +531,7 @@ struct CONTROLRIG_API FRigUnit_SetMetadataTag : public FRigUnitMutable
 	FRigUnit_SetMetadataTag()
 		: Item(NAME_None, ERigElementType::Bone)
 		, Tag(NAME_None)
-		, UseNameSpace(true)
+		, NameSpace(ERigMetaDataNameSpace::Self)
 		, CachedIndex()
 	{}
 
@@ -422,10 +551,10 @@ struct CONTROLRIG_API FRigUnit_SetMetadataTag : public FRigUnitMutable
 	FName Tag;
 
 	/**
-	 * If checked the metadata will be available only in this name space / module
+	 * Defines in which namespace the metadata will be looked up
 	 */
 	UPROPERTY(meta = (Input))
-	bool UseNameSpace;
+	ERigMetaDataNameSpace NameSpace;
 
 	// Used to cache the internally used index
 	UPROPERTY()
@@ -443,7 +572,7 @@ struct CONTROLRIG_API FRigUnit_SetMetadataTagArray : public FRigUnitMutable
 	FRigUnit_SetMetadataTagArray()
 		: Item(NAME_None, ERigElementType::Bone)
 		, Tags()
-		, UseNameSpace(true)
+		, NameSpace(ERigMetaDataNameSpace::Self)
 		, CachedIndex()
 	{}
 
@@ -463,10 +592,10 @@ struct CONTROLRIG_API FRigUnit_SetMetadataTagArray : public FRigUnitMutable
 	TArray<FName> Tags;
 	
 	/**
-	 * If checked the metadata will be available only in this name space / module
+	 * Defines in which namespace the metadata will be looked up
 	 */
 	UPROPERTY(meta = (Input))
-	bool UseNameSpace;
+	ERigMetaDataNameSpace NameSpace;
 
 	// Used to cache the internally used index
 	UPROPERTY()
@@ -484,7 +613,7 @@ struct CONTROLRIG_API FRigUnit_RemoveMetadataTag : public FRigUnitMutable
 	FRigUnit_RemoveMetadataTag()
 		: Item(NAME_None, ERigElementType::Bone)
 		, Tag(NAME_None)
-		, UseNameSpace(true)
+		, NameSpace(ERigMetaDataNameSpace::Self)
 		, Removed(false)
 		, CachedIndex()
 	{}
@@ -505,10 +634,10 @@ struct CONTROLRIG_API FRigUnit_RemoveMetadataTag : public FRigUnitMutable
 	FName Tag;
 
 	/**
-	 * If checked the metadata will be available only in this name space / module
+	 * Defines in which namespace the metadata will be looked up
 	 */
 	UPROPERTY(meta = (Input))
-	bool UseNameSpace;
+	ERigMetaDataNameSpace NameSpace;
 	
 	/**
 	 * Returns true if the removal was successful
@@ -532,7 +661,7 @@ struct CONTROLRIG_API FRigUnit_HasMetadataTag : public FRigUnit
 	FRigUnit_HasMetadataTag()
 		: Item(NAME_None, ERigElementType::Bone)
 		, Tag(NAME_None)
-		, UseNameSpace(true)
+		, NameSpace(ERigMetaDataNameSpace::Self)
 		, Found(false)
 		, CachedIndex()
 	{}
@@ -553,10 +682,10 @@ struct CONTROLRIG_API FRigUnit_HasMetadataTag : public FRigUnit
 	FName Tag;
 
 	/**
-	 * If checked the metadata will be available only in this name space / module
+	 * Defines in which namespace the metadata will be looked up
 	 */
 	UPROPERTY(meta = (Input))
-	bool UseNameSpace;
+	ERigMetaDataNameSpace NameSpace;
 	
 	// True if the item has the metadata
 	UPROPERTY(meta=(Output))
@@ -578,7 +707,7 @@ struct CONTROLRIG_API FRigUnit_HasMetadataTagArray : public FRigUnit
 	FRigUnit_HasMetadataTagArray()
 		: Item(NAME_None, ERigElementType::Bone)
 		, Tags()
-		, UseNameSpace(true)
+		, NameSpace(ERigMetaDataNameSpace::Self)
 		, Found(false)
 		, CachedIndex()
 	{}
@@ -599,10 +728,10 @@ struct CONTROLRIG_API FRigUnit_HasMetadataTagArray : public FRigUnit
 	TArray<FName> Tags;
 
 	/**
-	 * If checked the metadata will be available only in this name space / module
+	 * Defines in which namespace the metadata will be looked up
 	 */
 	UPROPERTY(meta = (Input))
-	bool UseNameSpace;
+	ERigMetaDataNameSpace NameSpace;
 	
 	// True if the item has the metadata
 	UPROPERTY(meta=(Output))
@@ -623,7 +752,7 @@ struct CONTROLRIG_API FRigUnit_FindItemsWithMetadataTag : public FRigUnit
 
 	FRigUnit_FindItemsWithMetadataTag()
 		: Tag(NAME_None)
-		, UseNameSpace(true)
+		, NameSpace(ERigMetaDataNameSpace::Self)
 		, Items()
 	{}
 
@@ -637,10 +766,10 @@ struct CONTROLRIG_API FRigUnit_FindItemsWithMetadataTag : public FRigUnit
 	FName Tag;
 
 	/**
-	 * If checked the metadata will be available only in this name space / module
+	 * Defines in which namespace the metadata will be looked up
 	 */
 	UPROPERTY(meta = (Input))
-	bool UseNameSpace;
+	ERigMetaDataNameSpace NameSpace;
 
 	// The items containing the metadata
 	UPROPERTY(meta=(Output))
@@ -657,7 +786,7 @@ struct CONTROLRIG_API FRigUnit_FindItemsWithMetadataTagArray : public FRigUnit
 
 	FRigUnit_FindItemsWithMetadataTagArray()
 		: Tags()
-		, UseNameSpace(true)
+		, NameSpace(ERigMetaDataNameSpace::Self)
 		, Items()
 	{}
 
@@ -671,10 +800,10 @@ struct CONTROLRIG_API FRigUnit_FindItemsWithMetadataTagArray : public FRigUnit
 	TArray<FName> Tags;
 
 	/**
-	 * If checked the metadata will be available only in this name space / module
+	 * Defines in which namespace the metadata will be looked up
 	 */
 	UPROPERTY(meta = (Input))
-	bool UseNameSpace;
+	ERigMetaDataNameSpace NameSpace;
 
 	// The items containing the metadata
 	UPROPERTY(meta=(Output))
@@ -692,7 +821,7 @@ struct CONTROLRIG_API FRigUnit_FilterItemsByMetadataTags : public FRigUnit
 	FRigUnit_FilterItemsByMetadataTags()
 		: Items()
 		, Tags()
-		, UseNameSpace(true)
+		, NameSpace(ERigMetaDataNameSpace::Self)
 		, Inclusive(true)
 		, Result()
 	{}
@@ -711,10 +840,10 @@ struct CONTROLRIG_API FRigUnit_FilterItemsByMetadataTags : public FRigUnit
 	TArray<FName> Tags;
 
 	/**
-	 * If checked the metadata will be available only in this name space / module
+	 * Defines in which namespace the metadata will be looked up
 	 */
 	UPROPERTY(meta = (Input))
-	bool UseNameSpace;
+	ERigMetaDataNameSpace NameSpace;
 
 	/**
      * If set to true only items with ALL of tags will be returned,
diff --git a/Engine/Plugins/Animation/ControlRig/Source/ControlRig/Public/Units/RigUnitContext.h b/Engine/Plugins/Animation/ControlRig/Source/ControlRig/Public/Units/RigUnitContext.h
--- a/Engine/Plugins/Animation/ControlRig/Source/ControlRig/Public/Units/RigUnitContext.h
+++ b/Engine/Plugins/Animation/ControlRig/Source/ControlRig/Public/Units/RigUnitContext.h
@@ -28,6 +28,22 @@ enum class EControlRigInteractionType : uint8
 	All = Translate | Rotate | Scale
 };
 
+UENUM(BlueprintType)
+enum class ERigMetaDataNameSpace : uint8
+{
+	// Use no namespace - store the metadata directly on the item
+	None,
+	// Store the metadata for item relative to its module
+	Self,
+	// Store the metadata relative to its parent model
+	Parent,
+	// Store the metadata under the root module
+	Root,
+	// (Only for remove) Remove all metadata data for all namespaces
+	All, 
+	Last UMETA(Hidden)
+};
+
 USTRUCT()
 struct CONTROLRIG_API FRigHierarchySettings
 {
@@ -173,6 +189,11 @@ public:
 		return RigModuleNameSpace;
 	}
 
+	/**
+	 * Returns the namespace given a namespace type
+	 */
+	FString GetElementNameSpace(ERigMetaDataNameSpace InNameSpaceType) const;
+	
 	/**
 	 * Returns the module this unit is running inside of (or nullptr)
 	 */
@@ -184,7 +205,7 @@ public:
 	/**
 	 * Adapts a metadata name according to rig module namespace.
 	 */
-	FName AdaptMetadataName(bool bUseNameSpace, const FName& InMetadataName) const;
+	FName AdaptMetadataName(ERigMetaDataNameSpace InNameSpaceType, const FName& InMetadataName) const;
 
 	/** The list of available asset user data object */
 	TArray<const UAssetUserData*> AssetUserData;
diff --git a/Engine/Plugins/Animation/ControlRig/Source/ControlRigEditor/Private/Graph/ControlRigGraphPanelPinFactory.cpp b/Engine/Plugins/Animation/ControlRig/Source/ControlRigEditor/Private/Graph/ControlRigGraphPanelPinFactory.cpp
--- a/Engine/Plugins/Animation/ControlRig/Source/ControlRigEditor/Private/Graph/ControlRigGraphPanelPinFactory.cpp
+++ b/Engine/Plugins/Animation/ControlRig/Source/ControlRigEditor/Private/Graph/ControlRigGraphPanelPinFactory.cpp
@@ -335,18 +335,18 @@ TSharedPtr<SGraphPin> FControlRigGraphPanelPinFactory::CreatePin_Internal(UEdGra
 									const FString MapHash = Blueprint->GetPathName();
 									const int32 MetadataVersion = ControlRig->GetHierarchy()->GetMetadataVersion();
 									
-									bool bUseShortNames = false;
+									ERigMetaDataNameSpace NameSpace = ERigMetaDataNameSpace::None;
 									if (const URigVMNode* ModelNode = InPin->GetNode())
 									{
-										if(const URigVMPin* UseNameSpacePin = ModelNode->FindPin(GET_MEMBER_NAME_STRING_CHECKED(FRigUnit_HasMetadata, UseNameSpace)))
+										if(const URigVMPin* NameSpacePin = ModelNode->FindPin(GET_MEMBER_NAME_STRING_CHECKED(FRigUnit_HasMetadata, NameSpace)))
 										{
-											if(UseNameSpacePin->GetDefaultValue().Equals(TEXT("true"), ESearchCase::IgnoreCase))
-											{
-												bUseShortNames = true;
-											}
+											NameSpace = (ERigMetaDataNameSpace)StaticEnum<ERigMetaDataNameSpace>()->GetValueByNameString(NameSpacePin->GetDefaultValue());
 										}
 									}
 
+									const bool bUseShortNames = (NameSpace != ERigMetaDataNameSpace::None) &&
+										(NameSpace != ERigMetaDataNameSpace::All);
+
 									static TMap<FString, FCachedMetadataNames> MetadataNameLists;
 									FCachedMetadataNames& MetadataNames = MetadataNameLists.FindOrAdd(MapHash);
 
@@ -374,7 +374,7 @@ TSharedPtr<SGraphPin> FControlRigGraphPanelPinFactory::CreatePin_Internal(UEdGra
 											if(bUseShortNames)
 											{
 												int32 Index = INDEX_NONE;
-												if(NameString.FindChar(TEXT(':'), Index))
+												if(NameString.FindLastChar(TEXT(':'), Index))
 												{
 													NameString.MidInline(Index + 1);
 												}
@@ -423,17 +423,16 @@ TSharedPtr<SGraphPin> FControlRigGraphPanelPinFactory::CreatePin_Internal(UEdGra
 									const FString MapHash = Blueprint->GetPathName();
 									const int32 MetadataTagVersion = ControlRig->GetHierarchy()->GetMetadataTagVersion(); 
 
-									bool bUseShortNames = false;
+									ERigMetaDataNameSpace NameSpace = ERigMetaDataNameSpace::None;
 									if (const URigVMNode* ModelNode = InPin->GetNode())
 									{
-										if(const URigVMPin* UseNameSpacePin = ModelNode->FindPin(GET_MEMBER_NAME_STRING_CHECKED(FRigUnit_HasMetadata, UseNameSpace)))
+										if(const URigVMPin* NameSpacePin = ModelNode->FindPin(GET_MEMBER_NAME_STRING_CHECKED(FRigUnit_HasMetadata, NameSpace)))
 										{
-											if(UseNameSpacePin->GetDefaultValue().Equals(TEXT("true"), ESearchCase::IgnoreCase))
-											{
-												bUseShortNames = true;
-											}
+											NameSpace = (ERigMetaDataNameSpace)StaticEnum<ERigMetaDataNameSpace>()->GetValueByNameString(NameSpacePin->GetDefaultValue());
 										}
 									}
+									const bool bUseShortNames = (NameSpace != ERigMetaDataNameSpace::None) &&
+										(NameSpace != ERigMetaDataNameSpace::All);
 
 									static TMap<FString, FCachedMetadataTagNames> MetadataTagNameLists;
 									FCachedMetadataTagNames& MetadataTagNames = MetadataTagNameLists.FindOrAdd(MapHash);
@@ -465,7 +464,7 @@ TSharedPtr<SGraphPin> FControlRigGraphPanelPinFactory::CreatePin_Internal(UEdGra
 											if(bUseShortNames)
 											{
 												int32 Index = INDEX_NONE;
-												if(TagString.FindChar(TEXT(':'), Index))
+												if(TagString.FindLastChar(TEXT(':'), Index))
 												{
 													TagString.MidInline(Index + 1);
 												}
