Constraints: fixed sequencer ptr not being updated

#rnx

diff --git a/Engine/Plugins/Animation/ControlRig/Source/ControlRigEditor/Private/Editor/Constraints/SConstraintsWidget.cpp b/Engine/Plugins/Animation/ControlRig/Source/ControlRigEditor/Private/Editor/Constraints/SConstraintsWidget.cpp
--- a/Engine/Plugins/Animation/ControlRig/Source/ControlRigEditor/Private/Editor/Constraints/SConstraintsWidget.cpp
+++ b/Engine/Plugins/Animation/ControlRig/Source/ControlRigEditor/Private/Editor/Constraints/SConstraintsWidget.cpp
@@ -1219,9 +1219,23 @@ void SConstraintsEditionWidget::RemoveItem(const TSharedPtr<FEditableConstraintI
 	RefreshConstraintList();
 }
 
+void SConstraintsEditionWidget::InvalidateConstraintList()
+{
+	FBaseConstraintListWidget::InvalidateConstraintList();
+	UpdateSequencer();
+}
+
+void SConstraintsEditionWidget::UpdateSequencer()
+{
+	if (!WeakSequencer.IsValid())
+	{
+		WeakSequencer = GetSequencerChecked();
+	}
+}
+
 FReply SConstraintsEditionWidget::OnBakeClicked()
 {
-	UTickableTransformConstraint* InConstraint = nullptr;
+	UpdateSequencer();
 	if (!WeakSequencer.IsValid() || ListItems.Num() < 1)
 	{
 		return FReply::Unhandled();
@@ -1253,6 +1267,8 @@ TSharedPtr<SWidget> SConstraintsEditionWidget::CreateContextMenu()
 	{
 		return SNullWidget::NullWidget;
 	}
+
+	UpdateSequencer();
 	
 	const FConstraintsManagerController& Controller = FConstraintsManagerController::Get(World);
 	UTickableConstraint* Constraint = ListItems[Index]->Constraint.Get();
diff --git a/Engine/Plugins/Animation/ControlRig/Source/ControlRigEditor/Private/Editor/Constraints/SConstraintsWidget.h b/Engine/Plugins/Animation/ControlRig/Source/ControlRigEditor/Private/Editor/Constraints/SConstraintsWidget.h
--- a/Engine/Plugins/Animation/ControlRig/Source/ControlRigEditor/Private/Editor/Constraints/SConstraintsWidget.h
+++ b/Engine/Plugins/Animation/ControlRig/Source/ControlRigEditor/Private/Editor/Constraints/SConstraintsWidget.h
@@ -245,7 +245,7 @@ public:
 	/* End FEditorUndoClient interface */
 
 	/** Invalidates the constraint list for further rebuild. */
-	void InvalidateConstraintList();
+	virtual void InvalidateConstraintList();
 
 	/** Rebuild the constraint list based on the current selection. */
 	virtual int32 RefreshConstraintList();
@@ -317,6 +317,9 @@ public:
 	/**  */
 	void RemoveItem(const TSharedPtr<FEditableConstraintItem>& Item);
 
+	// FBaseConstraintListWidget overrides
+	virtual void InvalidateConstraintList() override;
+	
 private:
 
 	/** Generates a widget for the specified item */
@@ -332,6 +335,8 @@ private:
 
 	FReply OnBakeClicked();
 
+	void UpdateSequencer();
+
 	//sequencer and it's time
 	TWeakPtr<ISequencer> WeakSequencer;
 	bool SequencerTimeChanged();
