Control Rig: ResolveConnector node to offer SkipSocket option

#rnx

diff --git a/Engine/Plugins/Animation/ControlRig/Source/ControlRig/Private/Units/Execution/RigUnit_RigModules.cpp b/Engine/Plugins/Animation/ControlRig/Source/ControlRig/Private/Units/Execution/RigUnit_RigModules.cpp
--- a/Engine/Plugins/Animation/ControlRig/Source/ControlRig/Private/Units/Execution/RigUnit_RigModules.cpp
+++ b/Engine/Plugins/Animation/ControlRig/Source/ControlRig/Private/Units/Execution/RigUnit_RigModules.cpp
@@ -13,6 +13,14 @@ FRigUnit_ResolveConnector_Execute()
 	if(const FRigBaseElement* ResolvedElement = ExecuteContext.Hierarchy->Find(Connector))
 	{
 		Result = ResolvedElement->GetKey();
+		if(SkipSocket && Result.IsValid() && Result.Type == ERigElementType::Socket)
+		{
+			const FRigElementKey ParentOfSocket = ExecuteContext.Hierarchy->GetFirstParent(Result);
+			if(ParentOfSocket.IsValid())
+			{
+				Result = ParentOfSocket;
+			}
+		}
 	}
 }
 
diff --git a/Engine/Plugins/Animation/ControlRig/Source/ControlRig/Public/Units/Execution/RigUnit_RigModules.h b/Engine/Plugins/Animation/ControlRig/Source/ControlRig/Public/Units/Execution/RigUnit_RigModules.h
--- a/Engine/Plugins/Animation/ControlRig/Source/ControlRig/Public/Units/Execution/RigUnit_RigModules.h
+++ b/Engine/Plugins/Animation/ControlRig/Source/ControlRig/Public/Units/Execution/RigUnit_RigModules.h
@@ -29,6 +29,7 @@ struct CONTROLRIG_API FRigUnit_ResolveConnector : public FRigUnit_RigModulesBase
 	FRigUnit_ResolveConnector()
 	{
 		Connector = Result = FRigElementKey(NAME_None, ERigElementType::Connector);
+		SkipSocket = false;
 	}
 
 	RIGVM_METHOD()
@@ -40,6 +41,13 @@ struct CONTROLRIG_API FRigUnit_ResolveConnector : public FRigUnit_RigModulesBase
 	UPROPERTY(meta = (Input, ExpandByDefault))
 	FRigElementKey Connector;
 
+	/*
+	 * If the connector is resolved to a socket the node
+	 * will return the socket's direct parent (skipping it).
+	 */
+	UPROPERTY(meta = (Input, ExpandByDefault))
+	bool SkipSocket;
+
 	/*
 	 * The resulting item the connector is resolved to
 	 */
