Control Rig: Functionality to get and set metadata per module

#rnx

diff --git a/Engine/Plugins/Animation/ControlRig/Source/ControlRig/Private/ModularRig.cpp b/Engine/Plugins/Animation/ControlRig/Source/ControlRig/Private/ModularRig.cpp
--- a/Engine/Plugins/Animation/ControlRig/Source/ControlRig/Private/ModularRig.cpp
+++ b/Engine/Plugins/Animation/ControlRig/Source/ControlRig/Private/ModularRig.cpp
@@ -135,8 +135,47 @@ const FRigModuleReference* FRigModuleInstance::GetModuleReference() const
 	return nullptr;
 }
 
+const FRigModuleInstance* FRigModuleInstance::GetParentModule() const
+{
+	if(ParentPath.IsEmpty())
+	{
+		return this;
+	}
+	if(const UControlRig* Rig = GetRig())
+	{
+		if(const UModularRig* ModularRig = Cast<UModularRig>(Rig->GetParentRig()))
+		{
+			return ModularRig->FindModule(ParentPath);
+		}
+	}
+	return nullptr;
+}
+
+const FRigModuleInstance* FRigModuleInstance::GetRootModule() const
+{
+	if(ParentPath.IsEmpty())
+	{
+		return this;
+	}
+	if(const UControlRig* Rig = GetRig())
+	{
+		if(const UModularRig* ModularRig = Cast<UModularRig>(Rig->GetParentRig()))
+		{
+			FString RootPath = ParentPath;
+			(void)URigHierarchy::SplitNameSpace(ParentPath, &RootPath, nullptr, false);
+			return ModularRig->FindModule(RootPath);
+		}
+	}
+	return nullptr;
+}
+
 const FRigConnectorElement* FRigModuleInstance::FindPrimaryConnector() const
 {
+	if(PrimaryConnector)
+	{
+		return PrimaryConnector;
+	}
+	
 	if(const UControlRig* Rig = GetRig())
 	{
 		if(const URigHierarchy* Hierarchy = Rig->GetHierarchy())
@@ -152,7 +191,8 @@ const FRigConnectorElement* FRigModuleInstance::FindPrimaryConnector() const
 					{
 						if(ModulePath.Equals(MyModulePath, ESearchCase::CaseSensitive))
 						{
-							return Connector;
+							PrimaryConnector = Connector;
+							return PrimaryConnector;
 						}
 					}
 				}
diff --git a/Engine/Plugins/Animation/ControlRig/Source/ControlRig/Private/Units/Hierarchy/RigUnit_Metadata.cpp b/Engine/Plugins/Animation/ControlRig/Source/ControlRig/Private/Units/Hierarchy/RigUnit_Metadata.cpp
--- a/Engine/Plugins/Animation/ControlRig/Source/ControlRig/Private/Units/Hierarchy/RigUnit_Metadata.cpp
+++ b/Engine/Plugins/Animation/ControlRig/Source/ControlRig/Private/Units/Hierarchy/RigUnit_Metadata.cpp
@@ -1,6 +1,8 @@
 // Copyright Epic Games, Inc. All Rights Reserved.
 
 #include "Units/Hierarchy/RigUnit_Metadata.h"
+
+#include "ModularRig.h"
 #include "RigVMTypeUtils.h"
 #include "RigVMCore/RigVMStruct.h"
 #include "Units/RigUnitContext.h"
@@ -30,8 +32,8 @@ FString FRigDispatch_MetadataBase::GetNodeTitle(const FRigVMTemplateTypeMap& InT
 		if(ValueTypeIndex != RigVMTypeUtils::TypeIndex::WildCard &&
 			ValueTypeIndex != RigVMTypeUtils::TypeIndex::WildCardArray)
 		{
-			static constexpr TCHAR GetMetadataFormat[] = TEXT("Get %s Metadata");
-			static constexpr TCHAR SetMetadataFormat[] = TEXT("Set %s Metadata");
+			static constexpr TCHAR GetMetadataFormat[] = TEXT("Get %s%s Metadata");
+			static constexpr TCHAR SetMetadataFormat[] = TEXT("Set %s%s Metadata");
 
 			const FRigVMTemplateArgumentType& ValueType = FRigVMRegistry::Get().GetType(ValueTypeIndex);
 			FString ValueName;
@@ -58,7 +60,7 @@ FString FRigDispatch_MetadataBase::GetNodeTitle(const FRigVMTemplateTypeMap& InT
 				ValueName += TEXT(" Array");
 			}
 
-			return FString::Printf(IsSetMetadata() ? SetMetadataFormat : GetMetadataFormat, *ValueName); 
+			return FString::Printf(IsSetMetadata() ? SetMetadataFormat : GetMetadataFormat, *GetNodeTitlePrefix(), *ValueName); 
 		}
 	}
 	return FRigDispatchFactory::GetNodeTitle(InTypes);
@@ -428,30 +430,6 @@ FRigUnit_RemoveMetadata_Execute()
 
 	if(CachedIndex.UpdateCache(Item, Hierarchy))
 	{
-		if(NameSpace == ERigMetaDataNameSpace::All)
-		{
-			Removed = true;
-			
-			// run remove for all relevant namespace types.
-			for(int32 NameSpaceTypeIndex = 0;
-				NameSpaceTypeIndex < (int32)ERigMetaDataNameSpace::Last;
-				NameSpaceTypeIndex++)
-			{
-				const ERigMetaDataNameSpace NameSpaceFromIndex = (ERigMetaDataNameSpace)NameSpaceTypeIndex;
-				if(NameSpaceFromIndex == ERigMetaDataNameSpace::All)
-				{
-					continue;
-				}
-					
-				bool bResult = false;
-				FRigUnit_RemoveMetadata::StaticExecute(ExecuteContext, Item, Name, NameSpaceFromIndex, bResult, CachedIndex);
-				if(!bResult)
-				{
-					Removed = false;
-				}
-			}
-			return;
-		}
 		if(FRigBaseElement* Element = Hierarchy->Get(CachedIndex))
 		{
 			const FName LocalName = ExecuteContext.AdaptMetadataName(NameSpace, Name);
@@ -476,31 +454,6 @@ FRigUnit_RemoveAllMetadata_Execute()
 
 	if(CachedIndex.UpdateCache(Item, Hierarchy))
 	{
-		if(NameSpace == ERigMetaDataNameSpace::All)
-		{
-			Removed = true;
-			
-			// run remove for all relevant namespace types.
-			for(int32 NameSpaceTypeIndex = 0;
-				NameSpaceTypeIndex < (int32)ERigMetaDataNameSpace::Last;
-				NameSpaceTypeIndex++)
-			{
-				const ERigMetaDataNameSpace NameSpaceFromIndex = (ERigMetaDataNameSpace)NameSpaceTypeIndex;
-				if(NameSpaceFromIndex == ERigMetaDataNameSpace::All)
-				{
-					continue;
-				}
-					
-				bool bResult = false;
-				FRigUnit_RemoveAllMetadata::StaticExecute(ExecuteContext, Item, NameSpaceFromIndex, bResult, CachedIndex);
-				if(!bResult)
-				{
-					Removed = false;
-				}
-			}
-			return;
-		}
-		
 		if(FRigBaseElement* Element = Hierarchy->Get(CachedIndex))
 		{
 			if(NameSpace != ERigMetaDataNameSpace::None)
@@ -764,7 +717,7 @@ FRigUnit_FindItemsWithMetadataTagArray_Execute()
 	TArrayView<const FName> LocalTags(Tags);
 	TArray<FName> AdaptedTags;
 
-	const bool bUseNameSpace = (NameSpace != ERigMetaDataNameSpace::None) && (NameSpace != ERigMetaDataNameSpace::All);
+	const bool bUseNameSpace = NameSpace != ERigMetaDataNameSpace::None;
 	if(bUseNameSpace && ExecuteContext.IsRigModule())
 	{
 		AdaptedTags.Reserve(Tags.Num());
@@ -817,7 +770,7 @@ FRigUnit_FilterItemsByMetadataTags_Execute()
 
 	TArrayView<const FName> LocalTags(Tags);
 	TArray<FName> AdaptedTags;
-	const bool bUseNameSpace = (NameSpace != ERigMetaDataNameSpace::None) && (NameSpace != ERigMetaDataNameSpace::All);
+	const bool bUseNameSpace = NameSpace != ERigMetaDataNameSpace::None;
 	if(bUseNameSpace && ExecuteContext.IsRigModule())
 	{
 		AdaptedTags.Reserve(Tags.Num());
@@ -882,3 +835,260 @@ FRigUnit_FilterItemsByMetadataTags_Execute()
 		}
 	}
 }
+
+const TArray<FRigVMTemplateArgumentInfo>& FRigDispatch_GetModuleMetadata::GetArgumentInfos() const
+{
+	if(ValueArgIndex == INDEX_NONE)
+	{
+		NameArgIndex = Infos.Emplace(NameArgName, ERigVMPinDirection::Input, RigVMTypeUtils::TypeIndex::FName);
+		NameSpaceArgIndex = Infos.Emplace(NameSpaceArgName, ERigVMPinDirection::Input, FRigVMRegistry::Get().GetTypeIndex<ERigMetaDataNameSpace>());
+		DefaultArgIndex = Infos.Emplace(DefaultArgName, ERigVMPinDirection::Input, GetValueTypes());
+		ValueArgIndex = Infos.Emplace(ValueArgName, ERigVMPinDirection::Output, GetValueTypes());
+		FoundArgIndex = Infos.Emplace(FoundArgName, ERigVMPinDirection::Output, RigVMTypeUtils::TypeIndex::Bool);
+	};
+	return Infos;
+}
+
+FRigBaseMetadata* FRigDispatch_GetModuleMetadata::FindMetadata(const FRigVMExtendedExecuteContext& InContext, const FName& InName, ERigMetadataType InType, ERigMetaDataNameSpace InNameSpace)
+{
+	const FControlRigExecuteContext& ExecuteContext = InContext.GetPublicData<FControlRigExecuteContext>();
+	if(const FRigModuleInstance* ModuleInstance = ExecuteContext.GetRigModuleInstance(InNameSpace))
+	{
+		if(const FRigConnectorElement* PrimaryConnector = ModuleInstance->FindPrimaryConnector())
+		{
+			// first try to find the metadata in the namespace
+			return ExecuteContext.Hierarchy->FindMetadataForElement(PrimaryConnector, InName, InType);
+		}
+	}
+	else if(ExecuteContext.IsRigModule())
+	{
+		// we are not in a rig module - but we still want to store the metadata for testing.
+		const TArray<FRigConnectorElement*> Connectors = ExecuteContext.Hierarchy->GetConnectors();
+		for(const FRigConnectorElement* Connector : Connectors)
+		{
+			if(Connector->IsPrimary())
+			{
+				const FName Name = ExecuteContext.AdaptMetadataName(InNameSpace, InName);
+				return ExecuteContext.Hierarchy->FindMetadataForElement(Connector, InName, InType);
+			}
+		}
+	}
+	return nullptr;
+}
+
+FRigVMFunctionPtr FRigDispatch_GetModuleMetadata::GetDispatchFunctionImpl(const FRigVMTemplateTypeMap& InTypes) const
+{
+	const FRigVMRegistry& Registry = FRigVMRegistry::Get();
+	const TRigVMTypeIndex& ValueTypeIndex = InTypes.FindChecked(TEXT("Value"));
+	
+	if(ValueTypeIndex == RigVMTypeUtils::TypeIndex::Bool)
+	{
+		return &FRigDispatch_GetModuleMetadata::GetModuleMetadataDispatch<bool, FRigBoolMetadata, ERigMetadataType::Bool>;
+	}
+	if(ValueTypeIndex == RigVMTypeUtils::TypeIndex::Float)
+	{
+		return &FRigDispatch_GetModuleMetadata::GetModuleMetadataDispatch<float, FRigFloatMetadata, ERigMetadataType::Float>;
+	}
+	if(ValueTypeIndex == RigVMTypeUtils::TypeIndex::Int32)
+	{
+		return &FRigDispatch_GetModuleMetadata::GetModuleMetadataDispatch<int32, FRigInt32Metadata, ERigMetadataType::Int32>;
+	}
+	if(ValueTypeIndex == RigVMTypeUtils::TypeIndex::FName)
+	{
+		return &FRigDispatch_GetModuleMetadata::GetModuleMetadataDispatch<FName, FRigNameMetadata, ERigMetadataType::Name>;
+	}
+	if(ValueTypeIndex == Registry.GetTypeIndex<FVector>(false))
+	{
+		return &FRigDispatch_GetModuleMetadata::GetModuleMetadataDispatch<FVector, FRigVectorMetadata, ERigMetadataType::Vector>;
+	}
+	if(ValueTypeIndex == Registry.GetTypeIndex<FRotator>(false))
+	{
+		return &FRigDispatch_GetModuleMetadata::GetModuleMetadataDispatch<FRotator, FRigRotatorMetadata, ERigMetadataType::Rotator>;
+	}
+	if(ValueTypeIndex == Registry.GetTypeIndex<FQuat>(false))
+	{
+		return &FRigDispatch_GetModuleMetadata::GetModuleMetadataDispatch<FQuat, FRigQuatMetadata, ERigMetadataType::Quat>;
+	}
+	if(ValueTypeIndex == Registry.GetTypeIndex<FTransform>(false))
+	{
+		return &FRigDispatch_GetModuleMetadata::GetModuleMetadataDispatch<FTransform, FRigTransformMetadata, ERigMetadataType::Transform>;
+	}
+	if(ValueTypeIndex == Registry.GetTypeIndex<FLinearColor>(false))
+	{
+		return &FRigDispatch_GetModuleMetadata::GetModuleMetadataDispatch<FLinearColor, FRigLinearColorMetadata, ERigMetadataType::LinearColor>;
+	}
+	if(ValueTypeIndex == Registry.GetTypeIndex<FRigElementKey>(false))
+	{
+		return &FRigDispatch_GetModuleMetadata::GetModuleMetadataDispatch<FRigElementKey, FRigElementKeyMetadata, ERigMetadataType::RigElementKey>;
+	}
+	if(ValueTypeIndex == RigVMTypeUtils::TypeIndex::BoolArray)
+	{
+		return &FRigDispatch_GetModuleMetadata::GetModuleMetadataDispatch<TArray<bool>, FRigBoolArrayMetadata, ERigMetadataType::BoolArray>;
+	}
+	if(ValueTypeIndex == RigVMTypeUtils::TypeIndex::FloatArray)
+	{
+		return &FRigDispatch_GetModuleMetadata::GetModuleMetadataDispatch<TArray<float>, FRigFloatArrayMetadata, ERigMetadataType::FloatArray>;
+	}
+	if(ValueTypeIndex == RigVMTypeUtils::TypeIndex::Int32Array)
+	{
+		return &FRigDispatch_GetModuleMetadata::GetModuleMetadataDispatch<TArray<int32>, FRigInt32ArrayMetadata, ERigMetadataType::Int32Array>;
+	}
+	if(ValueTypeIndex == RigVMTypeUtils::TypeIndex::FNameArray)
+	{
+		return &FRigDispatch_GetModuleMetadata::GetModuleMetadataDispatch<TArray<FName>, FRigNameArrayMetadata, ERigMetadataType::NameArray>;
+	}
+	if(ValueTypeIndex == Registry.GetTypeIndex<FVector>(true))
+	{
+		return &FRigDispatch_GetModuleMetadata::GetModuleMetadataDispatch<TArray<FVector>, FRigVectorArrayMetadata, ERigMetadataType::VectorArray>;
+	}
+	if(ValueTypeIndex == Registry.GetTypeIndex<FRotator>(true))
+	{
+		return &FRigDispatch_GetModuleMetadata::GetModuleMetadataDispatch<TArray<FRotator>, FRigRotatorArrayMetadata, ERigMetadataType::RotatorArray>;
+	}
+	if(ValueTypeIndex == Registry.GetTypeIndex<FQuat>(true))
+	{
+		return &FRigDispatch_GetModuleMetadata::GetModuleMetadataDispatch<TArray<FQuat>, FRigQuatArrayMetadata, ERigMetadataType::QuatArray>;
+	}
+	if(ValueTypeIndex == Registry.GetTypeIndex<FTransform>(true))
+	{
+		return &FRigDispatch_GetModuleMetadata::GetModuleMetadataDispatch<TArray<FTransform>, FRigTransformArrayMetadata, ERigMetadataType::TransformArray>;
+	}
+	if(ValueTypeIndex == Registry.GetTypeIndex<FLinearColor>(true))
+	{
+		return &FRigDispatch_GetModuleMetadata::GetModuleMetadataDispatch<TArray<FLinearColor>, FRigLinearColorArrayMetadata, ERigMetadataType::LinearColorArray>;
+	}
+	if(ValueTypeIndex == Registry.GetTypeIndex<FRigElementKey>(true))
+	{
+		return &FRigDispatch_GetModuleMetadata::GetModuleMetadataDispatch<TArray<FRigElementKey>, FRigElementKeyArrayMetadata, ERigMetadataType::RigElementKeyArray>;
+	}
+
+	return nullptr;
+}
+
+const TArray<FRigVMTemplateArgumentInfo>& FRigDispatch_SetModuleMetadata::GetArgumentInfos() const
+{
+	if(ValueArgIndex == INDEX_NONE)
+	{
+		NameArgIndex = Infos.Emplace(NameArgName, ERigVMPinDirection::Input, RigVMTypeUtils::TypeIndex::FName);
+		NameSpaceArgIndex = Infos.Emplace(NameSpaceArgName, ERigVMPinDirection::Input, FRigVMRegistry::Get().GetTypeIndex<ERigMetaDataNameSpace>());
+		ValueArgIndex = Infos.Emplace(ValueArgName, ERigVMPinDirection::Input, GetValueTypes());
+		SuccessArgIndex = Infos.Emplace(SuccessArgName, ERigVMPinDirection::Output, RigVMTypeUtils::TypeIndex::Bool);
+	};
+	return Infos;
+}
+
+FRigBaseMetadata* FRigDispatch_SetModuleMetadata::FindOrAddMetadata(const FControlRigExecuteContext& InContext, const FName& InName, ERigMetadataType InType, ERigMetaDataNameSpace InNameSpace)
+{
+	constexpr bool bNotify = true;
+	
+	if(const FRigModuleInstance* ModuleInstance = InContext.GetRigModuleInstance(InNameSpace))
+	{
+		if(const FRigConnectorElement* PrimaryConnector = ModuleInstance->FindPrimaryConnector())
+		{
+			return InContext.Hierarchy->GetMetadataForElement(const_cast<FRigConnectorElement*>(PrimaryConnector), InName, InType, bNotify);
+		}
+	}
+	else if(InContext.IsRigModule())
+	{
+		// we are not in a rig module - but we still want to store the metadata for testing.
+		const TArray<FRigConnectorElement*> Connectors = InContext.Hierarchy->GetConnectors();
+		for(const FRigConnectorElement* Connector : Connectors)
+		{
+			if(Connector->IsPrimary())
+			{
+				const FName Name = InContext.AdaptMetadataName(InNameSpace, InName);
+				return InContext.Hierarchy->GetMetadataForElement(const_cast<FRigConnectorElement*>(Connector), InName, InType, bNotify);
+			}
+		}
+	}
+	return nullptr;
+}
+
+FRigVMFunctionPtr FRigDispatch_SetModuleMetadata::GetDispatchFunctionImpl(const FRigVMTemplateTypeMap& InTypes) const
+{
+	const FRigVMRegistry& Registry = FRigVMRegistry::Get();
+	const TRigVMTypeIndex& ValueTypeIndex = InTypes.FindChecked(TEXT("Value"));
+	
+	if(ValueTypeIndex == RigVMTypeUtils::TypeIndex::Bool)
+	{
+		return &FRigDispatch_SetModuleMetadata::SetModuleMetadataDispatch<bool, FRigBoolMetadata, ERigMetadataType::Bool>;
+	}
+	if(ValueTypeIndex == RigVMTypeUtils::TypeIndex::Float)
+	{
+		return &FRigDispatch_SetModuleMetadata::SetModuleMetadataDispatch<float, FRigFloatMetadata, ERigMetadataType::Float>;
+	}
+	if(ValueTypeIndex == RigVMTypeUtils::TypeIndex::Int32)
+	{
+		return &FRigDispatch_SetModuleMetadata::SetModuleMetadataDispatch<int32, FRigInt32Metadata, ERigMetadataType::Int32>;
+	}
+	if(ValueTypeIndex == RigVMTypeUtils::TypeIndex::FName)
+	{
+		return &FRigDispatch_SetModuleMetadata::SetModuleMetadataDispatch<FName, FRigNameMetadata, ERigMetadataType::Name>;
+	}
+	if(ValueTypeIndex == Registry.GetTypeIndex<FVector>(false))
+	{
+		return &FRigDispatch_SetModuleMetadata::SetModuleMetadataDispatch<FVector, FRigVectorMetadata, ERigMetadataType::Vector>;
+	}
+	if(ValueTypeIndex == Registry.GetTypeIndex<FRotator>(false))
+	{
+		return &FRigDispatch_SetModuleMetadata::SetModuleMetadataDispatch<FRotator, FRigRotatorMetadata, ERigMetadataType::Rotator>;
+	}
+	if(ValueTypeIndex == Registry.GetTypeIndex<FQuat>(false))
+	{
+		return &FRigDispatch_SetModuleMetadata::SetModuleMetadataDispatch<FQuat, FRigQuatMetadata, ERigMetadataType::Quat>;
+	}
+	if(ValueTypeIndex == Registry.GetTypeIndex<FTransform>(false))
+	{
+		return &FRigDispatch_SetModuleMetadata::SetModuleMetadataDispatch<FTransform, FRigTransformMetadata, ERigMetadataType::Transform>;
+	}
+	if(ValueTypeIndex == Registry.GetTypeIndex<FLinearColor>(false))
+	{
+		return &FRigDispatch_SetModuleMetadata::SetModuleMetadataDispatch<FLinearColor, FRigLinearColorMetadata, ERigMetadataType::LinearColor>;
+	}
+	if(ValueTypeIndex == Registry.GetTypeIndex<FRigElementKey>(false))
+	{
+		return &FRigDispatch_SetModuleMetadata::SetModuleMetadataDispatch<FRigElementKey, FRigElementKeyMetadata, ERigMetadataType::RigElementKey>;
+	}
+	if(ValueTypeIndex == RigVMTypeUtils::TypeIndex::BoolArray)
+	{
+		return &FRigDispatch_SetModuleMetadata::SetModuleMetadataDispatch<TArray<bool>, FRigBoolArrayMetadata, ERigMetadataType::BoolArray>;
+	}
+	if(ValueTypeIndex == RigVMTypeUtils::TypeIndex::FloatArray)
+	{
+		return &FRigDispatch_SetModuleMetadata::SetModuleMetadataDispatch<TArray<float>, FRigFloatArrayMetadata, ERigMetadataType::FloatArray>;
+	}
+	if(ValueTypeIndex == RigVMTypeUtils::TypeIndex::Int32Array)
+	{
+		return &FRigDispatch_SetModuleMetadata::SetModuleMetadataDispatch<TArray<int32>, FRigInt32ArrayMetadata, ERigMetadataType::Int32Array>;
+	}
+	if(ValueTypeIndex == RigVMTypeUtils::TypeIndex::FNameArray)
+	{
+		return &FRigDispatch_SetModuleMetadata::SetModuleMetadataDispatch<TArray<FName>, FRigNameArrayMetadata, ERigMetadataType::NameArray>;
+	}
+	if(ValueTypeIndex == Registry.GetTypeIndex<FVector>(true))
+	{
+		return &FRigDispatch_SetModuleMetadata::SetModuleMetadataDispatch<TArray<FVector>, FRigVectorArrayMetadata, ERigMetadataType::VectorArray>;
+	}
+	if(ValueTypeIndex == Registry.GetTypeIndex<FRotator>(true))
+	{
+		return &FRigDispatch_SetModuleMetadata::SetModuleMetadataDispatch<TArray<FRotator>, FRigRotatorArrayMetadata, ERigMetadataType::RotatorArray>;
+	}
+	if(ValueTypeIndex == Registry.GetTypeIndex<FQuat>(true))
+	{
+		return &FRigDispatch_SetModuleMetadata::SetModuleMetadataDispatch<TArray<FQuat>, FRigQuatArrayMetadata, ERigMetadataType::QuatArray>;
+	}
+	if(ValueTypeIndex == Registry.GetTypeIndex<FTransform>(true))
+	{
+		return &FRigDispatch_SetModuleMetadata::SetModuleMetadataDispatch<TArray<FTransform>, FRigTransformArrayMetadata, ERigMetadataType::TransformArray>;
+	}
+	if(ValueTypeIndex == Registry.GetTypeIndex<FLinearColor>(true))
+	{
+		return &FRigDispatch_SetModuleMetadata::SetModuleMetadataDispatch<TArray<FLinearColor>, FRigLinearColorArrayMetadata, ERigMetadataType::LinearColorArray>;
+	}
+	if(ValueTypeIndex == Registry.GetTypeIndex<FRigElementKey>(true))
+	{
+		return &FRigDispatch_SetModuleMetadata::SetModuleMetadataDispatch<TArray<FRigElementKey>, FRigElementKeyArrayMetadata, ERigMetadataType::RigElementKeyArray>;
+	}
+
+	return nullptr;
+}
diff --git a/Engine/Plugins/Animation/ControlRig/Source/ControlRig/Private/Units/RigUnitContext.cpp b/Engine/Plugins/Animation/ControlRig/Source/ControlRig/Private/Units/RigUnitContext.cpp
--- a/Engine/Plugins/Animation/ControlRig/Source/ControlRig/Private/Units/RigUnitContext.cpp
+++ b/Engine/Plugins/Animation/ControlRig/Source/ControlRig/Private/Units/RigUnitContext.cpp
@@ -120,10 +120,39 @@ FString FControlRigExecuteContext::GetElementNameSpace(ERigMetaDataNameSpace InN
 	return FString();
 }
 
+const FRigModuleInstance* FControlRigExecuteContext::GetRigModuleInstance(ERigMetaDataNameSpace InNameSpaceType) const
+{
+	if(RigModuleInstance)
+	{
+		switch(InNameSpaceType)
+		{
+			case ERigMetaDataNameSpace::Self:
+			{
+				return RigModuleInstance;
+			}
+			case ERigMetaDataNameSpace::Parent:
+			{
+				return RigModuleInstance->GetParentModule();
+			}
+			case ERigMetaDataNameSpace::Root:
+			{
+				return RigModuleInstance->GetRootModule();
+			}
+			case ERigMetaDataNameSpace::None:
+			default:
+			{
+				break;
+			}
+		}
+		
+	}
+	return nullptr;
+}
+
 FName FControlRigExecuteContext::AdaptMetadataName(ERigMetaDataNameSpace InNameSpaceType, const FName& InMetadataName) const
 {
 	// only if we are within a rig module let's adapt the meta data name
-	const bool bUseNameSpace = (InNameSpaceType != ERigMetaDataNameSpace::None) && (InNameSpaceType != ERigMetaDataNameSpace::All);
+	const bool bUseNameSpace = InNameSpaceType != ERigMetaDataNameSpace::None;
 	if(bUseNameSpace && !InMetadataName.IsNone())
 	{
 		// if the metadata name already contains a namespace - we are just going
diff --git a/Engine/Plugins/Animation/ControlRig/Source/ControlRig/Public/ModularRig.h b/Engine/Plugins/Animation/ControlRig/Source/ControlRig/Public/ModularRig.h
--- a/Engine/Plugins/Animation/ControlRig/Source/ControlRig/Public/ModularRig.h
+++ b/Engine/Plugins/Animation/ControlRig/Source/ControlRig/Public/ModularRig.h
@@ -71,6 +71,7 @@ public:
 	TMap<FName, FRigVMExternalVariable> VariableBindings;
 
 	TArray<FRigModuleInstance*> CachedChildren;
+	mutable const FRigConnectorElement* PrimaryConnector = nullptr;
 
 	FString GetShortName() const;
 	FString GetLongName() const
@@ -83,6 +84,8 @@ public:
 	void SetRig(UControlRig* InRig);
 	bool ContainsRig(const UControlRig* InRig) const;
 	const FRigModuleReference* GetModuleReference() const;
+	const FRigModuleInstance* GetParentModule() const;
+	const FRigModuleInstance* GetRootModule() const;
 	const FRigConnectorElement* FindPrimaryConnector() const;
 	TArray<const FRigConnectorElement*> FindConnectors() const;
 	bool IsRootModule() const;
diff --git a/Engine/Plugins/Animation/ControlRig/Source/ControlRig/Public/Rigs/RigHierarchy.h b/Engine/Plugins/Animation/ControlRig/Source/ControlRig/Public/Rigs/RigHierarchy.h
--- a/Engine/Plugins/Animation/ControlRig/Source/ControlRig/Public/Rigs/RigHierarchy.h
+++ b/Engine/Plugins/Animation/ControlRig/Source/ControlRig/Public/Rigs/RigHierarchy.h
@@ -4786,6 +4786,8 @@ private:
 	friend struct FRigHierarchyRedirectorGuard;
 	friend struct FRigDispatch_GetMetadata;
 	friend struct FRigDispatch_SetMetadata;
+	friend struct FRigDispatch_GetModuleMetadata;
+	friend struct FRigDispatch_SetModuleMetadata;
 };
 
 struct CONTROLRIG_API FRigHierarchyInteractionBracket
diff --git a/Engine/Plugins/Animation/ControlRig/Source/ControlRig/Public/Units/Hierarchy/RigUnit_Metadata.h b/Engine/Plugins/Animation/ControlRig/Source/ControlRig/Public/Units/Hierarchy/RigUnit_Metadata.h
--- a/Engine/Plugins/Animation/ControlRig/Source/ControlRig/Public/Units/Hierarchy/RigUnit_Metadata.h
+++ b/Engine/Plugins/Animation/ControlRig/Source/ControlRig/Public/Units/Hierarchy/RigUnit_Metadata.h
@@ -41,6 +41,7 @@ struct CONTROLRIG_API FRigDispatch_MetadataBase : public FRigDispatchFactory
 
 #if WITH_EDITOR
 	virtual FString GetNodeTitle(const FRigVMTemplateTypeMap& InTypes) const override;;
+	virtual FString GetNodeTitlePrefix() const { return FString(); }
 #endif
 	virtual const TArray<FRigVMTemplateArgumentInfo>& GetArgumentInfos() const override;
 	virtual bool IsSetMetadata() const { return false; }
@@ -198,6 +199,7 @@ protected:
 	{
 		return CheckArgumentType(Handles[ItemArgIndex].IsType<FRigElementKey>(), ItemArgName) &&
 			CheckArgumentType(Handles[NameArgIndex].IsType<FName>(), NameArgName) &&
+			CheckArgumentType(Handles[NameSpaceArgIndex].IsType<ERigMetaDataNameSpace>(), NameSpaceArgName) &&
 			CheckArgumentType(Handles[CacheArgIndex].IsType<FCachedRigElement>(true), CacheArgName) &&
 			CheckArgumentType(Handles[DefaultArgIndex].IsType<ValueType>(), DefaultArgName) &&
 			CheckArgumentType(Handles[ValueArgIndex].IsType<ValueType>(), ValueArgName) &&
@@ -241,7 +243,7 @@ protected:
 };
 
 /*
- * Returns some metadata for the provided item
+ * Sets some metadata for the provided item
  */
 USTRUCT(meta=(DisplayName="Set Metadata"))
 struct CONTROLRIG_API FRigDispatch_SetMetadata : public FRigDispatch_MetadataBase
@@ -263,6 +265,7 @@ protected:
 	{
 		return CheckArgumentType(Handles[ItemArgIndex].IsType<FRigElementKey>(), ItemArgName) &&
 			CheckArgumentType(Handles[NameArgIndex].IsType<FName>(), NameArgName) &&
+			CheckArgumentType(Handles[NameSpaceArgIndex].IsType<ERigMetaDataNameSpace>(), NameSpaceArgName) &&
 			CheckArgumentType(Handles[CacheArgIndex].IsType<FCachedRigElement>(true), CacheArgName) &&
 			CheckArgumentType(Handles[ValueArgIndex].IsType<ValueType>(), ValueArgName) &&
 			CheckArgumentType(Handles[SuccessArgIndex].IsType<bool>(), SuccessArgName);
@@ -859,4 +862,132 @@ struct CONTROLRIG_API FRigUnit_FilterItemsByMetadataTags : public FRigUnit
 	// Used to cache the internally used indices
 	UPROPERTY()
 	TArray<FCachedRigElement> CachedIndices;
-};
\ No newline at end of file
+};
+
+/*
+ * Returns some metadata on a given module
+ */
+USTRUCT(meta=(DisplayName="Get Module Metadata"))
+struct CONTROLRIG_API FRigDispatch_GetModuleMetadata : public FRigDispatch_GetMetadata
+{
+	GENERATED_BODY()
+
+	virtual const TArray<FRigVMTemplateArgumentInfo>& GetArgumentInfos() const override;
+#if WITH_EDITOR
+	virtual FString GetNodeTitlePrefix() const override { return TEXT("Module "); }
+#endif
+	
+protected:
+
+	static FRigBaseMetadata* FindMetadata(const FRigVMExtendedExecuteContext& InContext, const FName& InName, ERigMetadataType InType, ERigMetaDataNameSpace InNameSpace);
+	virtual FRigVMFunctionPtr GetDispatchFunctionImpl(const FRigVMTemplateTypeMap& InTypes) const override;
+
+#if WITH_EDITOR
+	template<typename ValueType>
+	bool CheckArgumentTypes(FRigVMMemoryHandleArray Handles) const
+	{
+		return CheckArgumentType(Handles[NameArgIndex].IsType<FName>(), NameArgName) &&
+			CheckArgumentType(Handles[NameSpaceArgIndex].IsType<ERigMetaDataNameSpace>(), NameSpaceArgName) &&
+			CheckArgumentType(Handles[DefaultArgIndex].IsType<ValueType>(), DefaultArgName) &&
+			CheckArgumentType(Handles[ValueArgIndex].IsType<ValueType>(), ValueArgName) &&
+			CheckArgumentType(Handles[FoundArgIndex].IsType<bool>(), FoundArgName);
+	}
+#endif
+
+	template<typename ValueType, typename MetadataType, ERigMetadataType EnumValue>
+	static void GetModuleMetadataDispatch(FRigVMExtendedExecuteContext& InContext, FRigVMMemoryHandleArray Handles, FRigVMPredicateBranchArray Predicates)
+	{
+		const FRigDispatch_GetModuleMetadata* Factory = static_cast<const FRigDispatch_GetModuleMetadata*>(InContext.Factory);
+
+#if WITH_EDITOR
+		if(!Factory->CheckArgumentTypes<ValueType>(Handles))
+		{
+			return;
+		}
+#endif
+
+		// unpack the memory
+		const FName& Name = *reinterpret_cast<const FName*>(Handles[Factory->NameArgIndex].GetData());
+		const ERigMetaDataNameSpace NameSpace = *reinterpret_cast<const ERigMetaDataNameSpace*>(Handles[Factory->NameSpaceArgIndex].GetData());
+		const ValueType& Default = *reinterpret_cast<const ValueType*>(Handles[Factory->DefaultArgIndex].GetData());
+		ValueType& Value = *reinterpret_cast<ValueType*>(Handles[Factory->ValueArgIndex].GetData());
+		bool& Found = *reinterpret_cast<bool*>(Handles[Factory->FoundArgIndex].GetData());
+
+		// extract the metadata
+		if (const MetadataType* Md = Cast<MetadataType>(FindMetadata(InContext, Name, EnumValue, NameSpace)))
+		{
+			Value = Md->GetValue();
+			Found = true;
+		}
+		else
+		{
+			Value = Default;
+			Found = false;
+		}
+	}
+};
+
+/*
+ * Sets metadata on the module
+ */
+USTRUCT(meta=(DisplayName="Set Module Metadata"))
+struct CONTROLRIG_API FRigDispatch_SetModuleMetadata : public FRigDispatch_SetMetadata
+{
+	GENERATED_BODY()
+
+	virtual const TArray<FRigVMTemplateArgumentInfo>& GetArgumentInfos() const override;
+#if WITH_EDITOR
+	virtual FString GetNodeTitlePrefix() const override { return TEXT("Module "); }
+#endif
+
+protected:
+
+	static FRigBaseMetadata* FindOrAddMetadata(const FControlRigExecuteContext& InContext, const FName& InName, ERigMetadataType InType, ERigMetaDataNameSpace InNameSpace);
+	virtual FRigVMFunctionPtr GetDispatchFunctionImpl(const FRigVMTemplateTypeMap& InTypes) const override;
+
+#if WITH_EDITOR
+	template<typename ValueType>
+	bool CheckArgumentTypes(FRigVMMemoryHandleArray Handles) const
+	{
+		return CheckArgumentType(Handles[NameArgIndex].IsType<FName>(), NameArgName) &&
+			CheckArgumentType(Handles[NameSpaceArgIndex].IsType<ERigMetaDataNameSpace>(), NameSpaceArgName) &&
+			CheckArgumentType(Handles[ValueArgIndex].IsType<ValueType>(), ValueArgName) &&
+			CheckArgumentType(Handles[SuccessArgIndex].IsType<bool>(), SuccessArgName);
+	}
+#endif
+	
+	template<typename ValueType, typename MetadataType, ERigMetadataType EnumValue>
+	static void SetModuleMetadataDispatch(FRigVMExtendedExecuteContext& InContext, FRigVMMemoryHandleArray Handles, FRigVMPredicateBranchArray Predicates)
+	{
+		const FRigDispatch_SetModuleMetadata* Factory = static_cast<const FRigDispatch_SetModuleMetadata*>(InContext.Factory);
+
+#if WITH_EDITOR
+		if(!Factory->CheckArgumentTypes<ValueType>(Handles))
+		{
+			return;
+		}
+#endif
+
+		// unpack the memory
+		FControlRigExecuteContext& ExecuteContext = InContext.GetPublicData<FControlRigExecuteContext>();
+		const FName& Name = *(const FName*)Handles[Factory->NameArgIndex].GetData();
+		const ERigMetaDataNameSpace NameSpace = *(const ERigMetaDataNameSpace*)Handles[Factory->NameSpaceArgIndex].GetData();
+		ValueType& Value = *(ValueType*)Handles[Factory->ValueArgIndex].GetData();
+		bool& Success = *(bool*)Handles[Factory->SuccessArgIndex].GetData();
+
+		// extract the metadata
+		if (MetadataType* Md = Cast<MetadataType>(FindOrAddMetadata(ExecuteContext, Name, EnumValue, NameSpace)))
+		{
+			if(!Equals<ValueType>(Md->GetValue(), Value))
+			{
+				Md->GetValue() = Value;
+				ExecuteContext.Hierarchy->MetadataVersion++;
+			}
+			Success = true;
+		}
+		else
+		{
+			Success = false;
+		}
+	}
+};
diff --git a/Engine/Plugins/Animation/ControlRig/Source/ControlRig/Public/Units/RigUnitContext.h b/Engine/Plugins/Animation/ControlRig/Source/ControlRig/Public/Units/RigUnitContext.h
--- a/Engine/Plugins/Animation/ControlRig/Source/ControlRig/Public/Units/RigUnitContext.h
+++ b/Engine/Plugins/Animation/ControlRig/Source/ControlRig/Public/Units/RigUnitContext.h
@@ -39,8 +39,6 @@ enum class ERigMetaDataNameSpace : uint8
 	Parent,
 	// Store the metadata under the root module
 	Root,
-	// (Only for remove) Remove all metadata data for all namespaces
-	All, 
 	Last UMETA(Hidden)
 };
 
@@ -201,6 +199,11 @@ public:
 	{
 		return RigModuleInstance;
 	}
+	
+	/**
+	 * Returns the module this unit is running inside of (or nullptr)
+	 */
+	const FRigModuleInstance* GetRigModuleInstance(ERigMetaDataNameSpace InNameSpaceType) const;
 
 	/**
 	 * Adapts a metadata name according to rig module namespace.
diff --git a/Engine/Plugins/Animation/ControlRig/Source/ControlRigEditor/Private/Graph/ControlRigGraphPanelPinFactory.cpp b/Engine/Plugins/Animation/ControlRig/Source/ControlRigEditor/Private/Graph/ControlRigGraphPanelPinFactory.cpp
--- a/Engine/Plugins/Animation/ControlRig/Source/ControlRigEditor/Private/Graph/ControlRigGraphPanelPinFactory.cpp
+++ b/Engine/Plugins/Animation/ControlRig/Source/ControlRigEditor/Private/Graph/ControlRigGraphPanelPinFactory.cpp
@@ -344,8 +344,7 @@ TSharedPtr<SGraphPin> FControlRigGraphPanelPinFactory::CreatePin_Internal(UEdGra
 										}
 									}
 
-									const bool bUseShortNames = (NameSpace != ERigMetaDataNameSpace::None) &&
-										(NameSpace != ERigMetaDataNameSpace::All);
+									const bool bUseShortNames = NameSpace != ERigMetaDataNameSpace::None;
 
 									static TMap<FString, FCachedMetadataNames> MetadataNameLists;
 									FCachedMetadataNames& MetadataNames = MetadataNameLists.FindOrAdd(MapHash);
@@ -431,8 +430,7 @@ TSharedPtr<SGraphPin> FControlRigGraphPanelPinFactory::CreatePin_Internal(UEdGra
 											NameSpace = (ERigMetaDataNameSpace)StaticEnum<ERigMetaDataNameSpace>()->GetValueByNameString(NameSpacePin->GetDefaultValue());
 										}
 									}
-									const bool bUseShortNames = (NameSpace != ERigMetaDataNameSpace::None) &&
-										(NameSpace != ERigMetaDataNameSpace::All);
+									const bool bUseShortNames = NameSpace != ERigMetaDataNameSpace::None;
 
 									static TMap<FString, FCachedMetadataTagNames> MetadataTagNameLists;
 									FCachedMetadataTagNames& MetadataTagNames = MetadataTagNameLists.FindOrAdd(MapHash);
