Control Rig: Prevent modular rigs from collapsing when a module asset is autosaved
#rnx

diff --git a/Engine/Plugins/Animation/ControlRig/Source/ControlRig/Private/ModularRigController.cpp b/Engine/Plugins/Animation/ControlRig/Source/ControlRig/Private/ModularRigController.cpp
--- a/Engine/Plugins/Animation/ControlRig/Source/ControlRig/Private/ModularRigController.cpp
+++ b/Engine/Plugins/Animation/ControlRig/Source/ControlRig/Private/ModularRigController.cpp
@@ -288,23 +288,25 @@ bool UModularRigController::ConnectConnectorToElement(const FRigElementKey& InCo
 						OtherConnectorKey.Type = ERigElementType::Connector;
 						if (!Model->Connections.HasConnection(OtherConnectorKey))
 						{
-							const FRigConnectorElement* OtherConnectorElement = Cast<FRigConnectorElement>(Hierarchy->Find(OtherConnectorKey));
-							FModularRigResolveResult RuleResults = RuleManager->FindMatches(OtherConnectorElement, ModuleInstance, ModularRig->GetElementKeyRedirector());
-
-							if (RuleResults.GetMatches().Num() == 1)
-							{
-								Model->Connections.AddConnection(OtherConnectorKey, RuleResults.GetMatches()[0].GetKey());
-								Notify(EModularRigNotification::ConnectionChanged, Module);
-							}
-							else
+							if (const FRigConnectorElement* OtherConnectorElement = Cast<FRigConnectorElement>(Hierarchy->Find(OtherConnectorKey)))
 							{
-								for (const FRigElementResolveResult& Result : RuleResults.GetMatches())
+								FModularRigResolveResult RuleResults = RuleManager->FindMatches(OtherConnectorElement, ModuleInstance, ModularRig->GetElementKeyRedirector());
+
+								if (RuleResults.GetMatches().Num() == 1)
+								{
+									Model->Connections.AddConnection(OtherConnectorKey, RuleResults.GetMatches()[0].GetKey());
+									Notify(EModularRigNotification::ConnectionChanged, Module);
+								}
+								else
 								{
-									if (Result.GetState() == ERigElementResolveState::DefaultTarget)
+									for (const FRigElementResolveResult& Result : RuleResults.GetMatches())
 									{
-										Model->Connections.AddConnection(OtherConnectorKey, Result.GetKey());
-										Notify(EModularRigNotification::ConnectionChanged, Module);
-										break;
+										if (Result.GetState() == ERigElementResolveState::DefaultTarget)
+										{
+											Model->Connections.AddConnection(OtherConnectorKey, Result.GetKey());
+											Notify(EModularRigNotification::ConnectionChanged, Module);
+											break;
+										}
 									}
 								}
 							}
diff --git a/Engine/Plugins/Animation/ControlRig/Source/ControlRigDeveloper/Private/ControlRigBlueprint.cpp b/Engine/Plugins/Animation/ControlRig/Source/ControlRigDeveloper/Private/ControlRigBlueprint.cpp
--- a/Engine/Plugins/Animation/ControlRig/Source/ControlRigDeveloper/Private/ControlRigBlueprint.cpp
+++ b/Engine/Plugins/Animation/ControlRig/Source/ControlRigDeveloper/Private/ControlRigBlueprint.cpp
@@ -2148,11 +2148,21 @@ void UControlRigBlueprint::PropagateHierarchyFromBPToInstances() const
 			{
 				if (UControlRig* InstanceRig = Cast<UControlRig>(ArchetypeInstance))
 				{
-					InstanceRig->PostInitInstanceIfRequired();
-					InstanceRig->GetHierarchy()->CopyHierarchy(Hierarchy);
-					InstanceRig->HierarchySettings = HierarchySettings;
-					UpdateElementKeyRedirector(InstanceRig);
-					InstanceRig->Initialize(true);
+					if (InstanceRig->IsRigModuleInstance())
+					{
+						if (UModularRig* ModularRig = Cast<UModularRig>(InstanceRig->GetOuter()))
+						{
+							ModularRig->RequestInit();
+						}
+					}
+					else
+					{
+						InstanceRig->PostInitInstanceIfRequired();
+						InstanceRig->GetHierarchy()->CopyHierarchy(Hierarchy);
+						InstanceRig->HierarchySettings = HierarchySettings;
+						UpdateElementKeyRedirector(InstanceRig);
+						InstanceRig->Initialize(true);
+					}
 				}
 			}
 		}
