Control Rig: Do not yellow-tint optional connectors
#rnx

diff --git a/Engine/Plugins/Animation/ControlRig/Source/ControlRigEditor/Private/Editor/SModularRigTreeView.cpp b/Engine/Plugins/Animation/ControlRig/Source/ControlRigEditor/Private/Editor/SModularRigTreeView.cpp
--- a/Engine/Plugins/Animation/ControlRig/Source/ControlRigEditor/Private/Editor/SModularRigTreeView.cpp
+++ b/Engine/Plugins/Animation/ControlRig/Source/ControlRigEditor/Private/Editor/SModularRigTreeView.cpp
@@ -988,11 +988,7 @@ TPair<const FSlateBrush*, FSlateColor> FModularRigTreeElement::GetBrushAndColor(
 		const FModularRigModel& Model = InModularRig->GetModularRigModel();
 		const FString ConnectorPath = FString::Printf(TEXT("%s:%s"), *ModulePath, *ConnectorName);
 		bool bIsConnected = Model.Connections.HasConnection(FRigElementKey(*ConnectorPath, ERigElementType::Connector));
-		
-		if (!bIsConnected)
-		{
-			Color = FStyleColors::AccentYellow;
-		}
+		bool bConnectionWarning = !bIsConnected;
 		
 		if (const UControlRig* ModuleRig = ConnectorModule->GetRig())
 		{
@@ -1023,6 +1019,7 @@ TPair<const FSlateBrush*, FSlateColor> FModularRigTreeElement::GetBrushAndColor(
 				}
 				else if (Connector->Settings.bOptional)
 				{
+					bConnectionWarning = false;
 					Brush = FControlRigEditorStyle::Get().GetBrush("ControlRig.Schematic.ConnectorOptional");
 				}
 				else
@@ -1031,6 +1028,11 @@ TPair<const FSlateBrush*, FSlateColor> FModularRigTreeElement::GetBrushAndColor(
 				}
 			}
 		}
+
+		if (bConnectionWarning)
+		{
+			Color = FStyleColors::AccentYellow;
+		}
 	}
 	if (!Brush)
 	{
