Constraints: ensure that the constraint has not already been added to sequencer to avoid erasing the default active keyframe

#rnx

diff --git a/Engine/Plugins/Animation/ControlRig/Source/ControlRigEditor/Private/Editor/Constraints/SConstraintsWidget.cpp b/Engine/Plugins/Animation/ControlRig/Source/ControlRigEditor/Private/Editor/Constraints/SConstraintsWidget.cpp
--- a/Engine/Plugins/Animation/ControlRig/Source/ControlRigEditor/Private/Editor/Constraints/SConstraintsWidget.cpp
+++ b/Engine/Plugins/Animation/ControlRig/Source/ControlRigEditor/Private/Editor/Constraints/SConstraintsWidget.cpp
@@ -29,7 +29,6 @@
 #include "Styling/SlateIconFinder.h"
 #include "Widgets/Input/NumericTypeInterface.h"
 #include "FrameNumberDetailsCustomization.h"
-#include "SPrimaryButton.h"
 
 #define LOCTEXT_NAMESPACE "SConstraintsWidget"
 
@@ -439,7 +438,7 @@ void SConstraintMenuEntry::CreateConstraint(
 						{
 							Constraint->Evaluate();
 						}
-						FMovieSceneConstraintChannelHelper::SmartConstraintKey(WeakSequencer.Pin(), Constraint, TOptional<bool>(), TOptional<FFrameNumber>());
+						FMovieSceneConstraintChannelHelper::AddConstraintToSequencer(WeakSequencer.Pin(), Constraint);
 					}
 					else
 					{
@@ -1057,7 +1056,6 @@ void SConstraintsEditionWidget::Construct(const FArguments& InArgs)
 				.AutoHeight()
 				.Padding(FMargin(0.0f, 0.f))
 				[
-		
 					SNew(SHorizontalBox)
 					+ SHorizontalBox::Slot()
 					.FillWidth(1.f)
diff --git a/Engine/Plugins/Animation/ControlRig/Source/ControlRigEditor/Private/Sequencer/ControlRigParameterTrackEditor.cpp b/Engine/Plugins/Animation/ControlRig/Source/ControlRigEditor/Private/Sequencer/ControlRigParameterTrackEditor.cpp
--- a/Engine/Plugins/Animation/ControlRig/Source/ControlRigEditor/Private/Sequencer/ControlRigParameterTrackEditor.cpp
+++ b/Engine/Plugins/Animation/ControlRig/Source/ControlRigEditor/Private/Sequencer/ControlRigParameterTrackEditor.cpp
@@ -1947,32 +1947,21 @@ void FControlRigParameterTrackEditor::HandleActorAdded(AActor* Actor, FGuid Targ
 	{
 		return;
 	}
+
 	//test for constraint
-	FConstraintsManagerController& Controller = FConstraintsManagerController::Get(Actor->GetWorld());
-	TArray< TWeakObjectPtr<UTickableConstraint>> Constraints = Controller.GetAllConstraints();
-	for (TWeakObjectPtr<UTickableConstraint>& WeakConstraint : Constraints)
+	const FConstraintsManagerController& Controller = FConstraintsManagerController::Get(Actor->GetWorld());
+	const TArray< TWeakObjectPtr<UTickableConstraint>> Constraints = Controller.GetAllConstraints();
+	for (const TWeakObjectPtr<UTickableConstraint>& WeakConstraint : Constraints)
 	{
-		if (WeakConstraint.IsValid())
+		if (UTickableTransformConstraint* Constraint = WeakConstraint.IsValid() ? Cast<UTickableTransformConstraint>(WeakConstraint.Get()) : nullptr)
 		{
-			if (UTickableTransformConstraint* Constraint = Cast<UTickableTransformConstraint>(WeakConstraint.Get()))
+			if (UObject* Child = Constraint->ChildTRSHandle ? Constraint->ChildTRSHandle->GetTarget().Get() : nullptr)
 			{
-				if (Constraint->ChildTRSHandle)
+				const AActor* TargetActor = Child->IsA<AActor>() ? Cast<AActor>(Child) : Child->GetTypedOuter<AActor>();
+				if (TargetActor == Actor)
 				{
-					if (UObject* Child = Constraint->ChildTRSHandle->GetTarget().Get())
-					{
-						AActor* TargetActor = Cast<AActor>(Child);
-						if (TargetActor == nullptr)
-						{
-							TargetActor = Child->GetTypedOuter<AActor>();
-						}
-						if (TargetActor == Actor)
-						{
-							const TOptional<bool> bActive = true;
-							const TSharedPtr<ISequencer> SequencerPtr = GetSequencer();
-							FMovieSceneConstraintChannelHelper::SmartConstraintKey(SequencerPtr, Constraint, bActive, TOptional<FFrameNumber>());
-						}		
-					}
-				}
+					FMovieSceneConstraintChannelHelper::AddConstraintToSequencer(GetSequencer(), Constraint);
+				}		
 			}
 		}
 	}
diff --git a/Engine/Source/Editor/MovieSceneTools/Private/Constraints/MovieSceneConstraintChannelHelper.cpp b/Engine/Source/Editor/MovieSceneTools/Private/Constraints/MovieSceneConstraintChannelHelper.cpp
--- a/Engine/Source/Editor/MovieSceneTools/Private/Constraints/MovieSceneConstraintChannelHelper.cpp
+++ b/Engine/Source/Editor/MovieSceneTools/Private/Constraints/MovieSceneConstraintChannelHelper.cpp
@@ -879,6 +879,39 @@ void FMovieSceneConstraintChannelHelper::HandleConstraintKeyMoved(
 
 }
 
+bool FMovieSceneConstraintChannelHelper::AddConstraintToSequencer(
+	const TSharedPtr<ISequencer>& InSequencer,
+	UTickableTransformConstraint* InConstraint)
+{
+	if (!InSequencer.IsValid() || !InSequencer->GetFocusedMovieSceneSequence())
+	{
+		return false;
+	}
+	
+	ITransformConstraintChannelInterface* Interface = InConstraint ? GetHandleInterface(InConstraint->ChildTRSHandle) : nullptr;
+	if (!Interface)
+	{
+		return false;
+	}
+
+	// create bindings before smart keying so added to spawn copies
+	CreateBindingIDForHandle(InSequencer, InConstraint->ChildTRSHandle);
+	CreateBindingIDForHandle(InSequencer, InConstraint->ParentTRSHandle);
+
+	// adding the child to sequencer can trigger that same function so the constraint might already be added 
+	const bool IsOuterASection = !!Cast<IMovieSceneConstrainedSection>(InConstraint->GetOuter());
+	if (IsOuterASection)
+	{
+		return true;
+	}
+
+	const FFrameRate TickResolution = InSequencer->GetFocusedTickResolution();
+	const FFrameTime FrameTime = InSequencer->GetLocalTime().ConvertTo(TickResolution);
+	const FFrameNumber Time = FrameTime.GetFrame();
+	
+	return Interface->SmartConstraintKey(InConstraint, TOptional<bool>(), Time, InSequencer);
+}
+
 bool FMovieSceneConstraintChannelHelper::SmartConstraintKey(
 	const TSharedPtr<ISequencer>& InSequencer,
 	UTickableTransformConstraint* InConstraint, 
diff --git a/Engine/Source/Editor/MovieSceneTools/Public/Constraints/MovieSceneConstraintChannelHelper.h b/Engine/Source/Editor/MovieSceneTools/Public/Constraints/MovieSceneConstraintChannelHelper.h
--- a/Engine/Source/Editor/MovieSceneTools/Public/Constraints/MovieSceneConstraintChannelHelper.h
+++ b/Engine/Source/Editor/MovieSceneTools/Public/Constraints/MovieSceneConstraintChannelHelper.h
@@ -57,6 +57,10 @@ struct MOVIESCENETOOLS_API FMovieSceneConstraintChannelHelper
 {
 public:
 
+	static bool AddConstraintToSequencer(
+		const TSharedPtr<ISequencer>& InSequencer,
+		UTickableTransformConstraint* InConstraint);
+	
 	/** Adds an active key if needed and does the compensation when switching. Will use the optional active and time if set. 
 	Will return true if key is actually set, may not be if the value is the same.*/
 	static bool SmartConstraintKey(
