Control Rig: Autoresolve root module if possible

#rnx

diff --git a/Engine/Plugins/Animation/ControlRig/Source/ControlRigEditor/Private/Editor/ControlRigEditor.cpp b/Engine/Plugins/Animation/ControlRig/Source/ControlRigEditor/Private/Editor/ControlRigEditor.cpp
--- a/Engine/Plugins/Animation/ControlRig/Source/ControlRigEditor/Private/Editor/ControlRigEditor.cpp
+++ b/Engine/Plugins/Animation/ControlRig/Source/ControlRigEditor/Private/Editor/ControlRigEditor.cpp
@@ -5106,6 +5106,33 @@ void FControlRigEditor::OnPostConstruction_AnyThread(UControlRig* InRig, const F
 		{
 			RigBlueprint->Hierarchy->RestoreConnectorsFromStates(ConnectorStates);
 		}
+
+		if(RigBlueprint->IsModularRig())
+		{
+			// auto resolve the root module's primary connector
+			if(RigBlueprint->ModularRigModel.Connections.IsEmpty() && RigBlueprint->ModularRigModel.Modules.Num() == 1 && RigBlueprint->Hierarchy->Num(ERigElementType::Bone) > 0)
+			{
+				const FRigModuleReference& RootModule = RigBlueprint->ModularRigModel.Modules[0];
+
+				const FSoftObjectPath DefaultRootModulePath = UControlRigSettings::Get()->DefaultRootModule;
+				if(const UControlRigBlueprint* DefaultRootModule = Cast<UControlRigBlueprint>(DefaultRootModulePath.TryLoad()))
+				{
+					if(DefaultRootModule->GetControlRigClass() == RootModule.Class)
+					{
+						if(const FRigConnectorElement* PrimaryConnector = RootModule.FindPrimaryConnector(RigBlueprint->Hierarchy))
+						{
+							if(const FRigBoneElement* RootBone = RigBlueprint->Hierarchy->GetBones()[0])
+							{
+								if(UModularRigController* ModularRigController = RigBlueprint->ModularRigModel.GetController())
+								{
+									(void)ModularRigController->ConnectConnectorToElement(PrimaryConnector->GetKey(), RootBone->GetKey(), false);
+								}
+							}
+						}
+					}
+				}
+			}
+		}
 	}
 	
 	const int32 HierarchyHash = InRig->GetHierarchy()->GetTopologyHash(false);
