Control Rig: Fix zero transforms in sequencer not zeroing out the preferred angles
#rnx

diff --git a/Engine/Plugins/Animation/ControlRig/Source/ControlRigEditor/Private/EditMode/ControlRigEditMode.cpp b/Engine/Plugins/Animation/ControlRig/Source/ControlRigEditor/Private/EditMode/ControlRigEditMode.cpp
--- a/Engine/Plugins/Animation/ControlRig/Source/ControlRigEditor/Private/EditMode/ControlRigEditMode.cpp
+++ b/Engine/Plugins/Animation/ControlRig/Source/ControlRigEditor/Private/EditMode/ControlRigEditMode.cpp
@@ -3306,6 +3306,8 @@ void FControlRigEditMode::ZeroTransforms(bool bSelectionOnly)
 			if (ControlElement)
 			{
 				ControlRig->SetControlLocalTransform(ElementToReset.Name, InitialLocalTransform, true, FRigControlModifiedContext(), true, true);
+				const FVector InitialAngles = ControlRig->GetHierarchy()->GetControlPreferredEulerAngles(ControlElement, ControlElement->Settings.PreferredRotationOrder, true);
+				ControlRig->GetHierarchy()->SetControlPreferredEulerAngles(ControlElement, InitialAngles, ControlElement->Settings.PreferredRotationOrder);
 				NotifyDrivenControls(ControlRig, ElementToReset);
 				if (bHasNonDefaultParent == false)
 				{
