Sequencer: Anim Details: fixed slow update when scrubbing values
    - pre and post edit change had to be in the same order
    - the FPropertyChangedEvent has to use EPropertyChangeType::Interactive when scrubbing

#rnx

diff --git a/Engine/Plugins/Animation/ControlRig/Source/ControlRigEditor/Private/EditMode/AnimDetailsProxy.cpp b/Engine/Plugins/Animation/ControlRig/Source/ControlRigEditor/Private/EditMode/AnimDetailsProxy.cpp
--- a/Engine/Plugins/Animation/ControlRig/Source/ControlRigEditor/Private/EditMode/AnimDetailsProxy.cpp
+++ b/Engine/Plugins/Animation/ControlRig/Source/ControlRigEditor/Private/EditMode/AnimDetailsProxy.cpp
@@ -890,7 +890,7 @@ static FEulerTransform GetCurrentValue(UObject* InObject, TSharedPtr<FTrackInsta
 	return EulerTransform;
 }
 
-void UAnimDetailControlsProxyTransform::SetBindingValueFromCurrent(UObject* InObject, TSharedPtr<FTrackInstancePropertyBindings>& Binding, FRigControlModifiedContext& Context)
+void UAnimDetailControlsProxyTransform::SetBindingValueFromCurrent(UObject* InObject, TSharedPtr<FTrackInstancePropertyBindings>& Binding, FRigControlModifiedContext& Context, bool bInteractive)
 {
 	if (InObject && Binding.IsValid())
 	{
@@ -967,20 +967,23 @@ void UAnimDetailControlsProxyTransform::SetBindingValueFromCurrent(UObject* InOb
 					AxisProperty = FindFProperty<FDoubleProperty>(TBaseStructure<FVector>::Get(), GET_MEMBER_NAME_CHECKED(FVector, Z));
 				}
 			}
+
+			// build property chain. note that it won't be used if ValueProperty is null 
+			TArray<UObject*> ModifiedObjects = { InObject };
+			FPropertyChangedEvent PropertyChangedEvent(ValueProperty, bInteractive ? EPropertyChangeType::Interactive : EPropertyChangeType::ValueSet, MakeArrayView(ModifiedObjects));
+			FEditPropertyChain PropertyChain;
+			if (AxisProperty)
+			{
+				PropertyChain.AddHead(AxisProperty);
+			}
+			PropertyChain.AddHead(ValueProperty);
+			FPropertyChangedChainEvent PropertyChangedChainEvent(PropertyChain, PropertyChangedEvent);
+			
 			// Have to downcast here because of function overloading and inheritance not playing nicely
 			if (ValueProperty)
 			{
-				TArray<UObject*> ModifiedObjects = { InObject };
-				FPropertyChangedEvent PropertyChangedEvent(ValueProperty, EPropertyChangeType::ValueSet, MakeArrayView(ModifiedObjects));
-				FEditPropertyChain PropertyChain;
-
-				if (AxisProperty)
-				{
-					PropertyChain.AddHead(AxisProperty);
-				}
-				PropertyChain.AddHead(ValueProperty);
 				((UObject*)SceneComponentThatChanged)->PreEditChange(PropertyChain);
-
+			
 				if (ActorThatChanged && ActorThatChanged->GetRootComponent() == SceneComponentThatChanged)
 				{
 					((UObject*)ActorThatChanged)->PreEditChange(PropertyChain);
@@ -993,23 +996,11 @@ void UAnimDetailControlsProxyTransform::SetBindingValueFromCurrent(UObject* InOb
 
 			if (ValueProperty)
 			{
-				TArray<UObject*> ModifiedObjects = { InObject };
-				FPropertyChangedEvent PropertyChangedEvent(ValueProperty, EPropertyChangeType::ValueSet, MakeArrayView(ModifiedObjects));
-				FEditPropertyChain PropertyChain;
-
-				if (AxisProperty)
-				{
-					PropertyChain.AddHead(AxisProperty);
-				}
-				PropertyChain.AddHead(ValueProperty);
-				FPropertyChangedChainEvent PropertyChangedChainEvent(PropertyChain, PropertyChangedEvent);
-
+				((UObject*)SceneComponentThatChanged)->PostEditChangeChainProperty(PropertyChangedChainEvent);
 				if (ActorThatChanged && ActorThatChanged->GetRootComponent() == SceneComponentThatChanged)
 				{
 					((UObject*)ActorThatChanged)->PostEditChangeChainProperty(PropertyChangedChainEvent);
 				}
-				((UObject*)SceneComponentThatChanged)->PostEditChangeChainProperty(PropertyChangedChainEvent);
-
 			}
 			GUnrealEd->UpdatePivotLocationForSelection();
 		}
@@ -2192,7 +2183,7 @@ void UAnimDetailControlsProxyFloat::GetChannelSelectionState(TWeakPtr<FCurveEdit
 	LocationSelectionCache.XSelected = CachePropertySelection(CurveEditor, this, GET_MEMBER_NAME_CHECKED(FAnimDetailProxyFloat, Float));
 }
 
-void UAnimDetailControlsProxyFloat::SetBindingValueFromCurrent(UObject* InObject, TSharedPtr<FTrackInstancePropertyBindings>& Binding, FRigControlModifiedContext& Context)
+void UAnimDetailControlsProxyFloat::SetBindingValueFromCurrent(UObject* InObject, TSharedPtr<FTrackInstancePropertyBindings>& Binding, FRigControlModifiedContext& Context, bool bInteractive)
 {
 	if (InObject && Binding.IsValid())
 	{
@@ -2385,7 +2376,7 @@ void UAnimDetailControlsProxyBool::GetChannelSelectionState(TWeakPtr<FCurveEdito
 	LocationSelectionCache.XSelected = CachePropertySelection(CurveEditor, this, GET_MEMBER_NAME_CHECKED(FAnimDetailProxyBool, Bool));
 }
 
-void UAnimDetailControlsProxyBool::SetBindingValueFromCurrent(UObject* InObject, TSharedPtr<FTrackInstancePropertyBindings>& Binding, FRigControlModifiedContext& Context)
+void UAnimDetailControlsProxyBool::SetBindingValueFromCurrent(UObject* InObject, TSharedPtr<FTrackInstancePropertyBindings>& Binding, FRigControlModifiedContext& Context, bool bInteractive)
 {
 	if (InObject && Binding.IsValid())
 	{
@@ -2565,7 +2556,7 @@ void UAnimDetailControlsProxyInteger::GetChannelSelectionState(TWeakPtr<FCurveEd
 	LocationSelectionCache.XSelected = CachePropertySelection(CurveEditor, this, GET_MEMBER_NAME_CHECKED(FAnimDetailProxyInteger, Integer));
 }
 
-void UAnimDetailControlsProxyInteger::SetBindingValueFromCurrent(UObject* InObject, TSharedPtr<FTrackInstancePropertyBindings>& Binding, FRigControlModifiedContext& Context)
+void UAnimDetailControlsProxyInteger::SetBindingValueFromCurrent(UObject* InObject, TSharedPtr<FTrackInstancePropertyBindings>& Binding, FRigControlModifiedContext& Context, bool bInteractive)
 {
 	if (InObject && Binding.IsValid())
 	{
diff --git a/Engine/Plugins/Animation/ControlRig/Source/ControlRigEditor/Private/EditMode/AnimDetailsProxy.h b/Engine/Plugins/Animation/ControlRig/Source/ControlRigEditor/Private/EditMode/AnimDetailsProxy.h
--- a/Engine/Plugins/Animation/ControlRig/Source/ControlRigEditor/Private/EditMode/AnimDetailsProxy.h
+++ b/Engine/Plugins/Animation/ControlRig/Source/ControlRigEditor/Private/EditMode/AnimDetailsProxy.h
@@ -257,7 +257,7 @@ public:
 	virtual void GetChannelSelectionState(TWeakPtr<FCurveEditor>& CurveEditor, FAnimDetailVectorSelection& OutLocationSelection, FAnimDetailVectorSelection& OutRotationSelection,
 		FAnimDetailVectorSelection& OutScaleSelection) override;
 	virtual void UpdatePropertyNames(IDetailLayoutBuilder& DetailBuilder);
-	virtual void SetBindingValueFromCurrent(UObject* InObject, TSharedPtr<FTrackInstancePropertyBindings>& Binding, FRigControlModifiedContext& Context) override;
+	virtual void SetBindingValueFromCurrent(UObject* InObject, TSharedPtr<FTrackInstancePropertyBindings>& Binding, FRigControlModifiedContext& Context, bool bInteractive = false) override;
 
 	FAnimDetailPropertyState State;
 
@@ -285,7 +285,7 @@ public:
 	virtual void GetChannelSelectionState(TWeakPtr<FCurveEditor>& CurveEditor, FAnimDetailVectorSelection& OutLocationSelection, FAnimDetailVectorSelection& OutRotationSelection,
 		FAnimDetailVectorSelection& OutScaleSelection) override;
 	virtual void UpdatePropertyNames(IDetailLayoutBuilder& DetailBuilder);
-	virtual void SetBindingValueFromCurrent(UObject* InObject, TSharedPtr<FTrackInstancePropertyBindings>& Binding, FRigControlModifiedContext& Context) override;
+	virtual void SetBindingValueFromCurrent(UObject* InObject, TSharedPtr<FTrackInstancePropertyBindings>& Binding, FRigControlModifiedContext& Context, bool bInteractive = false) override;
 
 	
 	FAnimDetailPropertyState State;
@@ -314,7 +314,7 @@ public:
 	virtual void GetChannelSelectionState(TWeakPtr<FCurveEditor>& CurveEditor, FAnimDetailVectorSelection& OutLocationSelection, FAnimDetailVectorSelection& OutRotationSelection,
 		FAnimDetailVectorSelection& OutScaleSelection) override;
 	virtual void UpdatePropertyNames(IDetailLayoutBuilder& DetailBuilder);
-	virtual void SetBindingValueFromCurrent(UObject* InObject, TSharedPtr<FTrackInstancePropertyBindings>& Binding, FRigControlModifiedContext& Context) override;
+	virtual void SetBindingValueFromCurrent(UObject* InObject, TSharedPtr<FTrackInstancePropertyBindings>& Binding, FRigControlModifiedContext& Context, bool bInteractive = false) override;
 
 
 	FAnimDetailPropertyState State;
@@ -369,7 +369,7 @@ public:
 	virtual EControlRigContextChannelToKey GetChannelToKeyFromChannelName(const FString& InChannelName) const override;
 	virtual bool IsMultiple(const FName& InPropertyName) const override;
 	virtual void SetControlRigElementValueFromCurrent(UControlRig* ControlRig, FRigControlElement* ControlElement, const FRigControlModifiedContext& Context) override;
-	virtual void SetBindingValueFromCurrent(UObject* InObject, TSharedPtr<FTrackInstancePropertyBindings>& Binding, FRigControlModifiedContext& Context) override;
+	virtual void SetBindingValueFromCurrent(UObject* InObject, TSharedPtr<FTrackInstancePropertyBindings>& Binding, FRigControlModifiedContext& Context, bool bInteractive = false) override;
 	virtual void GetChannelSelectionState(TWeakPtr<FCurveEditor>& CurveEditor, FAnimDetailVectorSelection& OutLocationSelection, FAnimDetailVectorSelection& OutRotationSelection,
 		FAnimDetailVectorSelection& OutScaleSelection) override;
 
diff --git a/Engine/Plugins/Animation/ControlRig/Source/ControlRigEditor/Private/EditMode/ControlRigControlsProxy.cpp b/Engine/Plugins/Animation/ControlRig/Source/ControlRigEditor/Private/EditMode/ControlRigControlsProxy.cpp
--- a/Engine/Plugins/Animation/ControlRig/Source/ControlRigEditor/Private/EditMode/ControlRigControlsProxy.cpp
+++ b/Engine/Plugins/Animation/ControlRig/Source/ControlRigEditor/Private/EditMode/ControlRigControlsProxy.cpp
@@ -320,7 +320,7 @@ void UControlRigControlsProxy::PostEditChangeChainProperty(struct FPropertyChang
 				for (int32 Index = SItems.Value.Bindings.Num() - 1; Index >= 0; --Index)
 				{
 					FBindingAndTrack& Binding = SItems.Value.Bindings[Index];
-					SetBindingValueFromCurrent(SItems.Key.Get(), Binding.Binding, Context);
+					SetBindingValueFromCurrent(SItems.Key.Get(), Binding.Binding, Context, PropertyChangedEvent.ChangeType == EPropertyChangeType::Interactive);
 				}
 			}
 		}
diff --git a/Engine/Plugins/Animation/ControlRig/Source/ControlRigEditor/Private/EditMode/ControlRigControlsProxy.h b/Engine/Plugins/Animation/ControlRig/Source/ControlRigEditor/Private/EditMode/ControlRigControlsProxy.h
--- a/Engine/Plugins/Animation/ControlRig/Source/ControlRigEditor/Private/EditMode/ControlRigControlsProxy.h
+++ b/Engine/Plugins/Animation/ControlRig/Source/ControlRigEditor/Private/EditMode/ControlRigControlsProxy.h
@@ -81,7 +81,7 @@ public:
 	virtual EControlRigContextChannelToKey GetChannelToKeyFromChannelName(const FString& InChannelName) const { return EControlRigContextChannelToKey::AllTransform; }
 	virtual bool IsMultiple(const FName& InPropertyName) const { return false; }
 	virtual void SetControlRigElementValueFromCurrent(UControlRig* ControlRig, FRigControlElement* ControlElement, const FRigControlModifiedContext& Context) {};
-	virtual void SetBindingValueFromCurrent(UObject* InObject, TSharedPtr<FTrackInstancePropertyBindings>& Binding, FRigControlModifiedContext& Context) {};
+	virtual void SetBindingValueFromCurrent(UObject* InObject, TSharedPtr<FTrackInstancePropertyBindings>& Binding, FRigControlModifiedContext& Context, bool bInteractive = false) {};
 	virtual void GetChannelSelectionState(TWeakPtr<FCurveEditor>& CurveEditor, FAnimDetailVectorSelection& OutLocationSelection, FAnimDetailVectorSelection& OutRotationSelection,
 		FAnimDetailVectorSelection& OutScaleSelection) {};
 	virtual bool PropertyIsOnProxy(FProperty* Property, FProperty* MemberProperty){return false;}
