Control Rig: Fix tree wires in the Modular Hierarchy widget are disconnected
#rnx

diff --git a/Engine/Plugins/Animation/ControlRig/Source/ControlRigEditor/Private/Editor/SModularRigModel.cpp b/Engine/Plugins/Animation/ControlRig/Source/ControlRigEditor/Private/Editor/SModularRigModel.cpp
--- a/Engine/Plugins/Animation/ControlRig/Source/ControlRigEditor/Private/Editor/SModularRigModel.cpp
+++ b/Engine/Plugins/Animation/ControlRig/Source/ControlRigEditor/Private/Editor/SModularRigModel.cpp
@@ -112,7 +112,7 @@ void SModularRigModel::Construct(const FArguments& InArgs, TSharedRef<FControlRi
 		.DefaultLabel(FText::FromName(SModularRigTreeView::Column_Module))
 		.HAlignCell(HAlign_Left)
 		.HAlignHeader(HAlign_Left)
-		.VAlignCell(VAlign_Center)
+		.VAlignCell(VAlign_Fill)
 	);
 	HeaderRowWidget->AddColumn(
 		SHeaderRow::Column(SModularRigTreeView::Column_Connector)
diff --git a/Engine/Plugins/Animation/ControlRig/Source/ControlRigEditor/Private/Editor/SModularRigTreeView.cpp b/Engine/Plugins/Animation/ControlRig/Source/ControlRigEditor/Private/Editor/SModularRigTreeView.cpp
--- a/Engine/Plugins/Animation/ControlRig/Source/ControlRigEditor/Private/Editor/SModularRigTreeView.cpp
+++ b/Engine/Plugins/Animation/ControlRig/Source/ControlRigEditor/Private/Editor/SModularRigTreeView.cpp
@@ -259,7 +259,7 @@ TSharedRef<SWidget> SModularRigModelItem::GenerateWidgetForColumn(const FName& C
 		+SHorizontalBox::Slot()
 		.AutoWidth()
 		.Padding(6, 0, 0, 0)
-		.VAlign(VAlign_Center)
+		.VAlign(VAlign_Fill)
 		[
 			SNew(SExpanderArrow, SharedThis(this))
 			.IndentAmount(12)
