Control Rig Schematic: Draw all front layers with the node color

#rnx

diff --git a/Engine/Plugins/Animation/ControlRig/Source/ControlRigEditor/Private/Editor/ControlRigSchematicModel.cpp b/Engine/Plugins/Animation/ControlRig/Source/ControlRigEditor/Private/Editor/ControlRigSchematicModel.cpp
--- a/Engine/Plugins/Animation/ControlRig/Source/ControlRigEditor/Private/Editor/ControlRigSchematicModel.cpp
+++ b/Engine/Plugins/Animation/ControlRig/Source/ControlRigEditor/Private/Editor/ControlRigSchematicModel.cpp
@@ -899,10 +899,6 @@ FLinearColor FControlRigSchematicModel::GetColorForNode(const FSchematicGraphNod
 		{
 			return FLinearColor(0, 0, 0, 0.75);
 		}
-		if(InLayerIndex == 1) // outline
-		{
-			return FLinearColor::White;
-		}
 		
 		switch(ElementKeyNode->GetKey().Type)
 		{
