Control Rig: react to drag and drop of control rig module onto viewport

#rnx

diff --git a/Engine/Plugins/Animation/ControlRig/Source/ControlRigEditor/Private/Editor/ControlRigEditor.cpp b/Engine/Plugins/Animation/ControlRig/Source/ControlRigEditor/Private/Editor/ControlRigEditor.cpp
--- a/Engine/Plugins/Animation/ControlRig/Source/ControlRigEditor/Private/Editor/ControlRigEditor.cpp
+++ b/Engine/Plugins/Animation/ControlRig/Source/ControlRigEditor/Private/Editor/ControlRigEditor.cpp
@@ -248,6 +248,52 @@ bool FControlRigEditor::NewDocument_IsVisibleForType(ECreatedDocumentType GraphT
 	return true;
 }
 
+FReply FControlRigEditor::OnViewportDrop(const FGeometry& MyGeometry, const FDragDropEvent& DragDropEvent)
+{
+	const FReply SuperReply = IControlRigEditor::OnViewportDrop(MyGeometry, DragDropEvent);
+	if(SuperReply.IsEventHandled())
+	{
+		return SuperReply;
+	}
+
+	if(IsModularRig())
+	{
+		const TSharedPtr<FAssetDragDropOp> AssetDragDropOperation = DragDropEvent.GetOperationAs<FAssetDragDropOp>();
+		if (AssetDragDropOperation)
+		{
+			for (const FAssetData& AssetData : AssetDragDropOperation->GetAssets())
+			{
+				const UClass* AssetClass = AssetData.GetClass();
+				if (!AssetClass->IsChildOf(UControlRigBlueprint::StaticClass()))
+				{
+					continue;
+				}
+
+				if(const UControlRigBlueprint* AssetBlueprint = Cast<UControlRigBlueprint>(AssetData.GetAsset()))
+				{
+					UClass* ControlRigClass = AssetBlueprint->GetControlRigClass();
+					if(AssetBlueprint->IsControlRigModule() && ControlRigClass)
+					{
+						FSlateApplication::Get().DismissAllMenus();
+
+						UModularRigController* Controller = GetControlRigBlueprint()->GetModularRigController();
+						FString ClassName = ControlRigClass->GetName();
+						ClassName.RemoveFromEnd(TEXT("_C"));
+						const FRigName Name = Controller->GetSafeNewName(FString(), FRigName(ClassName));
+						const FString ModulePath = Controller->AddModule(Name, ControlRigClass, FString());
+						if (!ModulePath.IsEmpty())
+						{
+							return FReply::Handled();
+						}
+					}
+				}
+			}
+		}
+	}
+
+	return FReply::Unhandled();
+}
+
 void FControlRigEditor::CreateEmptyGraphContent(URigVMController* InController)
 {
 	URigVMNode* Node = InController->AddUnitNode(FRigUnit_BeginExecution::StaticStruct(), FRigUnit::GetMethodName(), FVector2D::ZeroVector, FString(), false);
diff --git a/Engine/Plugins/Animation/ControlRig/Source/ControlRigEditor/Private/Editor/ControlRigEditor.h b/Engine/Plugins/Animation/ControlRig/Source/ControlRigEditor/Private/Editor/ControlRigEditor.h
--- a/Engine/Plugins/Animation/ControlRig/Source/ControlRigEditor/Private/Editor/ControlRigEditor.h
+++ b/Engine/Plugins/Animation/ControlRig/Source/ControlRigEditor/Private/Editor/ControlRigEditor.h
@@ -84,6 +84,7 @@ public:
 	virtual void HandleVMCompiledEvent(UObject* InCompiledObject, URigVM* InVM) override {}
 	virtual void HandleVMCompiledEvent(UObject* InCompiledObject, URigVM* InVM, FRigVMExtendedExecuteContext& InContext) override;
 	virtual bool ShouldOpenGraphByDefault() const override { return !IsModularRig(); }
+	virtual FReply OnViewportDrop(const FGeometry& MyGeometry, const FDragDropEvent& DragDropEvent) override;
 
 	// allows the editor to fill an empty graph
 	virtual void CreateEmptyGraphContent(URigVMController* InController) override;
diff --git a/Engine/Source/Editor/AnimationEditorWidgets/Private/SchematicGraphPanel/SSchematicGraphPanel.cpp b/Engine/Source/Editor/AnimationEditorWidgets/Private/SchematicGraphPanel/SSchematicGraphPanel.cpp
--- a/Engine/Source/Editor/AnimationEditorWidgets/Private/SchematicGraphPanel/SSchematicGraphPanel.cpp
+++ b/Engine/Source/Editor/AnimationEditorWidgets/Private/SchematicGraphPanel/SSchematicGraphPanel.cpp
@@ -407,7 +407,7 @@ FReply SSchematicGraphNode::OnDrop(const FGeometry& MyGeometry, const FDragDropE
 	SNode::OnDrop(MyGeometry, DragDropEvent);
 	OnDropDelegate.ExecuteIfBound(this, DragDropEvent);
 	OnEndDragDelegate.ExecuteIfBound(this, DragDropEvent.GetOperation());
-	return FReply::Unhandled();
+	return FReply::Handled();
 }
 
 FReply SSchematicGraphNode::OnDragDetected(const FGeometry& MyGeometry, const FPointerEvent& MouseEvent)
diff --git a/Engine/Source/Editor/Persona/Private/SAnimationEditorViewport.cpp b/Engine/Source/Editor/Persona/Private/SAnimationEditorViewport.cpp
--- a/Engine/Source/Editor/Persona/Private/SAnimationEditorViewport.cpp
+++ b/Engine/Source/Editor/Persona/Private/SAnimationEditorViewport.cpp
@@ -181,6 +181,37 @@ void SAnimationEditorViewport::BindCommands()
 	}
 }
 
+void SAnimationEditorViewport::OnDragEnter(const FGeometry& MyGeometry, const FDragDropEvent& DragDropEvent)
+{
+	SEditorViewport::OnDragEnter(MyGeometry, DragDropEvent);
+	if(AssetEditorToolkitPtr.IsValid())
+	{
+		AssetEditorToolkitPtr.Pin()->OnViewportDragEnter(MyGeometry, DragDropEvent);
+	}
+}
+
+void SAnimationEditorViewport::OnDragLeave(const FDragDropEvent& DragDropEvent)
+{
+	SEditorViewport::OnDragLeave(DragDropEvent);
+	if(AssetEditorToolkitPtr.IsValid())
+	{
+		AssetEditorToolkitPtr.Pin()->OnViewportDragLeave(DragDropEvent);
+	}
+}
+
+FReply SAnimationEditorViewport::OnDrop(const FGeometry& MyGeometry, const FDragDropEvent& DragDropEvent)
+{
+	if(AssetEditorToolkitPtr.IsValid())
+	{
+		const FReply ReplyFromToolkit = AssetEditorToolkitPtr.Pin()->OnViewportDrop(MyGeometry, DragDropEvent);
+		if(ReplyFromToolkit.IsEventHandled())
+		{
+			return ReplyFromToolkit;
+		}
+	}
+	return SEditorViewport::OnDrop(MyGeometry, DragDropEvent);
+}
+
 //////////////////////////////////////////////////////////////////////////
 // SAnimationEditorViewportTabBody
 
diff --git a/Engine/Source/Editor/Persona/Private/SAnimationEditorViewport.h b/Engine/Source/Editor/Persona/Private/SAnimationEditorViewport.h
--- a/Engine/Source/Editor/Persona/Private/SAnimationEditorViewport.h
+++ b/Engine/Source/Editor/Persona/Private/SAnimationEditorViewport.h
@@ -107,7 +107,10 @@ protected:
 	virtual void PostRedo( bool bSuccess );
 
 	virtual void BindCommands() override;
-
+	virtual void OnDragEnter(const FGeometry& MyGeometry, const FDragDropEvent& DragDropEvent) override;
+	virtual void OnDragLeave(const FDragDropEvent& DragDropEvent) override;
+	virtual FReply OnDrop(const FGeometry& MyGeometry, const FDragDropEvent& DragDropEvent) override;
+	
 protected:
 	// Viewport client
 	TSharedPtr<class FAnimationViewportClient> LevelViewportClient;
diff --git a/Engine/Source/Editor/UnrealEd/Public/Toolkits/AssetEditorToolkit.h b/Engine/Source/Editor/UnrealEd/Public/Toolkits/AssetEditorToolkit.h
--- a/Engine/Source/Editor/UnrealEd/Public/Toolkits/AssetEditorToolkit.h
+++ b/Engine/Source/Editor/UnrealEd/Public/Toolkits/AssetEditorToolkit.h
@@ -302,6 +302,11 @@ public:
 
 	virtual void AddGraphEditorPinActionsToContextMenu(FToolMenuSection& InSection) const {};
 
+	/** React to a Drag&Drop events from the viewport */
+	virtual void OnViewportDragEnter(const FGeometry& MyGeometry, const FDragDropEvent& DragDropEvent) {}
+	virtual void OnViewportDragLeave(const FDragDropEvent& DragDropEvent) {}
+	virtual FReply OnViewportDrop(const FGeometry& MyGeometry, const FDragDropEvent& DragDropEvent) { return FReply::Unhandled(); }
+
 protected:
 	friend class UAssetEditorToolkitMenuContext;
 
