Anim Details: Undo crash fix.

diff --git a/Engine/Plugins/Animation/ControlRig/Source/ControlRigEditor/Private/EditMode/AnimDetailsProxy.cpp b/Engine/Plugins/Animation/ControlRig/Source/ControlRigEditor/Private/EditMode/AnimDetailsProxy.cpp
--- a/Engine/Plugins/Animation/ControlRig/Source/ControlRigEditor/Private/EditMode/AnimDetailsProxy.cpp
+++ b/Engine/Plugins/Animation/ControlRig/Source/ControlRigEditor/Private/EditMode/AnimDetailsProxy.cpp
@@ -968,20 +968,19 @@ void UAnimDetailControlsProxyTransform::SetBindingValueFromCurrent(UObject* InOb
 				}
 			}
 
-			// build property chain. note that it won't be used if ValueProperty is null 
-			TArray<UObject*> ModifiedObjects = { InObject };
-			FPropertyChangedEvent PropertyChangedEvent(ValueProperty, bInteractive ? EPropertyChangeType::Interactive : EPropertyChangeType::ValueSet, MakeArrayView(ModifiedObjects));
-			FEditPropertyChain PropertyChain;
-			if (AxisProperty)
-			{
-				PropertyChain.AddHead(AxisProperty);
-			}
-			PropertyChain.AddHead(ValueProperty);
-			FPropertyChangedChainEvent PropertyChangedChainEvent(PropertyChain, PropertyChangedEvent);
-			
+		
 			// Have to downcast here because of function overloading and inheritance not playing nicely
 			if (ValueProperty)
 			{
+				TArray<UObject*> ModifiedObjects = { InObject };
+				FPropertyChangedEvent PropertyChangedEvent(ValueProperty, bInteractive ? EPropertyChangeType::Interactive : EPropertyChangeType::ValueSet, MakeArrayView(ModifiedObjects));
+				FEditPropertyChain PropertyChain;
+				if (AxisProperty)
+				{
+					PropertyChain.AddHead(AxisProperty);
+				}
+				PropertyChain.AddHead(ValueProperty);
+				FPropertyChangedChainEvent PropertyChangedChainEvent(PropertyChain, PropertyChangedEvent);
 				((UObject*)SceneComponentThatChanged)->PreEditChange(PropertyChain);
 			
 				if (ActorThatChanged && ActorThatChanged->GetRootComponent() == SceneComponentThatChanged)
@@ -996,6 +995,15 @@ void UAnimDetailControlsProxyTransform::SetBindingValueFromCurrent(UObject* InOb
 
 			if (ValueProperty)
 			{
+				TArray<UObject*> ModifiedObjects = { InObject };
+				FPropertyChangedEvent PropertyChangedEvent(ValueProperty, bInteractive ? EPropertyChangeType::Interactive : EPropertyChangeType::ValueSet, MakeArrayView(ModifiedObjects));
+				FEditPropertyChain PropertyChain;
+				if (AxisProperty)
+				{
+					PropertyChain.AddHead(AxisProperty);
+				}
+				PropertyChain.AddHead(ValueProperty);
+				FPropertyChangedChainEvent PropertyChangedChainEvent(PropertyChain, PropertyChangedEvent);
 				((UObject*)SceneComponentThatChanged)->PostEditChangeChainProperty(PropertyChangedChainEvent);
 				if (ActorThatChanged && ActorThatChanged->GetRootComponent() == SceneComponentThatChanged)
 				{
