Control Rig: Fix get/set pose not setting the parent of the element
#rnx

diff --git a/Engine/Plugins/Animation/ControlRig/Source/ControlRig/Private/Rigs/RigHierarchy.cpp b/Engine/Plugins/Animation/ControlRig/Source/ControlRig/Private/Rigs/RigHierarchy.cpp
--- a/Engine/Plugins/Animation/ControlRig/Source/ControlRig/Private/Rigs/RigHierarchy.cpp
+++ b/Engine/Plugins/Animation/ControlRig/Source/ControlRig/Private/Rigs/RigHierarchy.cpp
@@ -2991,6 +2991,7 @@ FRigPose URigHierarchy::GetPose(bool bInitial, ERigElementType InElementType,
 		{
 			PoseElement.LocalTransform = GetTransform(TransformElement, bInitial ? ERigTransformType::InitialLocal : ERigTransformType::CurrentLocal);
 			PoseElement.GlobalTransform = GetTransform(TransformElement, bInitial ? ERigTransformType::InitialGlobal : ERigTransformType::CurrentGlobal);
+			PoseElement.ActiveParent = GetActiveParent(Element->GetKey());
 
 			if(FRigControlElement* ControlElement = Cast<FRigControlElement>(Element))
 			{
@@ -3076,6 +3077,7 @@ void URigHierarchy::SetPose(const FRigPose& InPose, ERigTransformType::Type InTr
 					TransformToSet = FControlRigMathLibrary::LerpTransform(PreviousTransform, TransformToSet, U);
 				}
 
+				SwitchToParent(Element->GetKey(), PoseElement.ActiveParent);
 				SetTransform(TransformElement, TransformToSet, InTransformType, true);
 			}
 			else if(FRigCurveElement* CurveElement = Cast<FRigCurveElement>(Element))
diff --git a/Engine/Plugins/Animation/ControlRig/Source/ControlRig/Private/Rigs/RigHierarchyPose.cpp b/Engine/Plugins/Animation/ControlRig/Source/ControlRig/Private/Rigs/RigHierarchyPose.cpp
--- a/Engine/Plugins/Animation/ControlRig/Source/ControlRig/Private/Rigs/RigHierarchyPose.cpp
+++ b/Engine/Plugins/Animation/ControlRig/Source/ControlRig/Private/Rigs/RigHierarchyPose.cpp
@@ -1,7 +1,16 @@
 // Copyright Epic Games, Inc. All Rights Reserved.
 
 #include "Rigs/RigHierarchyPose.h"
+#include "Rigs/RigHierarchy.h"
 
 #include UE_INLINE_GENERATED_CPP_BY_NAME(RigHierarchyPose)
 
 
+FRigPoseElement::FRigPoseElement(): Index()
+                                  , GlobalTransform(FTransform::Identity)
+                                  , LocalTransform(FTransform::Identity)
+                                  , PreferredEulerAngle(FVector::ZeroVector)
+                                  , ActiveParent(URigHierarchy::GetDefaultParentKey())
+                                  , CurveValue(0.f)
+{
+}
diff --git a/Engine/Plugins/Animation/ControlRig/Source/ControlRig/Public/Rigs/RigHierarchyPose.h b/Engine/Plugins/Animation/ControlRig/Source/ControlRig/Public/Rigs/RigHierarchyPose.h
--- a/Engine/Plugins/Animation/ControlRig/Source/ControlRig/Public/Rigs/RigHierarchyPose.h
+++ b/Engine/Plugins/Animation/ControlRig/Source/ControlRig/Public/Rigs/RigHierarchyPose.h
@@ -14,14 +14,7 @@ public:
 
 	GENERATED_BODY()
 
-	FRigPoseElement()
-	: Index()
-	, GlobalTransform(FTransform::Identity)
-	, LocalTransform(FTransform::Identity)
-	, PreferredEulerAngle(FVector::ZeroVector)
-	, CurveValue(0.f)
-	{
-	}
+	FRigPoseElement();
 
 	UPROPERTY()
 	FCachedRigElement Index;
@@ -35,6 +28,9 @@ public:
 	UPROPERTY()
 	FVector PreferredEulerAngle;
 
+	UPROPERTY()
+	FRigElementKey ActiveParent;
+
 	UPROPERTY()
 	float CurveValue;
 };
