Control Rig: Ensure to only allow one primary connector

#rnx

diff --git a/Engine/Plugins/Animation/ControlRig/Source/ControlRigDeveloper/Private/ControlRigBlueprint.cpp b/Engine/Plugins/Animation/ControlRig/Source/ControlRigDeveloper/Private/ControlRigBlueprint.cpp
--- a/Engine/Plugins/Animation/ControlRig/Source/ControlRigDeveloper/Private/ControlRigBlueprint.cpp
+++ b/Engine/Plugins/Animation/ControlRig/Source/ControlRigDeveloper/Private/ControlRigBlueprint.cpp
@@ -40,6 +40,7 @@
 #include "Editor/Transactor.h"
 #include "CookOnTheSide/CookOnTheFlyServer.h"
 #include "ScopedTransaction.h"
+#include "Algo/Count.h"
 #endif//WITH_EDITOR
 
 #define LOCTEXT_NAMESPACE "ControlRigBlueprint"
@@ -1165,6 +1166,34 @@ void UControlRigBlueprint::PostLoad()
 		}
 	}
 
+#if WITH_EDITOR
+	if(IsControlRigModule() && Hierarchy)
+	{
+		// backwards compat - makes sure to only ever allow one primary connector
+		TArray<FRigConnectorElement*> Connectors = Hierarchy->GetConnectors();
+		const int32 NumPrimaryConnectors = Algo::CountIf(Connectors, [](const FRigConnectorElement* InConnector) -> bool
+		{
+			return InConnector->IsPrimary();
+		});
+		if(NumPrimaryConnectors > 1)
+		{
+			bool bHasSeenPrimary = false;
+			for(FRigConnectorElement* Connector : Connectors)
+			{
+				if(bHasSeenPrimary)
+				{
+					Connector->Settings.Type = EConnectorType::Secondary;
+				}
+				else
+				{
+					bHasSeenPrimary = Connector->IsPrimary();
+				}
+			}
+			UpdateExposedModuleConnectors();
+		}
+	}
+#endif
+
 	ModularRigModel.PatchModelsOnLoad();
 	UpdateModularDependencyDelegates();
 }
