Control Rig: Allow to rename connector in details panel

#rnx

diff --git a/Engine/Plugins/Animation/ControlRig/Source/ControlRigEditor/Private/ControlRigElementDetails.cpp b/Engine/Plugins/Animation/ControlRig/Source/ControlRigEditor/Private/ControlRigElementDetails.cpp
--- a/Engine/Plugins/Animation/ControlRig/Source/ControlRigEditor/Private/ControlRigElementDetails.cpp
+++ b/Engine/Plugins/Animation/ControlRig/Source/ControlRigEditor/Private/ControlRigElementDetails.cpp
@@ -813,7 +813,7 @@ void FRigBaseElementDetails::CustomizeDetails(IDetailLayoutBuilder& DetailBuilde
 	}
 
 	const bool bAllControls = !IsAnyElementNotOfType(ERigElementType::Control);
-	const bool bAllAnimationChannels = !IsAnyControlNotOfAnimationType(ERigControlAnimationType::AnimationChannel);
+	const bool bAllAnimationChannels = bAllControls && !IsAnyControlNotOfAnimationType(ERigControlAnimationType::AnimationChannel);
 	if(bAllControls && bAllAnimationChannels)
 	{
 		GeneralCategory.AddCustomRow(FText::FromString(TEXT("Parent Control")))
