Schematic: Don't react to click when shift is pressed

#rnx

#changelist validated
#virtualized

diff --git a/Engine/Plugins/Animation/ControlRig/Source/ControlRigEditor/Private/Editor/ControlRigSchematicModel.cpp b/Engine/Plugins/Animation/ControlRig/Source/ControlRigEditor/Private/Editor/ControlRigSchematicModel.cpp
--- a/Engine/Plugins/Animation/ControlRig/Source/ControlRigEditor/Private/Editor/ControlRigSchematicModel.cpp
+++ b/Engine/Plugins/Animation/ControlRig/Source/ControlRigEditor/Private/Editor/ControlRigSchematicModel.cpp
@@ -1763,11 +1763,14 @@ void FControlRigSchematicModel::OnHierarchyModified(ERigHierarchyNotification In
 	{
 		case ERigHierarchyNotification::ElementSelected:
 		{
-			if(FControlRigSchematicRigElementKeyNode* Node = FindElementKeyNode(InElement->GetKey()))
+			if(InHierarchy->GetSelectedKeys().Num() == 1)
 			{
-				if(FSchematicGraphGroupNode* GroupNode = Node->GetGroupNode())
+				if(FControlRigSchematicRigElementKeyNode* Node = FindElementKeyNode(InElement->GetKey()))
 				{
-					GroupNode->SetExpanded(true);
+					if(FSchematicGraphGroupNode* GroupNode = Node->GetGroupNode())
+					{
+						GroupNode->SetExpanded(true);
+					}
 				}
 			}
 			break;
diff --git a/Engine/Source/Editor/AnimationEditorWidgets/Private/SchematicGraphPanel/SchematicGraphNode.cpp b/Engine/Source/Editor/AnimationEditorWidgets/Private/SchematicGraphPanel/SchematicGraphNode.cpp
--- a/Engine/Source/Editor/AnimationEditorWidgets/Private/SchematicGraphPanel/SchematicGraphNode.cpp
+++ b/Engine/Source/Editor/AnimationEditorWidgets/Private/SchematicGraphPanel/SchematicGraphNode.cpp
@@ -411,7 +411,10 @@ FReply FSchematicGraphGroupNode::OnClicked(const FPointerEvent& InMouseEvent)
 {
 	if(GetNumChildNodes() > 0)
 	{
-		SetExpanded(!IsExpanded());
+		if(!InMouseEvent.GetModifierKeys().AnyModifiersDown())
+		{
+			SetExpanded(!IsExpanded());
+		}
 		return FReply::Handled();
 	}
 	return FSchematicGraphNode::OnClicked(InMouseEvent);
