Sequencer: Anim Detail Fixes, reset to default and vector3d working

#changelist validated
#virtualized

diff --git a/Engine/Plugins/Animation/ControlRig/Source/ControlRigEditor/Private/EditMode/AnimDetailsProxy.cpp b/Engine/Plugins/Animation/ControlRig/Source/ControlRigEditor/Private/EditMode/AnimDetailsProxy.cpp
--- a/Engine/Plugins/Animation/ControlRig/Source/ControlRigEditor/Private/EditMode/AnimDetailsProxy.cpp
+++ b/Engine/Plugins/Animation/ControlRig/Source/ControlRigEditor/Private/EditMode/AnimDetailsProxy.cpp
@@ -1860,11 +1860,11 @@ static void SetVector2DValuesFromContext(UControlRig* ControlRig, FRigControlEle
 	FVector3f Value = ControlValue.Get<FVector3f>();
 
 	EControlRigContextChannelToKey ChannelsToKey = (EControlRigContextChannelToKey)Context.KeyMask;
-	if (EnumHasAnyFlags(ChannelsToKey, EControlRigContextChannelToKey::RotationX) == false)
+	if (EnumHasAnyFlags(ChannelsToKey, EControlRigContextChannelToKey::TranslationX) == false)
 	{
 		Val.X = Value.X;
 	}
-	if (EnumHasAnyFlags(ChannelsToKey, EControlRigContextChannelToKey::RotationZ) == false)
+	if (EnumHasAnyFlags(ChannelsToKey, EControlRigContextChannelToKey::TranslationY) == false)
 	{
 		Val.Y = Value.Y;
 	}
diff --git a/Engine/Plugins/Animation/ControlRig/Source/ControlRigEditor/Private/EditMode/ControlRigControlsProxy.cpp b/Engine/Plugins/Animation/ControlRig/Source/ControlRigEditor/Private/EditMode/ControlRigControlsProxy.cpp
--- a/Engine/Plugins/Animation/ControlRig/Source/ControlRigEditor/Private/EditMode/ControlRigControlsProxy.cpp
+++ b/Engine/Plugins/Animation/ControlRig/Source/ControlRigEditor/Private/EditMode/ControlRigControlsProxy.cpp
@@ -10,7 +10,6 @@
 #include "Components/SkeletalMeshComponent.h"
 #include "Subsystems/AssetEditorSubsystem.h"
 #include "MovieSceneCommonHelpers.h"
-#include "PropertyHandle.h"
 #include "Sequencer/ControlRigSequencerHelpers.h"
 #include "ISequencer.h"
 #include "CurveEditor.h"
@@ -219,6 +218,10 @@ void UControlRigControlsProxy::SelectionChanged(bool bInSelected)
 void UControlRigControlsProxy::PostEditChangeChainProperty(struct FPropertyChangedChainEvent& PropertyChangedEvent)
 {
 	Super::PostEditChangeChainProperty(PropertyChangedEvent);
+	if (PropertyChangedEvent.ChangeType == EPropertyChangeType::ToggleEditable)//hack so we can clear the reset cache for this property and not actually send this to our controls
+	{
+		return;
+	}
 	if (PropertyChangedEvent.Property && PropertyChangedEvent.Property->GetFName() == GET_MEMBER_NAME_CHECKED(UControlRigControlsProxy, bSelected))
 	{
 		if (OwnerControlElement && OwnerControlRig.IsValid())
diff --git a/Engine/Plugins/Animation/ControlRig/Source/ControlRigEditor/Private/EditMode/ControlRigEditMode.cpp b/Engine/Plugins/Animation/ControlRig/Source/ControlRigEditor/Private/EditMode/ControlRigEditMode.cpp
--- a/Engine/Plugins/Animation/ControlRig/Source/ControlRigEditor/Private/EditMode/ControlRigEditMode.cpp
+++ b/Engine/Plugins/Animation/ControlRig/Source/ControlRigEditor/Private/EditMode/ControlRigEditMode.cpp
@@ -3552,6 +3552,7 @@ bool FControlRigEditMode::CheckMovieSceneSig()
 				{
 					SetObjects_Internal();
 				}
+				DetailKeyFrameCache->ResetCachedData();
 			}
 		}
 	}
@@ -5481,6 +5482,11 @@ EPropertyKeyedStatus FDetailKeyFrameCacheAndHandler::GetPropertyKeyedStatus(cons
 	{
 		return *ExistingKeyedStatus;
 	}
+	//hack so we can get the reset cache state updated, use ToggleEditable state
+	{
+		IPropertyHandle* NotConst = const_cast<IPropertyHandle*>(&PropertyHandle);
+		NotConst->NotifyPostChange(EPropertyChangeType::ToggleEditable);
+	}
 
 	TSharedPtr<ISequencer> SequencerPtr = WeakSequencer.Pin();
 	UMovieSceneSequence* Sequence = SequencerPtr->GetFocusedMovieSceneSequence();
