Sequencer: Control Rig:Only reconstruct in game thread, from UDN.

#changelist validated
#virtualized

diff --git a/Engine/Plugins/Animation/ControlRig/Source/ControlRig/Private/Sequencer/MovieSceneControlRigParameterTrack.cpp b/Engine/Plugins/Animation/ControlRig/Source/ControlRig/Private/Sequencer/MovieSceneControlRigParameterTrack.cpp
--- a/Engine/Plugins/Animation/ControlRig/Source/ControlRig/Private/Sequencer/MovieSceneControlRigParameterTrack.cpp
+++ b/Engine/Plugins/Animation/ControlRig/Source/ControlRig/Private/Sequencer/MovieSceneControlRigParameterTrack.cpp
@@ -560,11 +560,11 @@ void UMovieSceneControlRigParameterTrack::HandleOnPostConstructed_GameThread()
 
 void UMovieSceneControlRigParameterTrack::HandleOnPostConstructed(UControlRig* Subject, const FName& InEventName)
 {
-	MovieSceneControlRigTrack::AsyncHelpers::ExecuteOnGameThread<void>([this]()
+	if(IsInGameThread())
 	{
 		HandleOnPostConstructed_GameThread();
 
-	}).Wait();
+	}
 }
 
 #if WITH_EDITORONLY_DATA
