Control Rig: Correct icon size, color and opacity on Modular Hierarchy widget
#rnx

#changelist validated
#virtualized

diff --git a/Engine/Plugins/Animation/ControlRig/Source/ControlRigEditor/Private/Editor/SModularRigTreeView.cpp b/Engine/Plugins/Animation/ControlRig/Source/ControlRigEditor/Private/Editor/SModularRigTreeView.cpp
--- a/Engine/Plugins/Animation/ControlRig/Source/ControlRigEditor/Private/Editor/SModularRigTreeView.cpp
+++ b/Engine/Plugins/Animation/ControlRig/Source/ControlRigEditor/Private/Editor/SModularRigTreeView.cpp
@@ -293,6 +293,7 @@ TSharedRef<SWidget> SModularRigModelItem::GenerateWidgetForColumn(const FName& C
 					}
 					return FSlateColor::UseForeground();
 				})
+				.DesiredSizeOverride(FVector2D(16, 16))
 			]
 		]
 
@@ -981,7 +982,8 @@ const TSharedPtr<FModularRigTreeElement>* SModularRigTreeView::FindItemAtPositio
 TPair<const FSlateBrush*, FSlateColor> FModularRigTreeElement::GetBrushAndColor(const UModularRig* InModularRig)
 {
 	const FSlateBrush* Brush = nullptr;
-	FSlateColor Color = FSlateColor::UseForeground();
+	FLinearColor Color = FSlateColor::UseForeground().GetColor(FWidgetStyle());
+	float Opacity = 1.f;
 
 	if (const FRigModuleInstance* ConnectorModule = InModularRig->FindModule(ModulePath))
 	{
@@ -1020,6 +1022,10 @@ TPair<const FSlateBrush*, FSlateColor> FModularRigTreeElement::GetBrushAndColor(
 				else if (Connector->Settings.bOptional)
 				{
 					bConnectionWarning = false;
+					if (!bIsConnected)
+					{
+						Opacity = 0.6;
+					}
 					Brush = FControlRigEditorStyle::Get().GetBrush("ControlRig.Schematic.ConnectorOptional");
 				}
 				else
@@ -1031,13 +1037,17 @@ TPair<const FSlateBrush*, FSlateColor> FModularRigTreeElement::GetBrushAndColor(
 
 		if (bConnectionWarning)
 		{
-			Color = FStyleColors::AccentYellow;
+			Color = FLinearColor(FColor::FromHex(TEXT("#FFB800")));
 		}
 	}
 	if (!Brush)
 	{
 		Brush = FControlRigEditorStyle::Get().GetBrush("ControlRig.Tree.RigidBody");
 	}
+
+	// Apply opacity
+	Color = Color.CopyWithNewOpacity(Opacity);
+	
 	return TPair<const FSlateBrush*, FSlateColor>(Brush, Color);
 }
 
