Control Rig: Offer in viewport switching of the preview mesh

#rnx

#changelist validated
#virtualized

diff --git a/Engine/Plugins/Animation/ControlRig/Source/ControlRigEditor/Private/Editor/ControlRigEditor.cpp b/Engine/Plugins/Animation/ControlRig/Source/ControlRigEditor/Private/Editor/ControlRigEditor.cpp
--- a/Engine/Plugins/Animation/ControlRig/Source/ControlRigEditor/Private/Editor/ControlRigEditor.cpp
+++ b/Engine/Plugins/Animation/ControlRig/Source/ControlRigEditor/Private/Editor/ControlRigEditor.cpp
@@ -76,6 +76,7 @@
 #include "UnrealExporter.h"
 #include "ControlRigElementDetails.h"
 #include "PropertyEditorModule.h"
+#include "PropertyCustomizationHelpers.h"
 #include "Settings/ControlRigSettings.h"
 #include "Widgets/Docking/SDockTab.h"
 #include "BlueprintCompilationManager.h"
@@ -1900,8 +1901,15 @@ void FControlRigEditor::HandleViewportCreated(const TSharedRef<class IPersonaVie
 
 	auto GetCompilationStateVisibility = [this]()
 	{
-		if (UBlueprint* Blueprint = GetBlueprintObj())
+		if (const UControlRigBlueprint* Blueprint = GetControlRigBlueprint())
 		{
+			if(Blueprint->IsModularRig())
+			{
+				if(Blueprint->GetPreviewMesh() == nullptr)
+				{
+					return EVisibility::Collapsed;
+				}
+			}
 			const bool bUpToDate = (Blueprint->Status == BS_UpToDate) || (Blueprint->Status == BS_UpToDateWithWarnings);
 			return bUpToDate ? EVisibility::Collapsed : EVisibility::Visible;
 		}
@@ -1911,11 +1919,10 @@ void FControlRigEditor::HandleViewportCreated(const TSharedRef<class IPersonaVie
 
 	auto GetCompileButtonVisibility = [this]()
 	{
-		if (UBlueprint* Blueprint = GetBlueprintObj())
+		if (const UControlRigBlueprint* Blueprint = GetControlRigBlueprint())
 		{
 			return (Blueprint->Status == BS_Dirty) ? EVisibility::Visible : EVisibility::Collapsed;
 		}
-
 		return EVisibility::Collapsed;
 	};
 
@@ -2008,21 +2015,21 @@ void FControlRigEditor::HandleViewportCreated(const TSharedRef<class IPersonaVie
 				.AutoWidth()
 				[
 					SAssignNew(DirectManipulationCombo, SComboBox<TSharedPtr<FString>>)
-                	.ContentPadding(FMargin(4.0f, 2.0f))
-                	.OptionsSource(&DirectManipulationTextList)
-                	.OnGenerateWidget_Lambda([this](TSharedPtr<FString> Item)
-                	{ 
-                		return SNew(SBox)
-                			.MaxDesiredWidth(600.0f)
-                			[
-                				SNew(STextBlock)
+					.ContentPadding(FMargin(4.0f, 2.0f))
+					.OptionsSource(&DirectManipulationTextList)
+					.OnGenerateWidget_Lambda([this](TSharedPtr<FString> Item)
+					{ 
+						return SNew(SBox)
+							.MaxDesiredWidth(600.0f)
+							[
+								SNew(STextBlock)
 								.TextStyle(FAppStyle::Get(), "AnimViewport.MessageText")
-                				.Text(FText::FromString(*Item))
-                			];
-                	} )	
-                	.OnSelectionChanged(this, &FControlRigEditor::OnDirectManipulationChanged)
-                	[
-			        	SNew(STextBlock)
+								.Text(FText::FromString(*Item))
+							];
+					} )	
+					.OnSelectionChanged(this, &FControlRigEditor::OnDirectManipulationChanged)
+					[
+						SNew(STextBlock)
 						.TextStyle(FAppStyle::Get(), "AnimViewport.MessageText")
 						.Text(this, &FControlRigEditor::GetDirectionManipulationText)
 					]
@@ -2161,7 +2168,82 @@ void FControlRigEditor::HandleViewportCreated(const TSharedRef<class IPersonaVie
 		],
 		FPersonaViewportNotificationOptions(TAttribute<EVisibility>::Create(GetCompilationStateVisibility))
 	);
-	
+
+	FPersonaViewportNotificationOptions ChangePreviewMeshNotificationOptions;
+	ChangePreviewMeshNotificationOptions.OnGetVisibility = IsModularRig() ? EVisibility::Visible : EVisibility::Collapsed;
+	//ChangePreviewMeshNotificationOptions.OnGetBrushOverride = TAttribute<const FSlateBrush*>(FControlRigEditorStyle::Get().GetBrush("ControlRig.Viewport.Notification.ChangeShapeTransform"));
+
+	// notification to allow to change the preview mesh directly in the viewport
+	InViewport->AddNotification(TAttribute<EMessageSeverity::Type>::CreateLambda([this]()
+		{
+			if(const UControlRigBlueprint* Blueprint = GetControlRigBlueprint())
+			{
+				if(Blueprint->GetPreviewMesh() == nullptr)
+				{
+					return EMessageSeverity::Warning;
+				}
+			}
+			return EMessageSeverity::Info;
+		}),
+		false,
+		SNew(SHorizontalBox)
+		+SHorizontalBox::Slot()
+		.AutoWidth()
+		.VAlign(VAlign_Center)
+		.Padding(4.0f, 4.0f)
+		[
+			SNew(STextBlock)
+			.Text(LOCTEXT("MissingPreviewMesh", "Please choose a preview mesh!"))
+			.TextStyle(FAppStyle::Get(), "AnimViewport.MessageText")
+			.Visibility_Lambda([this]()
+			{
+				if(const UControlRigBlueprint* Blueprint = GetControlRigBlueprint())
+				{
+					if(Blueprint->GetPreviewMesh())
+					{
+						return EVisibility::Collapsed;
+					}
+				}
+				return EVisibility::Visible;
+			})
+		]
+		+SHorizontalBox::Slot()
+		.AutoWidth()
+		.Padding(4.0f, 4.0f)
+		[
+			SNew(SObjectPropertyEntryBox)
+			.ObjectPath_Lambda([this]()
+			{
+				if(const UControlRigBlueprint* Blueprint = GetControlRigBlueprint())
+				{
+					if(const USkeletalMesh* PreviewMesh = Blueprint->GetPreviewMesh())
+					{
+						return PreviewMesh->GetPathName();
+					}
+				}
+				return FString();
+			})
+			.AllowedClass(USkeletalMesh::StaticClass())
+			.OnObjectChanged_Lambda([this](const FAssetData& InAssetData)
+			{
+				if(const UControlRigBlueprint* Blueprint = GetControlRigBlueprint())
+				{
+					if(USkeletalMesh* SkeletalMesh = Cast<USkeletalMesh>(InAssetData.GetAsset()))
+					{
+						const TSharedRef<IPersonaPreviewScene> CurrentPreviewScene = GetPersonaToolkit()->GetPreviewScene();
+						CurrentPreviewScene->SetPreviewMesh(SkeletalMesh);
+					}
+				}
+			})
+			.AllowCreate(false)
+			.AllowClear(false)
+			.DisplayUseSelected(false)
+			.DisplayBrowse(false)
+			.NewAssetFactories(TArray<UFactory*>())
+		],
+		ChangePreviewMeshNotificationOptions
+	);
+
 	FPersonaViewportNotificationOptions ChangeShapeTransformNotificationOptions;
 	ChangeShapeTransformNotificationOptions.OnGetVisibility = TAttribute<EVisibility>::Create(GetChangingShapeTransformTextVisibility);
 	ChangeShapeTransformNotificationOptions.OnGetBrushOverride = TAttribute<const FSlateBrush*>(FControlRigEditorStyle::Get().GetBrush("ControlRig.Viewport.Notification.ChangeShapeTransform"));
@@ -2719,6 +2801,20 @@ void FControlRigEditor::HandlePreviewMeshChanged(USkeletalMesh* InOldSkeletalMes
 		if (UControlRigBlueprint* ControlRigBP = GetControlRigBlueprint())
 		{
 			ControlRigBP->SetPreviewMesh(InNewSkeletalMesh);
+
+			if(IsModularRig())
+			{
+				{
+					TGuardValue<bool> SuspendBlueprintNotifs(ControlRigBP->bSuspendAllNotifications, true);
+					if(URigHierarchyController* Controller = ControlRigBP->GetHierarchyController())
+					{
+						Controller->ImportBones(InNewSkeletalMesh->GetSkeleton(), NAME_None, true, true, false, true, true);
+						Controller->ImportCurves(InNewSkeletalMesh->GetSkeleton(), NAME_None, false, true, true);
+					}
+				}
+				ControlRigBP->PropagateHierarchyFromBPToInstances();
+			}
+			
 			UpdateRigVMHost();
 			
 			if(UControlRig* DebuggedControlRig = Cast<UControlRig>(GetBlueprintObj()->GetObjectBeingDebugged()))
