Control Rig: make skeleton dependency editable

#rnx

#changelist validated
#virtualized

diff --git a/Engine/Plugins/Animation/ControlRig/Source/ControlRigDeveloper/Public/ControlRigBlueprint.h b/Engine/Plugins/Animation/ControlRig/Source/ControlRigDeveloper/Public/ControlRigBlueprint.h
--- a/Engine/Plugins/Animation/ControlRig/Source/ControlRigDeveloper/Public/ControlRigBlueprint.h
+++ b/Engine/Plugins/Animation/ControlRig/Source/ControlRigDeveloper/Public/ControlRigBlueprint.h
@@ -270,11 +270,11 @@ private:
 #endif
 
 	/** The skeleton from import into a hierarchy */
-	UPROPERTY(DuplicateTransient, AssetRegistrySearchable)
+	UPROPERTY(DuplicateTransient, AssetRegistrySearchable, EditAnywhere, Category="Control Rig Blueprint")
 	TSoftObjectPtr<UObject> SourceHierarchyImport;
 
 	/** The skeleton from import into a curve */
-	UPROPERTY(DuplicateTransient, AssetRegistrySearchable)
+	UPROPERTY(DuplicateTransient, AssetRegistrySearchable, EditAnywhere, Category="Control Rig Blueprint")
 	TSoftObjectPtr<UObject> SourceCurveImport;
 
 	/** If set to true, this control rig has animatable controls */
