Control Rig: Use dedicated icons for connectors in the schematic

#rnx

#changelist validated
#virtualized

diff --git a/Engine/Plugins/Animation/ControlRig/Source/ControlRigEditor/Private/ControlRigEditorStyle.h b/Engine/Plugins/Animation/ControlRig/Source/ControlRigEditor/Private/ControlRigEditorStyle.h
--- a/Engine/Plugins/Animation/ControlRig/Source/ControlRigEditor/Private/ControlRigEditorStyle.h
+++ b/Engine/Plugins/Animation/ControlRig/Source/ControlRigEditor/Private/ControlRigEditorStyle.h
@@ -147,13 +147,16 @@ public:
 			Set("ControlRig.TestData.Record", new IMAGE_BRUSH("Slate/RecordingIndicator", Icon32x32));
 		}
 
+		Set("ControlRig.ConnectorPrimary", new IMAGE_BRUSH_SVG("Slate/Connector_Primary", Icon128x128));
+		Set("ControlRig.ConnectorSecondary", new IMAGE_BRUSH_SVG("Slate/Connector_Secondary", Icon128x128));
+		Set("ControlRig.ConnectorOptional", new IMAGE_BRUSH_SVG("Slate/Connector_Optional", Icon128x128));
 		Set("ControlRig.ConnectorWarning", new IMAGE_BRUSH_SVG("Slate/Connector_Warning", Icon128x128));
 
 		// Schematic
 		{
-			Set("ControlRig.Schematic.ConnectorPrimary", new IMAGE_BRUSH_SVG("Slate/Connector_Primary", Icon128x128));
-			Set("ControlRig.Schematic.ConnectorSecondary", new IMAGE_BRUSH_SVG("Slate/Connector_Secondary", Icon128x128));
-			Set("ControlRig.Schematic.ConnectorOptional", new IMAGE_BRUSH_SVG("Slate/Connector_Optional", Icon128x128));
+			Set("ControlRig.Schematic.ConnectorPrimary", new IMAGE_BRUSH_SVG("Slate/Connector_Primary_Schematic", Icon128x128));
+			Set("ControlRig.Schematic.ConnectorSecondary", new IMAGE_BRUSH_SVG("Slate/Connector_Secondary_Schematic", Icon128x128));
+			Set("ControlRig.Schematic.ConnectorOptional", new IMAGE_BRUSH_SVG("Slate/Connector_Optional_Schematic", Icon128x128));
 			Set("ControlRig.Schematic.ConnectorWarning", new IMAGE_BRUSH_SVG("Slate/Connector_Warning_Schematic", Icon128x128));
 			Set("ControlRig.Schematic.Bone", new IMAGE_BRUSH_SVG("Slate/Bone_Schematic", Icon128x128));
 			Set("ControlRig.Schematic.Control", new IMAGE_BRUSH_SVG("Slate/Control_Schematic", Icon128x128));
diff --git a/Engine/Plugins/Animation/ControlRig/Source/ControlRigEditor/Private/Editor/SModularRigTreeView.cpp b/Engine/Plugins/Animation/ControlRig/Source/ControlRigEditor/Private/Editor/SModularRigTreeView.cpp
--- a/Engine/Plugins/Animation/ControlRig/Source/ControlRigEditor/Private/Editor/SModularRigTreeView.cpp
+++ b/Engine/Plugins/Animation/ControlRig/Source/ControlRigEditor/Private/Editor/SModularRigTreeView.cpp
@@ -1026,11 +1026,11 @@ TPair<const FSlateBrush*, FSlateColor> FModularRigTreeElement::GetBrushAndColor(
 					{
 						Opacity = 0.6;
 					}
-					Brush = FControlRigEditorStyle::Get().GetBrush("ControlRig.Schematic.ConnectorOptional");
+					Brush = FControlRigEditorStyle::Get().GetBrush("ControlRig.ConnectorOptional");
 				}
 				else
 				{
-					Brush = FControlRigEditorStyle::Get().GetBrush("ControlRig.Schematic.ConnectorSecondary");
+					Brush = FControlRigEditorStyle::Get().GetBrush("ControlRig.ConnectorSecondary");
 				}
 			}
 		}
