Control RIg: Fix bungee man test data unit test failing
#rnx

#changelist validated
#virtualized

diff --git a/Engine/Plugins/Animation/ControlRig/Source/ControlRig/Private/ControlRigTestData.cpp b/Engine/Plugins/Animation/ControlRig/Source/ControlRig/Private/ControlRigTestData.cpp
--- a/Engine/Plugins/Animation/ControlRig/Source/ControlRig/Private/ControlRigTestData.cpp
+++ b/Engine/Plugins/Animation/ControlRig/Source/ControlRig/Private/ControlRigTestData.cpp
@@ -2,6 +2,7 @@
 
 #include "ControlRigTestData.h"
 #include "ControlRig.h"
+#include "ControlRigObjectVersion.h"
 #include "HAL/PlatformTime.h"
 #if WITH_EDITOR
 #include "AssetToolsModule.h"
@@ -146,8 +147,32 @@ bool FControlRigTestDataFrame::RestoreVariables(UControlRig* InControlRig) const
 
 void UControlRigTestData::Serialize(FArchive& Ar)
 {
+	Ar.UsingCustomVersion(FControlRigObjectVersion::GUID);	
 	UObject::Serialize(Ar);
 	LastFrameIndex = INDEX_NONE;
+
+	// If pose is older than RigPoseWithParentKey, set the active parent of all poses to invalid key
+	if (GetLinkerCustomVersion(FControlRigObjectVersion::GUID) < FControlRigObjectVersion::RigPoseWithParentKey)
+	{
+		for (FRigPoseElement& Element : Initial.Pose)
+		{
+			Element.ActiveParent = FRigElementKey();
+		}
+		for (FControlRigTestDataFrame& Frame : InputFrames)
+		{
+			for (FRigPoseElement& Element : Frame.Pose)
+			{
+				Element.ActiveParent = FRigElementKey();
+			}
+		}
+		for (FControlRigTestDataFrame& Frame : OutputFrames)
+		{
+			for (FRigPoseElement& Element : Frame.Pose)
+			{
+				Element.ActiveParent = FRigElementKey();
+			}
+		}
+	}
 }
 
 UControlRigTestData* UControlRigTestData::CreateNewAsset(FString InDesiredPackagePath, FString InBlueprintPathName)
diff --git a/Engine/Plugins/Animation/ControlRig/Source/ControlRig/Private/Rigs/RigHierarchy.cpp b/Engine/Plugins/Animation/ControlRig/Source/ControlRig/Private/Rigs/RigHierarchy.cpp
--- a/Engine/Plugins/Animation/ControlRig/Source/ControlRig/Private/Rigs/RigHierarchy.cpp
+++ b/Engine/Plugins/Animation/ControlRig/Source/ControlRig/Private/Rigs/RigHierarchy.cpp
@@ -3077,7 +3077,10 @@ void URigHierarchy::SetPose(const FRigPose& InPose, ERigTransformType::Type InTr
 					TransformToSet = FControlRigMathLibrary::LerpTransform(PreviousTransform, TransformToSet, U);
 				}
 
-				SwitchToParent(Element->GetKey(), PoseElement.ActiveParent);
+				if (PoseElement.ActiveParent.IsValid())
+				{
+					SwitchToParent(Element->GetKey(), PoseElement.ActiveParent);
+				}
 				SetTransform(TransformElement, TransformToSet, InTransformType, true);
 			}
 			else if(FRigCurveElement* CurveElement = Cast<FRigCurveElement>(Element))
diff --git a/Engine/Plugins/Animation/ControlRig/Source/ControlRig/Public/ControlRigObjectVersion.h b/Engine/Plugins/Animation/ControlRig/Source/ControlRig/Public/ControlRigObjectVersion.h
--- a/Engine/Plugins/Animation/ControlRig/Source/ControlRig/Public/ControlRigObjectVersion.h
+++ b/Engine/Plugins/Animation/ControlRig/Source/ControlRig/Public/ControlRigObjectVersion.h
@@ -98,6 +98,9 @@ struct CONTROLRIG_API FControlRigObjectVersion
 		// Add type (primary, secondary) and optional bool to FRigConnectorSettings
 		ConnectorsWithType,
 
+		// Add parent key to control rig pose
+		RigPoseWithParentKey,
+
 		// -----<new versions can be added above this line>-------------------------------------------------
 		VersionPlusOne,
 		LatestVersion = VersionPlusOne - 1
