Control Rig: Fix shape exists node

#rnx

diff --git a/Engine/Plugins/Animation/ControlRig/Source/ControlRig/Private/ControlRig.cpp b/Engine/Plugins/Animation/ControlRig/Source/ControlRig/Private/ControlRig.cpp
--- a/Engine/Plugins/Animation/ControlRig/Source/ControlRig/Private/ControlRig.cpp
+++ b/Engine/Plugins/Animation/ControlRig/Source/ControlRig/Private/ControlRig.cpp
@@ -315,7 +315,7 @@ void UControlRig::OnAddShapeLibrary(const FControlRigExecuteContext* InContext,
 bool UControlRig::OnShapeExists(const FName& InShapeName) const
 {
 	const TArray<TSoftObjectPtr<UControlRigShapeLibrary>>& Libraries = GetShapeLibraries();
-	if (UControlRigShapeLibrary::GetShapeByName(InShapeName, GetShapeLibraries(), ShapeLibraryNameMap))
+	if (UControlRigShapeLibrary::GetShapeByName(InShapeName, GetShapeLibraries(), ShapeLibraryNameMap, false))
 	{
 		return true;
 	}
diff --git a/Engine/Plugins/Animation/ControlRig/Source/ControlRig/Private/ControlRigGizmoLibrary.cpp b/Engine/Plugins/Animation/ControlRig/Source/ControlRig/Private/ControlRigGizmoLibrary.cpp
--- a/Engine/Plugins/Animation/ControlRig/Source/ControlRig/Private/ControlRigGizmoLibrary.cpp
+++ b/Engine/Plugins/Animation/ControlRig/Source/ControlRig/Private/ControlRigGizmoLibrary.cpp
@@ -110,7 +110,7 @@ const FControlRigShapeDefinition* UControlRigShapeLibrary::GetShapeByName(const
 	return nullptr;
 }
 
-const FControlRigShapeDefinition* UControlRigShapeLibrary::GetShapeByName(const FName& InName, const TArray<TSoftObjectPtr<UControlRigShapeLibrary>>& InShapeLibraries, const TMap<FString, FString>& InLibraryNameMap)
+const FControlRigShapeDefinition* UControlRigShapeLibrary::GetShapeByName(const FName& InName, const TArray<TSoftObjectPtr<UControlRigShapeLibrary>>& InShapeLibraries, const TMap<FString, FString>& InLibraryNameMap, bool bUseDefaultIfNotFound)
 {
 	const FString InString = InName.ToString();
 	FString Left, Right;
@@ -149,8 +149,8 @@ const FControlRigShapeDefinition* UControlRigShapeLibrary::GetShapeByName(const
 			if(ShapeLibraryName.Equals(Left) || Left.IsEmpty())
 			{
 				// only fall back on the default shape for the very last shape library
-				const bool bUseDefaultShapeIfNotFound = (Pass > 0) && (Index == 0);
-				if(const FControlRigShapeDefinition* Shape = ShapeLibrary->GetShapeByName(RightName, bUseDefaultShapeIfNotFound))
+				const bool bFallBackToDefaultShape = (Pass > 0) && (Index == 0) && bUseDefaultIfNotFound;
+				if(const FControlRigShapeDefinition* Shape = ShapeLibrary->GetShapeByName(RightName, bFallBackToDefaultShape))
 				{
 					return Shape;
 				}
diff --git a/Engine/Plugins/Animation/ControlRig/Source/ControlRig/Public/ControlRigGizmoLibrary.h b/Engine/Plugins/Animation/ControlRig/Source/ControlRig/Public/ControlRigGizmoLibrary.h
--- a/Engine/Plugins/Animation/ControlRig/Source/ControlRig/Public/ControlRigGizmoLibrary.h
+++ b/Engine/Plugins/Animation/ControlRig/Source/ControlRig/Public/ControlRigGizmoLibrary.h
@@ -61,7 +61,7 @@ public:
 	TArray<FControlRigShapeDefinition> Shapes;
 
 	const FControlRigShapeDefinition* GetShapeByName(const FName& InName, bool bUseDefaultIfNotFound = false) const;
-	static const FControlRigShapeDefinition* GetShapeByName(const FName& InName, const TArray<TSoftObjectPtr<UControlRigShapeLibrary>>& InShapeLibraries, const TMap<FString, FString>& InLibraryNameMap);
+	static const FControlRigShapeDefinition* GetShapeByName(const FName& InName, const TArray<TSoftObjectPtr<UControlRigShapeLibrary>>& InShapeLibraries, const TMap<FString, FString>& InLibraryNameMap, bool bUseDefaultIfNotFound = true);
 	static const FString GetShapeName(const UControlRigShapeLibrary* InShapeLibrary, bool bUseNameSpace, const TMap<FString, FString>& InLibraryNameMap, const FControlRigShapeDefinition& InShape);
 
 #if WITH_EDITOR
