Control Rig: RigHierarchy design update

#rnx

diff --git a/Engine/Plugins/Animation/ControlRig/Source/ControlRigEditor/Private/Editor/SRigHierarchyTagWidget.cpp b/Engine/Plugins/Animation/ControlRig/Source/ControlRigEditor/Private/Editor/SRigHierarchyTagWidget.cpp
--- a/Engine/Plugins/Animation/ControlRig/Source/ControlRigEditor/Private/Editor/SRigHierarchyTagWidget.cpp
+++ b/Engine/Plugins/Animation/ControlRig/Source/ControlRigEditor/Private/Editor/SRigHierarchyTagWidget.cpp
@@ -19,6 +19,7 @@ void SRigHierarchyTagWidget::Construct(const FArguments& InArgs)
 {
 	Text = InArgs._Text;
 	Icon = InArgs._Icon;
+	IconColor = InArgs._IconColor;
 	Color = InArgs._Color;
 	Radius = InArgs._Radius;
 	Padding = InArgs._Padding;
@@ -49,6 +50,8 @@ void SRigHierarchyTagWidget::Construct(const FArguments& InArgs)
 				return Icon.Get() ? EVisibility::Visible : EVisibility::Collapsed;
 			})
 			.Image(Icon)
+			.ColorAndOpacity(IconColor)
+			.DesiredSizeOverride(InArgs._IconSize)
 		]
 
 		+ SHorizontalBox::Slot()
diff --git a/Engine/Plugins/Animation/ControlRig/Source/ControlRigEditor/Private/Editor/SRigHierarchyTagWidget.h b/Engine/Plugins/Animation/ControlRig/Source/ControlRigEditor/Private/Editor/SRigHierarchyTagWidget.h
--- a/Engine/Plugins/Animation/ControlRig/Source/ControlRigEditor/Private/Editor/SRigHierarchyTagWidget.h
+++ b/Engine/Plugins/Animation/ControlRig/Source/ControlRigEditor/Private/Editor/SRigHierarchyTagWidget.h
@@ -19,6 +19,7 @@ public:
 		, _TooltipText()
 		, _Icon(nullptr)
 		, _Color(FLinearColor(.2, .2, .2))
+		, _IconColor(FLinearColor::White)
 		, _TextColor(FLinearColor(.8, .8, .8))
 		, _Radius(5)
 		, _Padding(10, 0, 0, 0)
@@ -32,7 +33,9 @@ public:
 		SLATE_ATTRIBUTE(FText, Text)
 		SLATE_ATTRIBUTE(FText, TooltipText)
 		SLATE_ATTRIBUTE(const FSlateBrush*, Icon)
+		SLATE_ATTRIBUTE(TOptional<FVector2d>, IconSize)
 		SLATE_ATTRIBUTE(FLinearColor, Color)
+		SLATE_ATTRIBUTE(FSlateColor, IconColor)
 		SLATE_ATTRIBUTE(FSlateColor, TextColor)
 		SLATE_ARGUMENT(float, Radius)
 		SLATE_ARGUMENT(FMargin, Padding)
@@ -59,6 +62,7 @@ private:
 	TAttribute<FText> Text;
 	TAttribute<const FSlateBrush*> Icon;
 	TAttribute<FLinearColor> Color;
+	TAttribute<FSlateColor> IconColor;
 	float Radius = 5.f;
 	FMargin Padding;
 	FMargin ContentPadding;
diff --git a/Engine/Plugins/Animation/ControlRig/Source/ControlRigEditor/Private/Editor/SRigHierarchyTreeView.cpp b/Engine/Plugins/Animation/ControlRig/Source/ControlRigEditor/Private/Editor/SRigHierarchyTreeView.cpp
--- a/Engine/Plugins/Animation/ControlRig/Source/ControlRigEditor/Private/Editor/SRigHierarchyTreeView.cpp
+++ b/Engine/Plugins/Animation/ControlRig/Source/ControlRigEditor/Private/Editor/SRigHierarchyTreeView.cpp
@@ -21,6 +21,7 @@
 #include "Rigs/RigHierarchyController.h"
 #include "Styling/AppStyle.h"
 #include "Algo/Sort.h"
+#include "ControlRigBlueprint.h"
 
 #define LOCTEXT_NAMESPACE "SRigHierarchyTreeView"
 
@@ -162,6 +163,7 @@ void SRigHierarchyItem::Construct(const FArguments& InArgs, const TSharedRef<STa
 
 	STableRow<TSharedPtr<FRigTreeElement>>::Construct(
 		STableRow<TSharedPtr<FRigTreeElement>>::FArguments()
+		.Padding(FMargin(0, 1, 1, 1))
 		.OnDragDetected(Delegates.OnDragDetected)
 		.OnCanAcceptDrop(Delegates.OnCanAcceptDrop)
 		.OnAcceptDrop(Delegates.OnAcceptDrop)
@@ -171,7 +173,7 @@ void SRigHierarchyItem::Construct(const FArguments& InArgs, const TSharedRef<STa
 			SAssignNew(HorizontalBox, SHorizontalBox)
 			+ SHorizontalBox::Slot()
 			.MaxWidth(18)
-			.FillWidth(1.0)
+			.AutoWidth()
 			.HAlign(HAlign_Left)
 			.VAlign(VAlign_Center)
 			.Padding(FMargin(0.f, 0.f, 3.f, 0.f))
@@ -215,16 +217,25 @@ void SRigHierarchyItem::Construct(const FArguments& InArgs, const TSharedRef<STa
 			]
 		], OwnerTable);
 
-	for(const SRigHierarchyTagWidget::FArguments& TagArguments : InRigTreeElement->Tags)
+	if(!InRigTreeElement->Tags.IsEmpty())
 	{
-		TSharedRef<SRigHierarchyTagWidget> TagWidget = SArgumentNew(TagArguments, SRigHierarchyTagWidget);
-		TagWidget->OnElementKeyDragDetected().BindSP(InTreeView.Get(), &SRigHierarchyTreeView::OnElementKeyTagDragDetected);
-		
 		HorizontalBox->AddSlot()
-		.AutoWidth()
+		.FillWidth(1.f)
 		[
-			TagWidget
+			SNew(SSpacer)
 		];
+		
+		for(const SRigHierarchyTagWidget::FArguments& TagArguments : InRigTreeElement->Tags)
+		{
+			TSharedRef<SRigHierarchyTagWidget> TagWidget = SArgumentNew(TagArguments, SRigHierarchyTagWidget);
+			TagWidget->OnElementKeyDragDetected().BindSP(InTreeView.Get(), &SRigHierarchyTreeView::OnElementKeyTagDragDetected);
+			
+			HorizontalBox->AddSlot()
+			.AutoWidth()
+			[
+				TagWidget
+			];
+		}
 	}
 
 	InRigTreeElement->OnRenameRequested.BindSP(InlineWidget.Get(), &SInlineEditableTextBlock::EnterEditingMode);
@@ -570,6 +581,23 @@ bool SRigHierarchyTreeView::AddElement(const FRigBaseElement* InElement)
 
 								SRigHierarchyTagWidget::FArguments TagArguments;
 
+								static const FLinearColor BackgroundColor = FColor::FromHex(TEXT("#26BBFF"));
+								static const FLinearColor TextColor = FColor::FromHex(TEXT("#0F0F0F"));
+								static const FLinearColor IconColor = FColor::FromHex(TEXT("#1A1A1A"));
+
+								static const FSlateBrush* PrimaryBrush = FControlRigEditorStyle::Get().GetBrush("ControlRig.ConnectorPrimary");
+								static const FSlateBrush* SecondaryBrush = FControlRigEditorStyle::Get().GetBrush("ControlRig.ConnectorSecondary");
+								static const FSlateBrush* OptionalBrush = FControlRigEditorStyle::Get().GetBrush("ControlRig.ConnectorOptional");
+
+								const FSlateBrush* IconBrush = PrimaryBrush;
+								if(const FRigConnectorElement* ConnectorElement = Cast<FRigConnectorElement>(InElement))
+								{
+									if(ConnectorElement->Settings.Type == EConnectorType::Secondary)
+									{
+										IconBrush = ConnectorElement->Settings.bOptional ? OptionalBrush : SecondaryBrush;
+									}
+								}
+
 								FName Name = ConnectorKey.Name;
 								if (GetRigTreeDelegates().GetDisplaySettings().bUseShortName)
 								{
@@ -577,10 +605,11 @@ bool SRigHierarchyTreeView::AddElement(const FRigBaseElement* InElement)
 								}
 								TagArguments.Text(FText::FromName(Name));
 								TagArguments.TooltipText(FText::FromName(ConnectorKey.Name));
-								TagArguments.TextColor(FSlateColor(FLinearColor::White));
-
-								static const FLinearColor ConnectorColor = FControlRigEditorStyle::Get().ConnectorUserInterfaceColor;
-								TagArguments.Color(ConnectorColor);
+								TagArguments.Color(BackgroundColor);
+								TagArguments.IconColor(IconColor);
+								TagArguments.TextColor(TextColor);
+								TagArguments.Icon(IconBrush);
+								TagArguments.IconSize(FVector2d(16.f, 16.f));
 								TagArguments.AllowDragDrop(true);
 								FString Identifier;
 								FRigElementKey::StaticStruct()->ExportText(Identifier, &ConnectorKey, nullptr, nullptr, PPF_None, nullptr);
@@ -604,8 +633,6 @@ bool SRigHierarchyTreeView::AddElement(const FRigBaseElement* InElement)
 								}
 
 								TargetElementPtr->Get()->Tags.Add(TagArguments);
-
-								AddConnectorResolveWarningTag(*TargetElementPtr, InElement, Hierarchy);
 								return true;
 							}
 						}
@@ -1021,6 +1048,11 @@ void SRigHierarchyTreeView::AddConnectorResolveWarningTag(TSharedPtr<FRigTreeEle
 		TAttribute<FText> GetTooltipText = TAttribute<FText>::CreateSP(this,
 			&SRigHierarchyTreeView::GetConnectorWarningMessage, InTreeElement, ControlRigPtr, ConnectorKey);
 
+		static const FLinearColor BackgroundColor = FColor::FromHex(TEXT("#FFB800"));
+		static const FLinearColor TextColor = FColor::FromHex(TEXT("#0F0F0F"));
+		static const FLinearColor IconColor = FColor::FromHex(TEXT("#1A1A1A"));
+		static const FSlateBrush* WarningBrush = FControlRigEditorStyle::Get().GetBrush("ControlRig.ConnectorWarning");
+
 		SRigHierarchyTagWidget::FArguments TagArguments;
 		TagArguments.Visibility_Lambda([GetTooltipText]() -> EVisibility
 		{
@@ -1028,9 +1060,11 @@ void SRigHierarchyTreeView::AddConnectorResolveWarningTag(TSharedPtr<FRigTreeEle
 		});
 		TagArguments.Text(LOCTEXT("ConnectorWarningTagLabel", "Warning"));
 		TagArguments.ToolTipText(GetTooltipText);
-		TagArguments.TextColor(FSlateColor(FLinearColor::Black));
-		TagArguments.Color(FLinearColor(0.728f, 0.364f, 0.003f));
-		TagArguments.Icon(FAppStyle::GetBrush("AnimSlotManager.Warning"));
+		TagArguments.Color(BackgroundColor);
+		TagArguments.IconColor(IconColor);
+		TagArguments.TextColor(TextColor);
+		TagArguments.Icon(WarningBrush);
+		TagArguments.IconSize(FVector2d(16.f, 16.f));
 		InTreeElement->Tags.Add(TagArguments);
 	}
 }
@@ -1038,40 +1072,24 @@ void SRigHierarchyTreeView::AddConnectorResolveWarningTag(TSharedPtr<FRigTreeEle
 FText SRigHierarchyTreeView::GetConnectorWarningMessage(TSharedPtr<FRigTreeElement> InTreeElement,
 	TWeakObjectPtr<UControlRig> InControlRigPtr, const FRigElementKey InConnectorKey) const
 {
-	/*
 	if(UControlRig* ControlRig = InControlRigPtr.Get())
 	{
-		URigHierarchy* Hierarchy = ControlRig->GetHierarchy();
-				
-		static const FText NotResolvedWarning = LOCTEXT("ConnectorWarningConnectorNotResolved", "Connector is not resolved.");
-
-		FRigElementKeyRedirector& Redirector = ControlRig->GetElementKeyRedirector();
-		if(const FCachedRigElement* Cache = Redirector.Find(InConnectorKey))
+		if(const UControlRigBlueprint* ControlRigBlueprint = Cast<UControlRigBlueprint>(ControlRig->GetClass()->ClassGeneratedBy))
 		{
-			if(const_cast<FCachedRigElement*>(Cache)->UpdateCache(Hierarchy))
+			const FRigElementKey TargetKey = ControlRigBlueprint->ModularRigModel.Connections.FindTargetFromConnector(InConnectorKey);
+			if(TargetKey.IsValid())
 			{
-				if(const FRigConnectorElement* ConnectorElement = Hierarchy->Find<FRigConnectorElement>(InConnectorKey))
+				const URigHierarchy* Hierarchy = ControlRig->GetHierarchy();
+				if(Hierarchy->Contains(TargetKey))
 				{
-					const FRigConnectionInfo ConnectionInfo(&ControlRig->GetElementKeyRedirector(), Hierarchy);
-					FString FailureReason;
-					if(!ConnectorElement->CanConnect(&ConnectionInfo, &FailureReason))
-					{
-						return FText::FromString(FailureReason);
-					}
+					return FText();
 				}
 			}
-			else
-			{
-				return NotResolvedWarning;
-			}
-		}
-		else
-		{
-			return NotResolvedWarning;
 		}
 	}
-	*/
-	return FText();
+
+	static const FText NotResolvedWarning = LOCTEXT("ConnectorWarningConnectorNotResolved", "Connector is not resolved.");
+	return NotResolvedWarning;
 }
 
 bool SRigHierarchyItem::OnVerifyNameChanged(const FText& InText, FText& OutErrorMessage)
