ControlRig: Details panel to disallow editing of procedural sockets

#rnx

diff --git a/Engine/Plugins/Animation/ControlRig/Source/ControlRigEditor/Private/ControlRigElementDetails.cpp b/Engine/Plugins/Animation/ControlRig/Source/ControlRigEditor/Private/ControlRigElementDetails.cpp
--- a/Engine/Plugins/Animation/ControlRig/Source/ControlRigEditor/Private/ControlRigElementDetails.cpp
+++ b/Engine/Plugins/Animation/ControlRig/Source/ControlRigEditor/Private/ControlRigElementDetails.cpp
@@ -5305,6 +5305,8 @@ void FRigSocketElementDetails::CustomizeSettings(IDetailLayoutBuilder& DetailBui
 		return;
 	}
 
+	const bool bIsProcedural = IsAnyElementProcedural();
+	
 	IDetailCategoryBuilder& SettingsCategory = DetailBuilder.EditCategory(TEXT("Settings"), LOCTEXT("Settings", "Settings"));
 
 	SettingsCategory.AddCustomRow(FText::FromString(TEXT("Color")))
@@ -5317,6 +5319,7 @@ void FRigSocketElementDetails::CustomizeSettings(IDetailLayoutBuilder& DetailBui
 	.ValueContent()
 	[
 		SNew(SColorBlock)
+		.IsEnabled(!bIsProcedural)
 		//.Size(FVector2D(6.0, 38.0))
 		.Color(this, &FRigSocketElementDetails::GetSocketColor) 
 		.OnMouseButtonDown(this, &FRigSocketElementDetails::SetSocketColor)
@@ -5332,6 +5335,7 @@ void FRigSocketElementDetails::CustomizeSettings(IDetailLayoutBuilder& DetailBui
 	.ValueContent()
 	[
 		SNew(SEditableText)
+		.IsEnabled(!bIsProcedural)
 		.Font(IDetailLayoutBuilder::GetDetailFont())
 		.Text(this, &FRigSocketElementDetails::GetSocketDescription)
 		.OnTextCommitted(this, &FRigSocketElementDetails::SetSocketDescription)
