Control Rig: Schematic drop on group try to drop on any of the group children

#rnx

diff --git a/Engine/Plugins/Animation/ControlRig/Source/ControlRigEditor/Private/Editor/ControlRigSchematicModel.cpp b/Engine/Plugins/Animation/ControlRig/Source/ControlRigEditor/Private/Editor/ControlRigSchematicModel.cpp
--- a/Engine/Plugins/Animation/ControlRig/Source/ControlRigEditor/Private/Editor/ControlRigSchematicModel.cpp
+++ b/Engine/Plugins/Animation/ControlRig/Source/ControlRigEditor/Private/Editor/ControlRigSchematicModel.cpp
@@ -1423,12 +1423,32 @@ void FControlRigSchematicModel::HandleSchematicDrop(SSchematicGraphPanel* InPane
 		return;
 	}
 
-	const FControlRigSchematicRigElementKeyNode* ElementKeyNode =
-		Cast<FControlRigSchematicRigElementKeyNode>(InNode->GetNodeData());
-	if(ElementKeyNode == nullptr)
+	struct Local
 	{
-		return;
-	}
+		static void CollectTargetKeys(TArray<FRigElementKey>& OutKeys, const FSchematicGraphNode* InNode)
+		{
+			check(InNode);
+			
+			if(const FControlRigSchematicRigElementKeyNode* ElementKeyNode = Cast<FControlRigSchematicRigElementKeyNode>(InNode))
+			{
+				OutKeys.Add(ElementKeyNode->GetKey());
+			}
+
+			if(const FSchematicGraphGroupNode* GroupNode = Cast<FSchematicGraphGroupNode>(InNode))
+			{
+				for(int32 Index = 0; Index < GroupNode->GetNumChildNodes(); Index++)
+				{
+					if(const FSchematicGraphNode* ChildNode = GroupNode->GetChildNode(Index))
+					{
+						CollectTargetKeys(OutKeys, ChildNode);
+					}
+				}
+			}
+		}
+	};
+
+	TArray<FRigElementKey> TargetKeys;
+	Local::CollectTargetKeys(TargetKeys, InNode->GetNodeData());
 	
 	UControlRig* ControlRig = ControlRigBlueprint->GetDebuggedControlRig();
 	if (!ControlRig)
@@ -1442,13 +1462,6 @@ void FControlRigSchematicModel::HandleSchematicDrop(SSchematicGraphPanel* InPane
 		return;
 	}
 
-	const FRigElementKey& TargetKey = ElementKeyNode->GetKey();
-	FRigBaseElement* Target = Hierarchy->Find(TargetKey);
-	if (!Target)
-	{
-		return;
-	}
-
 	TSharedPtr<FAssetDragDropOp> AssetDragDropOp = InDragDropEvent.GetOperationAs<FAssetDragDropOp>();
 	TSharedPtr<FSchematicGraphNodeDragDropOp> SchematicDragDropOp = InDragDropEvent.GetOperationAs<FSchematicGraphNodeDragDropOp>();
 	TSharedPtr<FModularRigModuleDragDropOp> ModuleDragDropOperation = InDragDropEvent.GetOperationAs<FModularRigModuleDragDropOp>();
@@ -1464,7 +1477,7 @@ void FControlRigSchematicModel::HandleSchematicDrop(SSchematicGraphPanel* InPane
 
 			if(UControlRigBlueprint* AssetBlueprint = Cast<UControlRigBlueprint>(AssetData.GetAsset()))
 			{
-				FFunctionGraphTask::CreateAndDispatchWhenReady([this, AssetBlueprint, TargetKey]()
+				FFunctionGraphTask::CreateAndDispatchWhenReady([this, AssetBlueprint, TargetKeys]()
 				{
 					if (UModularRigController* Controller = ControlRigBlueprint->GetModularRigController())
 					{
@@ -1501,7 +1514,14 @@ void FControlRigSchematicModel::HandleSchematicDrop(SSchematicGraphPanel* InPane
 									}
 								}
 							}
-							Controller->ConnectConnectorToElement(PrimaryConnectorKey, TargetKey, true, ControlRig->GetModularRigSettings().bAutoResolve);
+
+							for(const FRigElementKey& TargetKey : TargetKeys)
+							{
+								if(Controller->ConnectConnectorToElement(PrimaryConnectorKey, TargetKey, true, ControlRig->GetModularRigSettings().bAutoResolve))
+								{
+									break;
+								}
+							}
 						}
 					}
 				}, TStatId(), NULL, ENamedThreads::GameThread);
@@ -1511,7 +1531,7 @@ void FControlRigSchematicModel::HandleSchematicDrop(SSchematicGraphPanel* InPane
 	else if(SchematicDragDropOp.IsValid())
 	{
 		const TArray<FGuid> Sources = SchematicDragDropOp->GetElements();
-		FFunctionGraphTask::CreateAndDispatchWhenReady([this, Sources, TargetKey]()
+		FFunctionGraphTask::CreateAndDispatchWhenReady([this, Sources, TargetKeys]()
 		{
 			UModularRig* ControlRig = Cast<UModularRig>(ControlRigBlueprint->GetDebuggedControlRig());
 			if (!ControlRig)
@@ -1533,7 +1553,13 @@ void FControlRigSchematicModel::HandleSchematicDrop(SSchematicGraphPanel* InPane
 					{
 						if (Hierarchy->Find<FRigConnectorElement>(Pair.Key))
 						{
-							Controller->ConnectConnectorToElement(Pair.Key, TargetKey, true, ControlRig->GetModularRigSettings().bAutoResolve);
+							for(const FRigElementKey& TargetKey : TargetKeys)
+							{
+								if(Controller->ConnectConnectorToElement(Pair.Key, TargetKey, true, ControlRig->GetModularRigSettings().bAutoResolve))
+								{
+									break;
+								}
+							}
 							break;
 						}
 					}
@@ -1544,7 +1570,7 @@ void FControlRigSchematicModel::HandleSchematicDrop(SSchematicGraphPanel* InPane
 	else if(ModuleDragDropOperation.IsValid())
 	{
 		const TArray<FString> Sources = ModuleDragDropOperation->GetElements();
-		FFunctionGraphTask::CreateAndDispatchWhenReady([this, Sources, TargetKey]()
+		FFunctionGraphTask::CreateAndDispatchWhenReady([this, Sources, TargetKeys]()
 		{
 			UModularRig* ControlRig = Cast<UModularRig>(ControlRigBlueprint->GetDebuggedControlRig());
 			if (!ControlRig)
@@ -1565,7 +1591,13 @@ void FControlRigSchematicModel::HandleSchematicDrop(SSchematicGraphPanel* InPane
 				{
 					if(const FRigConnectorElement* Connector = Hierarchy->Find<FRigConnectorElement>(FRigElementKey(*ModulePathOrConnectorName, ERigElementType::Connector)))
 					{
-						Controller->ConnectConnectorToElement(Connector->GetKey(), TargetKey, true, ControlRig->GetModularRigSettings().bAutoResolve);
+						for(const FRigElementKey& TargetKey : TargetKeys)
+						{
+							if(Controller->ConnectConnectorToElement(Connector->GetKey(), TargetKey, true, ControlRig->GetModularRigSettings().bAutoResolve))
+							{
+								break;
+							}
+						}
 						return;
 					}
 					if(const FRigModuleReference* Module = Controller->FindModule(ModulePathOrConnectorName))
@@ -1579,7 +1611,13 @@ void FControlRigSchematicModel::HandleSchematicDrop(SSchematicGraphPanel* InPane
 								if(Connector->IsPrimary())
 								{
 									const FRigElementKey ConnectorKey = Connector->GetKey();
-									Controller->ConnectConnectorToElement(ConnectorKey, TargetKey, true, ControlRig->GetModularRigSettings().bAutoResolve);
+									for(const FRigElementKey& TargetKey : TargetKeys)
+									{
+										if(Controller->ConnectConnectorToElement(ConnectorKey, TargetKey, true, ControlRig->GetModularRigSettings().bAutoResolve))
+										{
+											break;
+										}
+									}
 									return;
 								}
 							}
