Control Rig: Fix modular hierarchy colors
#rnx

diff --git a/Engine/Plugins/Animation/ControlRig/Source/ControlRigEditor/Private/Editor/SModularRigTreeView.cpp b/Engine/Plugins/Animation/ControlRig/Source/ControlRigEditor/Private/Editor/SModularRigTreeView.cpp
--- a/Engine/Plugins/Animation/ControlRig/Source/ControlRigEditor/Private/Editor/SModularRigTreeView.cpp
+++ b/Engine/Plugins/Animation/ControlRig/Source/ControlRigEditor/Private/Editor/SModularRigTreeView.cpp
@@ -982,7 +982,7 @@ const TSharedPtr<FModularRigTreeElement>* SModularRigTreeView::FindItemAtPositio
 TPair<const FSlateBrush*, FSlateColor> FModularRigTreeElement::GetBrushAndColor(const UModularRig* InModularRig)
 {
 	const FSlateBrush* Brush = nullptr;
-	FLinearColor Color = FSlateColor::UseForeground().GetColor(FWidgetStyle());
+	FLinearColor Color = FSlateColor(EStyleColor::Foreground).GetColor(FWidgetStyle());
 	float Opacity = 1.f;
 
 	if (const FRigModuleInstance* ConnectorModule = InModularRig->FindModule(ModulePath))
@@ -1024,7 +1024,8 @@ TPair<const FSlateBrush*, FSlateColor> FModularRigTreeElement::GetBrushAndColor(
 					bConnectionWarning = false;
 					if (!bIsConnected)
 					{
-						Opacity = 0.6;
+						Opacity = 0.7;
+						Color = FSlateColor(EStyleColor::Hover2).GetColor(FWidgetStyle());
 					}
 					Brush = FControlRigEditorStyle::Get().GetBrush("ControlRig.ConnectorOptional");
 				}
@@ -1037,7 +1038,7 @@ TPair<const FSlateBrush*, FSlateColor> FModularRigTreeElement::GetBrushAndColor(
 
 		if (bConnectionWarning)
 		{
-			Color = FLinearColor(FColor::FromHex(TEXT("#FFB800")));
+			Color = FSlateColor(EStyleColor::Warning).GetColor(FWidgetStyle());
 		}
 	}
 	if (!Brush)
