Control Rig: Fix profile time not showing for forward solve when post forward solve is present
#rnx

diff --git a/Engine/Plugins/Animation/ControlRig/Source/ControlRig/Private/ControlRig.cpp b/Engine/Plugins/Animation/ControlRig/Source/ControlRig/Private/ControlRig.cpp
--- a/Engine/Plugins/Animation/ControlRig/Source/ControlRig/Private/ControlRig.cpp
+++ b/Engine/Plugins/Animation/ControlRig/Source/ControlRig/Private/ControlRig.cpp
@@ -690,7 +690,10 @@ bool UControlRig::Execute(const FName& InEventName)
 
 	// Only top-level rigs should execute this function.
 	// Rig modules/nested rigs should run through Execute_Internal
-	ensureMsgf(GetTypedOuter<UControlRig>() == nullptr, TEXT("UControlRig::Execute running from a nested rig in %s"), *GetPackage()->GetPathName());
+	if (InEventName != FRigUnit_PrepareForExecution::EventName)
+	{
+		ensureMsgf(GetTypedOuter<UControlRig>() == nullptr, TEXT("UControlRig::Execute running from a nested rig in %s"), *GetPackage()->GetPathName());
+	}
 	ensureMsgf(InEventName != FRigUnit_PreBeginExecution::EventName &&
 						InEventName != FRigUnit_PostBeginExecution::EventName, TEXT("Requested execution of invalid event %s on top level rig in %s"), *InEventName.ToString(), *GetPackage()->GetPathName());
 
@@ -739,17 +742,6 @@ bool UControlRig::Execute(const FName& InEventName)
 	PublicContext.bDebugExecution = bDebugExecutionEnabled;
 #endif
 
-	if (VM)
-	{
-#if WITH_EDITOR
-		// default to always clear data after each execution
-		// only set a valid first entry event later when execution
-		// has passed the initialization stage and there are multiple events present in one evaluation
-		// first entry event is used to determined when to clear data during an evaluation
-		SetFirstEntryEventInEventQueue(ExtendedExecuteContext, NAME_None);
-#endif
-	}
-
 #if WITH_EDITOR
 	ExtendedExecuteContext.SetInstructionVisitInfo(&InstructionVisitInfo);
 
@@ -1077,11 +1069,10 @@ bool UControlRig::Execute(const FName& InEventName)
 	else
 	{
 #if WITH_EDITOR
-		// only set a valid first entry event when execution
-		// has passed the initialization stage and there are multiple events present
-		if (EventQueueToRun.Num() >= 2 && VM)
+		// only set a valid first entry event when none has been set
+		if (InstructionVisitInfo.GetFirstEntryEventInEventQueue() == NAME_None && !EventQueueToRun.IsEmpty() && VM)
 		{
-			SetFirstEntryEventInEventQueue(ExtendedExecuteContext, EventQueueToRun[0]);
+			InstructionVisitInfo.SetFirstEntryEventInEventQueue(EventQueueToRun[0]);
 		}
 
 		// Transform Overrride is generated using a Transient Control 
diff --git a/Engine/Plugins/Animation/ControlRig/Source/ControlRig/Private/ModularRig.cpp b/Engine/Plugins/Animation/ControlRig/Source/ControlRig/Private/ModularRig.cpp
--- a/Engine/Plugins/Animation/ControlRig/Source/ControlRig/Private/ModularRig.cpp
+++ b/Engine/Plugins/Animation/ControlRig/Source/ControlRig/Private/ModularRig.cpp
@@ -425,6 +425,10 @@ void UModularRig::ExecuteQueue()
 	FRigVMExtendedExecuteContext& Context = GetRigVMExtendedExecuteContext();
 	FControlRigExecuteContext& PublicContext = Context.GetPublicDataSafe<FControlRigExecuteContext>();
 	URigHierarchy* Hierarchy = GetHierarchy();
+
+#if WITH_EDITOR
+	TMap<FRigModuleInstance*, FFirstEntryEventGuard> FirstModuleEvent;
+#endif
 	
 	while(ExecutionQueue.IsValidIndex(ExecutionQueueFront))
 	{
@@ -467,6 +471,14 @@ void UModularRig::ExecuteQueue()
 				RigPublicContext.SetToWorldSpaceTransform(PublicContext.GetToWorldSpaceTransform());
 				RigPublicContext.RuntimeSettings = PublicContext.RuntimeSettings;
 
+#if WITH_EDITOR
+				if (!FirstModuleEvent.Contains(ModuleInstance))
+				{
+					//ModuleRig->InstructionVisitInfo.FirstEntryEventInQueue = ExecutionElement.EventName;
+					FirstModuleEvent.Add(ModuleInstance, FFirstEntryEventGuard(&ModuleRig->InstructionVisitInfo, ExecutionElement.EventName));
+				}
+#endif
+
 				// re-initialize the module in case only the VM side got recompiled.
 				// this happens when the user relies on auto recompilation when editing the
 				// module (dependency) graph - by changing a value, add / remove nodes or links.
diff --git a/Engine/Plugins/Animation/ControlRig/Source/ControlRig/Private/Rigs/AdditiveControlRig.cpp b/Engine/Plugins/Animation/ControlRig/Source/ControlRig/Private/Rigs/AdditiveControlRig.cpp
--- a/Engine/Plugins/Animation/ControlRig/Source/ControlRig/Private/Rigs/AdditiveControlRig.cpp
+++ b/Engine/Plugins/Animation/ControlRig/Source/ControlRig/Private/Rigs/AdditiveControlRig.cpp
@@ -20,13 +20,6 @@ UAdditiveControlRig::UAdditiveControlRig(const FObjectInitializer& ObjectInitial
 	bResetInitialTransformsBeforeConstruction = false;
 }
 
-#if WITH_EDITOR
-void UAdditiveControlRig::SetFirstEntryEventInEventQueue(FRigVMExtendedExecuteContext& Context, const FName& InFirstEventName)
-{
-	VM->SetFirstEntryEventInEventQueue(Context, NAME_None);
-}
-#endif
-
 FName UAdditiveControlRig::GetControlName(const FName& InBoneName)
 {
 	if (InBoneName != NAME_None)
diff --git a/Engine/Plugins/Animation/ControlRig/Source/ControlRig/Private/Rigs/FKControlRig.cpp b/Engine/Plugins/Animation/ControlRig/Source/ControlRig/Private/Rigs/FKControlRig.cpp
--- a/Engine/Plugins/Animation/ControlRig/Source/ControlRig/Private/Rigs/FKControlRig.cpp
+++ b/Engine/Plugins/Animation/ControlRig/Source/ControlRig/Private/Rigs/FKControlRig.cpp
@@ -260,13 +260,6 @@ bool UFKControlRig::Execute_Internal(const FName& InEventName)
 	return true;
 }
 
-#if WITH_EDITOR
-void UFKControlRig::SetFirstEntryEventInEventQueue(FRigVMExtendedExecuteContext& Context, const FName& InFirstEventName)
-{
-	VM->SetFirstEntryEventInEventQueue(Context, NAME_None);
-}
-#endif
-
 void UFKControlRig::SetBoneInitialTransformsFromSkeletalMeshComponent(USkeletalMeshComponent* InSkelMeshComp,
                                                                       bool bUseAnimInstance)
 {
diff --git a/Engine/Plugins/Animation/ControlRig/Source/ControlRig/Public/ControlRig.h b/Engine/Plugins/Animation/ControlRig/Source/ControlRig/Public/ControlRig.h
--- a/Engine/Plugins/Animation/ControlRig/Source/ControlRig/Public/ControlRig.h
+++ b/Engine/Plugins/Animation/ControlRig/Source/ControlRig/Public/ControlRig.h
@@ -629,8 +629,6 @@ protected:
 
 
 #if WITH_EDITOR
-	virtual void SetFirstEntryEventInEventQueue(FRigVMExtendedExecuteContext& Context, const FName& InFirstEventName) { VM->SetFirstEntryEventInEventQueue(Context, NAME_None); }
-
 	static void OnHierarchyTransformUndoRedoWeak(URigHierarchy* InHierarchy, const FRigElementKey& InKey, ERigTransformType::Type InTransformType, const FTransform& InTransform, bool bIsUndo, TWeakObjectPtr<UControlRig> WeakThis)
 	{
 		if(WeakThis.IsValid() && InHierarchy != nullptr)
diff --git a/Engine/Plugins/Animation/ControlRig/Source/ControlRig/Public/Rigs/AdditiveControlRig.h b/Engine/Plugins/Animation/ControlRig/Source/ControlRig/Public/Rigs/AdditiveControlRig.h
--- a/Engine/Plugins/Animation/ControlRig/Source/ControlRig/Public/Rigs/AdditiveControlRig.h
+++ b/Engine/Plugins/Animation/ControlRig/Source/ControlRig/Public/Rigs/AdditiveControlRig.h
@@ -23,9 +23,6 @@ public:
 	virtual void InitializeVMsFromCDO() override { URigVMHost::InitializeFromCDO(); }
 	virtual void RequestInitVMs() override { URigVMHost::RequestInit(); }
 	virtual bool Execute_Internal(const FName& InEventName) override;
-#if WITH_EDITOR
-	virtual void SetFirstEntryEventInEventQueue(FRigVMExtendedExecuteContext& Context, const FName& InFirstEventName) override;
-#endif
 	// END ControlRig
 
 	// utility function to 
diff --git a/Engine/Plugins/Animation/ControlRig/Source/ControlRig/Public/Rigs/FKControlRig.h b/Engine/Plugins/Animation/ControlRig/Source/ControlRig/Public/Rigs/FKControlRig.h
--- a/Engine/Plugins/Animation/ControlRig/Source/ControlRig/Public/Rigs/FKControlRig.h
+++ b/Engine/Plugins/Animation/ControlRig/Source/ControlRig/Public/Rigs/FKControlRig.h
@@ -49,9 +49,6 @@ public:
 	virtual void InitializeVMsFromCDO() override { URigVMHost::InitializeFromCDO(); }
 	virtual void RequestInitVMs() override { URigVMHost::RequestInit(); }
 	virtual bool Execute_Internal(const FName& InEventName) override;
-#if WITH_EDITOR
-	virtual void SetFirstEntryEventInEventQueue(FRigVMExtendedExecuteContext& Context, const FName& InFirstEventName) override;
-#endif
 	virtual void SetBoneInitialTransformsFromSkeletalMeshComponent(USkeletalMeshComponent* InSkelMeshComp, bool bUseAnimInstance = false) override;
 	// END ControlRig
 
diff --git a/Engine/Plugins/Runtime/RigVM/Source/RigVM/Private/RigVMHost.cpp b/Engine/Plugins/Runtime/RigVM/Source/RigVM/Private/RigVMHost.cpp
--- a/Engine/Plugins/Runtime/RigVM/Source/RigVM/Private/RigVMHost.cpp
+++ b/Engine/Plugins/Runtime/RigVM/Source/RigVM/Private/RigVMHost.cpp
@@ -331,6 +331,15 @@ void URigVMHost::Evaluate_AnyThread()
 			EventQueueToRun.Insert(Pair.Key, Pair.Value);
 		}
 	}
+
+#if WITH_EDITOR
+	FName FirstEvent = NAME_None;
+	if (!EventQueueToRun.IsEmpty())
+	{
+		FirstEvent = EventQueueToRun[0];
+	}
+	FFirstEntryEventGuard FirstEntryEventGuard(&InstructionVisitInfo, FirstEvent);
+#endif
 	
 	for (const FName& EventName : EventQueueToRun)
 	{
@@ -540,18 +549,6 @@ bool URigVMHost::Execute(const FName& InEventName)
 	}
 #endif
 
-	if (VM)
-	{
-#if WITH_EDITOR
-		// default to always clear data after each execution
-		// only set a valid first entry event later when execution
-		// has passed the initialization stage and there are multiple events present in one evaluation
-		// first entry event is used to determined when to clear data during an evaluation
-		VM->SetFirstEntryEventInEventQueue(GetRigVMExtendedExecuteContext(), NAME_None);
-		VM->SetFirstEntryEventInEventQueue(GetRigVMExtendedExecuteContext(), NAME_None);
-#endif
-	}
-
 	// setup the draw interface for debug drawing
 	if(!bIsEventInQueue || bIsEventFirstInQueue)
 	{
diff --git a/Engine/Plugins/Runtime/RigVM/Source/RigVM/Public/RigVMCore/RigVM.h b/Engine/Plugins/Runtime/RigVM/Source/RigVM/Public/RigVMCore/RigVM.h
--- a/Engine/Plugins/Runtime/RigVM/Source/RigVM/Public/RigVMCore/RigVM.h
+++ b/Engine/Plugins/Runtime/RigVM/Source/RigVM/Public/RigVMCore/RigVM.h
@@ -497,17 +497,6 @@ public:
 		return TArray<int32>();
 	}
 
-	UE_DEPRECATED(5.3, "Please, use SetFirstEntryEventInEventQueue with Context param")
-	const void SetFirstEntryEventInEventQueue(const FName& InFirstEventName) {}
-
-	const void SetFirstEntryEventInEventQueue(FRigVMExtendedExecuteContext& Context, const FName& InFirstEventName) 
-	{ 
-		if (FRigVMInstructionVisitInfo* InstructionVisitInfo = Context.GetRigVMInstructionVisitInfo())
-		{
-			InstructionVisitInfo->SetFirstEntryEventInEventQueue(InFirstEventName);
-		}
-	}
-
 	bool ResumeExecution(FRigVMExtendedExecuteContext& Context, const FName& InEntryName = NAME_None);
 
 	UE_DEPRECATED(5.3, "Please, use ResumeExecution with Context param")
diff --git a/Engine/Plugins/Runtime/RigVM/Source/RigVM/Public/RigVMCore/RigVMProfilingInfo.h b/Engine/Plugins/Runtime/RigVM/Source/RigVM/Public/RigVMCore/RigVMProfilingInfo.h
--- a/Engine/Plugins/Runtime/RigVM/Source/RigVM/Public/RigVMCore/RigVMProfilingInfo.h
+++ b/Engine/Plugins/Runtime/RigVM/Source/RigVM/Public/RigVMCore/RigVMProfilingInfo.h
@@ -27,7 +27,13 @@ struct RIGVM_API FRigVMInstructionVisitInfo
 
 	inline void ResetInstructionVisitedDuringLastRun(int32 NewSize = 0) { InstructionVisitedDuringLastRun.Reset(NewSize); }
 	inline void SetNumInstructionVisitedDuringLastRunZeroed(int32 Num) { InstructionVisitedDuringLastRun.SetNumZeroed(Num); }
-	inline void SetInstructionVisitedDuringLastRun(int32 InstructionIndex) { InstructionVisitedDuringLastRun[InstructionIndex]++; }
+	inline void SetInstructionVisitedDuringLastRun(int32 InstructionIndex)
+	{
+		if (InstructionVisitedDuringLastRun.IsValidIndex(InstructionIndex))
+		{
+			InstructionVisitedDuringLastRun[InstructionIndex]++;
+		}
+	}
 	inline int32 GetInstructionVisitedCountDuringLastRun(int32 InstructionIndex) const { return InstructionVisitedDuringLastRun.IsValidIndex(InstructionIndex) ? InstructionVisitedDuringLastRun[InstructionIndex] : 0; }
 	inline const TArray<int32>& GetInstructionVisitedCountDuringLastRun() const { return InstructionVisitedDuringLastRun; }
 
@@ -49,6 +55,28 @@ private:
 	// A RigVMHost can run multiple events per evaluation, such as the Backward&Forward Solve Mode,
 	// store the first event such that we know when to reset data for a new round of rig evaluation
 	FName FirstEntryEventInQueue = NAME_None;
+
+	friend class URigVMHost;
+	friend struct FFirstEntryEventGuard;
+};
+
+struct RIGVM_API FFirstEntryEventGuard
+{
+public:
+	FFirstEntryEventGuard(FRigVMInstructionVisitInfo* InVisitInfo, const FName& InFirstEvent)
+		: VisitInfo(InVisitInfo)
+	{
+		OldEntry = VisitInfo->FirstEntryEventInQueue;
+		VisitInfo->FirstEntryEventInQueue = InFirstEvent;
+	}
+
+	~FFirstEntryEventGuard()
+	{
+		VisitInfo->FirstEntryEventInQueue = OldEntry;
+	}
+
+	FName OldEntry;
+	FRigVMInstructionVisitInfo* VisitInfo;
 };
 
 USTRUCT()
diff --git a/Engine/Plugins/Runtime/RigVM/Source/RigVM/Public/RigVMDefines.h b/Engine/Plugins/Runtime/RigVM/Source/RigVM/Public/RigVMDefines.h
--- a/Engine/Plugins/Runtime/RigVM/Source/RigVM/Public/RigVMDefines.h
+++ b/Engine/Plugins/Runtime/RigVM/Source/RigVM/Public/RigVMDefines.h
@@ -23,7 +23,7 @@
 #endif
 
 #ifndef UE_RIGVM_DEBUG_TYPEINDEX
-#define UE_RIGVM_DEBUG_TYPEINDEX 0
+#define UE_RIGVM_DEBUG_TYPEINDEX 1
 #endif
 
 #ifndef UE_RIGVM_DEBUG_EXECUTION
