Control Rig: Don't show warning tags on optional connectors

#rnx

diff --git a/Engine/Plugins/Animation/ControlRig/Source/ControlRigEditor/Private/Editor/SRigHierarchyTreeView.cpp b/Engine/Plugins/Animation/ControlRig/Source/ControlRigEditor/Private/Editor/SRigHierarchyTreeView.cpp
--- a/Engine/Plugins/Animation/ControlRig/Source/ControlRigEditor/Private/Editor/SRigHierarchyTreeView.cpp
+++ b/Engine/Plugins/Animation/ControlRig/Source/ControlRigEditor/Private/Editor/SRigHierarchyTreeView.cpp
@@ -1040,6 +1040,14 @@ void SRigHierarchyTreeView::AddConnectorResolveWarningTag(TSharedPtr<FRigTreeEle
 	check(InRigElement);
 	check(InRigElement->GetType() == ERigElementType::Connector);
 
+	if(const FRigConnectorElement* ConnectorElement = Cast<FRigConnectorElement>(InRigElement))
+	{
+		if(ConnectorElement->IsOptional())
+		{
+			return;
+		}
+	}
+
 	if(UControlRig* ControlRig = InHierarchy->GetTypedOuter<UControlRig>())
 	{
 		TWeakObjectPtr<UControlRig> ControlRigPtr(ControlRig);
