Control Rig: Improve module updates when changing config settings

#rnx

diff --git a/Engine/Plugins/Animation/ControlRig/Source/ControlRigDeveloper/Private/ControlRigBlueprint.cpp b/Engine/Plugins/Animation/ControlRig/Source/ControlRigDeveloper/Private/ControlRigBlueprint.cpp
--- a/Engine/Plugins/Animation/ControlRig/Source/ControlRigDeveloper/Private/ControlRigBlueprint.cpp
+++ b/Engine/Plugins/Animation/ControlRig/Source/ControlRigDeveloper/Private/ControlRigBlueprint.cpp
@@ -2326,6 +2326,11 @@ void UControlRigBlueprint::OnModularDependencyVMCompiled(UObject* InBlueprint, U
 }
 
 void UControlRigBlueprint::OnModularDependencyChanged(URigVMBlueprint* InBlueprint)
+{
+	RequestConstructionOnAllModules();
+}
+
+void UControlRigBlueprint::RequestConstructionOnAllModules()
 {
 	// the rig will perform initialize itself - but we should request construction
 	check(IsModularRig());
@@ -2664,6 +2669,17 @@ void UControlRigBlueprint::HandleRigModulesModified(EModularRigNotification InNo
 			bRecompile = false;
 			break;
 		}
+		case EModularRigNotification::ModuleConfigValueChanged:
+		{
+			bRecompile = false;
+			PropagateModuleHierarchyFromBPToInstances();
+			RequestConstructionOnAllModules();
+			break;
+		}
+		default:
+		{
+			break;
+		}
 	}
 
 	if (bRecompile)
diff --git a/Engine/Plugins/Animation/ControlRig/Source/ControlRigDeveloper/Public/ControlRigBlueprint.h b/Engine/Plugins/Animation/ControlRig/Source/ControlRigDeveloper/Public/ControlRigBlueprint.h
--- a/Engine/Plugins/Animation/ControlRig/Source/ControlRigDeveloper/Public/ControlRigBlueprint.h
+++ b/Engine/Plugins/Animation/ControlRig/Source/ControlRigDeveloper/Public/ControlRigBlueprint.h
@@ -306,6 +306,7 @@ public:
 	void UpdateModularDependencyDelegates();
 	void OnModularDependencyVMCompiled(UObject* InBlueprint, URigVM* InVM, FRigVMExtendedExecuteContext& InExecuteContext);
 	void OnModularDependencyChanged(URigVMBlueprint* InBlueprint);
+	void RequestConstructionOnAllModules();
 	void RefreshModuleConnectors();
 	void RefreshModuleConnectors(const FRigModuleReference* InModule);
 
