Control Rig: Report error when trying to change parent fails
#rnx

diff --git a/Engine/Plugins/Animation/ControlRig/Source/ControlRigEditor/Private/EditMode/ControlRigEditMode.cpp b/Engine/Plugins/Animation/ControlRig/Source/ControlRigEditor/Private/EditMode/ControlRigEditMode.cpp
--- a/Engine/Plugins/Animation/ControlRig/Source/ControlRigEditor/Private/EditMode/ControlRigEditMode.cpp
+++ b/Engine/Plugins/Animation/ControlRig/Source/ControlRigEditor/Private/EditMode/ControlRigEditMode.cpp
@@ -3072,8 +3072,22 @@ void FControlRigEditMode::OpenSpacePickerWidget()
 			{
 				const FTransform Transform = InHierarchy->GetGlobalTransform(InControlKey);
 				URigHierarchy::TElementDependencyMap Dependencies = InHierarchy->GetDependenciesForVM(RuntimeRig->GetVM());
-				InHierarchy->SwitchToParent(InControlKey, InSpaceKey, false, true, Dependencies, nullptr);
-				InHierarchy->SetGlobalTransform(InControlKey, Transform);
+				FString OutFailureReason;
+				if (InHierarchy->SwitchToParent(InControlKey, InSpaceKey, false, true, Dependencies, &OutFailureReason))
+				{
+					InHierarchy->SetGlobalTransform(InControlKey, Transform);
+				}
+				else
+				{
+					if(URigHierarchyController* Controller = InHierarchy->GetController())
+					{
+						static constexpr TCHAR MessageFormat[] = TEXT("Could not switch %s to parent %s: %s");
+						Controller->ReportAndNotifyErrorf(MessageFormat,
+							*InControlKey.Name.ToString(),
+							*InSpaceKey.Name.ToString(),
+							*OutFailureReason);
+					}
+				}
 			}
 		}
 		
