Control Rig: When dragging onto secondary targets, make them smaller for visual differentiation

#rnx

diff --git a/Engine/Plugins/Animation/ControlRig/Source/ControlRigEditor/Private/Editor/ControlRigSchematicModel.cpp b/Engine/Plugins/Animation/ControlRig/Source/ControlRigEditor/Private/Editor/ControlRigSchematicModel.cpp
--- a/Engine/Plugins/Animation/ControlRig/Source/ControlRigEditor/Private/Editor/ControlRigSchematicModel.cpp
+++ b/Engine/Plugins/Animation/ControlRig/Source/ControlRigEditor/Private/Editor/ControlRigSchematicModel.cpp
@@ -1751,7 +1751,8 @@ void FControlRigSchematicModel::OnShowCandidatesForMatches(const FModularRigReso
 		// Create a temporary node that will be active only while this drag operation exists
 		if (!ContainsElementKeyNode(Match.GetKey()))
 		{
-			const FSchematicGraphNode* NewNode = AddElementKeyNode(Match.GetKey());
+			FSchematicGraphNode* NewNode = AddElementKeyNode(Match.GetKey());
+			NewNode->SetScaleOffset(0.6f);
 			TemporaryNodeGuids.Add(NewNode->GetGuid());
 		}
 	}
