Sequencer: Anim Details: fixed crash due to the CR Editor getting a reference to the current sequencer whereas it shouldn't

#rnx

diff --git a/Engine/Plugins/Animation/ControlRig/Source/ControlRigEditor/Private/EditMode/ControlRigEditMode.cpp b/Engine/Plugins/Animation/ControlRig/Source/ControlRigEditor/Private/EditMode/ControlRigEditMode.cpp
--- a/Engine/Plugins/Animation/ControlRig/Source/ControlRigEditor/Private/EditMode/ControlRigEditMode.cpp
+++ b/Engine/Plugins/Animation/ControlRig/Source/ControlRigEditor/Private/EditMode/ControlRigEditMode.cpp
@@ -391,9 +391,12 @@ void FControlRigEditMode::Enter()
 	LastMovieSceneSig = FGuid();
 	if (UsesToolkits())
 	{
-		if (WeakSequencer.IsValid() == false)
+		if (!AreEditingControlRigDirectly())
 		{
-			SetSequencer(FBakingHelper::GetSequencer());
+			if (WeakSequencer.IsValid() == false)
+			{
+				SetSequencer(FBakingHelper::GetSequencer());
+			}
 		}
 		if (!Toolkit.IsValid())
 		{
