Control Rig: Offer re-resolve functionality in modular rig tree

#rnx

diff --git a/Engine/Plugins/Animation/ControlRig/Source/ControlRig/Private/ModularRigController.cpp b/Engine/Plugins/Animation/ControlRig/Source/ControlRig/Private/ModularRigController.cpp
--- a/Engine/Plugins/Animation/ControlRig/Source/ControlRig/Private/ModularRigController.cpp
+++ b/Engine/Plugins/Animation/ControlRig/Source/ControlRig/Private/ModularRigController.cpp
@@ -263,7 +263,7 @@ bool UModularRigController::ConnectConnectorToElement(const FRigElementKey& InCo
 	if (CurrentTarget.IsValid())
 	{
 		const TGuardValue<bool> DisableAutomaticReparenting(bAutomaticReparenting, false);
-		DisconnectConnector(InConnectorKey, bSetupUndo);
+		DisconnectConnector(InConnectorKey, false, bSetupUndo);
 	}
 	
 	Model->Connections.AddConnection(InConnectorKey, InTargetKey);
@@ -278,38 +278,11 @@ bool UModularRigController::ConnectConnectorToElement(const FRigElementKey& InCo
 			{
 				if (bAutoResolveOtherConnectors)
 				{
-					UModularRigRuleManager* RuleManager = Hierarchy->GetRuleManager();
-					const FRigModuleInstance* ModuleInstance = ModularRig->FindModule(Module->GetPath());
-
-					for (const FRigModuleConnector& OtherConnector : RigCDO->GetRigModuleSettings().ExposedConnectors)
+					if(const FRigConnectorElement* PrimaryConnector = Module->FindPrimaryConnector(Hierarchy))
 					{
-						FRigElementKey OtherConnectorKey;
-						OtherConnectorKey.Name = *FString::Printf(TEXT("%s%s"), *Module->GetNamespace(), *OtherConnector.Name);
-						OtherConnectorKey.Type = ERigElementType::Connector;
-						if (!Model->Connections.HasConnection(OtherConnectorKey))
+						if(PrimaryConnector->GetKey() == InConnectorKey)
 						{
-							if (const FRigConnectorElement* OtherConnectorElement = Cast<FRigConnectorElement>(Hierarchy->Find(OtherConnectorKey)))
-							{
-								FModularRigResolveResult RuleResults = RuleManager->FindMatches(OtherConnectorElement, ModuleInstance, ModularRig->GetElementKeyRedirector());
-
-								if (RuleResults.GetMatches().Num() == 1)
-								{
-									Model->Connections.AddConnection(OtherConnectorKey, RuleResults.GetMatches()[0].GetKey());
-									Notify(EModularRigNotification::ConnectionChanged, Module);
-								}
-								else
-								{
-									for (const FRigElementResolveResult& Result : RuleResults.GetMatches())
-									{
-										if (Result.GetState() == ERigElementResolveState::DefaultTarget)
-										{
-											Model->Connections.AddConnection(OtherConnectorKey, Result.GetKey());
-											Notify(EModularRigNotification::ConnectionChanged, Module);
-											break;
-										}
-									}
-								}
-							}
+							(void)AutoConnectModules( {Module->GetPath()}, false, bSetupUndo);
 						}
 					}
 				}
@@ -342,7 +315,7 @@ bool UModularRigController::ConnectConnectorToElement(const FRigElementKey& InCo
 	return true;
 }
 
-bool UModularRigController::DisconnectConnector(const FRigElementKey& InConnectorKey, bool bSetupUndo)
+bool UModularRigController::DisconnectConnector(const FRigElementKey& InConnectorKey, bool bDisconnectSubModules, bool bSetupUndo)
 {
 	FString ConnectorModulePath, ConnectorName;
 	if (!URigHierarchy::SplitNameSpace(InConnectorKey.Name.ToString(), &ConnectorModulePath, &ConnectorName))
@@ -410,7 +383,7 @@ bool UModularRigController::DisconnectConnector(const FRigElementKey& InConnecto
 			Model->Connections.RemoveConnection(ToRemove);
 		}
 	}
-	else if (!ModuleConnector->IsOptional())
+	else if (!ModuleConnector->IsOptional() && bDisconnectSubModules)
 	{
 		// Remove connections from child modules
 		TArray<FRigElementKey> ConnectionsToRemove;
@@ -496,7 +469,7 @@ TArray<FRigElementKey> UModularRigController::DisconnectCyclicConnectors(bool bS
 
 	for(const FRigElementKey& ConnectorToDisconnect : ConnectorsToDisconnect)
 	{
-		if(DisconnectConnector(ConnectorToDisconnect, bSetupUndo))
+		if(DisconnectConnector(ConnectorToDisconnect, false, bSetupUndo))
 		{
 			DisconnectedConnectors.Add(ConnectorToDisconnect);
 		}
@@ -506,6 +479,210 @@ TArray<FRigElementKey> UModularRigController::DisconnectCyclicConnectors(bool bS
 	return DisconnectedConnectors;
 }
 
+bool UModularRigController::AutoConnectSecondaryConnectors(const TArray<FRigElementKey>& InConnectorKeys, bool bReplaceExistingConnections, bool bSetupUndo)
+{
+#if WITH_EDITOR
+
+	UBlueprint* Blueprint = Cast<UBlueprint>(GetOuter());
+	if(Blueprint == nullptr)
+	{
+		UE_LOG(LogControlRig, Error, TEXT("ModularRigController is not nested under blueprint."));
+		return false;
+	}
+
+	const UModularRig* ModularRig = Cast<UModularRig>(Blueprint->GetObjectBeingDebugged());
+	if (!ModularRig)
+	{
+		UE_LOG(LogControlRig, Error, TEXT("Could not find debugged modular rig in %s"), *Blueprint->GetPathName());
+		return false;
+	}
+	
+	URigHierarchy* Hierarchy = ModularRig->GetHierarchy();
+	if (!Hierarchy)
+	{
+		UE_LOG(LogControlRig, Error, TEXT("Could not find hierarchy in %s"), *ModularRig->GetPathName());
+		return false;
+	}
+
+	for(const FRigElementKey& ConnectorKey : InConnectorKeys)
+	{
+		if(ConnectorKey.Type != ERigElementType::Connector)
+		{
+			UE_LOG(LogControlRig, Error, TEXT("Could not find debugged modular rig in %s"), *Blueprint->GetPathName());
+			return false;
+		}
+		const FRigConnectorElement* Connector = Hierarchy->Find<FRigConnectorElement>(ConnectorKey);
+		if(Connector == nullptr)
+		{
+			UE_LOG(LogControlRig, Error, TEXT("Cannot find connector %s in %s"), *ConnectorKey.ToString(), *Blueprint->GetPathName());
+			return false;
+		}
+		if(Connector->IsPrimary())
+		{
+			UE_LOG(LogControlRig, Warning, TEXT("Provided connector %s in %s is a primary connector. It will be skipped during auto resolval."), *ConnectorKey.ToString(), *Blueprint->GetPathName());
+		}
+	}
+
+	TSharedPtr<FScopedTransaction> TransactionPtr;
+	if (bSetupUndo)
+	{
+		TransactionPtr = MakeShared<FScopedTransaction>(NSLOCTEXT("ModularRigController", "AutoResolveSecondaryConnectors", "Auto-Resolve Connectors"));
+	}
+
+	Blueprint->Modify();
+
+	bool bResolvedAllConnectors = true;
+	for(const FRigElementKey& ConnectorKey : InConnectorKeys)
+	{
+		const FString ModulePath = Hierarchy->GetModulePath(ConnectorKey);
+		if(ModulePath.IsEmpty())
+		{
+			UE_LOG(LogControlRig, Error, TEXT("Connector %s has no associated module path"), *ConnectorKey.ToString());
+			bResolvedAllConnectors = false;
+			continue;
+		}
+
+		const FRigModuleReference* Module = Model->FindModule(ModulePath);
+		if(Module == nullptr)
+		{
+			UE_LOG(LogControlRig, Error, TEXT("Could not find module %s"), *ModulePath);
+			bResolvedAllConnectors = false;
+			continue;
+		}
+
+		const FRigConnectorElement* PrimaryConnector = Module->FindPrimaryConnector(Hierarchy);
+		if(PrimaryConnector == nullptr)
+		{
+			UE_LOG(LogControlRig, Error, TEXT("Module %s has no primary connector"), *ModulePath);
+			bResolvedAllConnectors = false;
+			continue;
+		}
+		
+		const FRigElementKey PrimaryConnectorKey = PrimaryConnector->GetKey();
+		if(ConnectorKey == PrimaryConnectorKey)
+		{
+			// silently skip primary connectors
+			continue;
+		}
+		
+		if(!Model->Connections.HasConnection(PrimaryConnectorKey))
+		{
+			UE_LOG(LogControlRig, Warning, TEXT("Module %s's primary connector is not resolved"), *ModulePath);
+			bResolvedAllConnectors = false;
+			continue;
+		}
+		
+		const UControlRig* RigCDO = Module->Class->GetDefaultObject<UControlRig>();
+		if(RigCDO == nullptr)
+		{
+			UE_LOG(LogControlRig, Error, TEXT("Module %s has no default rig assigned"), *ModulePath);
+			bResolvedAllConnectors = false;
+			continue;
+		}
+
+		const UModularRigRuleManager* RuleManager = Hierarchy->GetRuleManager();
+		const FRigModuleInstance* ModuleInstance = ModularRig->FindModule(Module->GetPath());
+		
+		if (!Model->Connections.HasConnection(ConnectorKey) || bReplaceExistingConnections)
+		{
+			if (const FRigConnectorElement* OtherConnectorElement = Cast<FRigConnectorElement>(Hierarchy->Find(ConnectorKey)))
+			{
+				FModularRigResolveResult RuleResults = RuleManager->FindMatches(OtherConnectorElement, ModuleInstance, ModularRig->GetElementKeyRedirector());
+
+				bool bFoundMatch = false;
+				if (RuleResults.GetMatches().Num() == 1)
+				{
+					Model->Connections.AddConnection(ConnectorKey, RuleResults.GetMatches()[0].GetKey());
+					Notify(EModularRigNotification::ConnectionChanged, Module);
+					bFoundMatch = true;
+				}
+				else
+				{
+					for (const FRigElementResolveResult& Result : RuleResults.GetMatches())
+					{
+						if (Result.GetState() == ERigElementResolveState::DefaultTarget)
+						{
+							Model->Connections.AddConnection(ConnectorKey, Result.GetKey());
+							Notify(EModularRigNotification::ConnectionChanged, Module);
+							bFoundMatch = true;
+							break;
+						}
+					}
+				}
+
+				if(!bFoundMatch)
+				{
+					bResolvedAllConnectors = false;
+				}
+			}
+		}
+	}
+
+	TransactionPtr.Reset();
+
+	return bResolvedAllConnectors;
+
+#else
+	
+	return false;
+
+#endif
+}
+
+bool UModularRigController::AutoConnectModules(const TArray<FString>& InModulePaths, bool bReplaceExistingConnections, bool bSetupUndo)
+{
+#if WITH_EDITOR
+	TArray<FRigElementKey> ConnectorKeys;
+
+	const UBlueprint* Blueprint = Cast<UBlueprint>(GetOuter());
+	if(Blueprint == nullptr)
+	{
+		UE_LOG(LogControlRig, Error, TEXT("ModularRigController is not nested under blueprint."));
+		return false;
+	}
+
+	const UModularRig* ModularRig = Cast<UModularRig>(Blueprint->GetObjectBeingDebugged());
+	if (!ModularRig)
+	{
+		UE_LOG(LogControlRig, Error, TEXT("Could not find debugged modular rig in %s"), *Blueprint->GetPathName());
+		return false;
+	}
+	
+	const URigHierarchy* Hierarchy = ModularRig->GetHierarchy();
+	if (!Hierarchy)
+	{
+		UE_LOG(LogControlRig, Error, TEXT("Could not find hierarchy in %s"), *ModularRig->GetPathName());
+		return false;
+	}
+
+	for(const FString& ModulePath : InModulePaths)
+	{
+		const FRigModuleReference* Module = FindModule(ModulePath);
+		if (!Module)
+		{
+			UE_LOG(LogControlRig, Error, TEXT("Could not find module %s"), *ModulePath);
+			return false;
+		}
+
+		const TArray<const FRigConnectorElement*> Connectors = Module->FindConnectors(Hierarchy);
+		for(const FRigConnectorElement* Connector : Connectors)
+		{
+			if(Connector->IsSecondary())
+			{
+				ConnectorKeys.Add(Connector->GetKey());
+			}
+		}
+	}
+
+	return AutoConnectSecondaryConnectors(ConnectorKeys, bReplaceExistingConnections, bSetupUndo);
+
+#else
+
+	return false;
+
+#endif
+}
+
 bool UModularRigController::SetConfigValueInModule(const FString& InModulePath, const FName& InVariableName, const FString& InValue, bool bSetupUndo)
 {
 	FRigModuleReference* Module = FindModule(InModulePath);
diff --git a/Engine/Plugins/Animation/ControlRig/Source/ControlRig/Public/ModularRigController.h b/Engine/Plugins/Animation/ControlRig/Source/ControlRig/Public/ModularRigController.h
--- a/Engine/Plugins/Animation/ControlRig/Source/ControlRig/Public/ModularRigController.h
+++ b/Engine/Plugins/Animation/ControlRig/Source/ControlRig/Public/ModularRigController.h
@@ -34,11 +34,17 @@ class CONTROLRIG_API UModularRigController : public UObject
 	bool ConnectConnectorToElement(const FRigElementKey& InConnectorKey, const FRigElementKey& InTargetKey, bool bSetupUndo = true, bool bAutoResolveOtherConnectors = true);
 
 	UFUNCTION(BlueprintCallable, Category = "Control Rig | Modules")
-	bool DisconnectConnector(const FRigElementKey& InConnectorKey, bool bSetupUndo = true);
+	bool DisconnectConnector(const FRigElementKey& InConnectorKey, bool bDisconnectSubModules = false, bool bSetupUndo = true);
 	
 	UFUNCTION(BlueprintCallable, Category = "Control Rig | Modules")
 	TArray<FRigElementKey> DisconnectCyclicConnectors(bool bSetupUndo = true);
 
+	UFUNCTION(BlueprintCallable, Category = "Control Rig | Modules")
+	bool AutoConnectSecondaryConnectors(const TArray<FRigElementKey>& InConnectorKeys, bool bReplaceExistingConnections, bool bSetupUndo = true);
+
+	UFUNCTION(BlueprintCallable, Category = "Control Rig | Modules")
+	bool AutoConnectModules(const TArray<FString>& InModulePaths, bool bReplaceExistingConnections, bool bSetupUndo = true);
+
 	UFUNCTION(BlueprintCallable, Category = "Control Rig | Modules")
 	bool SetConfigValueInModule(const FString& InModulePath, const FName& InVariableName, const FString& InValue, bool bSetupUndo = true);
 
diff --git a/Engine/Plugins/Animation/ControlRig/Source/ControlRig/Public/ModularRigModel.h b/Engine/Plugins/Animation/ControlRig/Source/ControlRig/Public/ModularRigModel.h
--- a/Engine/Plugins/Animation/ControlRig/Source/ControlRig/Public/ModularRigModel.h
+++ b/Engine/Plugins/Animation/ControlRig/Source/ControlRig/Public/ModularRigModel.h
@@ -228,6 +228,16 @@ public:
 		return ConnectionList.IsValidIndex(FindConnectionIndex(InConnectorKey));
 	}
 
+	bool HasConnection(const FRigElementKey& InConnectorKey, const URigHierarchy* InHierarchy) const
+	{
+		if(HasConnection(InConnectorKey))
+		{
+			const FRigElementKey Target = FindTargetFromConnector(InConnectorKey);
+			return InHierarchy->Contains(Target);
+		}
+		return false;
+	}
+
 	/** Gets the connection map for a single module, where the connectors are identified without its namespace*/
 	TMap<FRigElementKey, FRigElementKey> GetModuleConnectionMap(const FString& InModulePath) const;
 };
diff --git a/Engine/Plugins/Animation/ControlRig/Source/ControlRigDeveloper/Private/ControlRigBlueprint.cpp b/Engine/Plugins/Animation/ControlRig/Source/ControlRigDeveloper/Private/ControlRigBlueprint.cpp
--- a/Engine/Plugins/Animation/ControlRig/Source/ControlRigDeveloper/Private/ControlRigBlueprint.cpp
+++ b/Engine/Plugins/Animation/ControlRig/Source/ControlRigDeveloper/Private/ControlRigBlueprint.cpp
@@ -1009,7 +1009,7 @@ bool UControlRigBlueprint::ResolveConnector(const FRigElementKey& DraggedKey, co
 			// Add connection to the model
 			if (UModularRigController* Controller = GetModularRigController())
 			{
-				Controller->DisconnectConnector(DraggedKey, bSetupUndoRedo);
+				Controller->DisconnectConnector(DraggedKey, false, bSetupUndoRedo);
 			}
 		}
 		else
diff --git a/Engine/Plugins/Animation/ControlRig/Source/ControlRigEditor/Private/ControlRigModularRigCommands.cpp b/Engine/Plugins/Animation/ControlRig/Source/ControlRigEditor/Private/ControlRigModularRigCommands.cpp
--- a/Engine/Plugins/Animation/ControlRig/Source/ControlRigEditor/Private/ControlRigModularRigCommands.cpp
+++ b/Engine/Plugins/Animation/ControlRig/Source/ControlRigEditor/Private/ControlRigModularRigCommands.cpp
@@ -10,6 +10,7 @@ void FControlRigModularRigCommands::RegisterCommands()
 	UI_COMMAND(RenameModuleItem, "Rename", "Rename module to the rig.", EUserInterfaceActionType::Button, FInputChord(EKeys::F2));
 	UI_COMMAND(DeleteModuleItem, "Delete", "Delete module from the rig.", EUserInterfaceActionType::Button, FInputChord(EKeys::Delete));
 	UI_COMMAND(MirrorModuleItem, "Mirror", "Mirror module from the rig.", EUserInterfaceActionType::Button, FInputChord());
+	UI_COMMAND(ReresolveModuleItem, "Auto Resolve", "Auto Resolve secondary connectors.", EUserInterfaceActionType::Button, FInputChord());
 }
 
 #undef LOCTEXT_NAMESPACE
diff --git a/Engine/Plugins/Animation/ControlRig/Source/ControlRigEditor/Private/ControlRigModularRigCommands.h b/Engine/Plugins/Animation/ControlRig/Source/ControlRigEditor/Private/ControlRigModularRigCommands.h
--- a/Engine/Plugins/Animation/ControlRig/Source/ControlRigEditor/Private/ControlRigModularRigCommands.h
+++ b/Engine/Plugins/Animation/ControlRig/Source/ControlRigEditor/Private/ControlRigModularRigCommands.h
@@ -30,6 +30,9 @@ public:
 	/** Mirror Module */
 	TSharedPtr< FUICommandInfo > MirrorModuleItem;
 
+	/** Reresolve Module */
+	TSharedPtr< FUICommandInfo > ReresolveModuleItem;
+
 	/**
 	 * Initialize commands
 	 */
diff --git a/Engine/Plugins/Animation/ControlRig/Source/ControlRigEditor/Private/Editor/ControlRigEditor.cpp b/Engine/Plugins/Animation/ControlRig/Source/ControlRigEditor/Private/Editor/ControlRigEditor.cpp
--- a/Engine/Plugins/Animation/ControlRig/Source/ControlRigEditor/Private/Editor/ControlRigEditor.cpp
+++ b/Engine/Plugins/Animation/ControlRig/Source/ControlRigEditor/Private/Editor/ControlRigEditor.cpp
@@ -2802,8 +2802,10 @@ void FControlRigEditor::HandlePreviewMeshChanged(USkeletalMesh* InOldSkeletalMes
 		{
 			ControlRigBP->SetPreviewMesh(InNewSkeletalMesh);
 
+			FModularRigConnections PreviousConnections;
 			if(IsModularRig())
 			{
+				PreviousConnections = ControlRigBP->ModularRigModel.Connections;
 				{
 					TGuardValue<bool> SuspendBlueprintNotifs(ControlRigBP->bSuspendAllNotifications, true);
 					if(URigHierarchyController* Controller = ControlRigBP->GetHierarchyController())
@@ -2817,14 +2819,98 @@ void FControlRigEditor::HandlePreviewMeshChanged(USkeletalMesh* InOldSkeletalMes
 			
 			UpdateRigVMHost();
 			
-			if(UControlRig* DebuggedControlRig = Cast<UControlRig>(GetBlueprintObj()->GetObjectBeingDebugged()))
+			if(UControlRig* DebuggedControlRig = Cast<UControlRig>(ControlRigBP->GetObjectBeingDebugged()))
 			{
 				DebuggedControlRig->GetHierarchy()->Notify(ERigHierarchyNotification::HierarchyReset, nullptr);
 				DebuggedControlRig->Initialize(true);
 			}
-		}
 
-		Compile();
+			Compile();
+
+			if(IsModularRig())
+			{
+				if(UControlRig* DebuggedControlRig = Cast<UControlRig>(ControlRigBP->GetObjectBeingDebugged()))
+				{
+					DebuggedControlRig->RequestConstruction();
+					DebuggedControlRig->Execute(FRigUnit_PrepareForExecution::EventName);
+					
+					if(URigHierarchy* Hierarchy = DebuggedControlRig->GetHierarchy())
+					{
+						FModularRigModel* Model = &ControlRigBP->ModularRigModel;
+						
+						// try to reestablish the connections.
+						UModularRigController* ModularRigController = ControlRigBP->GetModularRigController();
+						Model->ForEachModule(
+							[Model, Hierarchy, ModularRigController, PreviousConnections]
+							(const FRigModuleReference* Module) -> bool
+							{
+								bool bContinueResolval;
+								do
+								{
+									bContinueResolval = false;
+									
+									const TArray<const FRigConnectorElement*> Connectors = Module->FindConnectors(Hierarchy);
+									TArray<FRigElementKey> PrimaryConnectors, SecondaryConnectors, OptionalConnectors;
+									for(const FRigConnectorElement* ExistingConnector : Connectors)
+									{
+										if(ExistingConnector->IsPrimary())
+										{
+											PrimaryConnectors.Add(ExistingConnector->GetKey());
+										}
+										else if(ExistingConnector->IsOptional())
+										{
+											OptionalConnectors.Add(ExistingConnector->GetKey());
+										}
+										else
+										{
+											SecondaryConnectors.Add(ExistingConnector->GetKey());
+										}
+									}
+									TArray<FRigElementKey> ConnectorKeys;
+									ConnectorKeys.Append(PrimaryConnectors);
+									ConnectorKeys.Append(SecondaryConnectors);
+									ConnectorKeys.Append(OptionalConnectors);
+									
+									for(const FRigElementKey& ConnectorKey : ConnectorKeys)
+									{
+										const bool bIsPrimary = ConnectorKey == ConnectorKeys[0];
+										const bool bIsSecondary = !bIsPrimary;
+										
+										if(!Model->Connections.HasConnection(ConnectorKey, Hierarchy))
+										{
+											// try to reapply the connection
+											if(PreviousConnections.HasConnection(ConnectorKey, Hierarchy))
+											{
+												const FRigElementKey Target = PreviousConnections.FindTargetFromConnector(ConnectorKey);
+												if(ModularRigController->ConnectConnectorToElement(ConnectorKey, Target, true))
+												{
+													bContinueResolval = true;
+												}
+											}
+
+											// try to auto resolve it
+											if(!bContinueResolval && bIsSecondary)
+											{
+												if(ModularRigController->AutoConnectSecondaryConnectors({ConnectorKey}, true, true))
+												{
+													bContinueResolval = true;
+												}
+											}
+
+											// only do one connector at a time
+											break;
+										}
+									}
+								}
+								while (bContinueResolval);
+
+								return true; // continue to the next module
+							}
+						);
+					}
+				}
+			}
+		}
 	}
 }
 
diff --git a/Engine/Plugins/Animation/ControlRig/Source/ControlRigEditor/Private/Editor/SModularRigModel.cpp b/Engine/Plugins/Animation/ControlRig/Source/ControlRigEditor/Private/Editor/SModularRigModel.cpp
--- a/Engine/Plugins/Animation/ControlRig/Source/ControlRigEditor/Private/Editor/SModularRigModel.cpp
+++ b/Engine/Plugins/Animation/ControlRig/Source/ControlRigEditor/Private/Editor/SModularRigModel.cpp
@@ -225,6 +225,10 @@ void SModularRigModel::BindCommands()
 	CommandList->MapAction(Commands.MirrorModuleItem,
 		FExecuteAction::CreateSP(this, &SModularRigModel::HandleMirrorModules),
 		FCanExecuteAction());
+
+	CommandList->MapAction(Commands.ReresolveModuleItem,
+		FExecuteAction::CreateSP(this, &SModularRigModel::HandleReresolveModules),
+		FCanExecuteAction());
 }
 
 FReply SModularRigModel::OnKeyDown(const FGeometry& MyGeometry, const FKeyEvent& InKeyEvent)
@@ -390,6 +394,7 @@ void SModularRigModel::CreateContextMenu()
 					ModulesSection.AddMenuEntry(Commands.RenameModuleItem);
 					ModulesSection.AddMenuEntry(Commands.DeleteModuleItem);
 					ModulesSection.AddMenuEntry(Commands.MirrorModuleItem);
+					ModulesSection.AddMenuEntry(Commands.ReresolveModuleItem);
 				}
 			})
 		);
@@ -717,8 +722,6 @@ void SModularRigModel::HandleMirrorModules()
 		});
 		HandleMirrorModules(SelectedPaths);
 	}
-
-	return;
 }
 
 void SModularRigModel::HandleMirrorModules(const TArray<FString>& InPaths)
@@ -760,6 +763,104 @@ void SModularRigModel::HandleMirrorModules(const TArray<FString>& InPaths)
 	}
 }
 
+void SModularRigModel::HandleReresolveModules()
+{
+	if(!ControlRigEditor.IsValid())
+	{
+		return;
+	}
+
+	UModularRig* Rig = GetDefaultModularRig();
+	if (Rig)
+	{
+		TArray<TSharedPtr<FModularRigTreeElement>> SelectedItems = TreeView->GetSelectedItems();
+		TArray<FString> SelectedPaths;
+		Algo::Transform(SelectedItems, SelectedPaths, [](const TSharedPtr<FModularRigTreeElement>& Element)
+		{
+			if (Element.IsValid())
+			{
+				if(Element->ConnectorName.IsEmpty())
+				{
+					return Element->ModulePath;
+				}
+				return FString::Printf(TEXT("%s|%s"), *Element->ModulePath, *Element->ConnectorName);
+			}
+			return FString();
+		});
+		HandleReresolveModules(SelectedPaths);
+	}
+}
+
+void SModularRigModel::HandleReresolveModules(const TArray<FString>& InPaths)
+{
+	if (ControlRigBlueprint.IsValid())
+	{
+		UModularRigController* Controller = ControlRigBlueprint->GetModularRigController();
+		check(Controller);
+
+		const UModularRig* Rig = GetDefaultModularRig();
+		if (Rig == nullptr)
+		{
+			return;
+		}
+
+		const URigHierarchy* Hierarchy = Rig->GetHierarchy();
+		if(Hierarchy == nullptr)
+		{
+			return;
+		}
+
+		TArray<FRigElementKey> ConnectorKeys;
+		for (const FString& PathAndConnector : InPaths)
+		{
+			FString ModulePath = PathAndConnector;
+			FString ConnectorName;
+			(void)PathAndConnector.Split(TEXT("|"), &ModulePath, &ConnectorName, ESearchCase::CaseSensitive, ESearchDir::FromEnd);
+
+			const FRigModuleReference* Module = Controller->Model->FindModule(ModulePath);
+			if (!Module)
+			{
+				UE_LOG(LogControlRig, Error, TEXT("Could not find module %s"), *ModulePath);
+				return;
+			}
+
+			if(!ConnectorName.IsEmpty())
+			{
+				// if we are executing this on a primary connector we want to re-resolve all secondaries
+				const FRigConnectorElement* PrimaryConnector = Module->FindPrimaryConnector(Hierarchy);
+				const FName DesiredName = Hierarchy->GetNameMetadata(PrimaryConnector->GetKey(), URigHierarchy::DesiredNameMetadataName, NAME_None);
+				if(!DesiredName.IsNone() && DesiredName.ToString().Equals(ConnectorName, ESearchCase::CaseSensitive))
+				{
+					ConnectorName.Reset();
+				}
+			}
+
+			const TArray<const FRigConnectorElement*> Connectors = Module->FindConnectors(Hierarchy);
+			for(const FRigConnectorElement* Connector : Connectors)
+			{
+				if(Connector->IsSecondary())
+				{
+					if(ConnectorName.IsEmpty())
+					{
+						ConnectorKeys.AddUnique(Connector->GetKey());
+					}
+					else
+					{
+						const FName DesiredName = Hierarchy->GetNameMetadata(Connector->GetKey(), URigHierarchy::DesiredNameMetadataName, NAME_None);
+						if(!DesiredName.IsNone() && DesiredName.ToString().Equals(ConnectorName, ESearchCase::CaseSensitive))
+						{
+							ConnectorKeys.AddUnique(Connector->GetKey());
+							break;
+						}
+					}
+				}
+			}
+		}
+
+		Controller->AutoConnectSecondaryConnectors(ConnectorKeys, true, true);
+	}
+}
+
 void SModularRigModel::HandleConnectorResolved(const FRigElementKey& InConnector, const FRigElementKey& InTarget)
 {
 	if (ControlRigBlueprint.IsValid())
@@ -785,7 +886,7 @@ void SModularRigModel::HandleConnectorDisconnect(const FRigElementKey& InConnect
 		UModularRigController* Controller = ControlRigBlueprint->GetModularRigController();
 		check(Controller);
 
-		Controller->DisconnectConnector(InConnector, true);
+		Controller->DisconnectConnector(InConnector, false, true);
 	}
 }
 
diff --git a/Engine/Plugins/Animation/ControlRig/Source/ControlRigEditor/Private/Editor/SModularRigModel.h b/Engine/Plugins/Animation/ControlRig/Source/ControlRigEditor/Private/Editor/SModularRigModel.h
--- a/Engine/Plugins/Animation/ControlRig/Source/ControlRigEditor/Private/Editor/SModularRigModel.h
+++ b/Engine/Plugins/Animation/ControlRig/Source/ControlRigEditor/Private/Editor/SModularRigModel.h
@@ -78,6 +78,10 @@ private:
 	void HandleMirrorModules();
 	void HandleMirrorModules(const TArray<FString>& InPaths);
 
+	/** Reresolve items */
+	void HandleReresolveModules();
+	void HandleReresolveModules(const TArray<FString>& InPaths);
+
 	/** Resolve connector */
 	void HandleConnectorResolved(const FRigElementKey& InConnector, const FRigElementKey& InTarget);
 
diff --git a/Engine/Plugins/Animation/ControlRig/Source/ControlRigEditor/Private/Editor/SModularRigTreeView.cpp b/Engine/Plugins/Animation/ControlRig/Source/ControlRigEditor/Private/Editor/SModularRigTreeView.cpp
--- a/Engine/Plugins/Animation/ControlRig/Source/ControlRigEditor/Private/Editor/SModularRigTreeView.cpp
+++ b/Engine/Plugins/Animation/ControlRig/Source/ControlRigEditor/Private/Editor/SModularRigTreeView.cpp
@@ -151,6 +151,31 @@ void SModularRigModelItem::PopulateConnectorCurrentTarget(TSharedPtr<SVerticalBo
 {
 	static const FSlateBrush* RoundedBoxBrush = FControlRigEditorStyle::Get().GetBrush(TEXT("ControlRig.SpacePicker.RoundedRect"));
 
+	TAttribute<FSlateColor> TextColor = TAttribute<FSlateColor>::CreateLambda([this, InConnectorKey, InTargetKey]() -> FSlateColor
+	{
+		if(const UModularRig* ModularRig = Delegates.GetModularRig())
+		{
+			if(const URigHierarchy* Hierarchy = ModularRig->GetHierarchy())
+			{
+				if(!Hierarchy->Contains(InTargetKey))
+				{
+					if(!InTargetKey.IsValid())
+					{
+						if(const FRigConnectorElement* Connector = Hierarchy->Find<FRigConnectorElement>(InConnectorKey))
+						{
+							if(Connector->IsOptional())
+							{
+								return FSlateColor::UseForeground();
+							}
+						}
+					}
+					return FLinearColor::Red;
+				}
+			}
+		}
+		return FSlateColor::UseForeground();
+	});
+	
 	TSharedPtr<SHorizontalBox> RowBox, ButtonBox;
 	InListBox->AddSlot()
 	.AutoHeight()
@@ -208,6 +233,7 @@ void SModularRigModelItem::PopulateConnectorCurrentTarget(TSharedPtr<SVerticalBo
 							SNew( STextBlock )
 							.Text( InTitle )
 							.Font( IDetailLayoutBuilder::GetDetailFont() )
+							.ColorAndOpacity(TextColor)
 						]
 					]
 				]
