Control Rig: Drag and drop onto selected schematic nodes will create a module for each selected node
#rnx

diff --git a/Engine/Plugins/Animation/ControlRig/Source/ControlRig/Private/ModularRig.cpp b/Engine/Plugins/Animation/ControlRig/Source/ControlRig/Private/ModularRig.cpp
--- a/Engine/Plugins/Animation/ControlRig/Source/ControlRig/Private/ModularRig.cpp
+++ b/Engine/Plugins/Animation/ControlRig/Source/ControlRig/Private/ModularRig.cpp
@@ -238,6 +238,7 @@ void UModularRig::Serialize(FArchive& Ar)
 	if (Ar.IsLoading())
 	{
 		ModularRigModel.UpdateCachedChildren();
+		ModularRigModel.Connections.UpdateFromConnectionList();
 	}
 }
 
@@ -245,6 +246,7 @@ void UModularRig::PostLoad()
 {
 	Super::PostLoad();
 	ModularRigModel.UpdateCachedChildren();
+	ModularRigModel.Connections.UpdateFromConnectionList();
 	ResetShortestDisplayPathCache();
 }
 
diff --git a/Engine/Plugins/Animation/ControlRig/Source/ControlRig/Private/ModularRigController.cpp b/Engine/Plugins/Animation/ControlRig/Source/ControlRig/Private/ModularRigController.cpp
--- a/Engine/Plugins/Animation/ControlRig/Source/ControlRig/Private/ModularRigController.cpp
+++ b/Engine/Plugins/Animation/ControlRig/Source/ControlRig/Private/ModularRigController.cpp
@@ -1041,6 +1041,7 @@ bool UModularRigController::DeleteModule(const FString& InModulePath, bool bSetu
 		{
 			Model->Connections.RemoveConnection(KeyToRemove);
 		}
+		Model->Connections.UpdateFromConnectionList();
 	}
 
 	// Fix bindings
@@ -1303,6 +1304,7 @@ FString UModularRigController::MirrorModule(const FString& InModulePath, const F
 	{
 		NewModuleName = NewModuleName.Replace(*InSettings.SearchString, *InSettings.ReplaceString, ESearchCase::CaseSensitive);
 	}
+	
 
 	FString NewModulePath = AddModule(*NewModuleName, OriginalModule->Class.Get(), OriginalModule->ParentPath, bSetupUndo);
 	FRigModuleReference* NewModule = FindModule(NewModulePath);
diff --git a/Engine/Plugins/Animation/ControlRig/Source/ControlRig/Public/ModularRigModel.h b/Engine/Plugins/Animation/ControlRig/Source/ControlRig/Public/ModularRigModel.h
--- a/Engine/Plugins/Animation/ControlRig/Source/ControlRig/Public/ModularRigModel.h
+++ b/Engine/Plugins/Animation/ControlRig/Source/ControlRig/Public/ModularRigModel.h
@@ -256,6 +256,7 @@ public:
 		Modules = Other.Modules;
 		Connections = Other.Connections;
 		UpdateCachedChildren();
+		Connections.UpdateFromConnectionList();
 	}
 	
 	FModularRigModel& operator=(const FModularRigModel& Other)
@@ -263,6 +264,7 @@ public:
 		Modules = Other.Modules;
 		Connections = Other.Connections;
 		UpdateCachedChildren();
+		Connections.UpdateFromConnectionList();
 		return *this;
 	}
 	FModularRigModel(FModularRigModel&&) = delete;
diff --git a/Engine/Plugins/Animation/ControlRig/Source/ControlRigDeveloper/Private/ControlRigBlueprint.cpp b/Engine/Plugins/Animation/ControlRig/Source/ControlRigDeveloper/Private/ControlRigBlueprint.cpp
--- a/Engine/Plugins/Animation/ControlRig/Source/ControlRigDeveloper/Private/ControlRigBlueprint.cpp
+++ b/Engine/Plugins/Animation/ControlRig/Source/ControlRigDeveloper/Private/ControlRigBlueprint.cpp
@@ -902,6 +902,7 @@ void UControlRigBlueprint::Serialize(FArchive& Ar)
 		}
 
 		ModularRigModel.UpdateCachedChildren();
+		ModularRigModel.Connections.UpdateFromConnectionList();
 	}
 }
 
@@ -1390,6 +1391,7 @@ void UControlRigBlueprint::PostTransacted(const FTransactionObjectEvent& Transac
 			if (IsModularRig())
 			{
 				ModularRigModel.UpdateCachedChildren();
+				ModularRigModel.Connections.UpdateFromConnectionList();
 				RecompileModularRig();
 			}
 		}
@@ -1429,6 +1431,7 @@ void UControlRigBlueprint::PostDuplicate(bool bDuplicateForPIE)
 	}
 
 	ModularRigModel.UpdateCachedChildren();
+	ModularRigModel.Connections.UpdateFromConnectionList();
 }
 
 void UControlRigBlueprint::PostRename(UObject* OldOuter, const FName OldName)
diff --git a/Engine/Plugins/Animation/ControlRig/Source/ControlRigEditor/Private/Editor/ControlRigSchematicModel.cpp b/Engine/Plugins/Animation/ControlRig/Source/ControlRigEditor/Private/Editor/ControlRigSchematicModel.cpp
--- a/Engine/Plugins/Animation/ControlRig/Source/ControlRigEditor/Private/Editor/ControlRigSchematicModel.cpp
+++ b/Engine/Plugins/Animation/ControlRig/Source/ControlRigEditor/Private/Editor/ControlRigSchematicModel.cpp
@@ -18,6 +18,7 @@
 #include "SchematicGraphPanel/SSchematicGraphPanel.h"
 #include "Framework/Application/SlateApplication.h"
 #include "ScopedTransaction.h"
+#include "Dialog/SCustomDialog.h"
 
 #define LOCTEXT_NAMESPACE "ControlRigSchematicModel"
 
@@ -1261,9 +1262,17 @@ TArray<FRigElementKey> FControlRigSchematicModel::GetElementKeysFromDragDropEven
 
 void FControlRigSchematicModel::HandleSchematicNodeClicked(SSchematicGraphPanel* InPanel, SSchematicGraphNode* InNode, const FPointerEvent& InMouseEvent)
 {
+	const bool bClearSelection = !InMouseEvent.IsShiftDown() && !InMouseEvent.IsControlDown();
 	for (const TSharedPtr<FSchematicGraphNode>& Node : Nodes)
 	{
-		Node->SetSelected(Node->GetGuid() == InNode->GetGuid());
+		if (bClearSelection)
+		{
+			Node->SetSelected(Node->GetGuid() == InNode->GetGuid());
+		}
+		else if(Node->GetGuid() == InNode->GetGuid())
+		{
+			Node->SetSelected();
+		}
 	}
 
 	if(const FControlRigSchematicRigElementKeyNode* Node =
@@ -1275,7 +1284,6 @@ void FControlRigSchematicModel::HandleSchematicNodeClicked(SSchematicGraphPanel*
 			{
 				if (URigHierarchyController* Controller = Hierarchy->GetController())
 				{
-					const bool bClearSelection = !InMouseEvent.IsShiftDown() && !InMouseEvent.IsControlDown();
 					bool bSelect = true;
 					if(InMouseEvent.IsControlDown())
 					{
@@ -1448,7 +1456,54 @@ void FControlRigSchematicModel::HandleSchematicDrop(SSchematicGraphPanel* InPane
 	};
 
 	TArray<FRigElementKey> TargetKeys;
-	Local::CollectTargetKeys(TargetKeys, InNode->GetNodeData());
+	TArray<TSharedPtr<FSchematicGraphNode>> SelectedNodesShared = GetSelectedNodes();
+	TArray<FSchematicGraphNode*> SelectedNodes;
+	SelectedNodes.Reserve(SelectedNodesShared.Num());
+	for (TSharedPtr<FSchematicGraphNode>& Node : SelectedNodesShared)
+	{
+		if (Node.IsValid())
+		{
+			SelectedNodes.AddUnique(Node.Get());
+		}
+	}
+	SelectedNodes.AddUnique(InNode->GetNodeData());
+
+	if (SelectedNodes.Num() > 1)
+	{
+		FText Label;
+		if(FControlRigSchematicRigElementKeyNode* ExistingElementKeyNode = Cast<FControlRigSchematicRigElementKeyNode>(InNode->GetNodeData()))
+		{
+			Label = ExistingElementKeyNode->GetLabel();
+		}
+		else
+		{
+			Label = InNode->GetNodeData()->GetLabel();
+		}
+		TSharedRef<SCustomDialog> ChooseSelectedDialog = SNew(SCustomDialog)
+			.Title(LOCTEXT("Add modules", "Add modules to targets"))
+			.Buttons({
+				SCustomDialog::FButton(FText::Format(LOCTEXT("AddModuleToSingle", "Add module to {0}"), Label)),
+				SCustomDialog::FButton(LOCTEXT("AddModuleToSelected", "Add module to all selected"))
+				});
+
+		const int32 ButtonPressed = ChooseSelectedDialog->ShowModal();
+		if (ButtonPressed < 0)
+		{
+			return; // Window closed
+		}
+		if (ButtonPressed == 0)
+		{
+			SelectedNodes.Reset();
+			SelectedNodes.Add(InNode->GetNodeData());
+		}
+	}
+	
+	for (FSchematicGraphNode* Node : SelectedNodes)
+	{
+		TArray<FRigElementKey> NodeTargetKeys;
+		Local::CollectTargetKeys(NodeTargetKeys, Node);
+		TargetKeys.Append(NodeTargetKeys);
+	}
 	
 	UControlRig* ControlRig = ControlRigBlueprint->GetDebuggedControlRig();
 	if (!ControlRig)
@@ -1495,32 +1550,28 @@ void FControlRigSchematicModel::HandleSchematicDrop(SSchematicGraphPanel* InPane
 
 						FScopedTransaction Transaction(LOCTEXT("AddAndConnectModule", "Add and Connect Module"));
 
-						const FName ModuleName = Controller->GetSafeNewName(FString(), FRigName(AssetBlueprint->RigModuleSettings.Identifier.Name));
-						const FString ModulePath = Controller->AddModule(ModuleName, AssetBlueprint->GetControlRigClass(), FString());
-						if(!ModulePath.IsEmpty())
+						for(const FRigElementKey& TargetKey : TargetKeys)
 						{
-							FRigElementKey PrimaryConnectorKey;
-							TArray<FRigConnectorElement*> Connectors = Hierarchy->GetElementsOfType<FRigConnectorElement>();
-							for (FRigConnectorElement* Connector : Connectors)
+							const FName ModuleName = Controller->GetSafeNewName(FString(), FRigName(AssetBlueprint->RigModuleSettings.Identifier.Name));
+							const FString ModulePath = Controller->AddModule(ModuleName, AssetBlueprint->GetControlRigClass(), FString());
+							if(!ModulePath.IsEmpty())
 							{
-								if (Connector->IsPrimary())
+								FRigElementKey PrimaryConnectorKey;
+								TArray<FRigConnectorElement*> Connectors = Hierarchy->GetElementsOfType<FRigConnectorElement>();
+								for (FRigConnectorElement* Connector : Connectors)
 								{
-									FString Path, Name;
-									(void)URigHierarchy::SplitNameSpace(Connector->GetName(), &Path, &Name);
-									if (Path == ModulePath)
+									if (Connector->IsPrimary())
 									{
-										PrimaryConnectorKey = Connector->GetKey();
-										break;
+										FString Path, Name;
+										(void)URigHierarchy::SplitNameSpace(Connector->GetName(), &Path, &Name);
+										if (Path == ModulePath)
+										{
+											PrimaryConnectorKey = Connector->GetKey();
+											break;
+										}
 									}
 								}
-							}
-
-							for(const FRigElementKey& TargetKey : TargetKeys)
-							{
-								if(Controller->ConnectConnectorToElement(PrimaryConnectorKey, TargetKey, true, ControlRig->GetModularRigSettings().bAutoResolve))
-								{
-									break;
-								}
+								Controller->ConnectConnectorToElement(PrimaryConnectorKey, TargetKey, true, ControlRig->GetModularRigSettings().bAutoResolve);
 							}
 						}
 					}
@@ -1810,6 +1861,10 @@ void FControlRigSchematicModel::OnHierarchyModified(ERigHierarchyNotification In
 					{
 						GroupNode->SetExpanded(true);
 					}
+					else
+					{
+						Node->SetSelected(true);
+					}
 				}
 			}
 			break;
@@ -1819,6 +1874,7 @@ void FControlRigSchematicModel::OnHierarchyModified(ERigHierarchyNotification In
 			if(FControlRigSchematicRigElementKeyNode* Node = FindElementKeyNode(InElement->GetKey()))
 			{
 				Node->SetExpanded(false);
+				Node->SetSelected(false);
 			}
 			break;
 		}
diff --git a/Engine/Source/Editor/AnimationEditorWidgets/Private/SchematicGraphPanel/SchematicGraphModel.cpp b/Engine/Source/Editor/AnimationEditorWidgets/Private/SchematicGraphPanel/SchematicGraphModel.cpp
--- a/Engine/Source/Editor/AnimationEditorWidgets/Private/SchematicGraphPanel/SchematicGraphModel.cpp
+++ b/Engine/Source/Editor/AnimationEditorWidgets/Private/SchematicGraphPanel/SchematicGraphModel.cpp
@@ -444,6 +444,14 @@ bool FSchematicGraphModel::GetContextMenuForNode(const FSchematicGraphNode* InNo
 	return false;
 }
 
+const TArray<TSharedPtr<FSchematicGraphNode>> FSchematicGraphModel::GetSelectedNodes() const
+{
+	return Nodes.FilterByPredicate([](const TSharedPtr<FSchematicGraphNode>& Node)
+	{
+		return Node->IsSelected();
+	});
+}
+
 FLinearColor FSchematicGraphModel::GetBackgroundColorForTag(const FGuid& InNodeGuid, const FGuid& InTagGuid) const
 {
 	if(const FSchematicGraphNode* Node = FindNode(InNodeGuid))
diff --git a/Engine/Source/Editor/AnimationEditorWidgets/Public/SchematicGraphPanel/SchematicGraphModel.h b/Engine/Source/Editor/AnimationEditorWidgets/Public/SchematicGraphPanel/SchematicGraphModel.h
--- a/Engine/Source/Editor/AnimationEditorWidgets/Public/SchematicGraphPanel/SchematicGraphModel.h
+++ b/Engine/Source/Editor/AnimationEditorWidgets/Public/SchematicGraphPanel/SchematicGraphModel.h
@@ -115,6 +115,7 @@ public:
 	bool IsDragSupportedForNode(const FGuid& InNodeGuid) const;
 	virtual bool IsDragSupportedForNode(const FSchematicGraphNode* InNode) const;
 	virtual bool GetContextMenuForNode(const FSchematicGraphNode* InNode, FMenuBuilder& OutMenu) const;
+	const TArray<TSharedPtr<FSchematicGraphNode>> GetSelectedNodes() const;
 
 	FLinearColor GetBackgroundColorForTag(const FGuid& InNodeGuid, const FGuid& InTagGuid) const;
 	virtual FLinearColor GetBackgroundColorForTag(const FSchematicGraphTag* InTag) const;
