Control Rig: FIx DefaultParentPerElement invalid after topology change
#rnx

diff --git a/Engine/Plugins/Animation/ControlRig/Source/ControlRig/Private/Rigs/RigHierarchy.cpp b/Engine/Plugins/Animation/ControlRig/Source/ControlRig/Private/Rigs/RigHierarchy.cpp
--- a/Engine/Plugins/Animation/ControlRig/Source/ControlRig/Private/Rigs/RigHierarchy.cpp
+++ b/Engine/Plugins/Animation/ControlRig/Source/ControlRig/Private/Rigs/RigHierarchy.cpp
@@ -2005,6 +2005,12 @@ FRigBaseElementParentArray URigHierarchy::GetParents(const FRigBaseElement* InEl
 
 FRigElementKey URigHierarchy::GetDefaultParent(FRigElementKey InKey) const
 {
+	if (DefaultParentCacheTopologyVersion != GetTopologyVersion())
+	{
+		DefaultParentPerElement.Reset();
+		DefaultParentCacheTopologyVersion = GetTopologyVersion();
+	}
+	
 	FRigElementKey DefaultParent;
 	if(const FRigElementKey* DefaultParentPtr = DefaultParentPerElement.Find(InKey))
 	{
diff --git a/Engine/Plugins/Animation/ControlRig/Source/ControlRig/Public/Rigs/RigHierarchy.h b/Engine/Plugins/Animation/ControlRig/Source/ControlRig/Public/Rigs/RigHierarchy.h
--- a/Engine/Plugins/Animation/ControlRig/Source/ControlRig/Public/Rigs/RigHierarchy.h
+++ b/Engine/Plugins/Animation/ControlRig/Source/ControlRig/Public/Rigs/RigHierarchy.h
@@ -4700,6 +4700,7 @@ protected:
 	TObjectPtr<URigHierarchy> HierarchyForCacheValidation;
 
 	mutable TMap<FRigElementKey, FRigElementKey> DefaultParentPerElement;
+	mutable uint32 DefaultParentCacheTopologyVersion;
 
 	bool bUsePreferredEulerAngles;
 	mutable bool bAllowNameSpaceWhenSanitizingName;
