Contorl Rig: Fix mirror modules not working for config values
#rnx

diff --git a/Engine/Plugins/Animation/ControlRig/Source/ControlRig/Private/ModularRigController.cpp b/Engine/Plugins/Animation/ControlRig/Source/ControlRig/Private/ModularRigController.cpp
--- a/Engine/Plugins/Animation/ControlRig/Source/ControlRig/Private/ModularRigController.cpp
+++ b/Engine/Plugins/Animation/ControlRig/Source/ControlRig/Private/ModularRigController.cpp
@@ -1304,7 +1304,8 @@ FString UModularRigController::MirrorModule(const FString& InModulePath, const F
 	{
 		NewModuleName = NewModuleName.Replace(*InSettings.SearchString, *InSettings.ReplaceString, ESearchCase::CaseSensitive);
 	}
-	
+
+	FModularRigControllerCompileBracketScope CompileBracketScope(this);
 
 	FString NewModulePath = AddModule(*NewModuleName, OriginalModule->Class.Get(), OriginalModule->ParentPath, bSetupUndo);
 	FRigModuleReference* NewModule = FindModule(NewModulePath);
@@ -1704,3 +1705,25 @@ void UModularRigController::UpdateShortNames()
 		
 	}
 }
+
+FModularRigControllerCompileBracketScope::FModularRigControllerCompileBracketScope(UModularRigController* InController)
+	: Controller(InController), bSuspendNotifications(InController->bSuspendNotifications)
+{
+	check(InController);
+	
+	if (bSuspendNotifications)
+	{
+		return;
+	}
+	InController->Notify(EModularRigNotification::InteractionBracketOpened, nullptr);
+}
+
+FModularRigControllerCompileBracketScope::~FModularRigControllerCompileBracketScope()
+{
+	check(Controller);
+	if (bSuspendNotifications)
+	{
+		return;
+	}
+	Controller->Notify(EModularRigNotification::InteractionBracketClosed, nullptr);
+}
diff --git a/Engine/Plugins/Animation/ControlRig/Source/ControlRig/Public/ModularRigController.h b/Engine/Plugins/Animation/ControlRig/Source/ControlRig/Public/ModularRigController.h
--- a/Engine/Plugins/Animation/ControlRig/Source/ControlRig/Public/ModularRigController.h
+++ b/Engine/Plugins/Animation/ControlRig/Source/ControlRig/Public/ModularRigController.h
@@ -93,4 +93,18 @@ private:
 	bool bAutomaticReparenting;
 	
 	friend struct FModularRigModel;
-};
\ No newline at end of file
+	friend class FModularRigControllerCompileBracketScope;
+};
+
+class CONTROLRIG_API FModularRigControllerCompileBracketScope
+{
+public:
+   
+	FModularRigControllerCompileBracketScope(UModularRigController *InController);
+
+	~FModularRigControllerCompileBracketScope();
+
+private:
+	UModularRigController* Controller;
+	bool bSuspendNotifications;
+};
diff --git a/Engine/Plugins/Animation/ControlRig/Source/ControlRig/Public/ModularRigModel.h b/Engine/Plugins/Animation/ControlRig/Source/ControlRig/Public/ModularRigModel.h
--- a/Engine/Plugins/Animation/ControlRig/Source/ControlRig/Public/ModularRigModel.h
+++ b/Engine/Plugins/Animation/ControlRig/Source/ControlRig/Public/ModularRigModel.h
@@ -24,6 +24,12 @@ enum class EModularRigNotification : uint8
 
 	ModuleShortNameChanged,
 
+	InteractionBracketOpened, // A bracket has been opened
+	
+	InteractionBracketClosed, // A bracket has been opened
+	
+	InteractionBracketCanceled, // A bracket has been canceled
+
 	/** MAX - invalid */
 	Max UMETA(Hidden),
 };
diff --git a/Engine/Plugins/Animation/ControlRig/Source/ControlRigDeveloper/Private/ControlRigBlueprint.cpp b/Engine/Plugins/Animation/ControlRig/Source/ControlRigDeveloper/Private/ControlRigBlueprint.cpp
--- a/Engine/Plugins/Animation/ControlRig/Source/ControlRigDeveloper/Private/ControlRigBlueprint.cpp
+++ b/Engine/Plugins/Animation/ControlRig/Source/ControlRigDeveloper/Private/ControlRigBlueprint.cpp
@@ -81,6 +81,7 @@ UControlRigBlueprint::UControlRigBlueprint(const FObjectInitializer& ObjectIniti
 
 UControlRigBlueprint::UControlRigBlueprint()
 {
+	ModulesRecompilationBracket = 0;
 }
 
 UClass* UControlRigBlueprint::RegenerateClass(UClass* ClassToRegenerate, UObject* PreviousCDO)
@@ -2679,13 +2680,24 @@ void UControlRigBlueprint::HandleRigModulesModified(EModularRigNotification InNo
 			RequestConstructionOnAllModules();
 			break;
 		}
+		case EModularRigNotification::InteractionBracketOpened:
+		{
+			ModulesRecompilationBracket++;
+			break;
+		}
+		case EModularRigNotification::InteractionBracketClosed:
+		case EModularRigNotification::InteractionBracketCanceled:
+		{
+			ModulesRecompilationBracket--;
+			break;
+		}
 		default:
 		{
 			break;
 		}
 	}
 
-	if (bRecompile)
+	if (bRecompile && ModulesRecompilationBracket == 0)
 	{
 		RecompileModularRig();
 	}
diff --git a/Engine/Plugins/Animation/ControlRig/Source/ControlRigDeveloper/Public/ControlRigBlueprint.h b/Engine/Plugins/Animation/ControlRig/Source/ControlRigDeveloper/Public/ControlRigBlueprint.h
--- a/Engine/Plugins/Animation/ControlRig/Source/ControlRigDeveloper/Public/ControlRigBlueprint.h
+++ b/Engine/Plugins/Animation/ControlRig/Source/ControlRigDeveloper/Public/ControlRigBlueprint.h
@@ -329,6 +329,10 @@ private:
 	FOnRigVMRefreshEditorEvent ModularRigPreCompiled;
 	FOnRigVMRefreshEditorEvent ModularRigCompiled;
 
+	UPROPERTY(transient, DuplicateTransient)
+	int32 ModulesRecompilationBracket = 0;
+
+
 	void HandleHierarchyModified(ERigHierarchyNotification InNotification, URigHierarchy* InHierarchy, const FRigBaseElement* InElement);
 
 	void HandleRigModulesModified(EModularRigNotification InNotification, const FRigModuleReference* InModule);
