Control Rig: Remove weak hieararchy pointer in controller

#rnx

diff --git a/Engine/Plugins/Animation/ControlRig/Source/ControlRig/Private/Rigs/RigHierarchy.cpp b/Engine/Plugins/Animation/ControlRig/Source/ControlRig/Private/Rigs/RigHierarchy.cpp
--- a/Engine/Plugins/Animation/ControlRig/Source/ControlRig/Private/Rigs/RigHierarchy.cpp
+++ b/Engine/Plugins/Animation/ControlRig/Source/ControlRig/Private/Rigs/RigHierarchy.cpp
@@ -2857,6 +2857,11 @@ void URigHierarchy::PostEditUndo()
 	{
 		SetTransformStackIndex(DesiredStackIndex);
 	}
+
+	if(URigHierarchyController* Controller = GetController(false))
+	{
+		Controller->SetHierarchy(this);
+	}
 }
 
 #endif
diff --git a/Engine/Plugins/Animation/ControlRig/Source/ControlRig/Private/Rigs/RigHierarchyController.cpp b/Engine/Plugins/Animation/ControlRig/Source/ControlRig/Private/Rigs/RigHierarchyController.cpp
--- a/Engine/Plugins/Animation/ControlRig/Source/ControlRig/Private/Rigs/RigHierarchyController.cpp
+++ b/Engine/Plugins/Animation/ControlRig/Source/ControlRig/Private/Rigs/RigHierarchyController.cpp
@@ -36,29 +36,19 @@ void URigHierarchyController::Serialize(FArchive& Ar)
     }
 }
 
+URigHierarchy* URigHierarchyController::GetHierarchy() const
+{
+	return Cast<URigHierarchy>(GetOuter());;
+}
+
 void URigHierarchyController::SetHierarchy(URigHierarchy* InHierarchy)
 {
 	// since we changed the controller to be a property of the hierarchy,
 	// controlling a different hierarchy is no longer allowed
 	if (ensure(InHierarchy == GetOuter()))
 	{
-		// make sure making multiple valid SetHierarchy() calls won't lead to accumulated delegates
-		// though it should not happen in the first place
-		if (URigHierarchy* Hierarchy = WeakHierarchy.Get())
-		{
-			if(!Hierarchy->HasAnyFlags(RF_BeginDestroyed) && Hierarchy->IsValidLowLevel())
-			{
-				  Hierarchy->OnModified().RemoveAll(this);
-			}
-			Hierarchy = nullptr;
-		}
-		
-		URigHierarchy* OuterHierarchy = Cast<URigHierarchy>(GetOuter());
-		if (ensure(OuterHierarchy) && ::IsValid(OuterHierarchy))
-		{
-			WeakHierarchy = OuterHierarchy;
-			WeakHierarchy->OnModified().AddUObject(this, &URigHierarchyController::HandleHierarchyModified);
-		}
+		InHierarchy->OnModified().RemoveAll(this);
+		InHierarchy->OnModified().AddUObject(this, &URigHierarchyController::HandleHierarchyModified);
 	}
 	else
 	{
@@ -1999,9 +1989,11 @@ void URigHierarchyController::HandleHierarchyModified(ERigHierarchyNotification
 
 bool URigHierarchyController::IsValid() const
 {
-	// If we're pending kill, it's fine for our Hierarchy to also be pending kill
-	const bool bPendingKillAcceptable = ::IsValid(this);
-	return WeakHierarchy.IsValid(bPendingKillAcceptable);
+	if(::IsValid(this))
+	{
+		return ::IsValid(GetHierarchy());
+	}
+	return false;
 }
 
 FName URigHierarchyController::GetSafeNewName(const FName& InDesiredName, ERigElementType InElementType, bool bAllowNameSpace) const
diff --git a/Engine/Plugins/Animation/ControlRig/Source/ControlRig/Public/Rigs/RigHierarchyController.h b/Engine/Plugins/Animation/ControlRig/Source/ControlRig/Public/Rigs/RigHierarchyController.h
--- a/Engine/Plugins/Animation/ControlRig/Source/ControlRig/Public/Rigs/RigHierarchyController.h
+++ b/Engine/Plugins/Animation/ControlRig/Source/ControlRig/Public/Rigs/RigHierarchyController.h
@@ -18,7 +18,6 @@ public:
 
 	URigHierarchyController()
 	: bReportWarningsAndErrors(true)
-	, WeakHierarchy(nullptr)
 	, bSuspendAllNotifications(false)
 	, bSuspendSelectionNotifications(false)
 	, bSuspendPythonPrinting(false)
@@ -32,10 +31,7 @@ public:
 
 	// Returns the hierarchy currently linked to this controller
 	UFUNCTION(BlueprintCallable, Category = URigHierarchyController)
-	URigHierarchy* GetHierarchy() const
-	{
-		return WeakHierarchy.Get();
-	}
+	URigHierarchy* GetHierarchy() const;
 
 	// Sets the hierarchy currently linked to this controller
 	UFUNCTION(BlueprintCallable, Category = URigHierarchyController)
@@ -630,9 +626,6 @@ public:
 	
 private:
 
-	UPROPERTY(transient)
-	TWeakObjectPtr<URigHierarchy> WeakHierarchy;
-
 	FRigHierarchyModifiedEvent ModifiedEvent;
 	void Notify(ERigHierarchyNotification InNotifType, const FRigBaseElement* InElement);
 	void HandleHierarchyModified(ERigHierarchyNotification InNotifType, URigHierarchy* InHierarchy, const FRigBaseElement* InElement) const;
