Sequencer:  Details panel, support for constraints.

diff --git a/Engine/Plugins/Animation/ControlRig/Source/ControlRigEditor/Private/EditMode/AnimDetailsProxy.cpp b/Engine/Plugins/Animation/ControlRig/Source/ControlRigEditor/Private/EditMode/AnimDetailsProxy.cpp
--- a/Engine/Plugins/Animation/ControlRig/Source/ControlRigEditor/Private/EditMode/AnimDetailsProxy.cpp
+++ b/Engine/Plugins/Animation/ControlRig/Source/ControlRigEditor/Private/EditMode/AnimDetailsProxy.cpp
@@ -41,6 +41,12 @@
 #include "MVVM/ViewModels/TrackModel.h"
 #include "IKeyArea.h"
 #include "SequencerAddKeyOperation.h"
+#include "TransformConstraint.h"
+#include "LevelEditorViewport.h"
+#include "Units/Execution/RigUnit_BeginExecution.h"
+#include "ConstraintsManager.h"
+#include "Constraints/MovieSceneConstraintChannelHelper.h"
+#include "Constraints/ControlRigTransformableHandle.h"
 
 #include UE_INLINE_GENERATED_CPP_BY_NAME(AnimDetailsProxy)
 
@@ -734,6 +740,9 @@ void UAnimDetailControlsKeyedProxy::PostEditUndo()
 {
 	FRigControlModifiedContext Context;
 	Context.SetKey = EControlRigSetKey::Never;
+	UWorld* World = GCurrentLevelEditingViewportClient ? GCurrentLevelEditingViewportClient->GetWorld() : nullptr;
+	const FConstraintsManagerController& Controller = FConstraintsManagerController::Get(World);
+	Controller.EvaluateAllConstraints();
 	for (const TPair<TWeakObjectPtr<UControlRig>, FControlRigProxyItem>& Items : ControlRigItems)
 	{
 		if (UControlRig* ControlRig = Items.Value.ControlRig.Get())
@@ -762,6 +771,7 @@ void UAnimDetailControlsKeyedProxy::PostEditUndo()
 			}
 		}
 	}
+	ValueChanged();
 }
 #endif
 
@@ -812,28 +822,29 @@ static void SetValuesFromContext(const FEulerTransform& EulerTransform, const FR
 static FEulerTransform GetCurrentValue(UControlRig* ControlRig, FRigControlElement* ControlElement)
 {
 	FEulerTransform EulerTransform = FEulerTransform::Identity;
+
 	switch (ControlElement->Settings.ControlType)
 	{
-	case ERigControlType::Transform:
-	{
-		const FTransform NewTransform = ControlRig->GetControlValue(ControlElement, ERigControlValueType::Current).Get<FRigControlValue::FTransform_Float>().ToTransform();
-		EulerTransform = FEulerTransform(NewTransform);
-		break;
+		case ERigControlType::Transform:
+		{
+			const FTransform NewTransform = ControlRig->GetControlValue(ControlElement, ERigControlValueType::Current).Get<FRigControlValue::FTransform_Float>().ToTransform();
+			EulerTransform = FEulerTransform(NewTransform);
+			break;
 
-	}
-	case ERigControlType::TransformNoScale:
-	{
-		const FTransformNoScale NewTransform = ControlRig->GetControlValue(ControlElement, ERigControlValueType::Current).Get<FRigControlValue::FTransformNoScale_Float>().ToTransform();
-		EulerTransform.Location = NewTransform.Location;
-		EulerTransform.Rotation = FRotator(NewTransform.Rotation);
-		break;
+		}
+		case ERigControlType::TransformNoScale:
+		{
+			const FTransformNoScale NewTransform = ControlRig->GetControlValue(ControlElement, ERigControlValueType::Current).Get<FRigControlValue::FTransformNoScale_Float>().ToTransform();
+			EulerTransform.Location = NewTransform.Location;
+			EulerTransform.Rotation = FRotator(NewTransform.Rotation);
+			break;
 
-	}
-	case ERigControlType::EulerTransform:
-	{
-		EulerTransform = ControlRig->GetControlValue(ControlElement, ERigControlValueType::Current).Get<FRigControlValue::FEulerTransform_Float>().ToTransform();
-		break;
-	}
+		}
+		case ERigControlType::EulerTransform:
+		{
+			EulerTransform = ControlRig->GetControlValue(ControlElement, ERigControlValueType::Current).Get<FRigControlValue::FEulerTransform_Float>().ToTransform();
+			break;
+		}
 	};
 	EulerTransform.Rotation = ControlRig->GetHierarchy()->GetControlPreferredRotator(ControlElement);
 	return EulerTransform;
@@ -1015,6 +1026,72 @@ void UAnimDetailControlsProxyTransform::SetBindingValueFromCurrent(UObject* InOb
 	}
 }
 
+static FTransform GetControlRigComponentTransform(UControlRig* ControlRig)
+{
+	FTransform Transform = FTransform::Identity;
+	TSharedPtr<IControlRigObjectBinding> ObjectBinding = ControlRig->GetObjectBinding();
+	if (ObjectBinding.IsValid())
+	{
+		if (USceneComponent* BoundSceneComponent = Cast<USceneComponent>(ObjectBinding->GetBoundObject()))
+		{
+			return BoundSceneComponent->GetComponentTransform();
+		}
+	}
+	return Transform;
+}
+
+static bool SetConstrainedTransform(FTransform LocalTransform, UControlRig* ControlRig, FRigControlElement* ControlElement, const FRigControlModifiedContext& InContext)
+{
+	const FConstraintsManagerController& Controller = FConstraintsManagerController::Get(ControlRig->GetWorld());
+	const uint32 ControlHash = UTransformableControlHandle::ComputeHash(ControlRig, ControlElement->GetFName());
+	const TArray< TWeakObjectPtr<UTickableConstraint> > Constraints = Controller.GetParentConstraints(ControlHash, true);
+	if (Constraints.IsEmpty())
+	{
+		return false;
+	}
+	const int32 LastActiveIndex = FTransformConstraintUtils::GetLastActiveConstraintIndex(Constraints);
+	const bool bNeedsConstraintPostProcess = Constraints.IsValidIndex(LastActiveIndex);
+
+	if (!bNeedsConstraintPostProcess)
+	{
+		return false;
+	}
+	static constexpr bool bNotify = true, bFixEuler = true, bUndo = true;
+	FRigControlModifiedContext Context = InContext;
+	Context.EventName = FRigUnit_BeginExecution::EventName;
+	Context.bConstraintUpdate = true;
+	Context.SetKey = EControlRigSetKey::Never;
+
+	// set the global space, assumes it's attached to actor
+	// no need to compensate for constraints here, this will be done after when setting the control in the constraint space
+	{
+		TGuardValue<bool> CompensateGuard(FMovieSceneConstraintChannelHelper::bDoNotCompensate, true);
+		ControlRig->SetControlLocalTransform(
+			ControlElement->GetKey().Name, LocalTransform, bNotify, Context, bUndo, bFixEuler);
+	}
+	FTransform GlobalTransform = ControlRig->GetControlGlobalTransform(ControlElement->GetKey().Name);
+
+	// switch to constraint space
+	FTransform ToWorldTransform = GetControlRigComponentTransform(ControlRig);
+	const FTransform WorldTransform = GlobalTransform * ToWorldTransform;
+
+	const TOptional<FTransform> RelativeTransform =
+		FTransformConstraintUtils::GetConstraintsRelativeTransform(Constraints, LocalTransform, WorldTransform);
+	if (RelativeTransform)
+	{
+		LocalTransform = *RelativeTransform;
+	}
+
+	Context.bConstraintUpdate = false;
+	Context.SetKey = InContext.SetKey;
+	ControlRig->SetControlLocalTransform(ControlElement->GetKey().Name, LocalTransform, bNotify, Context, bUndo, bFixEuler);
+	ControlRig->Evaluate_AnyThread();
+	Controller.EvaluateAllConstraints();
+
+	return true;
+}
+
+
 void UAnimDetailControlsProxyTransform::SetControlRigElementValueFromCurrent(UControlRig* ControlRig, FRigControlElement* ControlElement, const FRigControlModifiedContext& Context)
 {
 	if (ControlElement && ControlRig)
@@ -1024,11 +1101,17 @@ void UAnimDetailControlsProxyTransform::SetControlRigElementValueFromCurrent(UCo
 		FVector TScale = Scale.ToVector();
 		FEulerTransform EulerTransform = GetCurrentValue(ControlRig, ControlElement);
 		SetValuesFromContext(EulerTransform, Context, TLocation, TRotation, TScale);
+		//constraints we just deal with FTransforms unfortunately, need to figure out how to handle rotation orders
+		FTransform RealTransform(TRotation, TLocation, TScale);
+		if (SetConstrainedTransform(RealTransform, ControlRig, ControlElement,Context))
+		{
+			ValueChanged();
+			return;
+		}
 		switch (ControlElement->Settings.ControlType)
 		{
 		case ERigControlType::Transform:
 		{
-			FTransform RealTransform(TRotation, TLocation, TScale);
 			FVector EulerAngle(TRotation.Roll, TRotation.Pitch, TRotation.Yaw);
 			ControlRig->GetHierarchy()->SetControlSpecifiedEulerAngle(ControlElement, EulerAngle);
 			ControlRig->SetControlValue<FRigControlValue::FTransform_Float>(ControlElement->GetKey().Name, RealTransform, true, Context, false);
@@ -2815,6 +2898,10 @@ const TArray<UControlRigControlsProxy*> UControlRigDetailPanelControlProxies::Ge
 void UControlRigDetailPanelControlProxies::ValuesChanged()
 {
 	//need to do all proxies
+	UWorld* World = GCurrentLevelEditingViewportClient ? GCurrentLevelEditingViewportClient->GetWorld() : nullptr;
+	const FConstraintsManagerController& Controller = FConstraintsManagerController::Get(World);
+	Controller.EvaluateAllConstraints();
+
 	for (UControlRigControlsProxy* P1 : SelectedControlRigProxies)
 	{
 		P1->ValueChanged();
@@ -2863,6 +2950,9 @@ UControlRigControlsProxy* UControlRigDetailPanelControlProxies::AddProxy(UContro
 			Proxy->bIsIndividual = bIsIndividual;
 			Proxy->SetFlags(RF_Transactional);
 			Proxy->Modify();
+			UWorld* World = GCurrentLevelEditingViewportClient ? GCurrentLevelEditingViewportClient->GetWorld() : nullptr;
+			const FConstraintsManagerController& Controller = FConstraintsManagerController::Get(World);
+			Controller.EvaluateAllConstraints();
 			Proxy->ValueChanged();
 		}
 	}
@@ -2902,6 +2992,9 @@ UControlRigControlsProxy* UControlRigDetailPanelControlProxies::AddProxy(UObject
 			Proxy->bIsIndividual = false;
 			Proxy->SetFlags(RF_Transactional);
 			Proxy->Modify();
+			UWorld* World = GCurrentLevelEditingViewportClient ? GCurrentLevelEditingViewportClient->GetWorld() : nullptr;
+			const FConstraintsManagerController& Controller = FConstraintsManagerController::Get(World);
+			Controller.EvaluateAllConstraints();
 			Proxy->ValueChanged();
 		}
 	}
diff --git a/Engine/Plugins/Animation/ControlRig/Source/ControlRigEditor/Private/EditMode/ControlRigControlsProxy.cpp b/Engine/Plugins/Animation/ControlRig/Source/ControlRigEditor/Private/EditMode/ControlRigControlsProxy.cpp
--- a/Engine/Plugins/Animation/ControlRig/Source/ControlRigEditor/Private/EditMode/ControlRigControlsProxy.cpp
+++ b/Engine/Plugins/Animation/ControlRig/Source/ControlRigEditor/Private/EditMode/ControlRigControlsProxy.cpp
@@ -31,6 +31,8 @@
 #include "DetailCategoryBuilder.h"
 #include "DetailWidgetRow.h"
 #include "SEnumCombo.h"
+#include "LevelEditorViewport.h"
+#include "ConstraintsManager.h"
 
 #include UE_INLINE_GENERATED_CPP_BY_NAME(ControlRigControlsProxy)
 
@@ -298,6 +300,9 @@ void UControlRigControlsProxy::PostEditChangeChainProperty(struct FPropertyChang
 			FRigControlModifiedContext Context;
 			Context.SetKey = EControlRigSetKey::DoNotCare;
 			Context.KeyMask = (uint32)GetChannelToKeyFromPropertyName(Property->GetFName());
+			UWorld* World = GCurrentLevelEditingViewportClient ? GCurrentLevelEditingViewportClient->GetWorld() : nullptr;
+			const FConstraintsManagerController& Controller = FConstraintsManagerController::Get(World);
+			Controller.EvaluateAllConstraints();
 
 			for (const TPair<TWeakObjectPtr<UControlRig>, FControlRigProxyItem>& Items : ControlRigItems)
 			{
@@ -327,6 +332,7 @@ void UControlRigControlsProxy::PostEditChangeChainProperty(struct FPropertyChang
 				}
 			}
 		}
+		ValueChanged();
 	}
 #endif
 }
diff --git a/Engine/Plugins/Animation/ControlRig/Source/ControlRigEditor/Private/EditMode/ControlsProxyDetailCustomization.cpp b/Engine/Plugins/Animation/ControlRig/Source/ControlRigEditor/Private/EditMode/ControlsProxyDetailCustomization.cpp
--- a/Engine/Plugins/Animation/ControlRig/Source/ControlRigEditor/Private/EditMode/ControlsProxyDetailCustomization.cpp
+++ b/Engine/Plugins/Animation/ControlRig/Source/ControlRigEditor/Private/EditMode/ControlsProxyDetailCustomization.cpp
@@ -786,10 +786,13 @@ TSharedRef<SWidget> FAnimDetailValueCustomization::MakeDoubleWidget(
 		.Value_Lambda([WeakHandlePtr,this,Proxy ]()
 			{
 				bool bIsMultiple = IsMultiple(Proxy, WeakHandlePtr.Pin()->GetProperty()->GetFName());
+				
+				/* wip attempt to try to get sliders working with multiple values,
+				one issue is that we need a bIsUsingSliderOnThisProxy, not a global one for al
 				if (bIsMultiple && bIsUsingSlider)
 				{
 					return TOptional<double>(MultipleValue);
-				}
+				}*/
 				double Value = 0.;
 				return (bIsMultiple == false && WeakHandlePtr.Pin()->GetValue(Value) == FPropertyAccess::Success) ?
 					TOptional<double>(Value) :
