Control Rig: Fix multiselection from the rig hierarchy not syncing properly with the schematic node selection
#rnx

diff --git a/Engine/Plugins/Animation/ControlRig/Source/ControlRigEditor/Private/Editor/ControlRigSchematicModel.cpp b/Engine/Plugins/Animation/ControlRig/Source/ControlRigEditor/Private/Editor/ControlRigSchematicModel.cpp
--- a/Engine/Plugins/Animation/ControlRig/Source/ControlRigEditor/Private/Editor/ControlRigSchematicModel.cpp
+++ b/Engine/Plugins/Animation/ControlRig/Source/ControlRigEditor/Private/Editor/ControlRigSchematicModel.cpp
@@ -1473,7 +1473,7 @@ void FControlRigSchematicModel::HandleSchematicDrop(SSchematicGraphPanel* InPane
 		FText Label;
 		if(FControlRigSchematicRigElementKeyNode* ExistingElementKeyNode = Cast<FControlRigSchematicRigElementKeyNode>(InNode->GetNodeData()))
 		{
-			Label = ExistingElementKeyNode->GetLabel();
+			Label = FText::FromName(ExistingElementKeyNode->GetKey().Name);
 		}
 		else
 		{
@@ -1483,7 +1483,7 @@ void FControlRigSchematicModel::HandleSchematicDrop(SSchematicGraphPanel* InPane
 			.Title(LOCTEXT("Add modules", "Add modules to targets"))
 			.Buttons({
 				SCustomDialog::FButton(FText::Format(LOCTEXT("AddModuleToSingle", "Add module to {0}"), Label)),
-				SCustomDialog::FButton(LOCTEXT("AddModuleToSelected", "Add module to all selected"))
+				SCustomDialog::FButton(FText::Format(LOCTEXT("AddModuleToSelected", "Add module to {0} selected"), SelectedNodes.Num()))
 				});
 
 		const int32 ButtonPressed = ChooseSelectedDialog->ShowModal();
@@ -1853,19 +1853,16 @@ void FControlRigSchematicModel::OnHierarchyModified(ERigHierarchyNotification In
 	{
 		case ERigHierarchyNotification::ElementSelected:
 		{
-			if(InHierarchy->GetSelectedKeys().Num() == 1)
+			if(FControlRigSchematicRigElementKeyNode* Node = FindElementKeyNode(InElement->GetKey()))
 			{
-				if(FControlRigSchematicRigElementKeyNode* Node = FindElementKeyNode(InElement->GetKey()))
+				if(InHierarchy->GetSelectedKeys().Num() == 1)
 				{
 					if(FSchematicGraphGroupNode* GroupNode = Node->GetGroupNode())
 					{
 						GroupNode->SetExpanded(true);
 					}
-					else
-					{
-						Node->SetSelected(true);
-					}
 				}
+				Node->SetSelected(true);
 			}
 			break;
 		}
