Control Rig: Guard against dead delegate handles on hierarchy

#rnx

diff --git a/Engine/Plugins/Animation/ControlRig/Source/ControlRigEditor/Private/ControlRigElementDetails.cpp b/Engine/Plugins/Animation/ControlRig/Source/ControlRigEditor/Private/ControlRigElementDetails.cpp
--- a/Engine/Plugins/Animation/ControlRig/Source/ControlRigEditor/Private/ControlRigElementDetails.cpp
+++ b/Engine/Plugins/Animation/ControlRig/Source/ControlRigEditor/Private/ControlRigElementDetails.cpp
@@ -1307,13 +1307,16 @@ void FRigBaseElementDetails::CustomizeMetadata(IDetailLayoutBuilder& DetailBuild
 			return;
 		}
 			
-		TSharedRef<IPropertyUtilities> PropertyUtilities = DetailBuilder.GetPropertyUtilities();
-		MetadataHandle = Hierarchy->OnMetadataChanged().AddLambda([this, PropertyUtilities](const FRigElementKey& InKey, const FName&)
+		TWeakPtr<IPropertyUtilities> WeakPropertyUtilities =  DetailBuilder.GetPropertyUtilities().ToWeakPtr();
+		MetadataHandle = Hierarchy->OnMetadataChanged().AddLambda([this, WeakPropertyUtilities](const FRigElementKey& InKey, const FName&)
 		{
-			const FRigBaseElement* Element = PerElementInfos.Num() == 1 ? PerElementInfos[0].GetElement() : nullptr;
-			if (InKey.Type == ERigElementType::All || (Element && Element->GetKey() == InKey))
+			if(WeakPropertyUtilities.IsValid())
 			{
-				PropertyUtilities->ForceRefresh();
+				const FRigBaseElement* Element = PerElementInfos.Num() == 1 ? PerElementInfos[0].GetElement() : nullptr;
+				if (InKey.Type == ERigElementType::All || (Element && Element->GetKey() == InKey))
+				{
+					WeakPropertyUtilities.Pin()->ForceRefresh();
+				}
 			}
 		});
 	}
