Control Rig: Guard against namespaced names being passed in

#rnx

diff --git a/Engine/Plugins/Animation/ControlRig/Source/ControlRig/Private/Rigs/RigHierarchy.cpp b/Engine/Plugins/Animation/ControlRig/Source/ControlRig/Private/Rigs/RigHierarchy.cpp
--- a/Engine/Plugins/Animation/ControlRig/Source/ControlRig/Private/Rigs/RigHierarchy.cpp
+++ b/Engine/Plugins/Animation/ControlRig/Source/ControlRig/Private/Rigs/RigHierarchy.cpp
@@ -1555,6 +1555,11 @@ FRigName URigHierarchy::GetSafeNewName(const FRigName& InPotentialNewName, ERigE
 		const FControlRigExecuteContext& CRContext = ExecuteContext->GetPublicData<FControlRigExecuteContext>();
 		if(CRContext.IsRigModule())
 		{
+			FRigName LastSegmentName;
+			if(SplitNameSpace(SanitizedName, nullptr, &LastSegmentName, true))
+			{
+				SanitizedName = LastSegmentName;
+			}
 			SanitizedName = CRContext.GetRigModuleNameSpace() + SanitizedName.GetName();
 			bAllowNameSpaceWhenSanitizingName = true;
 		}
diff --git a/Engine/Plugins/Animation/ControlRig/Source/ControlRig/Private/Rigs/RigHierarchyController.cpp b/Engine/Plugins/Animation/ControlRig/Source/ControlRig/Private/Rigs/RigHierarchyController.cpp
--- a/Engine/Plugins/Animation/ControlRig/Source/ControlRig/Private/Rigs/RigHierarchyController.cpp
+++ b/Engine/Plugins/Animation/ControlRig/Source/ControlRig/Private/Rigs/RigHierarchyController.cpp
@@ -2029,6 +2029,12 @@ int32 URigHierarchyController::AddElement(FRigBaseElement* InElementToAdd, FRigB
 
 	FRigName DesiredName = InDesiredName;
 	URigHierarchy::SanitizeName(DesiredName);
+	
+	FRigName LastSegmentName;
+	if(URigHierarchy::SplitNameSpace(DesiredName, nullptr, &LastSegmentName, true))
+	{
+		DesiredName = LastSegmentName;
+	}
 
 	if(!InDesiredName.IsNone() &&
 		!InElementToAdd->GetFName().IsEqual(DesiredName.GetFName(), ENameCase::CaseSensitive))
