Control Rig: Remove popup asking to apply module to single or multiple targets
#rnx

diff --git a/Engine/Plugins/Animation/ControlRig/Source/ControlRigEditor/Private/Editor/ControlRigSchematicModel.cpp b/Engine/Plugins/Animation/ControlRig/Source/ControlRigEditor/Private/Editor/ControlRigSchematicModel.cpp
--- a/Engine/Plugins/Animation/ControlRig/Source/ControlRigEditor/Private/Editor/ControlRigSchematicModel.cpp
+++ b/Engine/Plugins/Animation/ControlRig/Source/ControlRigEditor/Private/Editor/ControlRigSchematicModel.cpp
@@ -1468,35 +1468,6 @@ void FControlRigSchematicModel::HandleSchematicDrop(SSchematicGraphPanel* InPane
 	}
 	SelectedNodes.AddUnique(InNode->GetNodeData());
 
-	if (SelectedNodes.Num() > 1)
-	{
-		FText Label;
-		if(FControlRigSchematicRigElementKeyNode* ExistingElementKeyNode = Cast<FControlRigSchematicRigElementKeyNode>(InNode->GetNodeData()))
-		{
-			Label = FText::FromName(ExistingElementKeyNode->GetKey().Name);
-		}
-		else
-		{
-			Label = InNode->GetNodeData()->GetLabel();
-		}
-		TSharedRef<SCustomDialog> ChooseSelectedDialog = SNew(SCustomDialog)
-			.Title(LOCTEXT("Add modules", "Add modules to targets"))
-			.Buttons({
-				SCustomDialog::FButton(FText::Format(LOCTEXT("AddModuleToSingle", "Add module to {0}"), Label)),
-				SCustomDialog::FButton(FText::Format(LOCTEXT("AddModuleToSelected", "Add module to {0} selected"), SelectedNodes.Num()))
-				});
-
-		const int32 ButtonPressed = ChooseSelectedDialog->ShowModal();
-		if (ButtonPressed < 0)
-		{
-			return; // Window closed
-		}
-		if (ButtonPressed == 0)
-		{
-			SelectedNodes.Reset();
-			SelectedNodes.Add(InNode->GetNodeData());
-		}
-	}
 	
 	for (FSchematicGraphNode* Node : SelectedNodes)
 	{
