Control Rig: Unwinding of euler angles in edit mode
#rnx

diff --git a/Engine/Plugins/Animation/ControlRig/Source/ControlRigEditor/Private/EditMode/ControlRigEditMode.cpp b/Engine/Plugins/Animation/ControlRig/Source/ControlRigEditor/Private/EditMode/ControlRigEditMode.cpp
--- a/Engine/Plugins/Animation/ControlRig/Source/ControlRigEditor/Private/EditMode/ControlRigEditMode.cpp
+++ b/Engine/Plugins/Animation/ControlRig/Source/ControlRigEditor/Private/EditMode/ControlRigEditMode.cpp
@@ -4446,7 +4446,8 @@ void FControlRigEditMode::SetControlShapeTransform(
 	const FTransform& InGlobalTransform,
 	const FTransform& InToWorldTransform,
 	const FRigControlModifiedContext& InContext,
-	const bool bPrintPython) const
+	const bool bPrintPython,
+	const bool bFixEulerFlips) const
 {
 	UControlRig* ControlRig = InShapeActor->ControlRig.Get();
 	if (!ControlRig)
@@ -4454,12 +4455,12 @@ void FControlRigEditMode::SetControlShapeTransform(
 		return;
 	}
 
-	static constexpr bool bNotify = true, bFixEuler = true, bUndo = true;
+	static constexpr bool bNotify = true, bUndo = true;
 	if (AreEditingControlRigDirectly())
 	{
 		// assumes it's attached to actor
 		ControlRig->SetControlGlobalTransform(
-			InShapeActor->ControlName, InGlobalTransform, bNotify, InContext, bUndo, bPrintPython, bFixEuler);
+			InShapeActor->ControlName, InGlobalTransform, bNotify, InContext, bUndo, bPrintPython, bFixEulerFlips);
 		return;
 	}
 	
@@ -4475,7 +4476,7 @@ void FControlRigEditMode::SetControlShapeTransform(
 	{
 		TGuardValue<bool> CompensateGuard(FMovieSceneConstraintChannelHelper::bDoNotCompensate, true);
 		ControlRig->SetControlGlobalTransform(
-			InShapeActor->ControlName, InGlobalTransform, bNotify, InContext, bUndo, bPrintPython, bFixEuler);
+			InShapeActor->ControlName, InGlobalTransform, bNotify, InContext, bUndo, bPrintPython, bFixEulerFlips);
 	}
 
 	if (!bNeedsConstraintPostProcess)
@@ -4497,7 +4498,7 @@ void FControlRigEditMode::SetControlShapeTransform(
 	FRigControlModifiedContext Context = InContext;
 	Context.bConstraintUpdate = false;
 	
-	ControlRig->SetControlLocalTransform(InShapeActor->ControlName, LocalTransform, bNotify, Context, bUndo, bFixEuler);
+	ControlRig->SetControlLocalTransform(InShapeActor->ControlName, LocalTransform, bNotify, Context, bUndo, bFixEulerFlips);
 
 	const FConstraintsManagerController& Controller = FConstraintsManagerController::Get(ControlRig->GetWorld());
 	Controller.EvaluateAllConstraints();
@@ -4517,6 +4518,41 @@ void FControlRigEditMode::MoveControlShape(AControlRigShapeActor* ShapeActor, co
 	bool bUseLocal, bool bCalcLocal, FTransform& InOutLocal)
 {
 	bool bTransformChanged = false;
+
+	auto RotatorToStr = [](const FRotator& InRotator)
+	{
+		FRigPreferredEulerAngles EulerAngles;
+		EulerAngles.SetRotator(InRotator);
+		return EulerAngles.Current.ToString();
+	};
+
+	auto UpdatePreferredEulerAngles = [ShapeActor, bRotation, InRot, RotatorToStr](UControlRig* InControlRig)
+	{
+		if(bRotation)
+		{
+			if(FRigControlElement* ControlElement = InControlRig->GetHierarchy()->Find<FRigControlElement>(ShapeActor->GetElementKey()))
+			{
+				if(ControlElement->Settings.bUsePreferredRotationOrder)
+				{
+					FRotator Rot = ControlElement->PreferredEulerAngles.GetRotator();
+
+					// Split the current rotation between winding and remainder
+					FRotator CurrentRotWind, CurrentRotRem;
+					Rot.GetWindingAndRemainder(CurrentRotWind, CurrentRotRem);
+
+					// Apply the delta to the current remainder, and normalize
+					FRotator NewRotRem = CurrentRotRem + InRot;
+					NewRotRem.Normalize();
+
+					// Add the current winding to the new remainder
+					const FRotator Final = CurrentRotWind + NewRotRem;
+					
+					InControlRig->GetHierarchy()->SetControlPreferredRotator(ControlElement, Final, false, false);
+				}
+			}
+		}
+	};
+	
 	//first case is where we do all controls by the local diff.
 	if (bUseLocal)
 	{
@@ -4547,7 +4583,8 @@ void FControlRigEditMode::MoveControlShape(AControlRigShapeActor* ShapeActor, co
 				ControlRig->InteractionType = InteractionType;
 				ControlRig->ElementsBeingInteracted.AddUnique(ShapeActor->GetElementKey());
 				
-				ControlRig->SetControlLocalTransform(ShapeActor->ControlName, CurrentLocalTransform,true, FRigControlModifiedContext(), true, true);
+				ControlRig->SetControlLocalTransform(ShapeActor->ControlName, CurrentLocalTransform,true, FRigControlModifiedContext(), true, false);
+				UpdatePreferredEulerAngles(ControlRig);
 
 				FTransform CurrentTransform  = ControlRig->GetControlGlobalTransform(ShapeActor->ControlName);			// assumes it's attached to actor
 				CurrentTransform = ToWorldTransform * CurrentTransform;
@@ -4630,7 +4667,8 @@ void FControlRigEditMode::MoveControlShape(AControlRigShapeActor* ShapeActor, co
 				}
 				
 				ControlRig->Evaluate_AnyThread();
-				SetControlShapeTransform(ShapeActor, NewTransform, ToWorldTransform, Context, bPrintPythonCommands);
+				SetControlShapeTransform(ShapeActor, NewTransform, ToWorldTransform, Context, bPrintPythonCommands, false);
+				UpdatePreferredEulerAngles(ControlRig);
 				NotifyDrivenControls(ControlRig, ShapeActor->GetElementKey());
 				if(const FRigControlElement* ControlElement = ControlRig->FindControl(ShapeActor->ControlName))
 				{
diff --git a/Engine/Plugins/Animation/ControlRig/Source/ControlRigEditor/Public/EditMode/ControlRigEditMode.h b/Engine/Plugins/Animation/ControlRig/Source/ControlRigEditor/Public/EditMode/ControlRigEditMode.h
--- a/Engine/Plugins/Animation/ControlRig/Source/ControlRigEditor/Public/EditMode/ControlRigEditMode.h
+++ b/Engine/Plugins/Animation/ControlRig/Source/ControlRigEditor/Public/EditMode/ControlRigEditMode.h
@@ -485,7 +485,8 @@ private:
 	/* Check on tick to see if movie scene has changed, returns true if it has*/
 	bool CheckMovieSceneSig();
 	void SetControlShapeTransform( const AControlRigShapeActor* InShapeActor, const FTransform& InGlobalTransform,
-		const FTransform& InToWorldTransform, const FRigControlModifiedContext& InContext, const bool bPrintPython) const;
+		const FTransform& InToWorldTransform, const FRigControlModifiedContext& InContext, const bool bPrintPython,
+		const bool bFixEulerFlips = true) const;
 	static FTransform GetControlShapeTransform(const AControlRigShapeActor* ShapeActor);
 	void MoveControlShape(AControlRigShapeActor* ShapeActor, const bool bTranslation, FVector& InDrag, 
 		const bool bRotation, FRotator& InRot, const bool bScale, FVector& InScale, const FTransform& ToWorldTransform,
