Control Rig: Change order of secondary connector resolval and module reparenting

#rnx

diff --git a/Engine/Plugins/Animation/ControlRig/Source/ControlRig/Private/ModularRigController.cpp b/Engine/Plugins/Animation/ControlRig/Source/ControlRig/Private/ModularRigController.cpp
--- a/Engine/Plugins/Animation/ControlRig/Source/ControlRig/Private/ModularRigController.cpp
+++ b/Engine/Plugins/Animation/ControlRig/Source/ControlRig/Private/ModularRigController.cpp
@@ -276,15 +276,10 @@ bool UModularRigController::ConnectConnectorToElement(const FRigElementKey& InCo
 		{
 			if (URigHierarchy* Hierarchy = ModularRig->GetHierarchy())
 			{
-				if (bAutoResolveOtherConnectors)
+				bool bResolvedPrimaryConnector = false;
+				if(const FRigConnectorElement* PrimaryConnector = Module->FindPrimaryConnector(Hierarchy))
 				{
-					if(const FRigConnectorElement* PrimaryConnector = Module->FindPrimaryConnector(Hierarchy))
-					{
-						if(PrimaryConnector->GetKey() == InConnectorKey)
-						{
-							(void)AutoConnectModules( {Module->GetPath()}, false, bSetupUndo);
-						}
-					}
+					bResolvedPrimaryConnector = PrimaryConnector->GetKey() == InConnectorKey;
 				}
 
 				// automatically re-parent the module in the module tree as well
@@ -296,11 +291,20 @@ bool UModularRigController::ConnectConnectorToElement(const FRigElementKey& InCo
 						{
 							if(!TargetModulePathName.IsNone())
 							{
-								ReparentModule(Module->GetPath(), TargetModulePathName.ToString(), bSetupUndo);
+								const FString NewModulePath = ReparentModule(Module->GetPath(), TargetModulePathName.ToString(), bSetupUndo);
+								if(!NewModulePath.IsEmpty())
+								{
+									Module = FindModule(NewModulePath);
+								}
 							}
 						}
 					}
 				}
+
+				if (bAutoResolveOtherConnectors && bResolvedPrimaryConnector)
+				{
+					(void)AutoConnectModules( {Module->GetPath()}, false, bSetupUndo);
+				}
 			}
 		}
 	}
