Control Rig: Always set rotation angle
#trivial

diff --git a/Engine/Plugins/Animation/ControlRig/Source/ControlRigEditor/Private/EditMode/ControlRigEditMode.cpp b/Engine/Plugins/Animation/ControlRig/Source/ControlRigEditor/Private/EditMode/ControlRigEditMode.cpp
--- a/Engine/Plugins/Animation/ControlRig/Source/ControlRigEditor/Private/EditMode/ControlRigEditMode.cpp
+++ b/Engine/Plugins/Animation/ControlRig/Source/ControlRigEditor/Private/EditMode/ControlRigEditMode.cpp
@@ -4532,7 +4532,8 @@ void FControlRigEditMode::MoveControlShape(AControlRigShapeActor* ShapeActor, co
 		{
 			if(FRigControlElement* ControlElement = InControlRig->GetHierarchy()->Find<FRigControlElement>(ShapeActor->GetElementKey()))
 			{
-				if(ControlElement->Settings.bUsePreferredRotationOrder)
+				//if(ControlElement->Settings.bUsePreferredRotationOrder) always set rotation order since
+				//sequencer depends upon it
 				{
 					FRotator Rot = ControlElement->PreferredEulerAngles.GetRotator();
 
