Sequencer: Anim Details: Fix for anim details showing multiple curves with multiple actors, and fix for int64 sliders. Also revert preferred angle change for control rig editor, breaks auto key.

diff --git a/Engine/Plugins/Animation/ControlRig/Source/ControlRigEditor/Private/EditMode/ControlRigEditMode.cpp b/Engine/Plugins/Animation/ControlRig/Source/ControlRigEditor/Private/EditMode/ControlRigEditMode.cpp
--- a/Engine/Plugins/Animation/ControlRig/Source/ControlRigEditor/Private/EditMode/ControlRigEditMode.cpp
+++ b/Engine/Plugins/Animation/ControlRig/Source/ControlRigEditor/Private/EditMode/ControlRigEditMode.cpp
@@ -4584,7 +4584,7 @@ void FControlRigEditMode::MoveControlShape(AControlRigShapeActor* ShapeActor, co
 				ControlRig->InteractionType = InteractionType;
 				ControlRig->ElementsBeingInteracted.AddUnique(ShapeActor->GetElementKey());
 				
-				ControlRig->SetControlLocalTransform(ShapeActor->ControlName, CurrentLocalTransform,true, FRigControlModifiedContext(), true, false);
+				ControlRig->SetControlLocalTransform(ShapeActor->ControlName, CurrentLocalTransform,true, FRigControlModifiedContext(), true, /*fix eulers*/ false);
 				UpdatePreferredEulerAngles(ControlRig);
 
 				FTransform CurrentTransform  = ControlRig->GetControlGlobalTransform(ShapeActor->ControlName);			// assumes it's attached to actor
@@ -4668,8 +4668,8 @@ void FControlRigEditMode::MoveControlShape(AControlRigShapeActor* ShapeActor, co
 				}
 				
 				ControlRig->Evaluate_AnyThread();
-				SetControlShapeTransform(ShapeActor, NewTransform, ToWorldTransform, Context, bPrintPythonCommands, false);
-				UpdatePreferredEulerAngles(ControlRig);
+				SetControlShapeTransform(ShapeActor, NewTransform, ToWorldTransform, Context, bPrintPythonCommands, /*fix flips*/ true);
+				//UpdatePreferredEulerAngles(ControlRig);
 				NotifyDrivenControls(ControlRig, ShapeActor->GetElementKey());
 				if(const FRigControlElement* ControlElement = ControlRig->FindControl(ShapeActor->ControlName))
 				{
@@ -5419,6 +5419,7 @@ bool FDetailKeyFrameCacheAndHandler::IsPropertyKeyable(const UClass* InObjectCla
 		&& InObjectClass->IsChildOf(UAnimDetailControlsProxyVector2D::StaticClass())
 		&& InObjectClass->IsChildOf(UAnimDetailControlsProxyFloat::StaticClass())
 		&& InObjectClass->IsChildOf(UAnimDetailControlsProxyBool::StaticClass())
+		&& InObjectClass->IsChildOf(UAnimDetailControlsProxyInteger::StaticClass())
 		)
 	{
 		return true;
@@ -5457,11 +5458,18 @@ bool FDetailKeyFrameCacheAndHandler::IsPropertyKeyable(const UClass* InObjectCla
 		return true;
 	}
 
+	if (InObjectClass && InObjectClass->IsChildOf(UAnimDetailControlsProxyInteger::StaticClass()) && InPropertyHandle.GetProperty()
+		&& InPropertyHandle.GetProperty()->GetFName() == GET_MEMBER_NAME_CHECKED(UAnimDetailControlsProxyInteger, Integer))
+	{
+		return true;
+	}
+
 	if (InObjectClass && InObjectClass->IsChildOf(UAnimDetailControlsProxyBool::StaticClass()) && InPropertyHandle.GetProperty()
 		&& InPropertyHandle.GetProperty()->GetFName() == GET_MEMBER_NAME_CHECKED(UAnimDetailControlsProxyBool, Bool))
 	{
 		return true;
 	}
+
 	if (InObjectClass && InObjectClass->IsChildOf(UAnimDetailControlsProxyFloat::StaticClass()) && InPropertyHandle.GetProperty()
 		&& InPropertyHandle.GetProperty()->GetFName() == GET_MEMBER_NAME_CHECKED(UAnimDetailControlsProxyFloat, Float))
 	{
diff --git a/Engine/Plugins/Animation/ControlRig/Source/ControlRigEditor/Private/EditMode/ControlsProxyDetailCustomization.cpp b/Engine/Plugins/Animation/ControlRig/Source/ControlRigEditor/Private/EditMode/ControlsProxyDetailCustomization.cpp
--- a/Engine/Plugins/Animation/ControlRig/Source/ControlRigEditor/Private/EditMode/ControlsProxyDetailCustomization.cpp
+++ b/Engine/Plugins/Animation/ControlRig/Source/ControlRigEditor/Private/EditMode/ControlsProxyDetailCustomization.cpp
@@ -18,6 +18,7 @@
 #include "PropertyCustomizationHelpers.h"
 #include "PropertyEditorModule.h"
 #include "HAL/PlatformApplicationMisc.h"
+#include "Sequencer/MovieSceneControlRigParameterTrack.h"
 
 #include "Widgets/SBoxPanel.h"
 #include "Widgets/Layout/SBorder.h"
@@ -32,6 +33,7 @@
 #include "MVVM/ViewModels/OutlinerViewModel.h"
 #include "MVVM/ViewModels/ViewModelIterators.h"
 #include "MVVM/Extensions/IOutlinerExtension.h"
+#include "MVVM/ViewModels/TrackModel.h"
 #include "Rigs/RigControlHierarchy.h"
 #include "Tree/SCurveEditorTree.h"
 
@@ -562,26 +564,36 @@ void FAnimDetailValueCustomization::TogglePropertySelection(UControlRigControlsP
 						TParentFirstChildIterator<IOutlinerExtension> OutlinerExtenstionIt = OutlinerViewModel->GetRootItem()->GetDescendantsOfType<IOutlinerExtension>();
 						for (; OutlinerExtenstionIt; ++OutlinerExtenstionIt)
 						{
-							FName ID = OutlinerExtenstionIt->GetIdentifier();
-							FString Name = ID.ToString();
-							TArray<FString> StringArray;
-							Name.ParseIntoArray(StringArray, TEXT("."));
-
-							FString ChannelName;
-							if(GetChannelNameForCurve(StringArray, ControlElement, ChannelName))
+							if (TSharedPtr<FTrackModel> TrackModel = OutlinerExtenstionIt.GetCurrentItem()->FindAncestorOfType<FTrackModel>())
 							{
-								EControlRigContextChannelToKey ChannelToKeyFromCurve = Proxy->GetChannelToKeyFromChannelName(ChannelName);
-								if (ChannelToKey == ChannelToKeyFromCurve)
+								if (UMovieSceneControlRigParameterTrack* Track =  Cast<UMovieSceneControlRigParameterTrack>(TrackModel->GetTrack()) )
 								{
-									if (TViewModelPtr<ICurveEditorTreeItemExtension> CurveEditorItem = OutlinerExtenstionIt.GetCurrentItem().ImplicitCast())
+									if (Track->GetControlRig() != ControlRig)
 									{
-										FCurveEditorTreeItemID CurveEditorTreeItem = CurveEditorItem->GetCurveEditorItemID();
-										if (CurveEditorTreeItem != FCurveEditorTreeItemID::Invalid())
+										continue;
+									}
+									FName ID = OutlinerExtenstionIt->GetIdentifier();
+									FString Name = ID.ToString();
+									TArray<FString> StringArray;
+									Name.ParseIntoArray(StringArray, TEXT("."));
+
+									FString ChannelName;
+									if (GetChannelNameForCurve(StringArray, ControlElement, ChannelName))
+									{
+										EControlRigContextChannelToKey ChannelToKeyFromCurve = Proxy->GetChannelToKeyFromChannelName(ChannelName);
+										if (ChannelToKey == ChannelToKeyFromCurve)
 										{
-											const bool bSelected = bIsShiftDown ? true : !CurveEditorTreeView->IsItemSelected(CurveEditorTreeItem);
-											if (bIsCtrlDown || bSelected)
+											if (TViewModelPtr<ICurveEditorTreeItemExtension> CurveEditorItem = OutlinerExtenstionIt.GetCurrentItem().ImplicitCast())
 											{
-												CurveEditorTreeView->SetItemSelection(CurveEditorTreeItem, bSelected);
+												FCurveEditorTreeItemID CurveEditorTreeItem = CurveEditorItem->GetCurveEditorItemID();
+												if (CurveEditorTreeItem != FCurveEditorTreeItemID::Invalid())
+												{
+													const bool bSelected = bIsShiftDown ? true : !CurveEditorTreeView->IsItemSelected(CurveEditorTreeItem);
+													if (bIsCtrlDown || bSelected)
+													{
+														CurveEditorTreeView->SetItemSelection(CurveEditorTreeItem, bSelected);
+													}
+												}
 											}
 										}
 									}
@@ -602,33 +614,39 @@ void FAnimDetailValueCustomization::TogglePropertySelection(UControlRigControlsP
 					TParentFirstChildIterator<IOutlinerExtension> OutlinerExtenstionIt = OutlinerViewModel->GetRootItem()->GetDescendantsOfType<IOutlinerExtension>();
 					for (; OutlinerExtenstionIt; ++OutlinerExtenstionIt)
 					{
-						FName ID = OutlinerExtenstionIt->GetIdentifier();
-						FString Name = ID.ToString();
-						TArray<FString> StringArray;
-						Name.ParseIntoArray(StringArray, TEXT("."));
-
-						FString ChannelName;
-						if (StringArray.Num() == 2)
-						{
-							ChannelName = StringArray[0] + "." + StringArray[1];
-						}
-						else if (StringArray.Num() == 0)
+						if (TSharedPtr<FTrackModel> TrackModel = OutlinerExtenstionIt.GetCurrentItem()->FindAncestorOfType<FTrackModel>())
 						{
-							ChannelName = StringArray[0];
-						}
-
-						EControlRigContextChannelToKey ChannelToKeyFromCurve = Proxy->GetChannelToKeyFromChannelName(ChannelName);
-						if (ChannelToKey == ChannelToKeyFromCurve)
-						{
-							if (TViewModelPtr<ICurveEditorTreeItemExtension> CurveEditorItem = OutlinerExtenstionIt.GetCurrentItem().ImplicitCast())
+							if (TrackModel->GetTrack() == Element.WeakTrack.Get())
 							{
-								FCurveEditorTreeItemID CurveEditorTreeItem = CurveEditorItem->GetCurveEditorItemID();
-								if (CurveEditorTreeItem != FCurveEditorTreeItemID::Invalid())
+								FName ID = OutlinerExtenstionIt->GetIdentifier();
+								FString Name = ID.ToString();
+								TArray<FString> StringArray;
+								Name.ParseIntoArray(StringArray, TEXT("."));
+
+								FString ChannelName;
+								if (StringArray.Num() == 2)
 								{
-									const bool bSelected = bIsShiftDown ? true : !CurveEditorTreeView->IsItemSelected(CurveEditorTreeItem);
-									if (bIsCtrlDown || bSelected)
+									ChannelName = StringArray[0] + "." + StringArray[1];
+								}
+								else if (StringArray.Num() == 0)
+								{
+									ChannelName = StringArray[0];
+								}
+
+								EControlRigContextChannelToKey ChannelToKeyFromCurve = Proxy->GetChannelToKeyFromChannelName(ChannelName);
+								if (ChannelToKey == ChannelToKeyFromCurve)
+								{
+									if (TViewModelPtr<ICurveEditorTreeItemExtension> CurveEditorItem = OutlinerExtenstionIt.GetCurrentItem().ImplicitCast())
 									{
-										CurveEditorTreeView->SetItemSelection(CurveEditorTreeItem, bSelected);
+										FCurveEditorTreeItemID CurveEditorTreeItem = CurveEditorItem->GetCurveEditorItemID();
+										if (CurveEditorTreeItem != FCurveEditorTreeItemID::Invalid())
+										{
+											const bool bSelected = bIsShiftDown ? true : !CurveEditorTreeView->IsItemSelected(CurveEditorTreeItem);
+											if (bIsCtrlDown || bSelected)
+											{
+												CurveEditorTreeView->SetItemSelection(CurveEditorTreeItem, bSelected);
+											}
+										}
 									}
 								}
 							}
@@ -702,11 +720,11 @@ TSharedRef<SWidget> FAnimDetailValueCustomization::MakeIntegerWidget(
 			})
 		.Font(IDetailLayoutBuilder::GetDetailFont())
 		.UndeterminedString(NSLOCTEXT("PropertyEditor", "MultipleValues", "Multiple Values"))
-		.OnValueCommitted_Lambda([WeakHandlePtr](double Value, ETextCommit::Type)
+		.OnValueCommitted_Lambda([WeakHandlePtr](int64 Value, ETextCommit::Type)
 			{
 				WeakHandlePtr.Pin()->SetValue(Value, EPropertyValueSetFlags::DefaultFlags);
 			})
-		.OnValueChanged_Lambda([this, WeakHandlePtr](double Value)
+		.OnValueChanged_Lambda([this, WeakHandlePtr](int64 Value)
 			{
 				if (bIsUsingSlider)
 				{
@@ -718,7 +736,7 @@ TSharedRef<SWidget> FAnimDetailValueCustomization::MakeIntegerWidget(
 				bIsUsingSlider = true;
 				GEditor->BeginTransaction(LOCTEXT("SetVectorProperty", "Set Property"));
 			})
-		.OnEndSliderMovement_Lambda([this](double Value)
+		.OnEndSliderMovement_Lambda([this](int64 Value)
 			{
 				bIsUsingSlider = false;
 				GEditor->EndTransaction();
