Schematic: Crashing on dead brush pointers

#rnx

diff --git a/Engine/Plugins/Animation/ControlRig/Source/ControlRigEditor/Private/Editor/SModularRigTreeView.cpp b/Engine/Plugins/Animation/ControlRig/Source/ControlRigEditor/Private/Editor/SModularRigTreeView.cpp
--- a/Engine/Plugins/Animation/ControlRig/Source/ControlRigEditor/Private/Editor/SModularRigTreeView.cpp
+++ b/Engine/Plugins/Animation/ControlRig/Source/ControlRigEditor/Private/Editor/SModularRigTreeView.cpp
@@ -32,7 +32,7 @@
 
 #define LOCTEXT_NAMESPACE "SModularRigTreeView"
 
-TMap<FSoftObjectPath, FSlateBrush> FModularRigTreeElement::IconPathToBrush;
+TMap<FSoftObjectPath, TSharedPtr<FSlateBrush>> FModularRigTreeElement::IconPathToBrush;
 
 //////////////////////////////////////////////////////////////
 /// FModularRigTreeElement
@@ -1030,13 +1030,19 @@ TPair<const FSlateBrush*, FSlateColor> FModularRigTreeElement::GetBrushAndColor(
 				{
 					if (bIsConnected)
 					{
-						FSoftObjectPath IconPath = ModuleRig->GetRigModuleSettings().Icon;
-						Brush = IconPathToBrush.Find(IconPath);
-						if (!Brush)
+						const FSoftObjectPath IconPath = ModuleRig->GetRigModuleSettings().Icon;
+						const TSharedPtr<FSlateBrush>* ExistingBrush = IconPathToBrush.Find(IconPath);
+						if(ExistingBrush && ExistingBrush->IsValid())
+						{
+							Brush = ExistingBrush->Get();
+						}
+						else
 						{
 							if(UTexture2D* Icon = Cast<UTexture2D>(IconPath.TryLoad()))
 							{
-								Brush = &IconPathToBrush.Add(IconPath, UWidgetBlueprintLibrary::MakeBrushFromTexture(Icon, 16.0f, 16.0f));
+								const TSharedPtr<FSlateBrush> NewBrush = MakeShareable(new FSlateBrush(UWidgetBlueprintLibrary::MakeBrushFromTexture(Icon, 16.0f, 16.0f)));
+								IconPathToBrush.FindOrAdd(IconPath) = NewBrush;
+								Brush = NewBrush.Get();
 							}
 						}
 					}
diff --git a/Engine/Plugins/Animation/ControlRig/Source/ControlRigEditor/Private/Editor/SModularRigTreeView.h b/Engine/Plugins/Animation/ControlRig/Source/ControlRigEditor/Private/Editor/SModularRigTreeView.h
--- a/Engine/Plugins/Animation/ControlRig/Source/ControlRigEditor/Private/Editor/SModularRigTreeView.h
+++ b/Engine/Plugins/Animation/ControlRig/Source/ControlRigEditor/Private/Editor/SModularRigTreeView.h
@@ -121,7 +121,7 @@ public:
 	DECLARE_DELEGATE(FOnRenameRequested);
 	FOnRenameRequested OnRenameRequested;
 
-	static TMap<FSoftObjectPath, FSlateBrush> IconPathToBrush;
+	static TMap<FSoftObjectPath, TSharedPtr<FSlateBrush>> IconPathToBrush;
 
 	/** The brush to use when rendering an icon */
 	const FSlateBrush* IconBrush;
