Control Rig: Rename Resolve unit to Get Connection and add a bool output to inform if the connector is resolved or not
#rnx

diff --git a/Engine/Plugins/Animation/ControlRig/Source/ControlRig/Private/Units/Execution/RigUnit_RigModules.cpp b/Engine/Plugins/Animation/ControlRig/Source/ControlRig/Private/Units/Execution/RigUnit_RigModules.cpp
--- a/Engine/Plugins/Animation/ControlRig/Source/ControlRig/Private/Units/Execution/RigUnit_RigModules.cpp
+++ b/Engine/Plugins/Animation/ControlRig/Source/ControlRig/Private/Units/Execution/RigUnit_RigModules.cpp
@@ -22,6 +22,8 @@ FRigUnit_ResolveConnector_Execute()
 			}
 		}
 	}
+
+	bIsConnected = Result != Connector;
 }
 
 FRigUnit_GetCurrentNameSpace_Execute()
diff --git a/Engine/Plugins/Animation/ControlRig/Source/ControlRig/Public/Units/Execution/RigUnit_RigModules.h b/Engine/Plugins/Animation/ControlRig/Source/ControlRig/Public/Units/Execution/RigUnit_RigModules.h
--- a/Engine/Plugins/Animation/ControlRig/Source/ControlRig/Public/Units/Execution/RigUnit_RigModules.h
+++ b/Engine/Plugins/Animation/ControlRig/Source/ControlRig/Public/Units/Execution/RigUnit_RigModules.h
@@ -21,7 +21,7 @@ struct CONTROLRIG_API FRigUnit_RigModulesBaseMutable : public FRigUnitMutable
 /**
 * Returns the resolved item of the connector.
 */
-USTRUCT(meta=(DisplayName="Resolve", Keywords="Connector,GetResolved,Target", Varying))
+USTRUCT(meta=(DisplayName="Get Connection", Keywords="Connector,GetResolved,Target,Resolve", Varying))
 struct CONTROLRIG_API FRigUnit_ResolveConnector : public FRigUnit_RigModulesBase
 {
 	GENERATED_BODY()
@@ -30,6 +30,7 @@ struct CONTROLRIG_API FRigUnit_ResolveConnector : public FRigUnit_RigModulesBase
 	{
 		Connector = Result = FRigElementKey(NAME_None, ERigElementType::Connector);
 		SkipSocket = false;
+		bIsConnected = false;
 	}
 
 	RIGVM_METHOD()
@@ -53,6 +54,12 @@ struct CONTROLRIG_API FRigUnit_ResolveConnector : public FRigUnit_RigModulesBase
 	 */
 	UPROPERTY(meta = (Output))
 	FRigElementKey Result;
+
+	/*
+	 * Returns true if the connector is resolved to a target.
+	 */
+	UPROPERTY(meta = (Output))
+	bool bIsConnected;
 };
 
 /**
