Control Rig: Support for attached mesh actors

#rnx

diff --git a/Engine/Plugins/Animation/ControlRig/Source/ControlRigEditor/Private/Editor/ControlRigEditor.cpp b/Engine/Plugins/Animation/ControlRig/Source/ControlRigEditor/Private/Editor/ControlRigEditor.cpp
--- a/Engine/Plugins/Animation/ControlRig/Source/ControlRigEditor/Private/Editor/ControlRigEditor.cpp
+++ b/Engine/Plugins/Animation/ControlRig/Source/ControlRigEditor/Private/Editor/ControlRigEditor.cpp
@@ -408,7 +408,12 @@ void FControlRigEditor::CreatePersonaToolKitIfRequired()
 	FPersonaToolkitArgs PersonaToolkitArgs;
 	PersonaToolkitArgs.OnPreviewSceneCreated = FOnPreviewSceneCreated::FDelegate::CreateSP(this, &FControlRigEditor::HandlePreviewSceneCreated);
 	PersonaToolkitArgs.bPreviewMeshCanUseDifferentSkeleton = true;
-	PersonaToolkit = PersonaModule.CreatePersonaToolkit(ControlRigBlueprint, PersonaToolkitArgs);
+	USkeleton* Skeleton = nullptr;
+	if(USkeletalMesh* PreviewMesh = ControlRigBlueprint->GetPreviewMesh())
+	{
+		Skeleton = PreviewMesh->GetSkeleton();
+	}
+	PersonaToolkit = PersonaModule.CreatePersonaToolkit(ControlRigBlueprint, PersonaToolkitArgs, Skeleton);
 
 	// Set a default preview mesh, if any
 	PersonaToolkit->SetPreviewMesh(ControlRigBlueprint->GetPreviewMesh(), false);
