Control Rig: Fix selection not clearing when modules are applied, and only apply multiple modules when dropped on selected node
#rnx

diff --git a/Engine/Plugins/Animation/ControlRig/Source/ControlRigEditor/Private/Editor/ControlRigSchematicModel.cpp b/Engine/Plugins/Animation/ControlRig/Source/ControlRigEditor/Private/Editor/ControlRigSchematicModel.cpp
--- a/Engine/Plugins/Animation/ControlRig/Source/ControlRigEditor/Private/Editor/ControlRigSchematicModel.cpp
+++ b/Engine/Plugins/Animation/ControlRig/Source/ControlRigEditor/Private/Editor/ControlRigSchematicModel.cpp
@@ -146,6 +146,8 @@ void FControlRigSchematicModel::Tick(float InDeltaTime)
 
 	if(ControlRigBlueprint.IsValid())
 	{
+		//UE_LOG(LogControlRig, Display, TEXT("Selected nodes %d"), GetSelectedNodes().Num());
+		
 		const URigHierarchy* Hierarchy = ControlRigBlueprint->GetDebuggedControlRig()->GetHierarchy();
 		check(Hierarchy);
 		const FModularRigConnections& Connections = ControlRigBlueprint->ModularRigModel.Connections;
@@ -1459,16 +1461,26 @@ void FControlRigSchematicModel::HandleSchematicDrop(SSchematicGraphPanel* InPane
 	TArray<TSharedPtr<FSchematicGraphNode>> SelectedNodesShared = GetSelectedNodes();
 	TArray<FSchematicGraphNode*> SelectedNodes;
 	SelectedNodes.Reserve(SelectedNodesShared.Num());
+	bool bApplyToSelection = false;
 	for (TSharedPtr<FSchematicGraphNode>& Node : SelectedNodesShared)
 	{
 		if (Node.IsValid())
 		{
 			SelectedNodes.AddUnique(Node.Get());
+			if (Node.Get() == InNode->GetNodeData())
+			{
+				bApplyToSelection = true;
+			}
 		}
 	}
-	SelectedNodes.AddUnique(InNode->GetNodeData());
 
-	
+	// If the dropped target node is not part of the selection, ignore the selection
+	if (!bApplyToSelection)
+	{
+		SelectedNodes.Reset();
+		SelectedNodes.Add(InNode->GetNodeData());
+	}
+
 	for (FSchematicGraphNode* Node : SelectedNodes)
 	{
 		TArray<FRigElementKey> NodeTargetKeys;
@@ -1545,6 +1557,8 @@ void FControlRigSchematicModel::HandleSchematicDrop(SSchematicGraphPanel* InPane
 								Controller->ConnectConnectorToElement(PrimaryConnectorKey, TargetKey, true, ControlRig->GetModularRigSettings().bAutoResolve);
 							}
 						}
+
+						ClearSelection();
 					}
 				}, TStatId(), NULL, ENamedThreads::GameThread);
 			}
@@ -1586,6 +1600,8 @@ void FControlRigSchematicModel::HandleSchematicDrop(SSchematicGraphPanel* InPane
 						}
 					}
 				}
+
+				ClearSelection();
 			}
 		}, TStatId(), NULL, ENamedThreads::GameThread);
 	}
@@ -1620,6 +1636,7 @@ void FControlRigSchematicModel::HandleSchematicDrop(SSchematicGraphPanel* InPane
 								break;
 							}
 						}
+						ClearSelection();
 						return;
 					}
 					if(const FRigModuleReference* Module = Controller->FindModule(ModulePathOrConnectorName))
@@ -1640,6 +1657,7 @@ void FControlRigSchematicModel::HandleSchematicDrop(SSchematicGraphPanel* InPane
 											break;
 										}
 									}
+									ClearSelection();
 									return;
 								}
 							}
diff --git a/Engine/Source/Editor/AnimationEditorWidgets/Private/SchematicGraphPanel/SchematicGraphModel.cpp b/Engine/Source/Editor/AnimationEditorWidgets/Private/SchematicGraphPanel/SchematicGraphModel.cpp
--- a/Engine/Source/Editor/AnimationEditorWidgets/Private/SchematicGraphPanel/SchematicGraphModel.cpp
+++ b/Engine/Source/Editor/AnimationEditorWidgets/Private/SchematicGraphPanel/SchematicGraphModel.cpp
@@ -448,10 +448,25 @@ const TArray<TSharedPtr<FSchematicGraphNode>> FSchematicGraphModel::GetSelectedN
 {
 	return Nodes.FilterByPredicate([](const TSharedPtr<FSchematicGraphNode>& Node)
 	{
-		return Node->IsSelected();
+		if (Node.IsValid())
+		{
+			return Node->IsSelected();
+		}
+		return false;
 	});
 }
 
+void FSchematicGraphModel::ClearSelection()
+{
+	for (TSharedPtr<FSchematicGraphNode> Node : Nodes)
+	{
+		if (Node.IsValid())
+		{
+			Node->SetSelected(false);
+		}
+	}
+}
+
 FLinearColor FSchematicGraphModel::GetBackgroundColorForTag(const FGuid& InNodeGuid, const FGuid& InTagGuid) const
 {
 	if(const FSchematicGraphNode* Node = FindNode(InNodeGuid))
diff --git a/Engine/Source/Editor/AnimationEditorWidgets/Public/SchematicGraphPanel/SchematicGraphModel.h b/Engine/Source/Editor/AnimationEditorWidgets/Public/SchematicGraphPanel/SchematicGraphModel.h
--- a/Engine/Source/Editor/AnimationEditorWidgets/Public/SchematicGraphPanel/SchematicGraphModel.h
+++ b/Engine/Source/Editor/AnimationEditorWidgets/Public/SchematicGraphPanel/SchematicGraphModel.h
@@ -116,6 +116,7 @@ public:
 	virtual bool IsDragSupportedForNode(const FSchematicGraphNode* InNode) const;
 	virtual bool GetContextMenuForNode(const FSchematicGraphNode* InNode, FMenuBuilder& OutMenu) const;
 	const TArray<TSharedPtr<FSchematicGraphNode>> GetSelectedNodes() const;
+	void ClearSelection();
 
 	FLinearColor GetBackgroundColorForTag(const FGuid& InNodeGuid, const FGuid& InTagGuid) const;
 	virtual FLinearColor GetBackgroundColorForTag(const FSchematicGraphTag* InTag) const;
