Control Rig: Consolidate categories Modules and Rig Modules into a single category
#rnx

diff --git a/Engine/Plugins/Animation/ControlRig/Source/ControlRig/Public/Units/Execution/RigUnit_RigModules.h b/Engine/Plugins/Animation/ControlRig/Source/ControlRig/Public/Units/Execution/RigUnit_RigModules.h
--- a/Engine/Plugins/Animation/ControlRig/Source/ControlRig/Public/Units/Execution/RigUnit_RigModules.h
+++ b/Engine/Plugins/Animation/ControlRig/Source/ControlRig/Public/Units/Execution/RigUnit_RigModules.h
@@ -6,13 +6,13 @@
 #include "ControlRigDefines.h" 
 #include "RigUnit_RigModules.generated.h"
 
-USTRUCT(meta = (Abstract, NodeColor="0.262745, 0.8, 0, 0.229412", Category = "Rig Modules"))
+USTRUCT(meta = (Abstract, NodeColor="0.262745, 0.8, 0, 0.229412", Category = "Modules"))
 struct CONTROLRIG_API FRigUnit_RigModulesBase : public FRigUnit
 {
 	GENERATED_BODY()
 };
 
-USTRUCT(meta = (Abstract, NodeColor="0.262745, 0.8, 0, 0.229412", Category = "Rig Modules"))
+USTRUCT(meta = (Abstract, NodeColor="0.262745, 0.8, 0, 0.229412", Category = "Modules"))
 struct CONTROLRIG_API FRigUnit_RigModulesBaseMutable : public FRigUnitMutable
 {
 	GENERATED_BODY()
