Control RIg: Make sure when spawning animation channel the name does not contain any namespace character
#rnx

diff --git a/Engine/Plugins/Animation/ControlRig/Source/ControlRig/Private/Units/Execution/RigUnit_DynamicHierarchy.cpp b/Engine/Plugins/Animation/ControlRig/Source/ControlRig/Private/Units/Execution/RigUnit_DynamicHierarchy.cpp
--- a/Engine/Plugins/Animation/ControlRig/Source/ControlRig/Private/Units/Execution/RigUnit_DynamicHierarchy.cpp
+++ b/Engine/Plugins/Animation/ControlRig/Source/ControlRig/Private/Units/Execution/RigUnit_DynamicHierarchy.cpp
@@ -920,6 +920,12 @@ FRigUnit_HierarchyAddAnimationChannelBool_Execute()
 		return;
 	}
 
+	const FString NameStr = Name.ToString();
+	if (NameStr.Contains(TEXT(":")))
+	{
+		UE_CONTROLRIG_RIGUNIT_REPORT_ERROR(TEXT("Animation channel name %s contains invalid ':' character."), *NameStr);
+	}
+
 	Item.Reset();
 
 	if(URigHierarchyController* Controller = ExecuteContext.Hierarchy->GetController(true))
@@ -959,6 +965,12 @@ FRigUnit_HierarchyAddAnimationChannelFloat_Execute()
 		return;
 	}
 
+	const FString NameStr = Name.ToString();
+	if (NameStr.Contains(TEXT(":")))
+	{
+		UE_CONTROLRIG_RIGUNIT_REPORT_ERROR(TEXT("Animation channel name %s contains invalid ':' character."), *NameStr);
+	}
+
 	Item.Reset();
 
 	if(URigHierarchyController* Controller = ExecuteContext.Hierarchy->GetController(true))
@@ -999,6 +1011,13 @@ FRigUnit_HierarchyAddAnimationChannelScaleFloat_Execute()
 		return;
 	}
 
+	const FString NameStr = Name.ToString();
+	if (NameStr.Contains(TEXT(":")))
+	{
+		UE_CONTROLRIG_RIGUNIT_REPORT_ERROR(TEXT("Animation channel name %s contains invalid ':' character."), *NameStr);
+	}
+
+
 	Item.Reset();
 
 	if(URigHierarchyController* Controller = ExecuteContext.Hierarchy->GetController(true))
@@ -1039,6 +1058,13 @@ FRigUnit_HierarchyAddAnimationChannelInteger_Execute()
 		return;
 	}
 
+	const FString NameStr = Name.ToString();
+	if (NameStr.Contains(TEXT(":")))
+	{
+		UE_CONTROLRIG_RIGUNIT_REPORT_ERROR(TEXT("Animation channel name %s contains invalid ':' character."), *NameStr);
+	}
+
+
 	Item.Reset();
 
 	if(URigHierarchyController* Controller = ExecuteContext.Hierarchy->GetController(true))
@@ -1080,6 +1106,13 @@ FRigUnit_HierarchyAddAnimationChannelVector2D_Execute()
 		return;
 	}
 
+	const FString NameStr = Name.ToString();
+	if (NameStr.Contains(TEXT(":")))
+	{
+		UE_CONTROLRIG_RIGUNIT_REPORT_ERROR(TEXT("Animation channel name %s contains invalid ':' character."), *NameStr);
+	}
+
+
 	Item.Reset();
 
 	if(URigHierarchyController* Controller = ExecuteContext.Hierarchy->GetController(true))
@@ -1121,6 +1154,13 @@ FRigUnit_HierarchyAddAnimationChannelVector_Execute()
 		return;
 	}
 
+	const FString NameStr = Name.ToString();
+	if (NameStr.Contains(TEXT(":")))
+	{
+		UE_CONTROLRIG_RIGUNIT_REPORT_ERROR(TEXT("Animation channel name %s contains invalid ':' character."), *NameStr);
+	}
+
+
 	Item.Reset();
 
 	if(URigHierarchyController* Controller = ExecuteContext.Hierarchy->GetController(true))
@@ -1163,6 +1203,13 @@ FRigUnit_HierarchyAddAnimationChannelScaleVector_Execute()
 		return;
 	}
 
+	const FString NameStr = Name.ToString();
+	if (NameStr.Contains(TEXT(":")))
+	{
+		UE_CONTROLRIG_RIGUNIT_REPORT_ERROR(TEXT("Animation channel name %s contains invalid ':' character."), *NameStr);
+	}
+
+
 	Item.Reset();
 
 	if(URigHierarchyController* Controller = ExecuteContext.Hierarchy->GetController(true))
@@ -1205,6 +1252,13 @@ FRigUnit_HierarchyAddAnimationChannelRotator_Execute()
 		return;
 	}
 
+	const FString NameStr = Name.ToString();
+	if (NameStr.Contains(TEXT(":")))
+	{
+		UE_CONTROLRIG_RIGUNIT_REPORT_ERROR(TEXT("Animation channel name %s contains invalid ':' character."), *NameStr);
+	}
+
+
 	Item.Reset();
 
 	if(URigHierarchyController* Controller = ExecuteContext.Hierarchy->GetController(true))
