Control Rig: Fix zombie connectors left in the RigHierarchy when modules are applied to multiple targets
#rnx

diff --git a/Engine/Plugins/Animation/ControlRig/Source/ControlRig/Private/ModularRigController.cpp b/Engine/Plugins/Animation/ControlRig/Source/ControlRig/Private/ModularRigController.cpp
--- a/Engine/Plugins/Animation/ControlRig/Source/ControlRig/Private/ModularRigController.cpp
+++ b/Engine/Plugins/Animation/ControlRig/Source/ControlRig/Private/ModularRigController.cpp
@@ -1224,6 +1224,7 @@ FString UModularRigController::ReparentModule(const FString& InModulePath, const
 	// Reparent or unparent children
 	const FString OldPath = Module->GetPath();
 	Module->PreviousParentPath = Module->ParentPath;
+	Module->PreviousName = Module->Name;
 	Module->ParentPath = (NewParentModule) ? NewParentModule->GetPath() : FString();
 	Module->Name = GetSafeNewName(Module->ParentPath, FRigName(Module->Name));
 	const FString NewPath = Module->GetPath();
diff --git a/Engine/Plugins/Animation/ControlRig/Source/ControlRigDeveloper/Private/ControlRigBlueprint.cpp b/Engine/Plugins/Animation/ControlRig/Source/ControlRigDeveloper/Private/ControlRigBlueprint.cpp
--- a/Engine/Plugins/Animation/ControlRig/Source/ControlRigDeveloper/Private/ControlRigBlueprint.cpp
+++ b/Engine/Plugins/Animation/ControlRig/Source/ControlRigDeveloper/Private/ControlRigBlueprint.cpp
@@ -2595,7 +2595,7 @@ void UControlRigBlueprint::HandleRigModulesModified(EModularRigNotification InNo
 						}
 						else if(InNotification == EModularRigNotification::ModuleReparented)
 						{
-							OldNamespace = (InModule->PreviousParentPath.IsEmpty()) ? *InModule->Name.ToString() : URigHierarchy::JoinNameSpace(InModule->PreviousParentPath, InModule->Name.ToString());
+							OldNamespace = (InModule->PreviousParentPath.IsEmpty()) ? *InModule->PreviousName.ToString() : URigHierarchy::JoinNameSpace(InModule->PreviousParentPath, InModule->PreviousName.ToString());
 						}
 						FString NewNamespace = (InModule->ParentPath.IsEmpty()) ? *InModule->Name.ToString() : URigHierarchy::JoinNameSpace(InModule->ParentPath, InModule->Name.ToString());
 						OldNamespace.Append(UModularRig::NamespaceSeparator);
