Control Rig: ensure that the pivots are up-to-date when the gizmos need them
    - as both legacy and new gizmos tick before the mode, they use an obsolete coordinate system when redering. As the Control Rig Mode caches the pivots on tick, we have to make sure we get up-to-date transforms.

diff --git a/Engine/Plugins/Animation/ControlRig/Source/ControlRigEditor/Private/EditMode/ControlRigEditMode.cpp b/Engine/Plugins/Animation/ControlRig/Source/ControlRigEditor/Private/EditMode/ControlRigEditMode.cpp
--- a/Engine/Plugins/Animation/ControlRig/Source/ControlRigEditor/Private/EditMode/ControlRigEditMode.cpp
+++ b/Engine/Plugins/Animation/ControlRig/Source/ControlRigEditor/Private/EditMode/ControlRigEditMode.cpp
@@ -1323,6 +1323,9 @@ bool FControlRigEditMode::GetPivotForOrbit(FVector& OutPivot) const
 
 bool FControlRigEditMode::GetCustomDrawingCoordinateSystem(FMatrix& OutMatrix, void* InData)
 {
+	// check that the cached pivots are up-to-date and update them if needed
+	UpdatePivotTransformsIfNeeded();
+	
 	//since we strip translation just want the first one
 	for (const auto& Pairs : ControlRigShapeActors)
 	{
@@ -2638,6 +2641,16 @@ bool FControlRigEditMode::HasPivotTransformsChanged() const
 	return false;
 }
 
+void FControlRigEditMode::UpdatePivotTransformsIfNeeded()
+{
+	if (bPivotsNeedUpdate)
+	{
+		PostPoseUpdate();
+		UpdatePivotTransforms();
+		bPivotsNeedUpdate = false;
+	}
+}
+
 void FControlRigEditMode::HandleSelectionChanged()
 {
 	for (const auto& ShapeActors : ControlRigShapeActors)
@@ -4314,6 +4327,8 @@ void FControlRigEditMode::OnControlModified(UControlRig* Subject, FRigControlEle
 	const bool bModify = Context.SetKey != EControlRigSetKey::Never;
 	ControlProxy->ProxyChanged(Subject, InControlElement, bModify);
 
+	bPivotsNeedUpdate = true;
+	
 	/*
 	FScopedTransaction ScopedTransaction(LOCTEXT("ModifyControlTransaction", "Modify Control"),!GIsTransacting && Context.SetKey != EControlRigSetKey::Never);
 	ControlProxy->Modify();
diff --git a/Engine/Plugins/Animation/ControlRig/Source/ControlRigEditor/Public/EditMode/ControlRigEditMode.h b/Engine/Plugins/Animation/ControlRig/Source/ControlRigEditor/Public/EditMode/ControlRigEditMode.h
--- a/Engine/Plugins/Animation/ControlRig/Source/ControlRigEditor/Public/EditMode/ControlRigEditMode.h
+++ b/Engine/Plugins/Animation/ControlRig/Source/ControlRigEditor/Public/EditMode/ControlRigEditMode.h
@@ -453,6 +453,8 @@ private:
 	bool HasPivotTransformsChanged() const;
 	/** Set a RigElement's selection state */
 	void SetRigElementSelectionInternal(UControlRig* ControlRig, ERigElementType Type, const FName& InRigElementName, bool bSelected);
+	/** Updates the pivot transforms before ticking to ensure that they are up-to-date when needed. */
+	void UpdatePivotTransformsIfNeeded();
 	
 	FEditorViewportClient* CurrentViewportClient;
 	TArray<UE::Widget::EWidgetMode> RequestedWidgetModes;
@@ -573,6 +575,8 @@ private:
 	uint8 InteractionType;
 	bool bShowControlsAsOverlay;
 
+	bool bPivotsNeedUpdate = true;
+
 	bool bIsConstructionEventRunning;
 	TArray<uint32> LastHierarchyHash;
 	TArray<uint32> LastShapeLibraryHash;
