RigVM / ControlRig: Centralize garbage detection code

#rnx

diff --git a/Engine/Plugins/Animation/ControlRig/Source/ControlRig/Private/ControlRig.cpp b/Engine/Plugins/Animation/ControlRig/Source/ControlRig/Private/ControlRig.cpp
--- a/Engine/Plugins/Animation/ControlRig/Source/ControlRig/Private/ControlRig.cpp
+++ b/Engine/Plugins/Animation/ControlRig/Source/ControlRig/Private/ControlRig.cpp
@@ -144,7 +144,7 @@ void UControlRig::BeginDestroy()
 	{
 		if(UControlRig* CDO = GetClass()->GetDefaultObject<UControlRig>())
 		{
-			if (!CDO->HasAnyFlags(RF_BeginDestroyed))
+			if (!IsGarbageOrDestroyed(CDO))
 			{
 				if (CDO->GetHierarchy())
 				{
diff --git a/Engine/Plugins/Animation/ControlRig/Source/ControlRig/Private/ControlRigComponent.cpp b/Engine/Plugins/Animation/ControlRig/Source/ControlRig/Private/ControlRigComponent.cpp
--- a/Engine/Plugins/Animation/ControlRig/Source/ControlRig/Private/ControlRigComponent.cpp
+++ b/Engine/Plugins/Animation/ControlRig/Source/ControlRig/Private/ControlRigComponent.cpp
@@ -160,12 +160,12 @@ void UControlRigComponent::OnUnregister()
 {
 	Super::OnUnregister();
 
-	bool bBeginDestroyed = HasAnyFlags(RF_BeginDestroyed);
+	bool bBeginDestroyed = URigVMHost::IsGarbageOrDestroyed(this);
 	if (!bBeginDestroyed)
 	{
 		if (AActor* Actor = GetOwner())
 		{
-			bBeginDestroyed = Actor->HasAnyFlags(RF_BeginDestroyed);
+			bBeginDestroyed = URigVMHost::IsGarbageOrDestroyed(Actor);
 		}
 	}
 
@@ -176,7 +176,7 @@ void UControlRigComponent::OnUnregister()
 			if (Pair.Key)
 			{
 				if (Pair.Key->IsValidLowLevel() &&
-					!Pair.Key->HasAnyFlags(RF_BeginDestroyed) &&
+					!URigVMHost::IsGarbageOrDestroyed(Pair.Key) &&
 					IsValid(Pair.Key) &&
 					!Pair.Key->IsUnreachable())
 				{
diff --git a/Engine/Plugins/Animation/ControlRig/Source/ControlRig/Private/Sequencer/MovieSceneControlRigParameterTemplate.cpp b/Engine/Plugins/Animation/ControlRig/Source/ControlRig/Private/Sequencer/MovieSceneControlRigParameterTemplate.cpp
--- a/Engine/Plugins/Animation/ControlRig/Source/ControlRig/Private/Sequencer/MovieSceneControlRigParameterTemplate.cpp
+++ b/Engine/Plugins/Animation/ControlRig/Source/ControlRig/Private/Sequencer/MovieSceneControlRigParameterTemplate.cpp
@@ -645,7 +645,7 @@ void FControlRigBindingHelper::UnBindFromSequencerInstance(UControlRig* ControlR
 	check(ControlRig);
 
 	if (!ControlRig->IsValidLowLevel() ||
-	    ControlRig->HasAnyFlags(RF_BeginDestroyed) ||
+	    URigVMHost::IsGarbageOrDestroyed(ControlRig) ||
 		!IsValid(ControlRig))
 	{
 		return;
@@ -659,7 +659,7 @@ void FControlRigBindingHelper::UnBindFromSequencerInstance(UControlRig* ControlR
 	else if (USkeletalMeshComponent* SkeletalMeshComponent = Cast<USkeletalMeshComponent>(ControlRig->GetObjectBinding()->GetBoundObject()))
 	{
 		if (!SkeletalMeshComponent->IsValidLowLevel() ||
-			SkeletalMeshComponent->HasAnyFlags(RF_BeginDestroyed) ||
+			URigVMHost::IsGarbageOrDestroyed(SkeletalMeshComponent) ||
 			!IsValid(SkeletalMeshComponent))
 		{
 			return;
@@ -670,7 +670,7 @@ void FControlRigBindingHelper::UnBindFromSequencerInstance(UControlRig* ControlR
 		if (AnimInstance)
 		{
 			if (!AnimInstance->IsValidLowLevel() ||
-                AnimInstance->HasAnyFlags(RF_BeginDestroyed) ||
+                URigVMHost::IsGarbageOrDestroyed(AnimInstance) ||
                 !IsValid(AnimInstance))
 			{
 				return;
diff --git a/Engine/Plugins/Animation/ControlRig/Source/ControlRigEditor/Private/Editor/ControlRigEditor.cpp b/Engine/Plugins/Animation/ControlRig/Source/ControlRigEditor/Private/Editor/ControlRigEditor.cpp
--- a/Engine/Plugins/Animation/ControlRig/Source/ControlRigEditor/Private/Editor/ControlRigEditor.cpp
+++ b/Engine/Plugins/Animation/ControlRig/Source/ControlRigEditor/Private/Editor/ControlRigEditor.cpp
@@ -1055,7 +1055,7 @@ void FControlRigEditor::HandleSetObjectBeingDebugged(UObject* InObject)
 	UControlRig* DebuggedControlRig = Cast<UControlRig>(InObject);
 	if(UControlRig* PreviouslyDebuggedControlRig = Cast<UControlRig>(GetBlueprintObj()->GetObjectBeingDebugged()))
 	{
-		if(!PreviouslyDebuggedControlRig->HasAnyFlags(RF_BeginDestroyed))
+		if(!URigVMHost::IsGarbageOrDestroyed(PreviouslyDebuggedControlRig))
 		{
 			PreviouslyDebuggedControlRig->GetHierarchy()->OnModified().RemoveAll(this);
 			PreviouslyDebuggedControlRig->OnPreForwardsSolve_AnyThread().RemoveAll(this);
diff --git a/Engine/Plugins/Animation/ControlRig/Source/ControlRigEditor/Private/Editor/ControlRigSchematicModel.cpp b/Engine/Plugins/Animation/ControlRig/Source/ControlRigEditor/Private/Editor/ControlRigSchematicModel.cpp
--- a/Engine/Plugins/Animation/ControlRig/Source/ControlRigEditor/Private/Editor/ControlRigSchematicModel.cpp
+++ b/Engine/Plugins/Animation/ControlRig/Source/ControlRigEditor/Private/Editor/ControlRigSchematicModel.cpp
@@ -101,7 +101,7 @@ FControlRigSchematicModel::~FControlRigSchematicModel()
 	{
 		if(UControlRig* ControlRigBeingDebugged = ControlRigBeingDebuggedPtr.Get())
 		{
-			if(!ControlRigBeingDebugged->HasAnyFlags(RF_BeginDestroyed))
+			if(!URigVMHost::IsGarbageOrDestroyed(ControlRigBeingDebugged))
 			{
 				ControlRigBeingDebugged->GetHierarchy()->OnModified().RemoveAll(this);
 				ControlRigBeingDebugged->OnPostConstruction_AnyThread().RemoveAll(this);
@@ -577,7 +577,7 @@ void FControlRigSchematicModel::OnSetObjectBeingDebugged(UObject* InObject)
 	{
 		if(UControlRig* ControlRigBeingDebugged = ControlRigBeingDebuggedPtr.Get())
 		{
-			if(!ControlRigBeingDebugged->HasAnyFlags(RF_BeginDestroyed))
+			if(!URigVMHost::IsGarbageOrDestroyed(ControlRigBeingDebugged))
 			{
 				ControlRigBeingDebugged->GetHierarchy()->OnModified().RemoveAll(this);
 				ControlRigBeingDebugged->OnPostConstruction_AnyThread().RemoveAll(this);
diff --git a/Engine/Plugins/Animation/ControlRig/Source/ControlRigEditor/Private/Editor/ControlRigSkeletalMeshComponent.cpp b/Engine/Plugins/Animation/ControlRig/Source/ControlRigEditor/Private/Editor/ControlRigSkeletalMeshComponent.cpp
--- a/Engine/Plugins/Animation/ControlRig/Source/ControlRigEditor/Private/Editor/ControlRigSkeletalMeshComponent.cpp
+++ b/Engine/Plugins/Animation/ControlRig/Source/ControlRigEditor/Private/Editor/ControlRigSkeletalMeshComponent.cpp
@@ -183,7 +183,7 @@ void UControlRigSkeletalMeshComponent::SetControlRigBeingDebugged(UControlRig* I
 	{
 		if(UControlRig* ControlRigBeingDebugged = ControlRigBeingDebuggedPtr.Get())
 		{
-			if(!ControlRigBeingDebugged->HasAnyFlags(RF_BeginDestroyed))
+			if(!URigVMHost::IsGarbageOrDestroyed(ControlRigBeingDebugged))
 			{
 				ControlRigBeingDebugged->GetHierarchy()->OnModified().RemoveAll(this);
 			}
diff --git a/Engine/Plugins/Animation/ControlRig/Source/ControlRigEditor/Private/Editor/SControlRigAnimAttributeView.cpp b/Engine/Plugins/Animation/ControlRig/Source/ControlRigEditor/Private/Editor/SControlRigAnimAttributeView.cpp
--- a/Engine/Plugins/Animation/ControlRig/Source/ControlRigEditor/Private/Editor/SControlRigAnimAttributeView.cpp
+++ b/Engine/Plugins/Animation/ControlRig/Source/ControlRigEditor/Private/Editor/SControlRigAnimAttributeView.cpp
@@ -73,7 +73,7 @@ void SControlRigAnimAttributeView::StopObservingCurrentControlRig()
 	{
 		if(UControlRig* ControlRig = ControlRigBeingDebuggedPtr.Get())
 		{
-			if(!ControlRig->HasAnyFlags(RF_BeginDestroyed))
+			if(!URigVMHost::IsGarbageOrDestroyed(ControlRig))
 			{
 				ControlRig->OnPostForwardsSolve_AnyThread().RemoveAll(this);
 				ControlRig->SetEnableAnimAttributeTrace(false);
diff --git a/Engine/Plugins/Animation/ControlRig/Source/ControlRigEditor/Private/Editor/SRigHierarchy.cpp b/Engine/Plugins/Animation/ControlRig/Source/ControlRigEditor/Private/Editor/SRigHierarchy.cpp
--- a/Engine/Plugins/Animation/ControlRig/Source/ControlRigEditor/Private/Editor/SRigHierarchy.cpp
+++ b/Engine/Plugins/Animation/ControlRig/Source/ControlRigEditor/Private/Editor/SRigHierarchy.cpp
@@ -877,7 +877,7 @@ void SRigHierarchy::HandleSetObjectBeingDebugged(UObject* InObject)
 	{
 		if(UControlRig* ControlRigBeingDebugged = ControlRigBeingDebuggedPtr.Get())
 		{
-			if(!ControlRigBeingDebugged->HasAnyFlags(RF_BeginDestroyed))
+			if(!URigVMHost::IsGarbageOrDestroyed(ControlRigBeingDebugged))
 			{
 				ControlRigBeingDebugged->GetHierarchy()->OnModified().RemoveAll(this);
 			}
diff --git a/Engine/Plugins/Runtime/RigVM/Source/RigVM/Private/RigVMHost.cpp b/Engine/Plugins/Runtime/RigVM/Source/RigVM/Private/RigVMHost.cpp
--- a/Engine/Plugins/Runtime/RigVM/Source/RigVM/Private/RigVMHost.cpp
+++ b/Engine/Plugins/Runtime/RigVM/Source/RigVM/Private/RigVMHost.cpp
@@ -97,6 +97,16 @@ TArray<URigVMHost*> URigVMHost::FindRigVMHosts(UObject* Outer, TSubclassOf<URigV
 	return Result;
 }
 
+bool URigVMHost::IsGarbageOrDestroyed(const UObject* InObject)
+{
+	if(!IsValid(InObject))
+	{
+		return true;
+	}
+	return InObject->HasAnyFlags(RF_BeginDestroyed | RF_FinishDestroyed) ||
+		InObject->HasAnyInternalFlags(EInternalObjectFlags::Garbage);
+}
+
 UWorld* URigVMHost::GetWorld() const
 {
 	if (const UObject* Outer = GetOuter())
diff --git a/Engine/Plugins/Runtime/RigVM/Source/RigVM/Public/RigVMHost.h b/Engine/Plugins/Runtime/RigVM/Source/RigVM/Public/RigVMHost.h
--- a/Engine/Plugins/Runtime/RigVM/Source/RigVM/Public/RigVMHost.h
+++ b/Engine/Plugins/Runtime/RigVM/Source/RigVM/Public/RigVMHost.h
@@ -30,6 +30,8 @@ public:
 	UFUNCTION(BlueprintCallable, Category = RigVM)
 	static TArray<URigVMHost*> FindRigVMHosts(UObject* Outer, TSubclassOf<URigVMHost> OptionalClass);
 
+	static bool IsGarbageOrDestroyed(const UObject* InObject);
+
 	/** UObject interface */
 	virtual UWorld* GetWorld() const override;
 	virtual void Serialize(FArchive& Ar) override;
diff --git a/Engine/Plugins/Runtime/RigVM/Source/RigVMEditor/Private/Editor/RigVMEditor.cpp b/Engine/Plugins/Runtime/RigVM/Source/RigVMEditor/Private/Editor/RigVMEditor.cpp
--- a/Engine/Plugins/Runtime/RigVM/Source/RigVMEditor/Private/Editor/RigVMEditor.cpp
+++ b/Engine/Plugins/Runtime/RigVM/Source/RigVMEditor/Private/Editor/RigVMEditor.cpp
@@ -3693,7 +3693,7 @@ void FRigVMEditor::HandleSetObjectBeingDebugged(UObject* InObject)
 {
 	if(URigVMHost* PreviouslyDebuggedHost = Cast<URigVMHost>(GetBlueprintObj()->GetObjectBeingDebugged()))
 	{
-		if(!PreviouslyDebuggedHost->HasAnyFlags(RF_BeginDestroyed))
+		if(!URigVMHost::IsGarbageOrDestroyed(PreviouslyDebuggedHost))
 		{
 			PreviouslyDebuggedHost->OnExecuted_AnyThread().RemoveAll(this);
 		}
