Control Rig: ensure that the pivots are up-to-date when the gizmos need them
    - follow-up to 31493231 as it was still flicking in world space coordinate.

#rnx

diff --git a/Engine/Plugins/Animation/ControlRig/Source/ControlRigEditor/Private/EditMode/ControlRigEditMode.cpp b/Engine/Plugins/Animation/ControlRig/Source/ControlRigEditor/Private/EditMode/ControlRigEditMode.cpp
--- a/Engine/Plugins/Animation/ControlRig/Source/ControlRigEditor/Private/EditMode/ControlRigEditMode.cpp
+++ b/Engine/Plugins/Animation/ControlRig/Source/ControlRigEditor/Private/EditMode/ControlRigEditMode.cpp
@@ -1292,8 +1292,11 @@ FVector FControlRigEditMode::GetWidgetLocation() const
 		{
 			if (const FTransform* PivotTransform = PivotTransforms.Find(Pairs.Key))
 			{
-				FTransform ComponentTransform = GetHostingSceneComponentTransform(Pairs.Key);
-				PivotLocation += ComponentTransform.TransformPosition(PivotTransform->GetLocation());
+				// check that the cached pivot is up-to-date and update it if needed
+				FTransform Transform = *PivotTransform;
+				UpdatePivotTransformsIfNeeded(Pairs.Key, Transform);
+				const FTransform ComponentTransform = GetHostingSceneComponentTransform(Pairs.Key);
+				PivotLocation += ComponentTransform.TransformPosition(Transform.GetLocation());
 				++NumSelected;
 			}
 		}
@@ -1303,10 +1306,6 @@ FVector FControlRigEditMode::GetWidgetLocation() const
 		PivotLocation /= (NumSelected);
 		return PivotLocation;
 	}
-	else
-	{
-		PivotLocation = FVector(0, 0, 0);
-	}
 
 	return FEdMode::GetWidgetLocation();
 }
@@ -1323,9 +1322,6 @@ bool FControlRigEditMode::GetPivotForOrbit(FVector& OutPivot) const
 
 bool FControlRigEditMode::GetCustomDrawingCoordinateSystem(FMatrix& OutMatrix, void* InData)
 {
-	// check that the cached pivots are up-to-date and update them if needed
-	UpdatePivotTransformsIfNeeded();
-	
 	//since we strip translation just want the first one
 	for (const auto& Pairs : ControlRigShapeActors)
 	{
@@ -1333,7 +1329,10 @@ bool FControlRigEditMode::GetCustomDrawingCoordinateSystem(FMatrix& OutMatrix, v
 		{
 			if (const FTransform* PivotTransform = PivotTransforms.Find(Pairs.Key))
 			{
-				OutMatrix = PivotTransform->ToMatrixNoScale().RemoveTranslation();
+				// check that the cached pivot is up-to-date and update it if needed
+				FTransform Transform = *PivotTransform;
+				UpdatePivotTransformsIfNeeded(Pairs.Key, Transform);
+				OutMatrix = Transform.ToMatrixNoScale().RemoveTranslation();
 				return true;
 			}
 		}
@@ -2417,20 +2416,20 @@ ECoordSystem FControlRigEditMode::GetCoordSystemSpace() const
 	return GetModeManager()->GetCoordSystem();	
 }
 
-void FControlRigEditMode::UpdatePivotFromEditedShape(UControlRig* InControlRig)
+bool FControlRigEditMode::ComputePivotFromEditedShape(UControlRig* InControlRig, FTransform& OutTransform) const
 {
 	const URigHierarchy* Hierarchy = InControlRig ? InControlRig->GetHierarchy() : nullptr;
 	if (!Hierarchy)
 	{
-		return;
+		return false;
 	}
 
 	if (!ensure(bIsChangingControlShapeTransform))
 	{
-		return;
+		return false;
 	}
 	
-	FTransform PivotTransform = FTransform::Identity;
+	OutTransform = FTransform::Identity;
 	
 	if (auto* ShapeActors = ControlRigShapeActors.Find(InControlRig))
 	{
@@ -2444,25 +2443,25 @@ void FControlRigEditMode::UpdatePivotFromEditedShape(UControlRig* InControlRig)
 		{
 			if (FRigControlElement* ControlElement = InControlRig->FindControl((*ShapeActors)[Index]->ControlName))
 			{
-				PivotTransform = Hierarchy->GetControlShapeTransform(ControlElement, ERigTransformType::CurrentGlobal);
+				OutTransform = Hierarchy->GetControlShapeTransform(ControlElement, ERigTransformType::CurrentGlobal);
 			}				
 		}
 	}
 	
-	PivotTransforms.Add(InControlRig, MoveTemp(PivotTransform));
+	return true;
 }
 
-void FControlRigEditMode::UpdatePivotFromShapeActors(UControlRig* InControlRig, const bool bEachLocalSpace, const bool bIsParentSpace)
+bool FControlRigEditMode::ComputePivotFromShapeActors(UControlRig* InControlRig, const bool bEachLocalSpace, const bool bIsParentSpace, FTransform& OutTransform) const
 {
 	if (!ensure(!bIsChangingControlShapeTransform))
 	{
-		return;
+		return false;
 	}
 	
 	const URigHierarchy* Hierarchy = InControlRig ? InControlRig->GetHierarchy() : nullptr;
 	if (!Hierarchy)
 	{
-		return;
+		return false;
 	}
 	const FTransform ComponentTransform = GetHostingSceneComponentTransform(InControlRig);
 
@@ -2505,20 +2504,22 @@ void FControlRigEditMode::UpdatePivotFromShapeActors(UControlRig* InControlRig,
 	const FTransform WorldTransform = LastTransform * ComponentTransform;
 	PivotTransform.SetRotation(WorldTransform.GetRotation());
 	
-	PivotTransforms.Add(InControlRig, MoveTemp(PivotTransform));
+	OutTransform = PivotTransform;
+	
+	return true;
 }
 
-void FControlRigEditMode::UpdatePivotFromElements(UControlRig* InControlRig)
+bool FControlRigEditMode::ComputePivotFromElements(UControlRig* InControlRig, FTransform& OutTransform) const
 {
 	if (!ensure(!bIsChangingControlShapeTransform))
 	{
-		return;
+		return false;
 	}
 	
 	const URigHierarchy* Hierarchy = InControlRig ? InControlRig->GetHierarchy() : nullptr;
 	if (!Hierarchy)
 	{
-		return;
+		return false;
 	}
 	
 	const FTransform ComponentTransform = GetHostingSceneComponentTransform(InControlRig);
@@ -2551,7 +2552,9 @@ void FControlRigEditMode::UpdatePivotFromElements(UControlRig* InControlRig)
 	}
 		
 	PivotTransform.SetLocation(PivotLocation);
-	PivotTransforms.Add(InControlRig, MoveTemp(PivotTransform));
+	OutTransform = PivotTransform;
+
+	return true;
 }
 
 void FControlRigEditMode::UpdatePivotTransforms()
@@ -2566,26 +2569,34 @@ void FControlRigEditMode::UpdatePivotTransforms()
 	{
 		if (UControlRig* ControlRig = RuntimeRigPtr.Get())
 		{
+			bool bAdd = false;
+			FTransform Pivot = FTransform::Identity;
 			if (AreRigElementsSelected(ValidControlTypeMask(), ControlRig))
 			{
 				if (bIsChangingControlShapeTransform)
 				{
-					UpdatePivotFromEditedShape(ControlRig);
+					bAdd = ComputePivotFromEditedShape(ControlRig, Pivot);
 				}
 				else
 				{
-					UpdatePivotFromShapeActors(ControlRig, bEachLocalSpace, bIsParentSpace);			
+					bAdd = ComputePivotFromShapeActors(ControlRig, bEachLocalSpace, bIsParentSpace, Pivot);			
 				}
 			}
 			else if (AreRigElementSelectedAndMovable(ControlRig))
 			{
 				// do we even get in here ?!
 				// we will enter the if first as AreRigElementsSelected will return true before AreRigElementSelectedAndMovable does...
-				UpdatePivotFromElements(ControlRig);
+				bAdd = ComputePivotFromElements(ControlRig, Pivot);
+			}
+			if (bAdd)
+			{
+				PivotTransforms.Add(ControlRig, MoveTemp(Pivot));
 			}
 		}
 	}
 
+	bPivotsNeedUpdate = false;
+
 	//If in level editor and the transforms changed we need to force hit proxy invalidate so widget hit testing 
 	//doesn't work off of it's last transform.  Similar to what sequencer does on re-evaluation but do to how edit modes and widget ticks happen
 	//it doesn't work for control rig gizmo's
@@ -2641,13 +2652,51 @@ bool FControlRigEditMode::HasPivotTransformsChanged() const
 	return false;
 }
 
-void FControlRigEditMode::UpdatePivotTransformsIfNeeded()
+void FControlRigEditMode::UpdatePivotTransformsIfNeeded(UControlRig* InControlRig, FTransform& InOutTransform) const
 {
-	if (bPivotsNeedUpdate)
+	if (!bPivotsNeedUpdate)
 	{
-		PostPoseUpdate();
-		UpdatePivotTransforms();
-		bPivotsNeedUpdate = false;
+		return;
+	}
+
+	if (!InControlRig)
+	{
+		return;
+	}
+
+	// Update shape actors transforms
+	if (auto* ShapeActors = ControlRigShapeActors.Find(InControlRig))
+	{
+		FTransform ComponentTransform = FTransform::Identity;
+		if (!AreEditingControlRigDirectly())
+		{
+			ComponentTransform = GetHostingSceneComponentTransform(InControlRig);
+		}
+		for (AControlRigShapeActor* ShapeActor : *ShapeActors)
+		{
+			const FTransform Transform = InControlRig->GetControlGlobalTransform(ShapeActor->ControlName);
+			ShapeActor->SetActorTransform(Transform * ComponentTransform);
+		}
+	}
+
+	// Update pivot
+	if (AreRigElementsSelected(ValidControlTypeMask(), InControlRig))
+	{
+		if (bIsChangingControlShapeTransform)
+		{
+			ComputePivotFromEditedShape(InControlRig, InOutTransform);
+		}
+		else
+		{
+			const UControlRigEditModeSettings* Settings = GetDefault<UControlRigEditModeSettings>();
+			const bool bEachLocalSpace = Settings && Settings->bLocalTransformsInEachLocalSpace;
+			const bool bIsParentSpace = GetCoordSystemSpace() == COORD_Parent;
+			ComputePivotFromShapeActors(InControlRig, bEachLocalSpace, bIsParentSpace, InOutTransform);			
+		}
+	}
+	else if (AreRigElementSelectedAndMovable(InControlRig))
+	{
+		ComputePivotFromElements(InControlRig, InOutTransform);
 	}
 }
 
@@ -4890,7 +4939,7 @@ bool FControlRigEditMode::IsControlRigSkelMeshVisible(UControlRig* ControlRig) c
 	return true;
 }
 
-void FControlRigEditMode::TickControlShape(AControlRigShapeActor* ShapeActor, const FTransform& ComponentTransform)
+void FControlRigEditMode::TickControlShape(AControlRigShapeActor* ShapeActor, const FTransform& ComponentTransform) const
 {
 	const UControlRigEditModeSettings* Settings = GetDefault<UControlRigEditModeSettings>();
 	if (ShapeActor)
@@ -5321,7 +5370,7 @@ void FControlRigEditMode::OnPoseInitialized()
 	OnAnimSystemInitializedDelegate.Broadcast();
 }
 
-void FControlRigEditMode::PostPoseUpdate()
+void FControlRigEditMode::PostPoseUpdate() const
 {
 	for (auto& ShapeActors : ControlRigShapeActors)
 	{
diff --git a/Engine/Plugins/Animation/ControlRig/Source/ControlRigEditor/Public/EditMode/ControlRigEditMode.h b/Engine/Plugins/Animation/ControlRig/Source/ControlRigEditor/Public/EditMode/ControlRigEditMode.h
--- a/Engine/Plugins/Animation/ControlRig/Source/ControlRigEditor/Public/EditMode/ControlRigEditMode.h
+++ b/Engine/Plugins/Animation/ControlRig/Source/ControlRigEditor/Public/EditMode/ControlRigEditMode.h
@@ -242,9 +242,9 @@ protected:
 
 	/** Updates cached pivot transforms */
 	void UpdatePivotTransforms();
-	void UpdatePivotFromEditedShape(UControlRig* InControlRig);
-	void UpdatePivotFromShapeActors(UControlRig* InControlRig, const bool bEachLocalSpace, const bool bIsParentSpace);
-	void UpdatePivotFromElements(UControlRig* InControlRig);
+	bool ComputePivotFromEditedShape(UControlRig* InControlRig, FTransform& OutTransform) const;
+	bool ComputePivotFromShapeActors(UControlRig* InControlRig, const bool bEachLocalSpace, const bool bIsParentSpace, FTransform& OutTransform) const;
+	bool ComputePivotFromElements(UControlRig* InControlRig, FTransform& OutTransform) const;
 	
 	/** Get the current coordinate system space */
 	ECoordSystem GetCoordSystemSpace() const;
@@ -454,7 +454,7 @@ private:
 	/** Set a RigElement's selection state */
 	void SetRigElementSelectionInternal(UControlRig* ControlRig, ERigElementType Type, const FName& InRigElementName, bool bSelected);
 	/** Updates the pivot transforms before ticking to ensure that they are up-to-date when needed. */
-	void UpdatePivotTransformsIfNeeded();
+	void UpdatePivotTransformsIfNeeded(UControlRig* InControlRig, FTransform& InOutTransform) const;
 	
 	FEditorViewportClient* CurrentViewportClient;
 	TArray<UE::Widget::EWidgetMode> RequestedWidgetModes;
@@ -497,7 +497,7 @@ private:
 	void ChangeControlShapeTransform(AControlRigShapeActor* ShapeActor, const bool bTranslation, FVector& InDrag,
 		const bool bRotation, FRotator& InRot, const bool bScale, FVector& InScale, const FTransform& ToWorldTransform);
 
-	void TickControlShape(AControlRigShapeActor* ShapeActor, const FTransform& ComponentTransform);
+	void TickControlShape(AControlRigShapeActor* ShapeActor, const FTransform& ComponentTransform) const;
 	bool ModeSupportedByShapeActor(const AControlRigShapeActor* ShapeActor, UE::Widget::EWidgetMode InMode) const;
 
 
@@ -518,7 +518,7 @@ private:
 
 	// Post pose update handler
 	void OnPoseInitialized();
-	void PostPoseUpdate();
+	void PostPoseUpdate() const;
 	void NotifyDrivenControls(UControlRig* InControlRig, const FRigElementKey& InKey);
 	void UpdateSelectabilityOnSkeletalMeshes(UControlRig* InControlRig, bool bEnabled);
 
