Control Rig: Change path to default root module

#rnx

diff --git a/Engine/Plugins/Animation/ControlRig/Source/ControlRig/Private/Settings/ControlRigSettings.cpp b/Engine/Plugins/Animation/ControlRig/Source/ControlRig/Private/Settings/ControlRigSettings.cpp
--- a/Engine/Plugins/Animation/ControlRig/Source/ControlRig/Private/Settings/ControlRigSettings.cpp
+++ b/Engine/Plugins/Animation/ControlRig/Source/ControlRig/Private/Settings/ControlRigSettings.cpp
@@ -10,7 +10,7 @@ UControlRigSettings::UControlRigSettings(const FObjectInitializer& ObjectInitial
 {
 #if WITH_EDITORONLY_DATA
 	DefaultShapeLibrary = LoadObject<UControlRigShapeLibrary>(nullptr, TEXT("/ControlRig/Controls/DefaultGizmoLibraryNormalized.DefaultGizmoLibraryNormalized"));
-	DefaultRootModule = TEXT("/ControlRig/Modules/CR_Root.CR_Root");
+	DefaultRootModule = TEXT("/ControlRig/Modules/Root.Root");
 #endif
 }
 
