Control Rig: Use correct icon for connector in rig hierarchy

#rnx

diff --git a/Engine/Plugins/Animation/ControlRig/Source/ControlRigEditor/Private/Editor/SRigHierarchyTreeView.cpp b/Engine/Plugins/Animation/ControlRig/Source/ControlRigEditor/Private/Editor/SRigHierarchyTreeView.cpp
--- a/Engine/Plugins/Animation/ControlRig/Source/ControlRigEditor/Private/Editor/SRigHierarchyTreeView.cpp
+++ b/Engine/Plugins/Animation/ControlRig/Source/ControlRigEditor/Private/Editor/SRigHierarchyTreeView.cpp
@@ -195,6 +195,7 @@ void SRigHierarchyItem::Construct(const FArguments& InArgs, const TSharedRef<STa
 					}
 					return FSlateColor::UseForeground();
 				})
+				.DesiredSizeOverride(FVector2D(16, 16))
 			]
 			+ SHorizontalBox::Slot()
 			.AutoWidth()
@@ -1117,8 +1118,10 @@ TPair<const FSlateBrush*, FSlateColor> SRigHierarchyItem::GetBrushForElementType
 	static const FSlateBrush* RigidBodyBrush = FControlRigEditorStyle::Get().GetBrush("ControlRig.Tree.RigidBody");
 	static const FSlateBrush* SocketOpenBrush = FControlRigEditorStyle::Get().GetBrush("ControlRig.Tree.Socket_Open");
 	static const FSlateBrush* SocketClosedBrush = FControlRigEditorStyle::Get().GetBrush("ControlRig.Tree.Socket_Closed");
-	static const FSlateBrush* ConnectorBrush = FControlRigEditorStyle::Get().GetBrush("ControlRig.Tree.Connector");
-	
+	static const FSlateBrush* PrimaryConnectorBrush = FControlRigEditorStyle::Get().GetBrush("ControlRig.ConnectorPrimary");
+	static const FSlateBrush* SecondaryConnectorBrush = FControlRigEditorStyle::Get().GetBrush("ControlRig.ConnectorSecondary");
+	static const FSlateBrush* OptionalConnectorBrush = FControlRigEditorStyle::Get().GetBrush("ControlRig.ConnectorOptional");
+
 	const FSlateBrush* Brush = nullptr;
 	FSlateColor Color = FSlateColor::UseForeground();
 	switch (InKey.Type)
@@ -1222,7 +1225,14 @@ TPair<const FSlateBrush*, FSlateColor> SRigHierarchyItem::GetBrushForElementType
 		}
 		case ERigElementType::Connector:
 		{
-			Brush = ConnectorBrush;
+			Brush = PrimaryConnectorBrush;
+			if(const FRigConnectorElement* Connector = InHierarchy->Find<FRigConnectorElement>(InKey))
+			{
+				if(!Connector->IsPrimary())
+				{
+					Brush = Connector->IsOptional() ? OptionalConnectorBrush : SecondaryConnectorBrush;
+				}
+			}
 			break;
 		}
 		default:
