Control Rig: Fix invalid preferred angles when spawning constols or adding parent
#rnx

diff --git a/Engine/Plugins/Animation/ControlRig/Source/ControlRig/Private/Rigs/RigHierarchyController.cpp b/Engine/Plugins/Animation/ControlRig/Source/ControlRig/Private/Rigs/RigHierarchyController.cpp
--- a/Engine/Plugins/Animation/ControlRig/Source/ControlRig/Private/Rigs/RigHierarchyController.cpp
+++ b/Engine/Plugins/Animation/ControlRig/Source/ControlRig/Private/Rigs/RigHierarchyController.cpp
@@ -377,9 +377,7 @@ FRigElementKey URigHierarchyController::AddControl(
 		NewElement->Shape.Set(ERigTransformType::InitialLocal, InShapeTransform);  
 		Hierarchy->SetControlValue(NewElement, InValue, ERigControlValueType::Initial, false);
 		const FTransform LocalTransform = Hierarchy->GetTransform(NewElement, ERigTransformType::InitialLocal);
-
-		const FVector PreferredEulerAngles = AnimationCore::EulerFromQuat(LocalTransform.GetRotation(), NewElement->Settings.PreferredRotationOrder);
-		Hierarchy->SetControlPreferredEulerAngles(NewElement, PreferredEulerAngles, NewElement->Settings.PreferredRotationOrder, true);
+		Hierarchy->SetControlPreferredEulerAngles(NewElement, LocalTransform);
 
 		NewElement->Offset.MarkDirty(ERigTransformType::InitialGlobal);
 		NewElement->Pose.MarkDirty(ERigTransformType::InitialGlobal);
@@ -2850,6 +2848,13 @@ bool URigHierarchyController::AddParent(FRigBaseElement* InChild, FRigBaseElemen
 			}
 		}
 
+		if(FRigControlElement* ChildControlElement = Cast<FRigControlElement>(InChild))
+		{
+			const FTransform LocalTransform = Hierarchy->GetTransform(ChildControlElement, ERigTransformType::InitialLocal);
+			Hierarchy->SetControlPreferredEulerAngles(ChildControlElement, LocalTransform, true);
+			ChildControlElement->PreferredEulerAngles.Current = ChildControlElement->PreferredEulerAngles.Initial;
+		}
+
 		Notify(ERigHierarchyNotification::ParentChanged, MultiParentElement);
 
 		Hierarchy->EnsureCacheValidity();
diff --git a/Engine/Plugins/Animation/ControlRig/Source/ControlRig/Public/Rigs/RigHierarchy.h b/Engine/Plugins/Animation/ControlRig/Source/ControlRig/Public/Rigs/RigHierarchy.h
--- a/Engine/Plugins/Animation/ControlRig/Source/ControlRig/Public/Rigs/RigHierarchy.h
+++ b/Engine/Plugins/Animation/ControlRig/Source/ControlRig/Public/Rigs/RigHierarchy.h
@@ -2513,6 +2513,28 @@ public:
 		}
 	}
 
+	void SetControlPreferredEulerAngles(FRigControlElement* InControlElement, const FTransform& InTransform, bool bIsInitial = false)
+	{
+		FEulerTransform EulerTransform(InTransform);
+		if (InControlElement && InControlElement->Settings.ControlType == ERigControlType::Transform)
+		{
+			FVector EulerAngle(EulerTransform.Rotation.Roll, EulerTransform.Rotation.Pitch, EulerTransform.Rotation.Yaw);
+			SetControlSpecifiedEulerAngle(InControlElement, EulerAngle, bIsInitial);
+		}
+		else if (InControlElement && InControlElement->Settings.ControlType == ERigControlType::TransformNoScale)
+		{
+			const FTransformNoScale NoScale = EulerTransform.ToFTransform();
+			FVector EulerAngle(EulerTransform.Rotation.Roll, EulerTransform.Rotation.Pitch, EulerTransform.Rotation.Yaw);
+			SetControlSpecifiedEulerAngle(InControlElement, EulerAngle, bIsInitial);
+		}
+		else if (InControlElement && InControlElement->Settings.ControlType == ERigControlType::EulerTransform)
+		{
+			FVector EulerAngle(EulerTransform.Rotation.Roll, EulerTransform.Rotation.Pitch, EulerTransform.Rotation.Yaw);
+			FQuat Quat = GetControlQuaternion(InControlElement, EulerAngle);
+			SetControlSpecifiedEulerAngle(InControlElement, EulerAngle, bIsInitial);
+		}
+	}
+
 	FQuat GetControlQuaternion(const FRigControlElement* InControlElement, const FVector& InEulerAngle) const
 	{
 		if (InControlElement)
diff --git a/Engine/Plugins/Animation/ControlRig/Source/ControlRigEditor/Private/EditMode/SControlRigDetails.cpp b/Engine/Plugins/Animation/ControlRig/Source/ControlRigEditor/Private/EditMode/SControlRigDetails.cpp
--- a/Engine/Plugins/Animation/ControlRig/Source/ControlRigEditor/Private/EditMode/SControlRigDetails.cpp
+++ b/Engine/Plugins/Animation/ControlRig/Source/ControlRigEditor/Private/EditMode/SControlRigDetails.cpp
@@ -320,7 +320,7 @@ void SControlRigDetails::UpdateProxies()
 					const TArray<UControlRigControlsProxy*>& Proxies = ControlProxy->GetAllSelectedProxies();
 					for (UControlRigControlsProxy* Proxy : Proxies)
 					{
-						if (Proxy == nullptr)
+						if (Proxy == nullptr || !IsValid(Proxy))
 						{
 							continue;
 						}
