Control Rig: Fix layered rigs still applying preferred euler angles
#rnx

diff --git a/Engine/Plugins/Animation/ControlRig/Source/ControlRig/Public/Rigs/RigHierarchy.h b/Engine/Plugins/Animation/ControlRig/Source/ControlRig/Public/Rigs/RigHierarchy.h
--- a/Engine/Plugins/Animation/ControlRig/Source/ControlRig/Public/Rigs/RigHierarchy.h
+++ b/Engine/Plugins/Animation/ControlRig/Source/ControlRig/Public/Rigs/RigHierarchy.h
@@ -4727,6 +4727,9 @@ protected:
 	bool bUsePreferredEulerAngles;
 	mutable bool bAllowNameSpaceWhenSanitizingName;
 
+public:
+	bool UsesPreferredEulerAngles() const { return bUsePreferredEulerAngles; }
+
 private:
 	
 	void EnsureCacheValidityImpl();
diff --git a/Engine/Plugins/Animation/ControlRig/Source/ControlRigEditor/Private/EditMode/AnimDetailsProxy.cpp b/Engine/Plugins/Animation/ControlRig/Source/ControlRigEditor/Private/EditMode/AnimDetailsProxy.cpp
--- a/Engine/Plugins/Animation/ControlRig/Source/ControlRigEditor/Private/EditMode/AnimDetailsProxy.cpp
+++ b/Engine/Plugins/Animation/ControlRig/Source/ControlRigEditor/Private/EditMode/AnimDetailsProxy.cpp
@@ -846,7 +846,10 @@ static FEulerTransform GetCurrentValue(UControlRig* ControlRig, FRigControlEleme
 			break;
 		}
 	};
-	EulerTransform.Rotation = ControlRig->GetHierarchy()->GetControlPreferredRotator(ControlElement);
+	if (ControlRig->GetHierarchy()->UsesPreferredEulerAngles())
+	{
+		EulerTransform.Rotation = ControlRig->GetHierarchy()->GetControlPreferredRotator(ControlElement);
+	}
 	return EulerTransform;
 }
 
@@ -1129,14 +1132,22 @@ void UAnimDetailControlsProxyTransform::SetControlRigElementValueFromCurrent(UCo
 		case ERigControlType::EulerTransform:
 		{
 			URigHierarchy* Hierarchy = ControlRig->GetHierarchy();
-			FVector EulerAngle(TRotation.Roll, TRotation.Pitch, TRotation.Yaw);
-			FQuat Quat = Hierarchy->GetControlQuaternion(ControlElement, EulerAngle);
-			Hierarchy->SetControlSpecifiedEulerAngle(ControlElement, EulerAngle);
-			FRotator UERotator(Quat);
-			FEulerTransform UETransform(UERotator, TLocation, TScale);
-			UETransform.Rotation = UERotator;
-			ControlRig->SetControlValue<FRigControlValue::FEulerTransform_Float>(ControlElement->GetKey().Name, UETransform, true, Context, false);
-			Hierarchy->SetControlSpecifiedEulerAngle(ControlElement, EulerAngle);
+
+			if (Hierarchy->UsesPreferredEulerAngles())
+			{
+				FVector EulerAngle(TRotation.Roll, TRotation.Pitch, TRotation.Yaw);
+				FQuat Quat = Hierarchy->GetControlQuaternion(ControlElement, EulerAngle);
+				Hierarchy->SetControlSpecifiedEulerAngle(ControlElement, EulerAngle);
+				FRotator UERotator(Quat);
+				FEulerTransform UETransform(UERotator, TLocation, TScale);
+				UETransform.Rotation = UERotator;
+				ControlRig->SetControlValue<FRigControlValue::FEulerTransform_Float>(ControlElement->GetKey().Name, UETransform, true, Context, false);
+				Hierarchy->SetControlSpecifiedEulerAngle(ControlElement, EulerAngle);
+			}
+			else
+			{
+				ControlRig->SetControlValue<FRigControlValue::FEulerTransform_Float>(ControlElement->GetKey().Name, FEulerTransform(RealTransform), true, Context, false);
+			}
 			break;
 		}
 		}
