Sequencer: Anim Details: Don't update if compensating, and last fix from rotation order breakage

diff --git a/Engine/Plugins/Animation/ControlRig/Source/ControlRigEditor/Private/EditMode/ControlRigEditMode.cpp b/Engine/Plugins/Animation/ControlRig/Source/ControlRigEditor/Private/EditMode/ControlRigEditMode.cpp
--- a/Engine/Plugins/Animation/ControlRig/Source/ControlRigEditor/Private/EditMode/ControlRigEditMode.cpp
+++ b/Engine/Plugins/Animation/ControlRig/Source/ControlRigEditor/Private/EditMode/ControlRigEditMode.cpp
@@ -4648,8 +4648,8 @@ void FControlRigEditMode::MoveControlShape(AControlRigShapeActor* ShapeActor, co
 				ControlRig->InteractionType = InteractionType;
 				ControlRig->ElementsBeingInteracted.AddUnique(ShapeActor->GetElementKey());
 				
-				ControlRig->SetControlLocalTransform(ShapeActor->ControlName, CurrentLocalTransform,true, FRigControlModifiedContext(), true, /*fix eulers*/ false);
-				UpdatePreferredEulerAngles(ControlRig);
+				ControlRig->SetControlLocalTransform(ShapeActor->ControlName, CurrentLocalTransform,true, FRigControlModifiedContext(), true, /*fix eulers*/ true);
+				//UpdatePreferredEulerAngles(ControlRig);
 
 				FTransform CurrentTransform  = ControlRig->GetControlGlobalTransform(ShapeActor->ControlName);			// assumes it's attached to actor
 				CurrentTransform = ToWorldTransform * CurrentTransform;
@@ -5702,9 +5702,12 @@ void FDetailKeyFrameCacheAndHandler::OnChannelChanged(const FMovieSceneChannelMe
 void FDetailKeyFrameCacheAndHandler::ResetCachedData()
 {
 	CachedPropertyKeyedStatusMap.Reset();
-	if (EditMode && EditMode->GetControlProxy())
+	if (FMovieSceneConstraintChannelHelper::bDoNotCompensate == false) //if compensating don't reset this.
 	{
-		EditMode->GetControlProxy()->ValuesChanged();
+		if (EditMode && EditMode->GetControlProxy())
+		{
+			EditMode->GetControlProxy()->ValuesChanged();
+		}
 	}
 }
 
