Sequencer : Anim Details: Make sure to set sequencer on edit mode and pass that through to tabs. Also only check one type

diff --git a/Engine/Plugins/Animation/ControlRig/Source/ControlRigEditor/Private/EditMode/AnimDetailsProxy.cpp b/Engine/Plugins/Animation/ControlRig/Source/ControlRigEditor/Private/EditMode/AnimDetailsProxy.cpp
--- a/Engine/Plugins/Animation/ControlRig/Source/ControlRigEditor/Private/EditMode/AnimDetailsProxy.cpp
+++ b/Engine/Plugins/Animation/ControlRig/Source/ControlRigEditor/Private/EditMode/AnimDetailsProxy.cpp
@@ -648,6 +648,7 @@ static EPropertyKeyedStatus GetKeyedStatusInSection(const UMovieSceneSection* Se
 				SectionKeyedStatus = FMath::Max(SectionKeyedStatus, EPropertyKeyedStatus::PartiallyKeyed);
 			}
 		}
+		break; //just do it for one type
 	}
 
 	if (EmptyChannelCount == 0 && SectionKeyedStatus == EPropertyKeyedStatus::PartiallyKeyed)
diff --git a/Engine/Plugins/Animation/ControlRig/Source/ControlRigEditor/Private/EditMode/ControlRigEditMode.cpp b/Engine/Plugins/Animation/ControlRig/Source/ControlRigEditor/Private/EditMode/ControlRigEditMode.cpp
--- a/Engine/Plugins/Animation/ControlRig/Source/ControlRigEditor/Private/EditMode/ControlRigEditMode.cpp
+++ b/Engine/Plugins/Animation/ControlRig/Source/ControlRigEditor/Private/EditMode/ControlRigEditMode.cpp
@@ -11,6 +11,8 @@
 #include "ControlRig.h"
 #include "HitProxies.h"
 #include "EditMode/ControlRigEditModeSettings.h"
+#include "EditMode/SControlRigDetails.h"
+#include "EditMode/SControlRigOutliner.h"
 #include "ISequencer.h"
 #include "SequencerSettings.h"
 #include "Sections/MovieSceneSpawnSection.h"
@@ -65,6 +67,7 @@
 #include "EdModeInteractiveToolsContext.h"
 #include "Constraints/MovieSceneConstraintChannelHelper.h"
 #include "Editor/EditorPerProjectUserSettings.h"
+#include "Widgets/Docking/SDockTab.h"
 #include "TransformConstraint.h"
 #include "Animation/DebugSkelMeshComponent.h"
 #include "Materials/Material.h"
@@ -211,7 +214,7 @@ FControlRigEditMode::~FControlRigEditMode()
 
 }
 
-bool FControlRigEditMode:: SetSequencer(TWeakPtr<ISequencer> InSequencer)
+bool FControlRigEditMode::SetSequencer(TWeakPtr<ISequencer> InSequencer)
 {
 	if (InSequencer != WeakSequencer)
 	{
@@ -248,6 +251,14 @@ bool FControlRigEditMode:: SetSequencer(TWeakPtr<ISequencer> InSequencer)
 			ControlProxy->SetSequencer(WeakSequencer);
 		}
 		SetObjects_Internal();
+		if (FControlRigEditModeToolkit::Details.IsValid())
+		{
+			FControlRigEditModeToolkit::Details->SetEditMode(*this);
+		}
+		if (FControlRigEditModeToolkit::Outliner)
+		{
+			FControlRigEditModeToolkit::Outliner->SetEditMode(*this);
+		}
 	}
 	return false;
 }
diff --git a/Engine/Plugins/Animation/ControlRig/Source/ControlRigEditor/Private/EditMode/ControlRigEditModeToolkit.cpp b/Engine/Plugins/Animation/ControlRig/Source/ControlRigEditor/Private/EditMode/ControlRigEditModeToolkit.cpp
--- a/Engine/Plugins/Animation/ControlRig/Source/ControlRigEditor/Private/EditMode/ControlRigEditModeToolkit.cpp
+++ b/Engine/Plugins/Animation/ControlRig/Source/ControlRigEditor/Private/EditMode/ControlRigEditModeToolkit.cpp
@@ -46,6 +46,9 @@ const FName FControlRigEditModeToolkit::OutlinerTabName = FName(TEXT("ControlRig
 const FName FControlRigEditModeToolkit::DetailsTabName = FName(TEXT("ControlRigDetailsTab"));
 const FName FControlRigEditModeToolkit::SpacePickerTabName = FName(TEXT("ControlRigSpacePicker"));
 
+TSharedPtr<SControlRigDetails> FControlRigEditModeToolkit::Details = nullptr;
+TSharedPtr<SControlRigOutliner> FControlRigEditModeToolkit::Outliner = nullptr;
+
 void FControlRigEditModeToolkit::Init(const TSharedPtr<IToolkitHost>& InitToolkitHost)
 {
 	SAssignNew(ModeTools, SControlRigEditModeTools, SharedThis(this), EditMode);
@@ -190,7 +193,7 @@ TSharedRef<SDockTab> SpawnOutlinerTab(const FSpawnTabArgs& Args, FControlRigEdit
 {
 	return SNew(SDockTab)
 		[
-			SNew(SControlRigOutliner, *InEditorMode)
+			SAssignNew(FControlRigEditModeToolkit::Outliner, SControlRigOutliner, *InEditorMode)
 		];
 }
 
@@ -206,7 +209,7 @@ TSharedRef<SDockTab> SpawnDetailsTab(const FSpawnTabArgs& Args, FControlRigEditM
 {
 	return SNew(SDockTab)
 		[
-			SNew(SControlRigDetails, *InEditorMode)
+			SAssignNew(FControlRigEditModeToolkit::Details,SControlRigDetails, *InEditorMode)
 		];
 }
 
diff --git a/Engine/Plugins/Animation/ControlRig/Source/ControlRigEditor/Private/EditMode/ControlRigEditModeToolkit.h b/Engine/Plugins/Animation/ControlRig/Source/ControlRigEditor/Private/EditMode/ControlRigEditModeToolkit.h
--- a/Engine/Plugins/Animation/ControlRig/Source/ControlRigEditor/Private/EditMode/ControlRigEditModeToolkit.h
+++ b/Engine/Plugins/Animation/ControlRig/Source/ControlRigEditor/Private/EditMode/ControlRigEditModeToolkit.h
@@ -12,6 +12,8 @@
 #include "EditMode/ControlRigEditMode.h"
 
 class SControlRigTweenWidget;
+class SControlRigDetails;
+class SControlRigOutliner;
 
 class FControlRigEditModeToolkit : public FModeToolkit
 {
@@ -67,6 +69,9 @@ public:
 	static const FName OutlinerTabName;
 	static const FName SpacePickerTabName;
 
+	static TSharedPtr<SControlRigDetails> Details;
+	static TSharedPtr<SControlRigOutliner> Outliner;
+
 protected:
 
 	void CreateAndShowTweenOverlay();
@@ -89,15 +94,12 @@ protected:
 	static bool bPoseTabOpen;
 	static bool bSnapperTabOpen;
 	static bool bTweenOpen;
-
-
 private:
 	/** The edit mode we are bound to */
 	FControlRigEditMode& EditMode;
 	TSharedPtr<SWidget> TweenWidgetParent;
 	TSharedPtr<SControlRigTweenWidget> TweenWidget;
 
-
 	FVector2D InViewportTweenWidgetLocation;
 	/** The tools widget */
 	TSharedPtr<SControlRigEditModeTools> ModeTools;
diff --git a/Engine/Plugins/Animation/ControlRig/Source/ControlRigEditor/Private/EditMode/SControlRigDetails.cpp b/Engine/Plugins/Animation/ControlRig/Source/ControlRigEditor/Private/EditMode/SControlRigDetails.cpp
--- a/Engine/Plugins/Animation/ControlRig/Source/ControlRigEditor/Private/EditMode/SControlRigDetails.cpp
+++ b/Engine/Plugins/Animation/ControlRig/Source/ControlRigEditor/Private/EditMode/SControlRigDetails.cpp
@@ -94,7 +94,17 @@ void SControlRigDetails::Construct(const FArguments& InArgs, FControlRigEditMode
 		];
 
 	SetEditMode(InEditMode);
-	SequencerTracker.SetSequencerAndDetails(InEditMode.GetWeakSequencer(), this);
+
+}
+
+void SControlRigDetails::SetEditMode(FControlRigEditMode& InEditMode)
+{
+	FControlRigBaseDockableView::SetEditMode(InEditMode);
+	if (GetEditMode()->GetWeakSequencer().IsValid())
+	{
+		SequencerTracker.SetSequencerAndDetails(GetEditMode()->GetWeakSequencer(), this);
+		UpdateProxies();
+	}
 }
 
 SControlRigDetails::~SControlRigDetails()
diff --git a/Engine/Plugins/Animation/ControlRig/Source/ControlRigEditor/Private/EditMode/SControlRigDetails.h b/Engine/Plugins/Animation/ControlRig/Source/ControlRigEditor/Private/EditMode/SControlRigDetails.h
--- a/Engine/Plugins/Animation/ControlRig/Source/ControlRigEditor/Private/EditMode/SControlRigDetails.h
+++ b/Engine/Plugins/Animation/ControlRig/Source/ControlRigEditor/Private/EditMode/SControlRigDetails.h
@@ -74,6 +74,9 @@ class SControlRigDetails: public SCompoundWidget, public FControlRigBaseDockable
 
 	void Construct(const FArguments& InArgs, FControlRigEditMode& InEditMode);
 
+	//FControlRigBaseDockableView overrides
+	virtual void SetEditMode(FControlRigEditMode& InEditMode) override;
+
 	/** Display or edit set up for property */
 	bool ShouldShowPropertyOnDetailCustomization(const struct FPropertyAndParent& InPropertyAndParent) const;
 	bool IsReadOnlyPropertyOnDetailCustomization(const struct FPropertyAndParent& InPropertyAndParent) const;
diff --git a/Engine/Plugins/Animation/ControlRig/Source/ControlRigEditor/Private/Sequencer/ControlRigParameterTrackEditor.cpp b/Engine/Plugins/Animation/ControlRig/Source/ControlRigEditor/Private/Sequencer/ControlRigParameterTrackEditor.cpp
--- a/Engine/Plugins/Animation/ControlRig/Source/ControlRigEditor/Private/Sequencer/ControlRigParameterTrackEditor.cpp
+++ b/Engine/Plugins/Animation/ControlRig/Source/ControlRigEditor/Private/Sequencer/ControlRigParameterTrackEditor.cpp
@@ -2015,11 +2015,13 @@ void FControlRigParameterTrackEditor::OnActivateSequenceChanged(FMovieSceneSeque
 			{
 				GEditor->GetTimerManager()->SetTimerForNextTick([this, ControlRigEditMode]()
 				{
-					for(TWeakObjectPtr<UControlRig> &ControlRig: BoundControlRigs)
-					{ 
+					bool bSequencerSet = false;
+					for (TWeakObjectPtr<UControlRig>& ControlRig : BoundControlRigs)
+					{
 						if (ControlRig.IsValid())
 						{
 							ControlRigEditMode->AddControlRigObject(ControlRig.Get(), GetSequencer());
+							bSequencerSet = true;
 							for (int32 Index = 0; Index < PreviousSelectedControlRigs.Num(); ++Index)
 							{
 								if (ControlRig.Get()->GetClass() == PreviousSelectedControlRigs[Index].Key)
@@ -2031,9 +2033,13 @@ void FControlRigParameterTrackEditor::OnActivateSequenceChanged(FMovieSceneSeque
 									PreviousSelectedControlRigs.RemoveAt(Index);
 									break;
 								}
-							}	
+							}
 						}
 					}
+					if (bSequencerSet == false)
+					{
+						ControlRigEditMode->SetSequencer(GetSequencer());
+					}
 					PreviousSelectedControlRigs.Reset();
 				});
 			}
@@ -3144,7 +3150,7 @@ void FControlRigParameterTrackEditor::HandleControlSelected(UControlRig* Subject
 			//Force refresh later, not now
 			bSkipNextSelectionFromTimer = bSkipNextSelectionFromTimer ||
 				(bIsSelectingIndirectControl && ControlElement->Settings.AnimationType == ERigControlAnimationType::AnimationControl);
-			GetSequencer()->NotifyMovieSceneDataChanged(EMovieSceneDataChangeType::RefreshTree);
+
 		}
 	}
 }
