Control RIg: Fix icon size in Anim Outliner
#rnx

diff --git a/Engine/Plugins/Animation/ControlRig/Source/ControlRigEditor/Private/EditMode/SControlRigOutliner.cpp b/Engine/Plugins/Animation/ControlRig/Source/ControlRigEditor/Private/EditMode/SControlRigOutliner.cpp
--- a/Engine/Plugins/Animation/ControlRig/Source/ControlRigEditor/Private/EditMode/SControlRigOutliner.cpp
+++ b/Engine/Plugins/Animation/ControlRig/Source/ControlRigEditor/Private/EditMode/SControlRigOutliner.cpp
@@ -215,20 +215,18 @@ void SMultiRigHierarchyItem::Construct(const FArguments& InArgs, const TSharedRe
 		[
 			SNew(SHorizontalBox)
 			+ SHorizontalBox::Slot()
-		.MaxWidth(18)
-		.FillWidth(1.0)
-		.HAlign(HAlign_Left)
-		.VAlign(VAlign_Center)
-		.Padding(FMargin(0.f, 0.f, 3.f, 0.f))
-		[
-
-
-			SNew(SButton)
-			.ButtonStyle(FAppStyle::Get(), "NoBorder")
-			.OnClicked(this, &SMultiRigHierarchyItem::OnGetSelectedClicked)
+			.MaxWidth(24)
+			.FillWidth(1.0)
+			.HAlign(HAlign_Left)
+			.VAlign(VAlign_Center)
+			.Padding(FMargin(0.f, 0.f, 3.f, 0.f))
 			[
-				SNew(SImage)
-				.Image_Lambda([this]() -> const FSlateBrush*
+				SNew(SButton)
+				.ButtonStyle(FAppStyle::Get(), "NoBorder")
+				.OnClicked(this, &SMultiRigHierarchyItem::OnGetSelectedClicked)
+				[
+					SNew(SImage)
+					.Image_Lambda([this]() -> const FSlateBrush*
 					{
 						//if no key is set then it's the control rig so we get that based upon it's state
 						if (WeakRigTreeElement.Pin()->Data.Key.IsSet() == false)
@@ -252,16 +250,17 @@ void SMultiRigHierarchyItem::Construct(const FArguments& InArgs, const TSharedRe
 					{
 						return WeakRigTreeElement.Pin()->IconColor;
 					})
+					.DesiredSizeOverride(FVector2D(16, 16))
+				]
 			]
-		]
-		+ SHorizontalBox::Slot()
+			+ SHorizontalBox::Slot()
 			.AutoWidth()
 			.VAlign(VAlign_Center)
 			[
 				SAssignNew(InlineWidget, SInlineEditableTextBlock)
 				.Text(this, &SMultiRigHierarchyItem::GetDisplayName)
-			.MultiLine(false)
-			.ColorAndOpacity_Lambda([this]()
+				.MultiLine(false)
+				.ColorAndOpacity_Lambda([this]()
 				{
 					if (WeakRigTreeElement.IsValid())
 					{
@@ -270,7 +269,7 @@ void SMultiRigHierarchyItem::Construct(const FArguments& InArgs, const TSharedRe
 					return FSlateColor::UseForeground();
 				})
 			]
-			], OwnerTable);
+		], OwnerTable);
 
 }
 
