Constraints: fixed constraints listed even if with no sequencer opened

#rnx

diff --git a/Engine/Plugins/Animation/ControlRig/Source/ControlRigEditor/Private/Editor/Constraints/SConstraintsWidget.cpp b/Engine/Plugins/Animation/ControlRig/Source/ControlRigEditor/Private/Editor/Constraints/SConstraintsWidget.cpp
--- a/Engine/Plugins/Animation/ControlRig/Source/ControlRigEditor/Private/Editor/Constraints/SConstraintsWidget.cpp
+++ b/Engine/Plugins/Animation/ControlRig/Source/ControlRigEditor/Private/Editor/Constraints/SConstraintsWidget.cpp
@@ -975,6 +975,10 @@ int32 FBaseConstraintListWidget::RefreshConstraintList()
 					}
 				}
 			}
+			else
+			{
+				Constraints.Empty();
+			}
 		}
 		else if (ShowConstraints == EShowConstraints::ShowValid)
 		{
