Control Rig: Fix crash when the OwnerControlElement pointer of a UControlRigControlsProxy becomes invalid
#rnx

diff --git a/Engine/Plugins/Animation/ControlRig/Source/ControlRig/Public/Rigs/RigHierarchyDefines.h b/Engine/Plugins/Animation/ControlRig/Source/ControlRig/Public/Rigs/RigHierarchyDefines.h
--- a/Engine/Plugins/Animation/ControlRig/Source/ControlRig/Public/Rigs/RigHierarchyDefines.h
+++ b/Engine/Plugins/Animation/ControlRig/Source/ControlRig/Public/Rigs/RigHierarchyDefines.h
@@ -1493,7 +1493,7 @@ public:
 
 	bool IsValid() const
 	{
-		return Name != NAME_None && Type != ERigElementType::None;
+		return Name.IsValid() && Name != NAME_None && Type != ERigElementType::None;
 	}
 
 	explicit operator bool() const
diff --git a/Engine/Plugins/Animation/ControlRig/Source/ControlRigEditor/Private/EditMode/AnimDetailsProxy.cpp b/Engine/Plugins/Animation/ControlRig/Source/ControlRigEditor/Private/EditMode/AnimDetailsProxy.cpp
--- a/Engine/Plugins/Animation/ControlRig/Source/ControlRigEditor/Private/EditMode/AnimDetailsProxy.cpp
+++ b/Engine/Plugins/Animation/ControlRig/Source/ControlRigEditor/Private/EditMode/AnimDetailsProxy.cpp
@@ -2956,7 +2956,7 @@ UControlRigControlsProxy* UControlRigDetailPanelControlProxies::AddProxy(UContro
 		if (Proxy)
 		{
 			Proxy->Type = ControlElement->Settings.ControlType;;
-			Proxy->OwnerControlElement = ControlElement;
+			Proxy->OwnerControlElement.UpdateCache(ControlElement->GetKey(), ControlRig->GetHierarchy());
 			Proxy->OwnerControlRig = ControlRig;
 			Proxy->AddControlRigControl(ControlRig, ControlElement->GetKey().Name);
 			Proxy->bIsIndividual = bIsIndividual;
diff --git a/Engine/Plugins/Animation/ControlRig/Source/ControlRigEditor/Private/EditMode/ControlRigControlsProxy.cpp b/Engine/Plugins/Animation/ControlRig/Source/ControlRigEditor/Private/EditMode/ControlRigControlsProxy.cpp
--- a/Engine/Plugins/Animation/ControlRig/Source/ControlRigEditor/Private/EditMode/ControlRigControlsProxy.cpp
+++ b/Engine/Plugins/Animation/ControlRig/Source/ControlRigEditor/Private/EditMode/ControlRigControlsProxy.cpp
@@ -102,7 +102,10 @@ void UControlRigControlsProxy::ResetControlRigItems()
 	ControlRigItems.Reset();
 	if (OwnerControlRig.IsValid())
 	{
-		AddControlRigControl(OwnerControlRig.Get(), OwnerControlElement->GetFName());
+		if (OwnerControlElement.IsValid())
+		{
+			AddControlRigControl(OwnerControlRig.Get(), OwnerControlElement.GetKey().Name);
+		}
 	}
 }
 
@@ -114,9 +117,9 @@ void UControlRigControlsProxy::ResetItems()
 
 void UControlRigControlsProxy::AddItem(UControlRigControlsProxy* ControlProxy)
 {
-	if (ControlProxy->OwnerControlRig.IsValid())
+	if (ControlProxy->OwnerControlRig.IsValid() && ControlProxy->OwnerControlElement.IsValid())
 	{
-		AddControlRigControl(ControlProxy->OwnerControlRig.Get(), ControlProxy->OwnerControlElement->GetFName());
+		AddControlRigControl(ControlProxy->OwnerControlRig.Get(), ControlProxy->OwnerControlElement.GetKey().Name);
 	}
 	else if(ControlProxy->OwnerObject.IsValid())
 	{
@@ -208,12 +211,15 @@ void UControlRigControlsProxy::AddChildProxy(UControlRigControlsProxy* ControlPr
 
 void UControlRigControlsProxy::SelectionChanged(bool bInSelected)
 {
-	if (OwnerControlElement)
+	if (OwnerControlRig.IsValid())
 	{
-		Modify();
-		const FName PropertyName("bSelected");
-		FTrackInstancePropertyBindings Binding(PropertyName, PropertyName.ToString());
-		Binding.CallFunction<bool>(*this, bInSelected);
+		if (OwnerControlElement.UpdateCache(OwnerControlRig->GetHierarchy()))
+		{
+			Modify();
+			const FName PropertyName("bSelected");
+			FTrackInstancePropertyBindings Binding(PropertyName, PropertyName.ToString());
+			Binding.CallFunction<bool>(*this, bInSelected);
+		}
 	}
 }
 
@@ -226,10 +232,10 @@ void UControlRigControlsProxy::PostEditChangeChainProperty(struct FPropertyChang
 	}
 	if (PropertyChangedEvent.Property && PropertyChangedEvent.Property->GetFName() == GET_MEMBER_NAME_CHECKED(UControlRigControlsProxy, bSelected))
 	{
-		if (OwnerControlElement && OwnerControlRig.IsValid())
+		if (OwnerControlElement.IsValid() && OwnerControlRig.IsValid())
 		{
-			FControlRigInteractionScope InteractionScope(OwnerControlRig.Get(), OwnerControlElement->GetKey());
-			OwnerControlRig.Get()->SelectControl(OwnerControlElement->GetKey().Name, bSelected);
+			FControlRigInteractionScope InteractionScope(OwnerControlRig.Get(), OwnerControlElement.GetKey());
+			OwnerControlRig.Get()->SelectControl(OwnerControlElement.GetKey().Name, bSelected);
 			OwnerControlRig.Get()->Evaluate_AnyThread();
 		}
 	}
diff --git a/Engine/Plugins/Animation/ControlRig/Source/ControlRigEditor/Private/EditMode/ControlRigControlsProxy.h b/Engine/Plugins/Animation/ControlRig/Source/ControlRigEditor/Private/EditMode/ControlRigControlsProxy.h
--- a/Engine/Plugins/Animation/ControlRig/Source/ControlRigEditor/Private/EditMode/ControlRigControlsProxy.h
+++ b/Engine/Plugins/Animation/ControlRig/Source/ControlRigEditor/Private/EditMode/ControlRigControlsProxy.h
@@ -12,6 +12,7 @@
 #include "IPropertyTypeCustomization.h"
 #include "Rigs/RigHierarchyDefines.h"
 #include "MovieSceneCommonHelpers.h"
+#include "Rigs/RigHierarchyCache.h"
 
 #include "ControlRigControlsProxy.generated.h"
 
@@ -126,7 +127,7 @@ public:
 	//refactor
 	//We can set/get values form multiple control rig elements but only one owns this.
 	TWeakObjectPtr<UControlRig> OwnerControlRig;
-	FRigControlElement* OwnerControlElement = nullptr;
+	FCachedRigElement OwnerControlElement;
 	TMap<TWeakObjectPtr<UControlRig>, FControlRigProxyItem> ControlRigItems;
 	
 	TWeakObjectPtr<UObject> OwnerObject;
diff --git a/Engine/Plugins/Animation/ControlRig/Source/ControlRigEditor/Private/EditMode/SControlRigDetails.cpp b/Engine/Plugins/Animation/ControlRig/Source/ControlRigEditor/Private/EditMode/SControlRigDetails.cpp
--- a/Engine/Plugins/Animation/ControlRig/Source/ControlRigEditor/Private/EditMode/SControlRigDetails.cpp
+++ b/Engine/Plugins/Animation/ControlRig/Source/ControlRigEditor/Private/EditMode/SControlRigDetails.cpp
@@ -145,17 +145,31 @@ void SControlRigDetails::HandleControlSelected(UControlRig* Subject, FRigControl
 
 static UControlRigControlsProxy* GetParentProxy(UControlRigControlsProxy* ChildProxy, const TArray<UControlRigControlsProxy*>& Proxies)
 {
+	if (!ChildProxy->OwnerControlRig.IsValid())
+	{
+		return nullptr;
+	}
+	if (!ChildProxy->OwnerControlElement.UpdateCache(ChildProxy->OwnerControlRig->GetHierarchy()))
+	{
+		return nullptr;
+	}
 	FRigBaseElement* ChildParent = (ChildProxy && ChildProxy->OwnerControlRig.IsValid()) ?
-		ChildProxy->OwnerControlRig->GetHierarchy()->GetFirstParent(ChildProxy->OwnerControlElement) : nullptr;
+		ChildProxy->OwnerControlRig->GetHierarchy()->GetFirstParent(ChildProxy->OwnerControlElement.GetElement()) : nullptr;
 	if (ChildParent == nullptr)
 	{
 		return nullptr;
 	}
 	for (UControlRigControlsProxy* Proxy : Proxies)
 	{
-		if (Proxy && Proxy->OwnerControlRig.IsValid() && (ChildParent == Proxy->OwnerControlElement))
+		if (Proxy && Proxy->OwnerControlRig.IsValid())
 		{
-			return Proxy;
+			if (Proxy->OwnerControlElement.UpdateCache(Proxy->OwnerControlRig->GetHierarchy()))
+			{
+				if (ChildParent == Proxy->OwnerControlElement.GetElement())
+				{
+					return Proxy;
+				}
+			}
 		}
 	}
 	return nullptr;
@@ -342,7 +356,17 @@ void SControlRigDetails::UpdateProxies()
 						}
 						else
 						{
-							TObjectPtr<UEnum> EnumPtr = Proxy->OwnerControlElement ? Proxy->OwnerControlElement->Settings.ControlEnum : nullptr;
+							TObjectPtr<UEnum> EnumPtr = nullptr;
+							if (Proxy->OwnerControlRig.IsValid())
+							{
+								if (Proxy->OwnerControlElement.UpdateCache(Proxy->OwnerControlRig->GetHierarchy()))
+								{
+									if (const FRigControlElement* ControlElement = Cast<FRigControlElement>(Proxy->OwnerControlElement.GetElement()))
+									{
+										EnumPtr = ControlElement->Settings.ControlEnum;
+									}
+								}
+							}
 							if (UControlRigControlsProxy* ExistingProxy = GetProxyWithSameType(AllProxies, Proxy->Type, EnumPtr != nullptr))
 							{
 								ExistingProxy->AddItem(Proxy);
@@ -359,7 +383,17 @@ void SControlRigDetails::UpdateProxies()
 					{
 						if (UControlRigControlsProxy* ParentProxy = GetParentProxy(Proxy, Proxies))
 						{
-							TObjectPtr<UEnum> EnumPtr = ParentProxy->OwnerControlElement ? ParentProxy->OwnerControlElement->Settings.ControlEnum : nullptr;
+							TObjectPtr<UEnum> EnumPtr = nullptr;
+							if (ParentProxy->OwnerControlRig.IsValid())
+							{
+								if (ParentProxy->OwnerControlElement.UpdateCache(ParentProxy->OwnerControlRig->GetHierarchy()))
+								{
+									if (const FRigControlElement* ControlElement = Cast<FRigControlElement>(ParentProxy->OwnerControlElement.GetElement()))
+									{
+										EnumPtr = ControlElement->Settings.ControlEnum;
+									}
+								}
+							}
 							if (UControlRigControlsProxy* ExistingProxy = GetProxyWithSameType(AllProxies, ParentProxy->Type, EnumPtr != nullptr))
 							{
 								ParentProxy->AddChildProxy(Proxy);
