Sequencer: Anim Details: Fix crash on close

diff --git a/Engine/Plugins/Animation/ControlRig/Source/ControlRigEditor/Private/EditMode/ControlRigEditMode.cpp b/Engine/Plugins/Animation/ControlRig/Source/ControlRigEditor/Private/EditMode/ControlRigEditMode.cpp
--- a/Engine/Plugins/Animation/ControlRig/Source/ControlRigEditor/Private/EditMode/ControlRigEditMode.cpp
+++ b/Engine/Plugins/Animation/ControlRig/Source/ControlRigEditor/Private/EditMode/ControlRigEditMode.cpp
@@ -255,7 +255,7 @@ bool FControlRigEditMode::SetSequencer(TWeakPtr<ISequencer> InSequencer)
 		{
 			FControlRigEditModeToolkit::Details->SetEditMode(*this);
 		}
-		if (FControlRigEditModeToolkit::Outliner)
+		if (FControlRigEditModeToolkit::Outliner.IsValid())
 		{
 			FControlRigEditModeToolkit::Outliner->SetEditMode(*this);
 		}
@@ -493,6 +493,14 @@ void FControlRigEditMode::Exit()
 		bManipulatorMadeChange = false;
 	}
 
+	if (FControlRigEditModeToolkit::Details.IsValid())
+	{
+		FControlRigEditModeToolkit::Details.Reset();
+	}
+	if (FControlRigEditModeToolkit::Outliner.IsValid())
+	{
+		FControlRigEditModeToolkit::Outliner.Reset();
+	}
 	if (Toolkit.IsValid())
 	{
 		FToolkitManager::Get().CloseToolkit(Toolkit.ToSharedRef());
