Control Rig Edit Mode: notify the constraints widget from sequencer changes

#rnx

diff --git a/Engine/Plugins/Animation/ControlRig/Source/ControlRigEditor/Private/EditMode/SControlRigEditModeTools.cpp b/Engine/Plugins/Animation/ControlRig/Source/ControlRigEditor/Private/EditMode/SControlRigEditModeTools.cpp
--- a/Engine/Plugins/Animation/ControlRig/Source/ControlRigEditor/Private/EditMode/SControlRigEditModeTools.cpp
+++ b/Engine/Plugins/Animation/ControlRig/Source/ControlRigEditor/Private/EditMode/SControlRigEditModeTools.cpp
@@ -592,6 +592,10 @@ void SControlRigEditModeTools::SetVector2DDetailsObjects(const TArray<TWeakObjec
 void SControlRigEditModeTools::SetSequencer(TWeakPtr<ISequencer> InSequencer)
 {
 	WeakSequencer = InSequencer.Pin();
+	if (ConstraintsEditionWidget)
+	{
+		ConstraintsEditionWidget->SequencerChanged(InSequencer);
+	}
 }
 
 bool SControlRigEditModeTools::IsPropertyKeyable(const UClass* InObjectClass, const IPropertyHandle& InPropertyHandle) const
diff --git a/Engine/Plugins/Animation/ControlRig/Source/ControlRigEditor/Private/Editor/Constraints/SConstraintsWidget.cpp b/Engine/Plugins/Animation/ControlRig/Source/ControlRigEditor/Private/Editor/Constraints/SConstraintsWidget.cpp
--- a/Engine/Plugins/Animation/ControlRig/Source/ControlRigEditor/Private/Editor/Constraints/SConstraintsWidget.cpp
+++ b/Engine/Plugins/Animation/ControlRig/Source/ControlRigEditor/Private/Editor/Constraints/SConstraintsWidget.cpp
@@ -1227,6 +1227,15 @@ void SConstraintsEditionWidget::InvalidateConstraintList()
 	UpdateSequencer();
 }
 
+void SConstraintsEditionWidget::SequencerChanged(const TWeakPtr<ISequencer>& InNewSequencer)
+{
+	if (WeakSequencer != InNewSequencer)
+	{
+		WeakSequencer = InNewSequencer;
+		InvalidateConstraintList();
+	}
+}
+
 void SConstraintsEditionWidget::UpdateSequencer()
 {
 	if (!WeakSequencer.IsValid())
diff --git a/Engine/Plugins/Animation/ControlRig/Source/ControlRigEditor/Private/Editor/Constraints/SConstraintsWidget.h b/Engine/Plugins/Animation/ControlRig/Source/ControlRigEditor/Private/Editor/Constraints/SConstraintsWidget.h
--- a/Engine/Plugins/Animation/ControlRig/Source/ControlRigEditor/Private/Editor/Constraints/SConstraintsWidget.h
+++ b/Engine/Plugins/Animation/ControlRig/Source/ControlRigEditor/Private/Editor/Constraints/SConstraintsWidget.h
@@ -257,9 +257,10 @@ public:
 	EShowConstraints GetShowConstraints() const { return ShowConstraints; }
 	void SetShowConstraints(EShowConstraints InShowConstraints)
 	{
-		if (InShowConstraints != ShowConstraints) {
+		if (InShowConstraints != ShowConstraints)
+		{
 			ShowConstraints = InShowConstraints;
-			bNeedsRefresh = true;
+			InvalidateConstraintList();
 		}
 	}
 	FText GetShowConstraintsText(EShowConstraints Index) const;
@@ -319,6 +320,9 @@ public:
 
 	// FBaseConstraintListWidget overrides
 	virtual void InvalidateConstraintList() override;
+
+	/** Notify from sequencer changed. */
+	void SequencerChanged(const TWeakPtr<ISequencer>& InNewSequencer);
 	
 private:
 
