Control Rig: Fix static analysis warning
#rnx

diff --git a/Engine/Plugins/Animation/ControlRig/Source/ControlRigEditor/Private/EditMode/SControlRigDetails.cpp b/Engine/Plugins/Animation/ControlRig/Source/ControlRigEditor/Private/EditMode/SControlRigDetails.cpp
--- a/Engine/Plugins/Animation/ControlRig/Source/ControlRigEditor/Private/EditMode/SControlRigDetails.cpp
+++ b/Engine/Plugins/Animation/ControlRig/Source/ControlRigEditor/Private/EditMode/SControlRigDetails.cpp
@@ -145,7 +145,7 @@ void SControlRigDetails::HandleControlSelected(UControlRig* Subject, FRigControl
 
 static UControlRigControlsProxy* GetParentProxy(UControlRigControlsProxy* ChildProxy, const TArray<UControlRigControlsProxy*>& Proxies)
 {
-	if (!ChildProxy->OwnerControlRig.IsValid())
+	if (!ChildProxy || !ChildProxy->OwnerControlRig.IsValid())
 	{
 		return nullptr;
 	}
