Sequencer : Anim Details:  Move updating values when we are dirtied to only happen on next tick so we only do it once instead of possibly thousands of times.

diff --git a/Engine/Plugins/Animation/ControlRig/Source/ControlRigEditor/Private/EditMode/ControlRigEditMode.cpp b/Engine/Plugins/Animation/ControlRig/Source/ControlRigEditor/Private/EditMode/ControlRigEditMode.cpp
--- a/Engine/Plugins/Animation/ControlRig/Source/ControlRigEditor/Private/EditMode/ControlRigEditMode.cpp
+++ b/Engine/Plugins/Animation/ControlRig/Source/ControlRigEditor/Private/EditMode/ControlRigEditMode.cpp
@@ -698,6 +698,7 @@ void FControlRigEditMode::Tick(FEditorViewportClient* ViewportClient, float Delt
 			}
 		}
 	}
+	DetailKeyFrameCache->UpdateIfDirty();
 }
 
 //Hit proxy for FK Rigs and bones.
@@ -5721,15 +5722,23 @@ void FDetailKeyFrameCacheAndHandler::OnChannelChanged(const FMovieSceneChannelMe
 void FDetailKeyFrameCacheAndHandler::ResetCachedData()
 {
 	CachedPropertyKeyedStatusMap.Reset();
-	if (FMovieSceneConstraintChannelHelper::bDoNotCompensate == false) //if compensating don't reset this.
+	bValuesDirty = true;
+}
+
+void FDetailKeyFrameCacheAndHandler::UpdateIfDirty()
+{
+	if (bValuesDirty == true)
 	{
-		if (EditMode && EditMode->GetControlProxy())
+		if (FMovieSceneConstraintChannelHelper::bDoNotCompensate == false) //if compensating don't reset this.
 		{
-			EditMode->GetControlProxy()->ValuesChanged();
+			if (EditMode && EditMode->GetControlProxy())
+			{
+				EditMode->GetControlProxy()->ValuesChanged();
+			}
+			bValuesDirty = false;
 		}
 	}
 }
 
 
-
 #undef LOCTEXT_NAMESPACE
diff --git a/Engine/Plugins/Animation/ControlRig/Source/ControlRigEditor/Public/EditMode/ControlRigEditMode.h b/Engine/Plugins/Animation/ControlRig/Source/ControlRigEditor/Public/EditMode/ControlRigEditMode.h
--- a/Engine/Plugins/Animation/ControlRig/Source/ControlRigEditor/Public/EditMode/ControlRigEditMode.h
+++ b/Engine/Plugins/Animation/ControlRig/Source/ControlRigEditor/Public/EditMode/ControlRigEditMode.h
@@ -106,12 +106,13 @@ struct FDetailKeyFrameCacheAndHandler: public IDetailKeyframeHandler
 
 	void SetDelegates(TWeakPtr<ISequencer>& InWeakSequencer, FControlRigEditMode* InEditMode);
 	void UnsetDelegates();
-
+	void ResetCachedData();
+	void UpdateIfDirty();
 	/** Map to the last calculated property keyed status. Resets when Scrubbing, changing Movie Scene Data, etc */
 	mutable TMap<const IPropertyHandle*, EPropertyKeyedStatus> CachedPropertyKeyedStatusMap;
 
-	void ResetCachedData();
-
+	/* flag to specify that we need to update values, will poll this on edit mode tick for performance */
+	bool bValuesDirty = false;
 private:
 	TWeakPtr<ISequencer> WeakSequencer;
 	FControlRigEditMode* EditMode = nullptr;
