Sequencer: Bring back deleted line

diff --git a/Engine/Plugins/Animation/ControlRig/Source/ControlRigEditor/Private/Sequencer/ControlRigParameterTrackEditor.cpp b/Engine/Plugins/Animation/ControlRig/Source/ControlRigEditor/Private/Sequencer/ControlRigParameterTrackEditor.cpp
--- a/Engine/Plugins/Animation/ControlRig/Source/ControlRigEditor/Private/Sequencer/ControlRigParameterTrackEditor.cpp
+++ b/Engine/Plugins/Animation/ControlRig/Source/ControlRigEditor/Private/Sequencer/ControlRigParameterTrackEditor.cpp
@@ -3150,6 +3150,7 @@ void FControlRigParameterTrackEditor::HandleControlSelected(UControlRig* Subject
 			//Force refresh later, not now
 			bSkipNextSelectionFromTimer = bSkipNextSelectionFromTimer ||
 				(bIsSelectingIndirectControl && ControlElement->Settings.AnimationType == ERigControlAnimationType::AnimationControl);
+			GetSequencer()->NotifyMovieSceneDataChanged(EMovieSceneDataChangeType::RefreshTree);
 
 		}
 	}
