Control Rig: Fix crash when searching for RigHierarchy element dependencies on a modular rig
#rnx

diff --git a/Engine/Plugins/Animation/ControlRig/Source/ControlRig/Private/Rigs/RigHierarchy.cpp b/Engine/Plugins/Animation/ControlRig/Source/ControlRig/Private/Rigs/RigHierarchy.cpp
--- a/Engine/Plugins/Animation/ControlRig/Source/ControlRig/Private/Rigs/RigHierarchy.cpp
+++ b/Engine/Plugins/Animation/ControlRig/Source/ControlRig/Private/Rigs/RigHierarchy.cpp
@@ -6024,8 +6024,27 @@ URigHierarchy::TElementDependencyMap URigHierarchy::GetDependenciesForVM(const U
 
 	typedef TTuple<int32, int32> TInt32Tuple;
 	TArray<TArray<TInt32Tuple>> ReadTransformPerInstruction, WrittenTransformsPerInstruction;
-	ReadTransformPerInstruction.AddZeroed(Instructions.Num());
-	WrittenTransformsPerInstruction.AddZeroed(Instructions.Num());
+	
+	// Find the max instruction index
+	{
+		/** TODO: UE-207320
+		 * This function needs to be rewritten to account for modular rigs. Rig Hierarchies will have reads/writes coming from
+		 * different VMs, so we cannot depend on the algorithm used in this function to find dependencies. */
+		
+		int32 MaxInstructionIndex = Instructions.Num();
+		for(int32 RecordType = 0; RecordType < 2; RecordType++)
+		{
+			const TArray<TInstructionSliceElement>& Records = RecordType == 0 ? ReadTransformsAtRuntime : WrittenTransformsAtRuntime; 
+			for(int32 RecordIndex = 0; RecordIndex < Records.Num(); RecordIndex++)
+			{
+				const TInstructionSliceElement& Record = Records[RecordIndex];
+				const int32& InstructionIndex = Record.Get<0>();
+				MaxInstructionIndex = FMath::Max(MaxInstructionIndex, InstructionIndex);
+			}
+		}
+		ReadTransformPerInstruction.AddZeroed(MaxInstructionIndex);
+		WrittenTransformsPerInstruction.AddZeroed(MaxInstructionIndex);
+	}
 
 	// fill lookup tables per instruction / element
 	for(int32 RecordType = 0; RecordType < 2; RecordType++)
