Control Rig: Fix crash when trying to bind a module config value to a variable
#rnx

diff --git a/Engine/Plugins/Animation/ControlRig/Source/ControlRig/Private/ModularRigController.cpp b/Engine/Plugins/Animation/ControlRig/Source/ControlRig/Private/ModularRigController.cpp
--- a/Engine/Plugins/Animation/ControlRig/Source/ControlRig/Private/ModularRigController.cpp
+++ b/Engine/Plugins/Animation/ControlRig/Source/ControlRig/Private/ModularRigController.cpp
@@ -796,15 +796,22 @@ TArray<FString> UModularRigController::GetPossibleBindings(const FString& InModu
 		const FString CurModulePath = InModule->GetPath();
 		if (InModulePath != CurModulePath && !CurModulePath.StartsWith(InvalidModulePrefix, ESearchCase::CaseSensitive))
 		{
-			TArray<FRigVMExternalVariable> Variables = InModule->Class->GetDefaultObject<UControlRig>()->GetExternalVariables();
-			for (const FRigVMExternalVariable& Variable : Variables)
+			if (!InModule->Class.IsValid())
 			{
-				FText ErrorMessage;
-				const FString SourceVariablePath = URigHierarchy::JoinNameSpace(CurModulePath, Variable.Name.ToString());
-				if (CanBindModuleVariable(InModulePath, InVariableName, SourceVariablePath, ErrorMessage))
-				{
-					PossibleBindings.Add(SourceVariablePath);
-				}
+				InModule->Class.LoadSynchronous();
+			}
+			if (InModule->Class.IsValid())
+			{
+				TArray<FRigVMExternalVariable> Variables = InModule->Class->GetDefaultObject<UControlRig>()->GetExternalVariables();
+			   for (const FRigVMExternalVariable& Variable : Variables)
+			   {
+				   FText ErrorMessage;
+				   const FString SourceVariablePath = URigHierarchy::JoinNameSpace(CurModulePath, Variable.Name.ToString());
+				   if (CanBindModuleVariable(InModulePath, InVariableName, SourceVariablePath, ErrorMessage))
+				   {
+					   PossibleBindings.Add(SourceVariablePath);
+				   }
+			   }
 			}
 		}		
 		return true;
