FK Control Rigs: Fix channels keys not showing up on creation or load do to the track/section models not getting recreated. Need to mark as changed and wrap the change in a transaction for the view model to make the change. This will probably also fix some load time issues with AnimDataModels.

diff --git a/Engine/Plugins/Animation/ControlRig/Source/ControlRig/Private/Sequencer/MovieSceneControlRigParameterSection.cpp b/Engine/Plugins/Animation/ControlRig/Source/ControlRig/Private/Sequencer/MovieSceneControlRigParameterSection.cpp
--- a/Engine/Plugins/Animation/ControlRig/Source/ControlRig/Private/Sequencer/MovieSceneControlRigParameterSection.cpp
+++ b/Engine/Plugins/Animation/ControlRig/Source/ControlRig/Private/Sequencer/MovieSceneControlRigParameterSection.cpp
@@ -1511,6 +1511,7 @@ void UMovieSceneControlRigParameterSection::StoreLastControlsUsedToReconstruct(c
 void UMovieSceneControlRigParameterSection::ReconstructChannelProxy()
 {
 	ChannelProxy.Reset();
+	BroadcastChanged();
 }
 
 EMovieSceneChannelProxyType UMovieSceneControlRigParameterSection::CacheChannelProxy()
@@ -2508,8 +2509,7 @@ void UMovieSceneControlRigParameterSection::RecreateWithThisControlRig(UControlR
 			break;
 		}
 	}
-	CacheChannelProxy();
-	BroadcastChanged();
+	ReconstructChannelProxy();
 }
 
 void UMovieSceneControlRigParameterSection::SetControlRig(UControlRig* InControlRig)
diff --git a/Engine/Plugins/Animation/ControlRig/Source/ControlRig/Private/Sequencer/MovieSceneControlRigParameterTrack.cpp b/Engine/Plugins/Animation/ControlRig/Source/ControlRig/Private/Sequencer/MovieSceneControlRigParameterTrack.cpp
--- a/Engine/Plugins/Animation/ControlRig/Source/ControlRig/Private/Sequencer/MovieSceneControlRigParameterTrack.cpp
+++ b/Engine/Plugins/Animation/ControlRig/Source/ControlRig/Private/Sequencer/MovieSceneControlRigParameterTrack.cpp
@@ -12,6 +12,9 @@
 #include "UObject/Package.h"
 #include "Async/Async.h"
 
+#if WITH_EDITOR
+#include "ScopedTransaction.h"
+#endif//WITH_EDITOR
 
 #include UE_INLINE_GENERATED_CPP_BY_NAME(MovieSceneControlRigParameterTrack)
 
@@ -540,6 +543,10 @@ void UMovieSceneControlRigParameterTrack::HandleOnPostConstructed_GameThread()
 	{
 		TArray<FRigControlElement*> SortedControls;
 		ControlRig->GetControlsInOrder(SortedControls);
+#if WITH_EDITOR
+		const FScopedTransaction PostConstructTransation(NSLOCTEXT("ControlRig", "PostConstructTransation", "Post Construct"));
+#endif		
+		bool bSectionWasDifferent = false;
 		for (UMovieSceneSection* BaseSection : GetAllSections())
 		{
 			if (UMovieSceneControlRigParameterSection* Section = Cast<UMovieSceneControlRigParameterSection>(BaseSection))
@@ -547,9 +554,14 @@ void UMovieSceneControlRigParameterTrack::HandleOnPostConstructed_GameThread()
 				if (Section->IsDifferentThanLastControlsUsedToReconstruct(SortedControls))
 				{
 					Section->RecreateWithThisControlRig(ControlRig, Section->GetBlendType() == EMovieSceneBlendType::Absolute);
+					bSectionWasDifferent = true;
 				}
 			}
 		}
+		if (bSectionWasDifferent)
+		{
+			BroadcastChanged();
+		}
 		if (SortedControls.Num() > 0) //really set up
 		{
 			TArray<FName> Names = GetControlsWithDifferentRotationOrders();
