Control Rig: Fix module config values edited in the details panel not propagating properly
#rnx

diff --git a/Engine/Plugins/Animation/ControlRig/Source/ControlRigDeveloper/Private/ControlRigBlueprint.cpp b/Engine/Plugins/Animation/ControlRig/Source/ControlRigDeveloper/Private/ControlRigBlueprint.cpp
--- a/Engine/Plugins/Animation/ControlRig/Source/ControlRigDeveloper/Private/ControlRigBlueprint.cpp
+++ b/Engine/Plugins/Animation/ControlRig/Source/ControlRigDeveloper/Private/ControlRigBlueprint.cpp
@@ -1657,16 +1657,6 @@ void UControlRigBlueprint::RecompileModularRig()
 	{
 		if (UModularRig* DefaultObject = Cast<UModularRig>(MyControlRigClass->GetDefaultObject(false)))
 		{
-			DefaultObject->ResetModules();
-
-			// copy the model over to the CDO.
-			// non-CDO instances are going to instantiate the model into a
-			// UObject module instance tree. CDO's are data only to avoid bugs / 
-			// behaviors in the blueprint re-instancer - which is disregarding any
-			// object under a CDO.
-			DefaultObject->ModularRigModel = ModularRigModel;
-			DefaultObject->ModularRigModel.SetOuterClientHost(DefaultObject);
-			DefaultObject->ModularRigSettings = ModularRigSettings;
 			PropagateModuleHierarchyFromBPToInstances();
 		}
 	}
@@ -2283,8 +2273,22 @@ void UControlRigBlueprint::PropagateModuleHierarchyFromBPToInstances() const
 {
 	if (const UClass* MyControlRigClass = GeneratedClass)
 	{
-		if (UControlRig* DefaultObject = Cast<UControlRig>(MyControlRigClass->GetDefaultObject(false)))
+		if (UModularRig* DefaultObject = Cast<UModularRig>(MyControlRigClass->GetDefaultObject(false)))
 		{
+			// We need to first transfer the model from the blueprint to the CDO
+			// We then ask instances to initialize which will provoke a call UModularRig::UpdateModuleHierarchyFromCDO
+			
+			DefaultObject->ResetModules();
+
+			// copy the model over to the CDO.
+			// non-CDO instances are going to instantiate the model into a
+			// UObject module instance tree. CDO's are data only to avoid bugs / 
+			// behaviors in the blueprint re-instancer - which is disregarding any
+			// object under a CDO.
+			DefaultObject->ModularRigModel = ModularRigModel;
+			DefaultObject->ModularRigModel.SetOuterClientHost(DefaultObject);
+			DefaultObject->ModularRigSettings = ModularRigSettings;
+			
 			TArray<UObject*> ArchetypeInstances;
 			DefaultObject->GetArchetypeInstances(ArchetypeInstances);
 			for (UObject* ArchetypeInstance : ArchetypeInstances)
