Control Rig: Fix crash when compiling a module while a modular rig which uses that module is present in a level
#rnx

diff --git a/Engine/Plugins/Animation/ControlRig/Source/ControlRig/Private/ModularRig.cpp b/Engine/Plugins/Animation/ControlRig/Source/ControlRig/Private/ModularRig.cpp
--- a/Engine/Plugins/Animation/ControlRig/Source/ControlRig/Private/ModularRig.cpp
+++ b/Engine/Plugins/Animation/ControlRig/Source/ControlRig/Private/ModularRig.cpp
@@ -313,7 +313,7 @@ void UModularRig::UpdateModuleHierarchyFromCDO()
 		Model.ForEachModule([this, Model](const FRigModuleReference* InModuleReference) -> bool
 		{
 			check(InModuleReference);
-			if (!InModuleReference->Class.IsValid())
+			if (IsInGameThread() && !InModuleReference->Class.IsValid())
 			{
 				(void)InModuleReference->Class.LoadSynchronous();
 			}
