Control Rig: Fix crash when right clicking FRigUnit_ResolveConnector node
#rnx

diff --git a/Engine/Plugins/Animation/ControlRig/Source/ControlRig/Private/Units/Execution/RigUnit_RigModules.cpp b/Engine/Plugins/Animation/ControlRig/Source/ControlRig/Private/Units/Execution/RigUnit_RigModules.cpp
--- a/Engine/Plugins/Animation/ControlRig/Source/ControlRig/Private/Units/Execution/RigUnit_RigModules.cpp
+++ b/Engine/Plugins/Animation/ControlRig/Source/ControlRig/Private/Units/Execution/RigUnit_RigModules.cpp
@@ -9,16 +9,19 @@
 FRigUnit_ResolveConnector_Execute()
 {
 	Result = Connector;
-	
-	if(const FRigBaseElement* ResolvedElement = ExecuteContext.Hierarchy->Find(Connector))
+
+	if (const URigHierarchy* Hierarchy = ExecuteContext.Hierarchy)
 	{
-		Result = ResolvedElement->GetKey();
-		if(SkipSocket && Result.IsValid() && Result.Type == ERigElementType::Socket)
+		if(const FRigBaseElement* ResolvedElement = Hierarchy->Find(Connector))
 		{
-			const FRigElementKey ParentOfSocket = ExecuteContext.Hierarchy->GetFirstParent(Result);
-			if(ParentOfSocket.IsValid())
+			Result = ResolvedElement->GetKey();
+			if(SkipSocket && Result.IsValid() && Result.Type == ERigElementType::Socket)
 			{
-				Result = ParentOfSocket;
+				const FRigElementKey ParentOfSocket = Hierarchy->GetFirstParent(Result);
+				if(ParentOfSocket.IsValid())
+				{
+					Result = ParentOfSocket;
+				}
 			}
 		}
 	}
@@ -90,6 +93,11 @@ FRigUnit_IsItemInCurrentNameSpace_Execute()
 FRigUnit_GetItemsInNameSpace_Execute()
 {
 	const URigHierarchy* Hierarchy = ExecuteContext.Hierarchy;
+	if (!Hierarchy)
+	{
+		Items.Reset();
+		return;
+	}
 
 	FString NameSpace;
 	FRigUnit_GetCurrentNameSpace::StaticExecute(ExecuteContext, NameSpace);
