Control Rig: Reset recorded transforms when forward solve is executed
#rnx

diff --git a/Engine/Plugins/Animation/ControlRig/Source/ControlRig/Private/ControlRig.cpp b/Engine/Plugins/Animation/ControlRig/Source/ControlRig/Private/ControlRig.cpp
--- a/Engine/Plugins/Animation/ControlRig/Source/ControlRig/Private/ControlRig.cpp
+++ b/Engine/Plugins/Animation/ControlRig/Source/ControlRig/Private/ControlRig.cpp
@@ -107,16 +107,13 @@ void UControlRig::ResetRecordedTransforms(const FName& InEventName)
 	{
 		if(Hierarchy->bRecordTransformsAtRuntime)
 		{
-			bool bResetRecordedTransforms = SupportsEvent(InEventName);;
-			if(InEventName == FRigUnit_PostBeginExecution::EventName)
+			bool bResetRecordedTransforms = SupportsEvent(InEventName);
+			if(InEventName == FRigUnit_PostBeginExecution::EventName ||
+				InEventName == FRigUnit_PreBeginExecution::EventName)
 			{
 				bResetRecordedTransforms = false;
 			}
-			else if(InEventName == FRigUnit_BeginExecution::EventName)
-			{
-				bResetRecordedTransforms = !SupportsEvent(FRigUnit_PreBeginExecution::EventName);
-			}
-
+			
 			if(bResetRecordedTransforms)
 			{
 				Hierarchy->ReadTransformsAtRuntime.Reset();
