Control Rig: Fix crash accessing recorded transforms
#rnx

diff --git a/Engine/Plugins/Animation/ControlRig/Source/ControlRig/Private/Rigs/RigHierarchy.cpp b/Engine/Plugins/Animation/ControlRig/Source/ControlRig/Private/Rigs/RigHierarchy.cpp
--- a/Engine/Plugins/Animation/ControlRig/Source/ControlRig/Private/Rigs/RigHierarchy.cpp
+++ b/Engine/Plugins/Animation/ControlRig/Source/ControlRig/Private/Rigs/RigHierarchy.cpp
@@ -6042,8 +6042,8 @@ URigHierarchy::TElementDependencyMap URigHierarchy::GetDependenciesForVM(const U
 				MaxInstructionIndex = FMath::Max(MaxInstructionIndex, InstructionIndex);
 			}
 		}
-		ReadTransformPerInstruction.AddZeroed(MaxInstructionIndex);
-		WrittenTransformsPerInstruction.AddZeroed(MaxInstructionIndex);
+		ReadTransformPerInstruction.AddZeroed(MaxInstructionIndex+1);
+		WrittenTransformsPerInstruction.AddZeroed(MaxInstructionIndex+1);
 	}
 
 	// fill lookup tables per instruction / element
