Control Rig: Lock the rig hierarchy when resetting its elements
#rnx

diff --git a/Engine/Plugins/Animation/ControlRig/Source/ControlRig/Private/Rigs/RigHierarchy.cpp b/Engine/Plugins/Animation/ControlRig/Source/ControlRig/Private/Rigs/RigHierarchy.cpp
--- a/Engine/Plugins/Animation/ControlRig/Source/ControlRig/Private/Rigs/RigHierarchy.cpp
+++ b/Engine/Plugins/Animation/ControlRig/Source/ControlRig/Private/Rigs/RigHierarchy.cpp
@@ -398,6 +398,8 @@ void URigHierarchy::Reset_Impl(bool bResetElements)
 
 	if(bResetElements)
 	{
+		FScopeLock Lock(&ElementsLock);
+		
 		// walk in reverse since certain elements might not have been allocated themselves
 		for(int32 ElementIndex = Elements.Num() - 1; ElementIndex >= 0; ElementIndex--)
 		{
