Control Rig: Avoid re-entry during evaluation
#rnx

diff --git a/Engine/Plugins/Animation/ControlRig/Source/ControlRig/Private/Constraints/ControlRigTransformableHandle.cpp b/Engine/Plugins/Animation/ControlRig/Source/ControlRig/Private/Constraints/ControlRigTransformableHandle.cpp
--- a/Engine/Plugins/Animation/ControlRig/Source/ControlRig/Private/Constraints/ControlRigTransformableHandle.cpp
+++ b/Engine/Plugins/Animation/ControlRig/Source/ControlRig/Private/Constraints/ControlRigTransformableHandle.cpp
@@ -58,7 +58,7 @@ bool UTransformableControlHandle::IsValid(const bool bDeepCheck) const
 
 void UTransformableControlHandle::PreEvaluate(const bool bTick) const
 {
-	if (!ControlRig.IsValid())
+	if (!ControlRig.IsValid() || ControlRig->IsEvaluating())
 	{
 		return;
 	}
diff --git a/Engine/Plugins/Animation/ControlRig/Source/ControlRig/Private/ControlRig.cpp b/Engine/Plugins/Animation/ControlRig/Source/ControlRig/Private/ControlRig.cpp
--- a/Engine/Plugins/Animation/ControlRig/Source/ControlRig/Private/ControlRig.cpp
+++ b/Engine/Plugins/Animation/ControlRig/Source/ControlRig/Private/ControlRig.cpp
@@ -706,14 +706,17 @@ bool UControlRig::Execute(const FName& InEventName)
 		bJustRanInit = true;
 	}
 
-	if(EventQueueToRun.IsEmpty())
+	// The EventQueueToRun should only be modified in URigVMHost::Evaluate_AnyThread
+	// We create a temporary queue for the execution of only this event
+	TArray<FName> TempEventQueueToRun = EventQueueToRun;
+	if(TempEventQueueToRun.IsEmpty())
 	{
-		EventQueueToRun = EventQueue;
+		TempEventQueueToRun = EventQueue;
 	}
 
-	const bool bIsEventInQueue = EventQueueToRun.Contains(InEventName);
-	const bool bIsEventFirstInQueue = !EventQueueToRun.IsEmpty() && EventQueueToRun[0] == InEventName; 
-	const bool bIsEventLastInQueue = !EventQueueToRun.IsEmpty() && EventQueueToRun.Last() == InEventName;
+	const bool bIsEventInQueue = TempEventQueueToRun.Contains(InEventName);
+	const bool bIsEventFirstInQueue = !TempEventQueueToRun.IsEmpty() && TempEventQueueToRun[0] == InEventName; 
+	const bool bIsEventLastInQueue = !TempEventQueueToRun.IsEmpty() && TempEventQueueToRun.Last() == InEventName;
 	const bool bIsConstructionEvent = InEventName == FRigUnit_PrepareForExecution::EventName;
 	const bool bIsForwardSolve = InEventName == FRigUnit_BeginExecution::EventName;
 	const bool bIsInteractionEvent = InEventName == FRigUnit_InteractionExecution::EventName;
diff --git a/Engine/Plugins/Runtime/RigVM/Source/RigVM/Private/RigVMHost.cpp b/Engine/Plugins/Runtime/RigVM/Source/RigVM/Private/RigVMHost.cpp
--- a/Engine/Plugins/Runtime/RigVM/Source/RigVM/Private/RigVMHost.cpp
+++ b/Engine/Plugins/Runtime/RigVM/Source/RigVM/Private/RigVMHost.cpp
@@ -327,9 +327,12 @@ void URigVMHost::Evaluate_AnyThread()
 	// and we don't want other systems to see that.
 	FScopeLock EvaluateLock(&GetEvaluateMutex());
 	
+	// The EventQueueToRun should only be modified in this function
+	ensureMsgf(EventQueueToRun.IsEmpty(), TEXT("Detected a recursive call to the control rig evaluation function %s"), *GetPackage()->GetPathName());
+	
 	// create a copy since we need to change it here temporarily,
 	// and UI / the rig may change the event queue while it is running
-	EventQueueToRun = EventQueue;
+	TGuardValue<TArray<FName>> EventQueueToRunGuard(EventQueueToRun, EventQueue);
 
 	AdaptEventQueueForEvaluate(EventQueueToRun);
 	
@@ -391,8 +394,6 @@ void URigVMHost::Evaluate_AnyThread()
 		}
 #endif
 	}
-
-	EventQueueToRun.Reset();
 }
 
 TArray<FRigVMExternalVariable> URigVMHost::GetExternalVariables() const
diff --git a/Engine/Plugins/Runtime/RigVM/Source/RigVM/Public/RigVMHost.h b/Engine/Plugins/Runtime/RigVM/Source/RigVM/Public/RigVMHost.h
--- a/Engine/Plugins/Runtime/RigVM/Source/RigVM/Public/RigVMHost.h
+++ b/Engine/Plugins/Runtime/RigVM/Source/RigVM/Public/RigVMHost.h
@@ -396,6 +396,9 @@ public:
 	TArray<FName> EventQueueToRun;
 	TMap<FName, int32> EventsToRunOnce;
 
+	/** Returns true if Evaluate_AnyThread is currently executing */
+	bool IsEvaluating() const { return !EventQueueToRun.IsEmpty(); }
+
 	/** Copy the VM from the default object */
 	void InstantiateVMFromCDO();
 	
