Control Rig: Fix ControlRig.Debug.Snapshot unit test not passing
#rnx

diff --git a/Engine/Plugins/Animation/ControlRig/Source/ControlRig/Private/ControlRig.cpp b/Engine/Plugins/Animation/ControlRig/Source/ControlRig/Private/ControlRig.cpp
--- a/Engine/Plugins/Animation/ControlRig/Source/ControlRig/Private/ControlRig.cpp
+++ b/Engine/Plugins/Animation/ControlRig/Source/ControlRig/Private/ControlRig.cpp
@@ -708,15 +708,15 @@ bool UControlRig::Execute(const FName& InEventName)
 
 	// The EventQueueToRun should only be modified in URigVMHost::Evaluate_AnyThread
 	// We create a temporary queue for the execution of only this event
-	TArray<FName> TempEventQueueToRun = EventQueueToRun;
-	if(TempEventQueueToRun.IsEmpty())
+	TArray<FName> LocalEventQueueToRun = EventQueueToRun;
+	if(LocalEventQueueToRun.IsEmpty())
 	{
-		TempEventQueueToRun = EventQueue;
+		LocalEventQueueToRun = EventQueue;
 	}
 
-	const bool bIsEventInQueue = TempEventQueueToRun.Contains(InEventName);
-	const bool bIsEventFirstInQueue = !TempEventQueueToRun.IsEmpty() && TempEventQueueToRun[0] == InEventName; 
-	const bool bIsEventLastInQueue = !TempEventQueueToRun.IsEmpty() && TempEventQueueToRun.Last() == InEventName;
+	const bool bIsEventInQueue = LocalEventQueueToRun.Contains(InEventName);
+	const bool bIsEventFirstInQueue = !LocalEventQueueToRun.IsEmpty() && LocalEventQueueToRun[0] == InEventName; 
+	const bool bIsEventLastInQueue = !LocalEventQueueToRun.IsEmpty() && LocalEventQueueToRun.Last() == InEventName;
 	const bool bIsConstructionEvent = InEventName == FRigUnit_PrepareForExecution::EventName;
 	const bool bIsForwardSolve = InEventName == FRigUnit_BeginExecution::EventName;
 	const bool bIsInteractionEvent = InEventName == FRigUnit_InteractionExecution::EventName;
@@ -1070,9 +1070,9 @@ bool UControlRig::Execute(const FName& InEventName)
 	{
 #if WITH_EDITOR
 		// only set a valid first entry event when none has been set
-		if (InstructionVisitInfo.GetFirstEntryEventInEventQueue() == NAME_None && !EventQueueToRun.IsEmpty() && VM)
+		if (InstructionVisitInfo.GetFirstEntryEventInEventQueue() == NAME_None && !LocalEventQueueToRun.IsEmpty() && VM)
 		{
-			InstructionVisitInfo.SetFirstEntryEventInEventQueue(EventQueueToRun[0]);
+			InstructionVisitInfo.SetFirstEntryEventInEventQueue(LocalEventQueueToRun[0]);
 		}
 
 		// Transform Overrride is generated using a Transient Control 
@@ -1426,6 +1426,14 @@ bool UControlRig::Execute_Internal(const FName& InEventName)
 		return false;
 	}
 
+	// The EventQueueToRun should only be modified in URigVMHost::Evaluate_AnyThread
+	// We create a temporary queue for the execution of only this event
+	TArray<FName> LocalEventQueueToRun = EventQueueToRun;
+	if(LocalEventQueueToRun.IsEmpty())
+	{
+		LocalEventQueueToRun = EventQueue;
+	}
+
 	// make sure to initialize here as well - just in case this gets
 	// called without a call to ::Execute. This is already tackled by
 	// the UModularRig::ExecuteQueue, but added it here as well nevertheless
@@ -1495,8 +1503,8 @@ bool UControlRig::Execute_Internal(const FName& InEventName)
 		{
 			if(URigVM* SnapShotVM = GetSnapshotVM(false)) // don't create it for normal runs
 			{
-				const bool bIsEventFirstInQueue = !EventQueueToRun.IsEmpty() && EventQueueToRun[0] == InEventName; 
-				const bool bIsEventLastInQueue = !EventQueueToRun.IsEmpty() && EventQueueToRun.Last() == InEventName;
+				const bool bIsEventFirstInQueue = !LocalEventQueueToRun.IsEmpty() && LocalEventQueueToRun[0] == InEventName; 
+				const bool bIsEventLastInQueue = !LocalEventQueueToRun.IsEmpty() && LocalEventQueueToRun.Last() == InEventName;
 
 				if (GetHaltedAtBreakpoint().IsValid())
 				{
