Control Rig: Fix mirror modules sometimes crashing
#rnx

diff --git a/Engine/Plugins/Animation/ControlRig/Source/ControlRig/Private/ModularRigController.cpp b/Engine/Plugins/Animation/ControlRig/Source/ControlRig/Private/ModularRigController.cpp
--- a/Engine/Plugins/Animation/ControlRig/Source/ControlRig/Private/ModularRigController.cpp
+++ b/Engine/Plugins/Animation/ControlRig/Source/ControlRig/Private/ModularRigController.cpp
@@ -224,10 +224,10 @@ bool UModularRigController::CanConnectConnectorToElement(const FRigElementKey& I
 	return true;
 }
 
-bool UModularRigController::ConnectConnectorToElement(const FRigElementKey& InConnectorKey, const FRigElementKey& InTargetKey, bool bSetupUndo, bool bAutoResolveOtherConnectors)
+bool UModularRigController::ConnectConnectorToElement(const FRigElementKey& InConnectorKey, const FRigElementKey& InTargetKey, bool bSetupUndo, bool bAutoResolveOtherConnectors, bool bCheckValidConnection)
 {
 	FText ErrorMessage;
-	if (!CanConnectConnectorToElement(InConnectorKey, InTargetKey, ErrorMessage))
+	if (bCheckValidConnection && !CanConnectConnectorToElement(InConnectorKey, InTargetKey, ErrorMessage))
 	{
 		UE_LOG(LogControlRig, Error, TEXT("Could not connect %s to %s: %s"), *InConnectorKey.ToString(), *InTargetKey.ToString(), *ErrorMessage.ToString());
 		return false;
@@ -1315,6 +1315,7 @@ FString UModularRigController::MirrorModule(const FString& InModulePath, const F
 	if (!InSettings.SearchString.IsEmpty())
 	{
 		NewModuleName = NewModuleName.Replace(*InSettings.SearchString, *InSettings.ReplaceString, ESearchCase::CaseSensitive);
+		NewModuleName = GetSafeNewName(OriginalModule->ParentPath, FRigName(NewModuleName)).ToString();
 	}
 
 	FModularRigControllerCompileBracketScope CompileBracketScope(this);
@@ -1335,7 +1336,10 @@ FString UModularRigController::MirrorModule(const FString& InModulePath, const F
 
 		FString NewConnectorPath = URigHierarchy::JoinNameSpace(NewModulePath, Pair.Key.Name.ToString());
 		FRigElementKey NewConnectorKey(*NewConnectorPath, ERigElementType::Connector);
-		ConnectConnectorToElement(NewConnectorKey, NewTargetKey, bSetupUndo);
+		ConnectConnectorToElement(NewConnectorKey, NewTargetKey, bSetupUndo, false, false);
+
+		// Path might change after connecting
+		NewModulePath = NewModule->GetPath();
 	}
 
 	for (const TPair<FName, FString>& Pair : OriginalModule->Bindings)
diff --git a/Engine/Plugins/Animation/ControlRig/Source/ControlRig/Public/ModularRigController.h b/Engine/Plugins/Animation/ControlRig/Source/ControlRig/Public/ModularRigController.h
--- a/Engine/Plugins/Animation/ControlRig/Source/ControlRig/Public/ModularRigController.h
+++ b/Engine/Plugins/Animation/ControlRig/Source/ControlRig/Public/ModularRigController.h
@@ -31,7 +31,7 @@ class CONTROLRIG_API UModularRigController : public UObject
 	bool CanConnectConnectorToElement(const FRigElementKey& InConnectorKey, const FRigElementKey& InTargetKey, FText& OutErrorMessage);
 
 	UFUNCTION(BlueprintCallable, Category = "Control Rig | Modules")
-	bool ConnectConnectorToElement(const FRigElementKey& InConnectorKey, const FRigElementKey& InTargetKey, bool bSetupUndo = true, bool bAutoResolveOtherConnectors = true);
+	bool ConnectConnectorToElement(const FRigElementKey& InConnectorKey, const FRigElementKey& InTargetKey, bool bSetupUndo = true, bool bAutoResolveOtherConnectors = true, bool bCheckValidConnection = true);
 
 	UFUNCTION(BlueprintCallable, Category = "Control Rig | Modules")
 	bool DisconnectConnector(const FRigElementKey& InConnectorKey, bool bDisconnectSubModules = false, bool bSetupUndo = true);
diff --git a/Engine/Plugins/Animation/ControlRig/Source/ControlRigEditor/Private/Editor/ControlRigSchematicModel.cpp b/Engine/Plugins/Animation/ControlRig/Source/ControlRigEditor/Private/Editor/ControlRigSchematicModel.cpp
--- a/Engine/Plugins/Animation/ControlRig/Source/ControlRigEditor/Private/Editor/ControlRigSchematicModel.cpp
+++ b/Engine/Plugins/Animation/ControlRig/Source/ControlRigEditor/Private/Editor/ControlRigSchematicModel.cpp
@@ -605,7 +605,7 @@ void FControlRigSchematicModel::HandleModularRigModified(EModularRigNotification
 	{
 		case EModularRigNotification::ConnectionChanged:
 		{
-			if (ControlRigBeingDebuggedPtr.IsValid())
+			if (ControlRigBeingDebuggedPtr.IsValid() && InModule)
 			{
 				if (const URigHierarchy* Hierarchy = ControlRigBeingDebuggedPtr->GetHierarchy())
 				{
