Control Rig: Fix modular rig not propagating shape library delegates into module execute context
#rnx

diff --git a/Engine/Plugins/Animation/ControlRig/Source/ControlRig/Private/ModularRig.cpp b/Engine/Plugins/Animation/ControlRig/Source/ControlRig/Private/ModularRig.cpp
--- a/Engine/Plugins/Animation/ControlRig/Source/ControlRig/Private/ModularRig.cpp
+++ b/Engine/Plugins/Animation/ControlRig/Source/ControlRig/Private/ModularRig.cpp
@@ -471,6 +471,8 @@ void UModularRig::ExecuteQueue()
 #endif
 				RigPublicContext.SetFramesPerSecond(PublicContext.GetFramesPerSecond());
 				RigPublicContext.SetToWorldSpaceTransform(PublicContext.GetToWorldSpaceTransform());
+				RigPublicContext.OnAddShapeLibraryDelegate = PublicContext.OnAddShapeLibraryDelegate;
+				RigPublicContext.OnShapeExistsDelegate = PublicContext.OnShapeExistsDelegate;
 				RigPublicContext.RuntimeSettings = PublicContext.RuntimeSettings;
 
 #if WITH_EDITOR
