Control Rig: First and Last rig element types are being shown as options of RigElementKey
#rnx

diff --git a/Engine/Plugins/Animation/ControlRig/Source/ControlRig/Public/Rigs/RigHierarchyDefines.h b/Engine/Plugins/Animation/ControlRig/Source/ControlRig/Public/Rigs/RigHierarchyDefines.h
--- a/Engine/Plugins/Animation/ControlRig/Source/ControlRig/Public/Rigs/RigHierarchyDefines.h
+++ b/Engine/Plugins/Animation/ControlRig/Source/ControlRig/Public/Rigs/RigHierarchyDefines.h
@@ -31,8 +31,8 @@ enum class ERigElementType : uint8
 	Connector = 0x040,
 	Socket = 0x080,
 	
-	First = Bone,
-	Last = Socket,
+	First = Bone UMETA(Hidden), 
+	Last = Socket UMETA(Hidden), 
 	All = Bone | Null | Control | Curve | RigidBody | Reference | Connector | Socket,
 	ToResetAfterConstructionEvent = Bone | Control | Curve | Socket UMETA(Hidden),
 };
