Control Rig: Animation mode enables on opening Widget Blueprint
#rnx

diff --git a/Engine/Plugins/Animation/ControlRig/Source/ControlRigEditor/Private/Sequencer/ControlRigParameterTrackEditor.cpp b/Engine/Plugins/Animation/ControlRig/Source/ControlRigEditor/Private/Sequencer/ControlRigParameterTrackEditor.cpp
--- a/Engine/Plugins/Animation/ControlRig/Source/ControlRigEditor/Private/Sequencer/ControlRigParameterTrackEditor.cpp
+++ b/Engine/Plugins/Animation/ControlRig/Source/ControlRigEditor/Private/Sequencer/ControlRigParameterTrackEditor.cpp
@@ -2005,7 +2005,7 @@ void FControlRigParameterTrackEditor::OnActivateSequenceChanged(FMovieSceneSeque
 		
 		return false;
 	});
-	if (bControlRigEditModeWasOpen)
+	if (bControlRigEditModeWasOpen && GetSequencer()->IsLevelEditorSequencer())
 	{
 		GEditor->GetTimerManager()->SetTimerForNextTick([this]()
 		{
diff --git a/Engine/Source/Editor/Sequencer/Private/Sequencer.h b/Engine/Source/Editor/Sequencer/Private/Sequencer.h
--- a/Engine/Source/Editor/Sequencer/Private/Sequencer.h
+++ b/Engine/Source/Editor/Sequencer/Private/Sequencer.h
@@ -518,7 +518,7 @@ public:
 	const FSequencerHostCapabilities& GetHostCapabilities() const { return HostCapabilities; }
 
 	/** @return Whether or not this sequencer is used in the level editor */
-	bool IsLevelEditorSequencer() const { return bIsEditingWithinLevelEditor; }
+	virtual bool IsLevelEditorSequencer() const override { return bIsEditingWithinLevelEditor; }
 
 	/** @return Whether to show the curve editor or not */
 	virtual void SetShowCurveEditor(bool bInShowCurveEditor) override;
diff --git a/Engine/Source/Editor/Sequencer/Public/ISequencer.h b/Engine/Source/Editor/Sequencer/Public/ISequencer.h
--- a/Engine/Source/Editor/Sequencer/Public/ISequencer.h
+++ b/Engine/Source/Editor/Sequencer/Public/ISequencer.h
@@ -740,6 +740,9 @@ public:
 	/** Whether the sequence is read-only */
 	virtual bool IsReadOnly() const = 0;
 
+	/** @return Whether or not this sequencer is used in the level editor */
+	virtual bool IsLevelEditorSequencer() const = 0;
+
 	/**
 	 * Create a widget containing the spinboxes for setting the working and playback range
 	 * 
