Control Rig: Fix bake to control rig not working for modular rigs
#rnx

diff --git a/Engine/Plugins/Animation/ControlRig/Source/ControlRig/Private/ModularRig.cpp b/Engine/Plugins/Animation/ControlRig/Source/ControlRig/Private/ModularRig.cpp
--- a/Engine/Plugins/Animation/ControlRig/Source/ControlRig/Private/ModularRig.cpp
+++ b/Engine/Plugins/Animation/ControlRig/Source/ControlRig/Private/ModularRig.cpp
@@ -232,6 +232,15 @@ bool FRigModuleInstance::IsRootModule() const
 	return ParentPath.IsEmpty();
 }
 
+void UModularRig::PostInitProperties()
+{
+	Super::PostInitProperties();
+
+	ModularRigModel.UpdateCachedChildren();
+	ModularRigModel.Connections.UpdateFromConnectionList();
+	UpdateSupportedEvents();
+}
+
 void UModularRig::Serialize(FArchive& Ar)
 {
 	Super::Serialize(Ar);
@@ -409,6 +418,20 @@ void UModularRig::Evaluate_AnyThread()
 	Super::Evaluate_AnyThread();
 }
 
+bool UModularRig::SupportsEvent(const FName& InEventName) const
+{
+	return GetSupportedEvents().Contains(InEventName);
+}
+
+const TArray<FName>& UModularRig::GetSupportedEvents() const
+{
+	if (SupportedEvents.IsEmpty())
+	{
+		UpdateSupportedEvents();
+	}
+	return SupportedEvents;
+}
+
 const FModularRigSettings& UModularRig::GetModularRigSettings() const
 {
 	if(HasAnyFlags(RF_ClassDefaultObject))
@@ -670,17 +693,20 @@ void UModularRig::UpdateCachedChildren()
 	}
 }
 
-void UModularRig::UpdateSupportedEvents()
+void UModularRig::UpdateSupportedEvents() const
 {
 	SupportedEvents.Reset();
-	ForEachModule([this](const FRigModuleInstance* Module) -> bool
+	ModularRigModel.ForEachModule([this](const FRigModuleReference* Module) -> bool
 	{
-		if (const UControlRig* ModuleRig = Module->GetRig())
+		if (Module->Class.IsValid())
 		{
-			const TArray<FName>& ModuleEvents = ModuleRig->GetSupportedEvents();
-			for (const FName& EventName : ModuleEvents)
+			if (UControlRig* CDO = Module->Class->GetDefaultObject<UControlRig>())
 			{
-				SupportedEvents.AddUnique(EventName);
+				TArray<FName> Events = CDO->GetSupportedEvents();
+				for (const FName& Event : Events)
+				{
+					SupportedEvents.AddUnique(Event);
+				}
 			}
 		}
 		return true;
diff --git a/Engine/Plugins/Animation/ControlRig/Source/ControlRig/Public/ModularRig.h b/Engine/Plugins/Animation/ControlRig/Source/ControlRig/Public/ModularRig.h
--- a/Engine/Plugins/Animation/ControlRig/Source/ControlRig/Public/ModularRig.h
+++ b/Engine/Plugins/Animation/ControlRig/Source/ControlRig/Public/ModularRig.h
@@ -131,10 +131,12 @@ class CONTROLRIG_API UModularRig : public UControlRig
 	TArray<FRigModuleInstance> Modules;
 	TArray<FRigModuleInstance*> RootModules;
 
-	TArray<FName> SupportedEvents;
+	mutable TArray<FName> SupportedEvents;
 
 public:
 
+	virtual void PostInitProperties() override;
+
 	// BEGIN ControlRig
 	virtual void Serialize(FArchive& Ar) override;
 	virtual void PostLoad() override;
@@ -147,8 +149,8 @@ public:
 	virtual void Evaluate_AnyThread() override;
 	virtual FRigElementKeyRedirector& GetElementKeyRedirector() override { return ElementKeyRedirector; }
 	virtual FRigElementKeyRedirector GetElementKeyRedirector() const override { return ElementKeyRedirector; }
-	virtual bool SupportsEvent(const FName& InEventName) const override { return SupportedEvents.Contains(InEventName); }
-	virtual const TArray<FName>& GetSupportedEvents() const override { return SupportedEvents; }
+	virtual bool SupportsEvent(const FName& InEventName) const override;
+	virtual const TArray<FName>& GetSupportedEvents() const override;
 	// END ControlRig
 
 	UPROPERTY()
@@ -177,7 +179,7 @@ public:
 	const FModularRigModel& GetModularRigModel() const;
 	void UpdateModuleHierarchyFromCDO();
 	void UpdateCachedChildren();
-	void UpdateSupportedEvents();
+	void UpdateSupportedEvents() const;
 
 	const FRigModuleInstance* FindModule(const FString& InPath) const;
 	const FRigModuleInstance* FindModule(const UControlRig* InModuleInstance) const;
