Control Rig: FIx crash when accessing invalid bone index during UpdateInputOutputMappingIfRequired
#rnx

diff --git a/Engine/Plugins/Animation/ControlRig/Source/ControlRig/Private/AnimNode_ControlRigBase.cpp b/Engine/Plugins/Animation/ControlRig/Source/ControlRig/Private/AnimNode_ControlRigBase.cpp
--- a/Engine/Plugins/Animation/ControlRig/Source/ControlRig/Private/AnimNode_ControlRigBase.cpp
+++ b/Engine/Plugins/Animation/ControlRig/Source/ControlRig/Private/AnimNode_ControlRigBase.cpp
@@ -690,10 +690,13 @@ void FAnimNode_ControlRigBase::UpdateInputOutputMappingIfRequired(UControlRig* I
 					{
 						continue;
 					}
-					const FName& BoneName = InRefSkeleton.GetBoneName(InRequiredBonesArray[InputBoneToTransfer.BoneIndex]);
-					if (NodeNames.Contains(BoneName))
+					if (InRequiredBonesArray.IsValidIndex(InputBoneToTransfer.BoneIndex))
 					{
-						OutMapping.Add(BoneName, InputBoneToTransfer.BoneIndex);
+						const FName& BoneName = InRefSkeleton.GetBoneName(InRequiredBonesArray[InputBoneToTransfer.BoneIndex]);
+						if (NodeNames.Contains(BoneName))
+						{
+							OutMapping.Add(BoneName, InputBoneToTransfer.BoneIndex);
+						}
 					}
 				}
 			}
