Control Rig: Fix RIgCachedElement not updating the element
#rnx

diff --git a/Engine/Plugins/Animation/ControlRig/Source/ControlRig/Private/Rigs/RigHierarchyCache.cpp b/Engine/Plugins/Animation/ControlRig/Source/ControlRig/Private/Rigs/RigHierarchyCache.cpp
--- a/Engine/Plugins/Animation/ControlRig/Source/ControlRig/Private/Rigs/RigHierarchyCache.cpp
+++ b/Engine/Plugins/Animation/ControlRig/Source/ControlRig/Private/Rigs/RigHierarchyCache.cpp
@@ -19,7 +19,7 @@ bool FCachedRigElement::UpdateCache(const URigHierarchy* InHierarchy)
 {
 	if(InHierarchy)
 	{
-		if(!IsValid() || InHierarchy->GetTopologyVersionHash() != ContainerVersion)
+		if(!IsValid() || InHierarchy->GetTopologyVersionHash() != ContainerVersion || Element != InHierarchy->Get(Index))
 		{
 			return UpdateCache(GetKey(), InHierarchy);
 		}
@@ -32,7 +32,7 @@ bool FCachedRigElement::UpdateCache(const FRigElementKey& InKey, const URigHiera
 {
 	if(InHierarchy)
 	{
-		if(!IsValid() || !IsIdentical(InKey, InHierarchy))
+		if(!IsValid() || !IsIdentical(InKey, InHierarchy) || Element != InHierarchy->Get(Index))
 		{
 			// have to create a copy since Reset below
 			// potentially resets the InKey as well.
