Control Rig: Make sure not to create new object or call StaticFindObjectFast while garbage collecting
#rnx

diff --git a/Engine/Plugins/Animation/ControlRig/Source/ControlRig/Private/ControlRig.cpp b/Engine/Plugins/Animation/ControlRig/Source/ControlRig/Private/ControlRig.cpp
--- a/Engine/Plugins/Animation/ControlRig/Source/ControlRig/Private/ControlRig.cpp
+++ b/Engine/Plugins/Animation/ControlRig/Source/ControlRig/Private/ControlRig.cpp
@@ -3341,6 +3341,7 @@ void UControlRig::PostInitInstance(URigVMHost* InCDO)
 		// did save their hiearchies. If that's the case, let's mark them as garbage and rename them before creating our own hierarchy.
 		if (!HasAnyFlags(RF_ClassDefaultObject))
 		{
+			FGCScopeGuard Guard;
 			UObject* ObjectFound = StaticFindObjectFast(URigHierarchy::StaticClass(), this, TEXT("DynamicHierarchy"), true, RF_DefaultSubObject);
 			if (ObjectFound)
 			{
diff --git a/Engine/Plugins/Animation/ControlRig/Source/ControlRig/Private/Rigs/RigHierarchy.cpp b/Engine/Plugins/Animation/ControlRig/Source/ControlRig/Private/Rigs/RigHierarchy.cpp
--- a/Engine/Plugins/Animation/ControlRig/Source/ControlRig/Private/Rigs/RigHierarchy.cpp
+++ b/Engine/Plugins/Animation/ControlRig/Source/ControlRig/Private/Rigs/RigHierarchy.cpp
@@ -442,8 +442,8 @@ void URigHierarchy::Reset_Impl(bool bResetElements)
 	ChildElementCacheTopologyVersion = std::numeric_limits<uint32>::max();
 
 
-	if(!IsGarbageCollecting())
 	{
+		FGCScopeGuard Guard;
 		Notify(ERigHierarchyNotification::HierarchyReset, nullptr);
 	}
 
@@ -2925,8 +2925,8 @@ URigHierarchyController* URigHierarchy::GetController(bool bCreateIfNeeded)
 	}
 	else if(bCreateIfNeeded)
 	{
-		 if(ensure(!IsGarbageCollecting()))
 		 {
+			 FGCScopeGuard Guard;
 			 HierarchyController = NewObject<URigHierarchyController>(this, TEXT("HierarchyController"), RF_Transient);
 			 // In case we create this object from async loading thread
 			 HierarchyController->ClearInternalFlags(EInternalObjectFlags::Async);
@@ -2946,8 +2946,8 @@ UModularRigRuleManager* URigHierarchy::GetRuleManager(bool bCreateIfNeeded)
 	}
 	else if(bCreateIfNeeded)
 	{
-		if(ensure(!IsGarbageCollecting()))
 		{
+			FGCScopeGuard Guard;
 			RuleManager = NewObject<UModularRigRuleManager>(this, TEXT("RuleManager"), RF_Transient);
 			// In case we create this object from async loading thread
 			RuleManager->ClearInternalFlags(EInternalObjectFlags::Async);
