Control Rig: Fix crash after running control rig unit tests and closing the editor
#rnx

diff --git a/Engine/Plugins/Animation/ControlRig/Source/ControlRig/Private/Units/Hierarchy/RigUnit_SetBoneRotation.cpp b/Engine/Plugins/Animation/ControlRig/Source/ControlRig/Private/Units/Hierarchy/RigUnit_SetBoneRotation.cpp
--- a/Engine/Plugins/Animation/ControlRig/Source/ControlRig/Private/Units/Hierarchy/RigUnit_SetBoneRotation.cpp
+++ b/Engine/Plugins/Animation/ControlRig/Source/ControlRig/Private/Units/Hierarchy/RigUnit_SetBoneRotation.cpp
@@ -87,7 +87,7 @@ IMPLEMENT_RIGUNIT_AUTOMATION_TEST(FRigUnit_SetBoneRotation)
 	const FRigElementKey BoneA = Controller->AddBone(TEXT("BoneA"), Root, FTransform(FQuat(FVector(-1.f, 0.f, 0.f), 0.5f)), true, ERigBoneType::User);
 	const FRigElementKey BoneB = Controller->AddBone(TEXT("BoneB"), BoneA, FTransform(FQuat(FVector(-1.f, 0.f, 0.f), 0.7f)), true, ERigBoneType::User);
 
-	Unit.ExecuteContext.Hierarchy = Hierarchy;
+	Unit.ExecuteContext.Hierarchy = Hierarchy.Get();
 
 	Hierarchy->ResetPoseToInitial(ERigElementType::Bone);
 	Unit.Bone = TEXT("Root");
diff --git a/Engine/Plugins/Animation/ControlRig/Source/ControlRig/Private/Units/Hierarchy/RigUnit_SetBoneTransform.cpp b/Engine/Plugins/Animation/ControlRig/Source/ControlRig/Private/Units/Hierarchy/RigUnit_SetBoneTransform.cpp
--- a/Engine/Plugins/Animation/ControlRig/Source/ControlRig/Private/Units/Hierarchy/RigUnit_SetBoneTransform.cpp
+++ b/Engine/Plugins/Animation/ControlRig/Source/ControlRig/Private/Units/Hierarchy/RigUnit_SetBoneTransform.cpp
@@ -84,7 +84,7 @@ IMPLEMENT_RIGUNIT_AUTOMATION_TEST(FRigUnit_SetBoneTransform)
 	const FRigElementKey Root = Controller->AddBone(TEXT("Root"), FRigElementKey(), FTransform(FVector(1.f, 0.f, 0.f)), true, ERigBoneType::User);
 	const FRigElementKey BoneA = Controller->AddBone(TEXT("BoneA"), Root, FTransform(FVector(1.f, 2.f, 3.f)), true, ERigBoneType::User);
 	const FRigElementKey BoneB = Controller->AddBone(TEXT("BoneB"), BoneA, FTransform(FVector(1.f, 5.f, 3.f)), true, ERigBoneType::User);
-	Unit.ExecuteContext.Hierarchy = Hierarchy;
+	Unit.ExecuteContext.Hierarchy = Hierarchy.Get();
 
 	Hierarchy->ResetPoseToInitial(ERigElementType::Bone);
 	Unit.Bone = TEXT("Root");
diff --git a/Engine/Plugins/Animation/ControlRig/Source/ControlRig/Private/Units/Hierarchy/RigUnit_SetControlTransform.cpp b/Engine/Plugins/Animation/ControlRig/Source/ControlRig/Private/Units/Hierarchy/RigUnit_SetControlTransform.cpp
--- a/Engine/Plugins/Animation/ControlRig/Source/ControlRig/Private/Units/Hierarchy/RigUnit_SetControlTransform.cpp
+++ b/Engine/Plugins/Animation/ControlRig/Source/ControlRig/Private/Units/Hierarchy/RigUnit_SetControlTransform.cpp
@@ -355,7 +355,7 @@ IMPLEMENT_RIGUNIT_AUTOMATION_TEST(FRigUnit_SetMultiControlBool)
 	
 	Controller->AddControl(TEXT("Control1"), FRigElementKey(), Settings, FRigControlValue(), FTransform::Identity, FTransform::Identity);
 	Controller->AddControl(TEXT("Control2"), FRigElementKey(), Settings, FRigControlValue(), FTransform::Identity, FTransform::Identity);
-	Unit.ExecuteContext.Hierarchy = Hierarchy;
+	Unit.ExecuteContext.Hierarchy = Hierarchy.Get();
 
 	// this unit has an empty entry by default, clear it for testing purpose
 	Unit.Entries.Reset();
@@ -387,7 +387,7 @@ IMPLEMENT_RIGUNIT_AUTOMATION_TEST(FRigUnit_SetMultiControlFloat)
 	
 	Controller->AddControl(TEXT("Control1"), FRigElementKey(), Settings, FRigControlValue(), FTransform::Identity, FTransform::Identity);
 	Controller->AddControl(TEXT("Control2"), FRigElementKey(), Settings, FRigControlValue(), FTransform::Identity, FTransform::Identity);
-	Unit.ExecuteContext.Hierarchy = Hierarchy;
+	Unit.ExecuteContext.Hierarchy = Hierarchy.Get();
 
 	// this unit has an empty entry by default, clear it for testing purpose
 	Unit.Entries.Reset();
@@ -447,7 +447,7 @@ IMPLEMENT_RIGUNIT_AUTOMATION_TEST(FRigUnit_SetMultiControlInteger)
 	
 	Controller->AddControl(TEXT("Control1"), FRigElementKey(), Settings, FRigControlValue(), FTransform::Identity, FTransform::Identity);
 	Controller->AddControl(TEXT("Control2"), FRigElementKey(), Settings, FRigControlValue(), FTransform::Identity, FTransform::Identity);
-	Unit.ExecuteContext.Hierarchy = Hierarchy;
+	Unit.ExecuteContext.Hierarchy = Hierarchy.Get();
 
 	// this unit has an empty entry by default, clear it for testing purpose
 	Unit.Entries.Reset();
@@ -479,7 +479,7 @@ IMPLEMENT_RIGUNIT_AUTOMATION_TEST(FRigUnit_SetMultiControlVector2D)
 	
 	Controller->AddControl(TEXT("Control1"), FRigElementKey(), Settings, FRigControlValue(), FTransform::Identity, FTransform::Identity);
 	Controller->AddControl(TEXT("Control2"), FRigElementKey(), Settings, FRigControlValue(), FTransform::Identity, FTransform::Identity);
-	Unit.ExecuteContext.Hierarchy = Hierarchy;
+	Unit.ExecuteContext.Hierarchy = Hierarchy.Get();
 
 	// this unit has an empty entry by default, clear it for testing purpose
 	Unit.Entries.Reset();
@@ -513,7 +513,7 @@ IMPLEMENT_RIGUNIT_AUTOMATION_TEST(FRigUnit_SetMultiControlRotator)
 	
 	Controller->AddControl(TEXT("Control1"), FRigElementKey(), Settings, FRigControlValue(), FTransform::Identity, FTransform::Identity);
 	Controller->AddControl(TEXT("Control2"), FRigElementKey(), Settings, FRigControlValue(), FTransform::Identity, FTransform::Identity);
-	Unit.ExecuteContext.Hierarchy = Hierarchy;
+	Unit.ExecuteContext.Hierarchy = Hierarchy.Get();
 
 	// this unit has an empty entry by default, clear it for testing purpose
 	Unit.Entries.Reset();
diff --git a/Engine/Plugins/Animation/ControlRig/Source/ControlRig/Private/Units/Hierarchy/RigUnit_SetCurveValue.cpp b/Engine/Plugins/Animation/ControlRig/Source/ControlRig/Private/Units/Hierarchy/RigUnit_SetCurveValue.cpp
--- a/Engine/Plugins/Animation/ControlRig/Source/ControlRig/Private/Units/Hierarchy/RigUnit_SetCurveValue.cpp
+++ b/Engine/Plugins/Animation/ControlRig/Source/ControlRig/Private/Units/Hierarchy/RigUnit_SetCurveValue.cpp
@@ -26,7 +26,7 @@ IMPLEMENT_RIGUNIT_AUTOMATION_TEST(FRigUnit_SetCurveValue)
 {
 	const FRigElementKey CurveA = Controller->AddCurve(TEXT("CurveA"), 0.f);
 	const FRigElementKey CurveB = Controller->AddCurve(TEXT("CurveB"), 0.f);
-	Unit.ExecuteContext.Hierarchy = Hierarchy;
+	Unit.ExecuteContext.Hierarchy = Hierarchy.Get();
 	
 	Hierarchy->ResetCurveValues();
 	Unit.Curve = TEXT("CurveA");
diff --git a/Engine/Plugins/Animation/ControlRig/Source/ControlRig/Private/Units/Hierarchy/RigUnit_SetRelativeBoneTransform.cpp b/Engine/Plugins/Animation/ControlRig/Source/ControlRig/Private/Units/Hierarchy/RigUnit_SetRelativeBoneTransform.cpp
--- a/Engine/Plugins/Animation/ControlRig/Source/ControlRig/Private/Units/Hierarchy/RigUnit_SetRelativeBoneTransform.cpp
+++ b/Engine/Plugins/Animation/ControlRig/Source/ControlRig/Private/Units/Hierarchy/RigUnit_SetRelativeBoneTransform.cpp
@@ -65,7 +65,7 @@ IMPLEMENT_RIGUNIT_AUTOMATION_TEST(FRigUnit_SetRelativeBoneTransform)
 	const FRigElementKey BoneA = Controller->AddBone(TEXT("BoneA"), Root, FTransform(FVector(1.f, 2.f, 3.f)), true, ERigBoneType::User);
 	const FRigElementKey BoneB = Controller->AddBone(TEXT("BoneB"), BoneA, FTransform(FVector(1.f, 5.f, 3.f)), true, ERigBoneType::User);
 	const FRigElementKey BoneC = Controller->AddBone(TEXT("BoneC"), Root, FTransform(FVector(-4.f, 0.f, 0.f)), true, ERigBoneType::User);
-	Unit.ExecuteContext.Hierarchy = Hierarchy;
+	Unit.ExecuteContext.Hierarchy = Hierarchy.Get();
 
 	Hierarchy->ResetPoseToInitial(ERigElementType::Bone);
 	Unit.Bone = TEXT("BoneA");
diff --git a/Engine/Plugins/Animation/ControlRig/Source/ControlRig/Private/Units/Hierarchy/RigUnit_UnsetCurveValue.cpp b/Engine/Plugins/Animation/ControlRig/Source/ControlRig/Private/Units/Hierarchy/RigUnit_UnsetCurveValue.cpp
--- a/Engine/Plugins/Animation/ControlRig/Source/ControlRig/Private/Units/Hierarchy/RigUnit_UnsetCurveValue.cpp
+++ b/Engine/Plugins/Animation/ControlRig/Source/ControlRig/Private/Units/Hierarchy/RigUnit_UnsetCurveValue.cpp
@@ -25,7 +25,7 @@ FRigUnit_UnsetCurveValue_Execute()
 IMPLEMENT_RIGUNIT_AUTOMATION_TEST(FRigUnit_UnsetCurveValue)
 {
 	const FRigElementKey Curve = Controller->AddCurve(TEXT("Curve"), 0.f);
-	Unit.ExecuteContext.Hierarchy = Hierarchy;
+	Unit.ExecuteContext.Hierarchy = Hierarchy.Get();
 	
 	Hierarchy->ResetCurveValues();
 	Unit.Curve = TEXT("Curve");
diff --git a/Engine/Plugins/Animation/ControlRig/Source/ControlRig/Private/Units/Highlevel/Hierarchy/RigUnit_AimBone.cpp b/Engine/Plugins/Animation/ControlRig/Source/ControlRig/Private/Units/Highlevel/Hierarchy/RigUnit_AimBone.cpp
--- a/Engine/Plugins/Animation/ControlRig/Source/ControlRig/Private/Units/Highlevel/Hierarchy/RigUnit_AimBone.cpp
+++ b/Engine/Plugins/Animation/ControlRig/Source/ControlRig/Private/Units/Highlevel/Hierarchy/RigUnit_AimBone.cpp
@@ -684,7 +684,7 @@ IMPLEMENT_RIGUNIT_AUTOMATION_TEST(FRigUnit_AimConstraintLocalSpaceOffset)
     	const FRigElementKey Parent1 = Controller->AddBone(TEXT("Parent1"), FRigElementKey(), FTransform(FVector(10.f, 10.f, 10.f)), true, ERigBoneType::User);
     	const FRigElementKey Parent2 = Controller->AddBone(TEXT("Parent2"), FRigElementKey(), FTransform(FVector(-10.f,10.f, 10.f)), true, ERigBoneType::User);
     	
-    	Unit.ExecuteContext.Hierarchy = Hierarchy;
+    	Unit.ExecuteContext.Hierarchy = Hierarchy.Get();
 		Unit.AimAxis = FVector(0,1, 0);
 		Unit.UpAxis = FVector(0, 0,1);
     	
diff --git a/Engine/Plugins/Animation/ControlRig/Source/ControlRig/Private/Units/Highlevel/Hierarchy/RigUnit_DistributeRotation.cpp b/Engine/Plugins/Animation/ControlRig/Source/ControlRig/Private/Units/Highlevel/Hierarchy/RigUnit_DistributeRotation.cpp
--- a/Engine/Plugins/Animation/ControlRig/Source/ControlRig/Private/Units/Highlevel/Hierarchy/RigUnit_DistributeRotation.cpp
+++ b/Engine/Plugins/Animation/ControlRig/Source/ControlRig/Private/Units/Highlevel/Hierarchy/RigUnit_DistributeRotation.cpp
@@ -279,7 +279,7 @@ IMPLEMENT_RIGUNIT_AUTOMATION_TEST(FRigUnit_DistributeRotation)
 	const FRigElementKey BoneB = Controller->AddBone(TEXT("BoneB"), BoneA, FTransform(FVector(2.f, 0.f, 0.f)), true, ERigBoneType::User);
 	const FRigElementKey BoneC = Controller->AddBone(TEXT("BoneC"), BoneB, FTransform(FVector(2.f, 0.f, 0.f)), true, ERigBoneType::User);
 	const FRigElementKey BoneD = Controller->AddBone(TEXT("BoneD"), BoneC, FTransform(FVector(2.f, 0.f, 0.f)), true, ERigBoneType::User);
-	Unit.ExecuteContext.Hierarchy = Hierarchy;
+	Unit.ExecuteContext.Hierarchy = Hierarchy.Get();
 
 	Unit.StartBone = TEXT("Root");
 	Unit.EndBone = TEXT("BoneD");
diff --git a/Engine/Plugins/Animation/ControlRig/Source/ControlRig/Private/Units/Highlevel/Hierarchy/RigUnit_ModifyBoneTransforms.cpp b/Engine/Plugins/Animation/ControlRig/Source/ControlRig/Private/Units/Highlevel/Hierarchy/RigUnit_ModifyBoneTransforms.cpp
--- a/Engine/Plugins/Animation/ControlRig/Source/ControlRig/Private/Units/Highlevel/Hierarchy/RigUnit_ModifyBoneTransforms.cpp
+++ b/Engine/Plugins/Animation/ControlRig/Source/ControlRig/Private/Units/Highlevel/Hierarchy/RigUnit_ModifyBoneTransforms.cpp
@@ -40,7 +40,7 @@ IMPLEMENT_RIGUNIT_AUTOMATION_TEST(FRigUnit_ModifyBoneTransforms)
 	const FRigElementKey Root = Controller->AddBone(TEXT("Root"), FRigElementKey(), FTransform(FVector(1.f, 0.f, 0.f)), true, ERigBoneType::User);
 	const FRigElementKey BoneA = Controller->AddBone(TEXT("BoneA"), Root, FTransform(FVector(1.f, 2.f, 3.f)), true, ERigBoneType::User);
 	const FRigElementKey BoneB = Controller->AddBone(TEXT("BoneB"), Root, FTransform(FVector(5.f, 6.f, 7.f)), true, ERigBoneType::User);
-	Unit.ExecuteContext.Hierarchy = Hierarchy;
+	Unit.ExecuteContext.Hierarchy = Hierarchy.Get();
 
 	Unit.BoneToModify.SetNumZeroed(2);
 	Unit.BoneToModify[0].Bone = TEXT("BoneA");
diff --git a/Engine/Plugins/Animation/ControlRig/Source/ControlRig/Private/Units/Highlevel/Hierarchy/RigUnit_MultiFABRIK.cpp b/Engine/Plugins/Animation/ControlRig/Source/ControlRig/Private/Units/Highlevel/Hierarchy/RigUnit_MultiFABRIK.cpp
--- a/Engine/Plugins/Animation/ControlRig/Source/ControlRig/Private/Units/Highlevel/Hierarchy/RigUnit_MultiFABRIK.cpp
+++ b/Engine/Plugins/Animation/ControlRig/Source/ControlRig/Private/Units/Highlevel/Hierarchy/RigUnit_MultiFABRIK.cpp
@@ -479,7 +479,7 @@ IMPLEMENT_RIGUNIT_AUTOMATION_TEST(FRigUnit_MultiFABRIK)
 	const FRigElementKey Chain1_1 = Controller->AddBone(TEXT("Chain1_1"), Chain1_0, FTransform(FVector(3.f, 2.f, 0.f)), true, ERigBoneType::User);
 	const FRigElementKey Chain2_0 = Controller->AddBone(TEXT("Chain2_0"), Root, FTransform(FVector(-2.f, 0.f, 0.f)), true, ERigBoneType::User);
 	const FRigElementKey Chain2_1 = Controller->AddBone(TEXT("Chain2_1"), Chain2_0, FTransform(FVector(-2.f, 3.f, 0.f)), true, ERigBoneType::User);
-	Unit.ExecuteContext.Hierarchy = Hierarchy;
+	Unit.ExecuteContext.Hierarchy = Hierarchy.Get();
 
 	// first validation test
 	// make sure this doesn't crash
diff --git a/Engine/Plugins/Animation/ControlRig/Source/ControlRig/Private/Units/Highlevel/Hierarchy/RigUnit_TransformConstraint.cpp b/Engine/Plugins/Animation/ControlRig/Source/ControlRig/Private/Units/Highlevel/Hierarchy/RigUnit_TransformConstraint.cpp
--- a/Engine/Plugins/Animation/ControlRig/Source/ControlRig/Private/Units/Highlevel/Hierarchy/RigUnit_TransformConstraint.cpp
+++ b/Engine/Plugins/Animation/ControlRig/Source/ControlRig/Private/Units/Highlevel/Hierarchy/RigUnit_TransformConstraint.cpp
@@ -381,7 +381,7 @@ FRigUnit_ParentConstraint_Execute()
 	const FRigElementKey Parent3 = Controller->AddBone(TEXT("Parent3"), FRigElementKey(), FTransform( AnimationCore::QuatFromEuler( FVector(-50, 50, -50), Order), FVector(60.f, 60.f, 60.f)), true, ERigBoneType::User);
 	const FRigElementKey Parent4 = Controller->AddBone(TEXT("Parent4"), FRigElementKey(), FTransform( AnimationCore::QuatFromEuler( FVector(60, 60, 60), Order), FVector(80.f, 80.f, 80.f)), true, ERigBoneType::User);
 	
-	Unit.ExecuteContext.Hierarchy = Hierarchy;
+	Unit.ExecuteContext.Hierarchy = Hierarchy.Get();
 	Unit.Child = Child;
 
 	Unit.Parents.Add(FConstraintParent(Parent1, 1.0));
@@ -565,7 +565,7 @@ FRigUnit_ParentConstraintMath_Execute()
 	const FRigElementKey Parent3 = Controller->AddBone(TEXT("Parent3"), FRigElementKey(), FTransform( AnimationCore::QuatFromEuler( FVector(-50, 50, -50), Order), FVector(60.f, 60.f, 60.f)), true, ERigBoneType::User);
 	const FRigElementKey Parent4 = Controller->AddBone(TEXT("Parent4"), FRigElementKey(), FTransform( AnimationCore::QuatFromEuler( FVector(60, 60, 60), Order), FVector(80.f, 80.f, 80.f)), true, ERigBoneType::User);
 	
-	ExecuteContext.Hierarchy = Hierarchy;
+	ExecuteContext.Hierarchy = Hierarchy.Get();
 	Unit.Input = Hierarchy->GetInitialGlobalTransform(Child);
 
 	Unit.Parents.Add(FConstraintParent(Parent1, 1.0));
@@ -747,7 +747,7 @@ FRigVMStructUpgradeInfo FRigUnit_PositionConstraint::GetUpgradeInfo() const
 	const FRigElementKey Parent2 = Controller->AddBone(TEXT("Parent2"), FRigElementKey(), FTransform(FVector(40.f, 40.f, 40.f)), true, ERigBoneType::User);
 	const FRigElementKey Parent3 = Controller->AddBone(TEXT("Parent3"), FRigElementKey(), FTransform(FVector(60.f, 60.f, 60.f)), true, ERigBoneType::User);
 	
-	Unit.ExecuteContext.Hierarchy = Hierarchy;
+	Unit.ExecuteContext.Hierarchy = Hierarchy.Get();
 	Unit.Child = Child;
 
 	Unit.Parents.Add(FConstraintParent(Parent1, 1.0));
@@ -952,7 +952,7 @@ IMPLEMENT_RIGUNIT_AUTOMATION_TEST(FRigUnit_PositionConstraintLocalSpaceOffset)
     	const FRigElementKey Parent2 = Controller->AddBone(TEXT("Parent2"), FRigElementKey(), FTransform(FVector(40.f, 40.f, 40.f)), true, ERigBoneType::User);
     	const FRigElementKey Parent3 = Controller->AddBone(TEXT("Parent3"), FRigElementKey(), FTransform(FVector(60.f, 60.f, 60.f)), true, ERigBoneType::User);
     	
-    	Unit.ExecuteContext.Hierarchy = Hierarchy;
+    	Unit.ExecuteContext.Hierarchy = Hierarchy.Get();
     	Unit.Child = Child;
     
     	Unit.Parents.Add(FConstraintParent(Parent1, 1.0));
@@ -985,7 +985,7 @@ IMPLEMENT_RIGUNIT_AUTOMATION_TEST(FRigUnit_PositionConstraintLocalSpaceOffset)
 		const FRigElementKey Knee = Controller->AddBone(TEXT("Knee"), FRigElementKey(TEXT("Thigh"), ERigElementType::Bone), FTransform(FVector(10.f, 0.f, -50.f)), true, ERigBoneType::User);
 		const FRigElementKey Corrective = Controller->AddBone(TEXT("Corrective"), FRigElementKey(TEXT("Knee"), ERigElementType::Bone), FTransform(FVector(-50.f, 0.f, 0.f)), false, ERigBoneType::User);
 
-		Unit.ExecuteContext.Hierarchy = Hierarchy;
+		Unit.ExecuteContext.Hierarchy = Hierarchy.Get();
 		Unit.Child = Corrective;
     
 		Unit.Parents.Add(FConstraintParent(Knee, 1.0));
@@ -1194,7 +1194,7 @@ IMPLEMENT_RIGUNIT_AUTOMATION_TEST(FRigUnit_RotationConstraint)
 	const FRigElementKey Parent3 = Controller->AddBone(TEXT("Parent3"), FRigElementKey(), FTransform( AnimationCore::QuatFromEuler( FVector(-50, 50, -50), Order), FVector(60.f, 60.f, 60.f)), true, ERigBoneType::User);
 	const FRigElementKey Parent4 = Controller->AddBone(TEXT("Parent4"), FRigElementKey(), FTransform( AnimationCore::QuatFromEuler( FVector(60, 60, 60), Order), FVector(80.f, 80.f, 80.f)), true, ERigBoneType::User);
 	
-	Unit.ExecuteContext.Hierarchy = Hierarchy;
+	Unit.ExecuteContext.Hierarchy = Hierarchy.Get();
 	Unit.Child = Child;
 
 	Unit.Parents.Add(FConstraintParent(Parent1, 1.0));
@@ -1514,7 +1514,7 @@ IMPLEMENT_RIGUNIT_AUTOMATION_TEST(FRigUnit_RotationConstraintLocalSpaceOffset)
 		const FRigElementKey Parent3 = Controller->AddBone(TEXT("Parent3"), FRigElementKey(), FTransform( AnimationCore::QuatFromEuler( FVector(-50, 50, -50), Order), FVector(60.f, 60.f, 60.f)), true, ERigBoneType::User);
 		const FRigElementKey Parent4 = Controller->AddBone(TEXT("Parent4"), FRigElementKey(), FTransform( AnimationCore::QuatFromEuler( FVector(60, 60, 60), Order), FVector(80.f, 80.f, 80.f)), true, ERigBoneType::User);
 		
-		Unit.ExecuteContext.Hierarchy = Hierarchy;
+		Unit.ExecuteContext.Hierarchy = Hierarchy.Get();
 		Unit.Child = Child;
 
 		Unit.Parents.Add(FConstraintParent(Parent1, 1.0));
@@ -1584,7 +1584,7 @@ IMPLEMENT_RIGUNIT_AUTOMATION_TEST(FRigUnit_RotationConstraintLocalSpaceOffset)
 		const FRigElementKey Child = Controller->AddBone(TEXT("Child"), FRigElementKey(TEXT("Parent2"), ERigElementType::Bone),
 			FTransform(AnimationCore::QuatFromEuler(FVector(0,0,-45), Order), FVector(0.f, 0.f, 50.f)), false, ERigBoneType::User);
 
-		Unit.ExecuteContext.Hierarchy = Hierarchy;
+		Unit.ExecuteContext.Hierarchy = Hierarchy.Get();
 		Unit.Child = Child;
     
 		Unit.Parents.Add(FConstraintParent(Parent1, 1.0));
@@ -1803,7 +1803,7 @@ IMPLEMENT_RIGUNIT_AUTOMATION_TEST(FRigUnit_ScaleConstraint)
 	const FRigElementKey Parent1 = Controller->AddBone(TEXT("Parent1"), FRigElementKey(), FTransform(FQuat::Identity, FVector::ZeroVector, FVector(4,4,4)), true, ERigBoneType::User);
 	const FRigElementKey Parent2 = Controller->AddBone(TEXT("Parent2"), FRigElementKey(), FTransform(FQuat::Identity, FVector::ZeroVector, FVector(1,1,1)), true, ERigBoneType::User);
 	
-	Unit.ExecuteContext.Hierarchy = Hierarchy;
+	Unit.ExecuteContext.Hierarchy = Hierarchy.Get();
 	Unit.Child = Child;
 
 	Unit.Parents.Add(FConstraintParent(Parent1, 1.0));
@@ -2049,7 +2049,7 @@ IMPLEMENT_RIGUNIT_AUTOMATION_TEST(FRigUnit_ScaleConstraintLocalSpaceOffset)
 		const FRigElementKey Parent1 = Controller->AddBone(TEXT("Parent1"), FRigElementKey(), FTransform(FQuat::Identity, FVector::ZeroVector, FVector(4,4,4)), true, ERigBoneType::User);
 		const FRigElementKey Parent2 = Controller->AddBone(TEXT("Parent2"), FRigElementKey(), FTransform(FQuat::Identity, FVector::ZeroVector, FVector(1,1,1)), true, ERigBoneType::User);
 	
-		Unit.ExecuteContext.Hierarchy = Hierarchy;
+		Unit.ExecuteContext.Hierarchy = Hierarchy.Get();
 		Unit.Child = Child;
 
 		Unit.Parents.Add(FConstraintParent(Parent1, 1.0));
@@ -2090,7 +2090,7 @@ IMPLEMENT_RIGUNIT_AUTOMATION_TEST(FRigUnit_ScaleConstraintLocalSpaceOffset)
 			FTransform(AnimationCore::QuatFromEuler(FVector(0,0,0), Order), FVector(0.f, 0.f, 50.f)),
 			false, ERigBoneType::User);
 
-		Unit.ExecuteContext.Hierarchy = Hierarchy;
+		Unit.ExecuteContext.Hierarchy = Hierarchy.Get();
 		Unit.Child = Child;
     
 		Unit.Parents.Add(FConstraintParent(Parent1, 1.0));
diff --git a/Engine/Plugins/Animation/ControlRig/Source/ControlRig/Public/Units/RigUnitTest.h b/Engine/Plugins/Animation/ControlRig/Source/ControlRig/Public/Units/RigUnitTest.h
--- a/Engine/Plugins/Animation/ControlRig/Source/ControlRig/Public/Units/RigUnitTest.h
+++ b/Engine/Plugins/Animation/ControlRig/Source/ControlRig/Public/Units/RigUnitTest.h
@@ -21,7 +21,7 @@ public:
 
 	~FControlRigUnitTestBase()
 	{
-		if (Hierarchy)
+		if (Hierarchy.IsValid())
 		{
 			// we no longer add/remove the controller to/from root since controller is now part of the hierarchy
 			Hierarchy->RemoveFromRoot();
@@ -40,11 +40,11 @@ public:
 			Hierarchy->AddToRoot();
 			// we no longer add the controller to root since controller is now part of the hierarchy
 		}
-		ExecuteContext.Hierarchy = Hierarchy;
+		ExecuteContext.Hierarchy = Hierarchy.Get();
 	}
 
-	URigHierarchy* Hierarchy;
-	URigHierarchyController* Controller;
+	TSoftObjectPtr<URigHierarchy> Hierarchy;
+	TSoftObjectPtr<URigHierarchyController> Controller;
 };
 
 #define CONTROLRIG_RIGUNIT_STRINGIFY(Content) #Content
diff --git a/Engine/Plugins/Animation/RigLogic/Source/RigLogicModule/Private/RigUnit_RigLogic_Test.cpp b/Engine/Plugins/Animation/RigLogic/Source/RigLogicModule/Private/RigUnit_RigLogic_Test.cpp
--- a/Engine/Plugins/Animation/RigLogic/Source/RigLogicModule/Private/RigUnit_RigLogic_Test.cpp
+++ b/Engine/Plugins/Animation/RigLogic/Source/RigLogicModule/Private/RigUnit_RigLogic_Test.cpp
@@ -273,7 +273,7 @@ IMPLEMENT_RIGUNIT_AUTOMATION_TEST(FRigUnit_RigLogic)
 	}
 
 	Test.GetData()->SkelMeshComponent = SkelMeshComponent;
-	Unit.ExecuteContext.Hierarchy = Hierarchy;
+	Unit.ExecuteContext.Hierarchy = Hierarchy.Get();
 
 	//Test
 	Execute();
