Control Rig: When reporting to the log, make sure to specify the module that reported that message
#rnx

diff --git a/Engine/Plugins/Animation/ControlRig/Source/ControlRig/Private/Rigs/RigHierarchy.cpp b/Engine/Plugins/Animation/ControlRig/Source/ControlRig/Private/Rigs/RigHierarchy.cpp
--- a/Engine/Plugins/Animation/ControlRig/Source/ControlRig/Private/Rigs/RigHierarchy.cpp
+++ b/Engine/Plugins/Animation/ControlRig/Source/ControlRig/Private/Rigs/RigHierarchy.cpp
@@ -3545,7 +3545,7 @@ void URigHierarchy::SetTransform(FRigTransformElement* InTransformElement, const
 										MessageFormat,
 										*InTransformElement->GetName(),
 										*Child->GetName());
-									CRContext.GetLog()->Report(
+									CRContext.Report(
 										EMessageSeverity::Info,
 										ExecuteContext->GetPublicData<>().GetFunctionName(),
 										ExecuteContext->GetPublicData<>().GetInstructionIndex(),
diff --git a/Engine/Plugins/Animation/ControlRig/Source/ControlRig/Private/Units/Core/RigUnit_UserData.cpp b/Engine/Plugins/Animation/ControlRig/Source/ControlRig/Private/Units/Core/RigUnit_UserData.cpp
--- a/Engine/Plugins/Animation/ControlRig/Source/ControlRig/Private/Units/Core/RigUnit_UserData.cpp
+++ b/Engine/Plugins/Animation/ControlRig/Source/ControlRig/Private/Units/Core/RigUnit_UserData.cpp
@@ -109,7 +109,7 @@ void FRigDispatch_GetUserData::Execute(FRigVMExtendedExecuteContext& InContext,
 #if WITH_EDITOR
 		if(ControlRigContext.GetLog())
 		{
-			ControlRigContext.GetLog()->Report(EMessageSeverity::Info, ControlRigContext.GetFunctionName(), ControlRigContext.GetInstructionIndex(), Message);
+			ControlRigContext.Report(EMessageSeverity::Info, ControlRigContext.GetFunctionName(), ControlRigContext.GetInstructionIndex(), Message);
 		}
 		else
 #endif
@@ -179,7 +179,7 @@ FRigUnit_SetupShapeLibraryFromUserData_Execute()
 #if WITH_EDITOR
 		if(ExecuteContext.GetLog())
 		{
-			ExecuteContext.GetLog()->Report(EMessageSeverity::Info, ExecuteContext.GetFunctionName(), ExecuteContext.GetInstructionIndex(), Message);
+			ExecuteContext.Report(EMessageSeverity::Info, ExecuteContext.GetFunctionName(), ExecuteContext.GetInstructionIndex(), Message);
 		}
 		else
 #endif
diff --git a/Engine/Plugins/Animation/ControlRig/Source/ControlRig/Private/Units/Execution/RigUnit_Hierarchy.cpp b/Engine/Plugins/Animation/ControlRig/Source/ControlRig/Private/Units/Execution/RigUnit_Hierarchy.cpp
--- a/Engine/Plugins/Animation/ControlRig/Source/ControlRig/Private/Units/Execution/RigUnit_Hierarchy.cpp
+++ b/Engine/Plugins/Animation/ControlRig/Source/ControlRig/Private/Units/Execution/RigUnit_Hierarchy.cpp
@@ -214,10 +214,7 @@ FRigUnit_HierarchyGetChainItemArray_Execute()
 				{
 					Keys.Reset();
 #if WITH_EDITOR
-					if(ExecuteContext.GetLog())
-					{
-						ExecuteContext.GetLog()->Report(EMessageSeverity::Info, ExecuteContext.GetFunctionName(), ExecuteContext.GetInstructionIndex(), TEXT("Start and End are not part of the same chain."));
-					}
+					ExecuteContext.Report(EMessageSeverity::Info, ExecuteContext.GetFunctionName(), ExecuteContext.GetInstructionIndex(), TEXT("Start and End are not part of the same chain."));
 #endif
 				}
 				else if(bIncludeStart)
diff --git a/Engine/Plugins/Animation/ControlRig/Source/ControlRig/Private/Units/Execution/RigUnit_RigModules.cpp b/Engine/Plugins/Animation/ControlRig/Source/ControlRig/Private/Units/Execution/RigUnit_RigModules.cpp
--- a/Engine/Plugins/Animation/ControlRig/Source/ControlRig/Private/Units/Execution/RigUnit_RigModules.cpp
+++ b/Engine/Plugins/Animation/ControlRig/Source/ControlRig/Private/Units/Execution/RigUnit_RigModules.cpp
@@ -34,11 +34,10 @@ FRigUnit_GetCurrentNameSpace_Execute()
 	if(!ExecuteContext.IsRigModule())
 	{
 #if WITH_EDITOR
-		if(ExecuteContext.GetLog() != nullptr)
-		{
-			static const FString Message = TEXT("This node should only be used in a Rig Module."); 
-			ExecuteContext.GetLog()->Report(EMessageSeverity::Warning, ExecuteContext.GetFunctionName(), ExecuteContext.GetInstructionIndex(), Message);
-		}
+		
+		static const FString Message = TEXT("This node should only be used in a Rig Module."); 
+		ExecuteContext.Report(EMessageSeverity::Warning, ExecuteContext.GetFunctionName(), ExecuteContext.GetInstructionIndex(), Message);
+		
 #endif
 	}
 	NameSpace = *ExecuteContext.GetRigModuleNameSpace();
diff --git a/Engine/Plugins/Animation/ControlRig/Source/ControlRig/Public/Units/RigUnitContext.h b/Engine/Plugins/Animation/ControlRig/Source/ControlRig/Public/Units/RigUnitContext.h
--- a/Engine/Plugins/Animation/ControlRig/Source/ControlRig/Public/Units/RigUnitContext.h
+++ b/Engine/Plugins/Animation/ControlRig/Source/ControlRig/Public/Units/RigUnitContext.h
@@ -222,6 +222,22 @@ public:
 	FRigUnitContext UnitContext;
 	URigHierarchy* Hierarchy;
 
+#if WITH_EDITOR
+	virtual void Report(const FRigVMLogSettings& InLogSettings, const FName& InFunctionName, int32 InInstructionIndex, const FString& InMessage) const override
+	{
+		FString Prefix = GetRigModuleNameSpace();
+		if (!Prefix.IsEmpty())
+		{
+			const FString Name = FString::Printf(TEXT("%s %s"), *Prefix, *InFunctionName.ToString());
+			FRigVMExecuteContext::Report(InLogSettings, *Name, InInstructionIndex, InMessage);
+		}
+		else
+		{
+			FRigVMExecuteContext::Report(InLogSettings, InFunctionName, InInstructionIndex, InMessage);
+		}
+	}
+#endif
+
 private:
 	FString RigModuleNameSpace;
 	uint32 RigModuleNameSpaceHash;
@@ -247,10 +263,8 @@ private:
 
 #if WITH_EDITOR
 #define UE_CONTROLRIG_RIGUNIT_REPORT(Severity, Format, ...) \
-if(ExecuteContext.GetLog() != nullptr) \
-{ \
-	ExecuteContext.GetLog()->Report(EMessageSeverity::Severity, ExecuteContext.GetFunctionName(), ExecuteContext.GetInstructionIndex(), FString::Printf((Format), ##__VA_ARGS__)); \
-}
+ExecuteContext.Report(EMessageSeverity::Severity, ExecuteContext.GetFunctionName(), ExecuteContext.GetInstructionIndex(), FString::Printf((Format), ##__VA_ARGS__)); 
+
 #define UE_CONTROLRIG_RIGUNIT_LOG_MESSAGE(Format, ...) UE_CONTROLRIG_RIGUNIT_REPORT(Info, (Format), ##__VA_ARGS__)
 #define UE_CONTROLRIG_RIGUNIT_REPORT_WARNING(Format, ...) UE_CONTROLRIG_RIGUNIT_REPORT(Warning, (Format), ##__VA_ARGS__)
 #define UE_CONTROLRIG_RIGUNIT_REPORT_ERROR(Format, ...) UE_CONTROLRIG_RIGUNIT_REPORT(Error, (Format), ##__VA_ARGS__)
diff --git a/Engine/Plugins/Runtime/RigVM/Source/RigVM/Private/RigVMFunctions/RigVMFunction_String.cpp b/Engine/Plugins/Runtime/RigVM/Source/RigVM/Private/RigVMFunctions/RigVMFunction_String.cpp
--- a/Engine/Plugins/Runtime/RigVM/Source/RigVM/Private/RigVMFunctions/RigVMFunction_String.cpp
+++ b/Engine/Plugins/Runtime/RigVM/Source/RigVM/Private/RigVMFunctions/RigVMFunction_String.cpp
@@ -418,7 +418,7 @@ void FRigDispatch_FromString::Execute(FRigVMExtendedExecuteContext& InContext, F
 #if WITH_EDITOR
 			if(ExecuteContext.GetLog() != nullptr)
 			{
-				ExecuteContext.GetLog()->Report(EMessageSeverity::Error, InContext.GetPublicData<>().GetFunctionName(), InContext.GetPublicData<>().GetInstructionIndex(), Error);
+				ExecuteContext.Report(EMessageSeverity::Error, InContext.GetPublicData<>().GetFunctionName(), InContext.GetPublicData<>().GetInstructionIndex(), Error);
 			}
 #endif
 			FString ObjectPath;
diff --git a/Engine/Plugins/Runtime/RigVM/Source/RigVM/Public/RigVMCore/RigVMExecuteContext.h b/Engine/Plugins/Runtime/RigVM/Source/RigVM/Public/RigVMCore/RigVMExecuteContext.h
--- a/Engine/Plugins/Runtime/RigVM/Source/RigVM/Public/RigVMCore/RigVMExecuteContext.h
+++ b/Engine/Plugins/Runtime/RigVM/Source/RigVM/Public/RigVMCore/RigVMExecuteContext.h
@@ -394,6 +394,13 @@ struct FRigVMExecuteContext
 #if WITH_EDITOR
 	FRigVMLog* GetLog() const { return LogPtr; }
 	void SetLog(FRigVMLog* InLog) { LogPtr = InLog; }
+	virtual void Report(const FRigVMLogSettings& InLogSettings, const FName& InFunctionName, int32 InInstructionIndex, const FString& InMessage) const
+	{
+		if (FRigVMLog* Log = GetLog())
+		{
+			Log->Report(InLogSettings, InFunctionName, InInstructionIndex, InMessage);
+		}
+	}
 #endif
 
 	FRigVMDrawInterface* GetDrawInterface() const { return DrawInterfacePtr; }
diff --git a/Engine/Plugins/Runtime/RigVM/Source/RigVM/Public/RigVMFunctions/RigVMFunctionDefines.h b/Engine/Plugins/Runtime/RigVM/Source/RigVM/Public/RigVMFunctions/RigVMFunctionDefines.h
--- a/Engine/Plugins/Runtime/RigVM/Source/RigVM/Public/RigVMFunctions/RigVMFunctionDefines.h
+++ b/Engine/Plugins/Runtime/RigVM/Source/RigVM/Public/RigVMFunctions/RigVMFunctionDefines.h
@@ -36,7 +36,7 @@ namespace ERigVMClampSpatialMode
 #define UE_RIGVMSTRUCT_REPORT(Severity, Format, ...) \
 if(ExecuteContext.GetLog() != nullptr) \
 { \
-ExecuteContext.GetLog()->Report(EMessageSeverity::Severity, ExecuteContext.GetFunctionName(), ExecuteContext.GetInstructionIndex(), FString::Printf((Format), ##__VA_ARGS__)); \
+ExecuteContext.Report(EMessageSeverity::Severity, ExecuteContext.GetFunctionName(), ExecuteContext.GetInstructionIndex(), FString::Printf((Format), ##__VA_ARGS__)); \
 }
 #define UE_RIGVMSTRUCT_LOG_MESSAGE(Format, ...) UE_RIGVMSTRUCT_REPORT(Info, (Format), ##__VA_ARGS__)
 #define UE_RIGVMSTRUCT_REPORT_WARNING(Format, ...) UE_RIGVMSTRUCT_REPORT(Warning, (Format), ##__VA_ARGS__)
