Control Rig: Fig bug when mirroring tree of modules
#rnx

diff --git a/Engine/Plugins/Animation/ControlRig/Source/ControlRigEditor/Private/Editor/SModularRigModel.cpp b/Engine/Plugins/Animation/ControlRig/Source/ControlRigEditor/Private/Editor/SModularRigModel.cpp
--- a/Engine/Plugins/Animation/ControlRig/Source/ControlRigEditor/Private/Editor/SModularRigModel.cpp
+++ b/Engine/Plugins/Animation/ControlRig/Source/ControlRigEditor/Private/Editor/SModularRigModel.cpp
@@ -752,9 +752,8 @@ void SModularRigModel::HandleMirrorModules(const TArray<FString>& InPaths)
 			UModularRigController* Controller = ControlRigBlueprint->GetModularRigController();
 			check(Controller);
 
-			// Make sure we delete the modules from children to root
+			// Make sure we mirror the modules from root to children
 			TArray<FString> SortedPaths = Controller->Model->SortPaths(InPaths);
-			Algo::Reverse(SortedPaths);
 			for (const FString& Path : SortedPaths)
 			{
 				Controller->MirrorModule(Path, Settings);
