ControRig: Add cooked CRs to sequencer tracks

diff --git a/Engine/Plugins/Animation/ControlRig/Source/ControlRigEditor/Private/Sequencer/ControlRigParameterTrackEditor.cpp b/Engine/Plugins/Animation/ControlRig/Source/ControlRigEditor/Private/Sequencer/ControlRigParameterTrackEditor.cpp
--- a/Engine/Plugins/Animation/ControlRig/Source/ControlRigEditor/Private/Sequencer/ControlRigParameterTrackEditor.cpp
+++ b/Engine/Plugins/Animation/ControlRig/Source/ControlRigEditor/Private/Sequencer/ControlRigParameterTrackEditor.cpp
@@ -1899,14 +1899,14 @@ void FControlRigParameterTrackEditor::OnChannelChanged(const FMovieSceneChannelM
 	}
 }
 
-void FControlRigParameterTrackEditor::AddTrackForComponent(USceneComponent* InComponent, FGuid InBinding)
+void FControlRigParameterTrackEditor::AddTrackForComponent(USceneComponent* InComponent, FGuid InBinding) 
 {
 	if (USkeletalMeshComponent* SkelMeshComp = Cast<USkeletalMeshComponent>(InComponent))
 	{
 		if(bAutoGenerateControlRigTrack && !SkelMeshComp->GetDefaultAnimatingRig().IsNull())
 		{
 			UObject* Object = SkelMeshComp->GetDefaultAnimatingRig().LoadSynchronous();
-			if (Object != nullptr && (Object->IsA<UControlRigBlueprint>() || Object->IsA<UControlRigComponent>()))
+			if (Object != nullptr && (Object->IsA<UControlRigBlueprint>() || Object->IsA<UControlRigComponent>() || Object->IsA<URigVMBlueprintGeneratedClass>()))
 			{
 				FGuid Binding = InBinding.IsValid() ? InBinding : GetSequencer()->GetHandleToObject(InComponent, true /*bCreateHandle*/);
 				if (Binding.IsValid())
@@ -1916,14 +1916,21 @@ void FControlRigParameterTrackEditor::AddTrackForComponent(USceneComponent* InCo
 					UMovieSceneControlRigParameterTrack* Track = Cast<UMovieSceneControlRigParameterTrack>(OwnerMovieScene->FindTrack(UMovieSceneControlRigParameterTrack::StaticClass(), Binding, NAME_None));
 					if (Track == nullptr)
 					{
+						URigVMBlueprintGeneratedClass* RigClass = nullptr;
 						if (UControlRigBlueprint* BPControlRig = Cast<UControlRigBlueprint>(Object))
 						{
-							if (URigVMBlueprintGeneratedClass* RigClass = BPControlRig->GetRigVMBlueprintGeneratedClass())
+							RigClass = BPControlRig->GetRigVMBlueprintGeneratedClass();
+						}
+						else
+						{
+							RigClass = Cast<URigVMBlueprintGeneratedClass>(Object);
+						}
+
+						if (RigClass)
+						{
+							if (UControlRig* CDO = Cast<UControlRig>(RigClass->GetDefaultObject(true /* create if needed */)))
 							{
-								if (UControlRig* CDO = Cast<UControlRig>(RigClass->GetDefaultObject(true /* create if needed */)))
-								{
-									AddControlRig(CDO->GetClass(), InComponent, Binding);
-								}
+								AddControlRig(CDO->GetClass(), InComponent, Binding);
 							}
 						}
 					}
