Control Rig: Snapper:  We had an optimization to not re-evaluate a snap if we are on the current Sequencer frame but that will break do to how we gather frames (we could have gathered the last frame most recently). So just remove that check.

diff --git a/Engine/Plugins/Animation/ControlRig/Source/ControlRigEditor/Private/Tools/ControlRigSnapper.cpp b/Engine/Plugins/Animation/ControlRig/Source/ControlRigEditor/Private/Tools/ControlRigSnapper.cpp
--- a/Engine/Plugins/Animation/ControlRig/Source/ControlRigEditor/Private/Tools/ControlRigSnapper.cpp
+++ b/Engine/Plugins/Animation/ControlRig/Source/ControlRigEditor/Private/Tools/ControlRigSnapper.cpp
@@ -168,7 +168,7 @@ bool FControlRigSnapper::GetControlRigControlTransforms(ISequencer* Sequencer,
 	{
 		FMovieSceneSequenceIDRef Template = Sequencer->GetFocusedTemplateID();
 		FMovieSceneSequenceTransform RootToLocalTransform = Sequencer->GetFocusedMovieSceneSequenceTransform();
-		TOptional<FFrameNumber> FrameNumber = (Sequencer->GetLocalTime().Time * RootToLocalTransform.InverseNoLooping()).RoundToFrame();
+		TOptional<FFrameNumber> FrameNumber;
 		return LocalGetControlRigControlTransforms(Sequencer, FrameNumber, Sequencer->GetFocusedMovieSceneSequence(), Template, RootToLocalTransform,
 			ControlRig, ControlName, Frames, ParentTransforms, OutTransforms);
 	
