Sequencer: Tween tools: Fix divide by zero crash with new Push Pull tween and change blend to ease function based on animators feedback

diff --git a/Engine/Plugins/Animation/ControlRig/Source/ControlRigEditor/Private/Tools/ControlRigTweener.cpp b/Engine/Plugins/Animation/ControlRig/Source/ControlRigEditor/Private/Tools/ControlRigTweener.cpp
--- a/Engine/Plugins/Animation/ControlRig/Source/ControlRigEditor/Private/Tools/ControlRigTweener.cpp
+++ b/Engine/Plugins/Animation/ControlRig/Source/ControlRigEditor/Private/Tools/ControlRigTweener.cpp
@@ -632,6 +632,10 @@ bool FControlsToTween::Setup(TWeakPtr<ISequencer>& InSequencer, TWeakPtr<FContro
 
 double FPushPullSlider::DoBlend(const FBlendStruct& BlendStruct)
 {
+	if (FMath::IsNearlyEqual(BlendStruct.NextTime, BlendStruct.PreviousTime))
+	{
+		return BlendStruct.CurrentValue;
+	}
 	const double T = (BlendStruct.CurrentTime - BlendStruct.PreviousTime) / (BlendStruct.NextTime - BlendStruct.PreviousTime);
 	const double ValueAtT = BlendStruct.PreviousValue + T * (BlendStruct.NextValue - BlendStruct.PreviousValue);
 	double NewValue;
@@ -733,25 +737,51 @@ static float ExpOut(float InTime)
 {
 	return 1.f - ExpIn(1.f - InTime);
 }
+
+static float SCurve(float X, float Slope, float Width, float Height, float XShift, float YShift)
+{
+	if (X > (XShift + Width))
+	{
+		return(Height + YShift);
+	}
+	else if (X < XShift)
+	{
+		return YShift;
+	}
+
+	float Val = Height * (FMath::Pow((X - XShift), Slope) / (FMath::Pow(X - XShift, Slope) + FMath::Pow((Width - (X - XShift)), Slope))) + YShift;
+	return Val;
+}
 }
+
 double FBlendToEaseSlider::DoBlend(const FBlendStruct& BlendStruct)
 {
-	double NewValue = BlendStruct.CurrentValue;
+	if (FMath::IsNearlyEqual(BlendStruct.NextTime, BlendStruct.PreviousTime))
+	{
+		return BlendStruct.CurrentValue;
+	}
+	const double Source = BlendStruct.CurrentValue;
 	const double FullTimeDiff = BlendStruct.NextTime - BlendStruct.PreviousTime;
-	if (BlendStruct.BlendValue < 0.0)
+	const double AbsValue = FMath::Abs(BlendStruct.BlendValue);
+	const double X = BlendStruct.CurrentTime - BlendStruct.PreviousTime;
+	const double Ratio = X / FullTimeDiff;
+	double Shift = 0.0, Delta = 0.0, Base = 0.0;
+	if (BlendStruct.BlendValue > 0)
 	{
-		const double MyTimeDiff = ((BlendStruct.CurrentTime - BlendStruct.PreviousTime) / FullTimeDiff);
-		double NewBlend = FMath::Clamp((- 1.0 * BlendStruct.BlendValue) / MyTimeDiff, 0.0, 1.0);
-		NewBlend *= (1.0 - BlendToEase::ExpIn(MyTimeDiff));
-		NewValue = BlendStruct.CurrentValue + NewBlend * (BlendStruct.PreviousValue - BlendStruct.CurrentValue);
+		Shift = -1.0;
+		Delta = BlendStruct.NextValue - Source;
+		Base = Source;
 	}
 	else
 	{
-		const double MyTimeDiff =  ((BlendStruct.NextTime - BlendStruct.CurrentTime) / FullTimeDiff);
-		double NewBlend = FMath::Clamp(BlendStruct.BlendValue / MyTimeDiff, 0.0, 1.0);
-		NewBlend *= (1.0 - BlendToEase::ExpIn(MyTimeDiff));
-		NewValue = BlendStruct.CurrentValue + NewBlend * (BlendStruct.NextValue - BlendStruct.CurrentValue);
+		Shift = 0.0;
+		Delta = Source - BlendStruct.PreviousValue;
+		Base = BlendStruct.PreviousValue;
 	}
+	const double Slope = 5.0 * AbsValue;
+	const double EaseY = BlendToEase::SCurve(Ratio, Slope, 2.0, 2.0, Shift, Shift);
+	const double NewValue = Base + (Delta * EaseY);
+	
 	return NewValue;
 }
 
