Control Rig Edit Mode: shape creation / update improvements
    - fixed RecreateControlShapeActors recreating shapes too often ( if(Controls.Num() == ShapeActors->Num()) was tested after the controls were filtered...)
    - always try to reuse existing shapes instead of creating new ones. adding new actors to complex levels/sequences can have slow side effects.
    - FControlRigParameterTrackEditor::OnSelectionChanged will only request shape updates if the rig has been added.
    - added UControlRigEditModeSettings::OnSettingsChange to listen for settings changes instead of checking them for each tick (this should probably be used in other places)
    - added some shape utility functions to factorize shape related code

TODO/FIXME: UControlRigEditModeSettings::bHideControlShapes will actually also override bDisplayHierarchy, bDisplayNulls, bDisplaySockets etc. in the Render function which probably not expected

diff --git a/Engine/Plugins/Animation/ControlRig/Source/ControlRigEditor/Private/EditMode/ControlRigEditMode.cpp b/Engine/Plugins/Animation/ControlRig/Source/ControlRigEditor/Private/EditMode/ControlRigEditMode.cpp
--- a/Engine/Plugins/Animation/ControlRig/Source/ControlRigEditor/Private/EditMode/ControlRigEditMode.cpp
+++ b/Engine/Plugins/Animation/ControlRig/Source/ControlRigEditor/Private/EditMode/ControlRigEditMode.cpp
@@ -14,8 +14,6 @@
 #include "EditMode/SControlRigDetails.h"
 #include "EditMode/SControlRigOutliner.h"
 #include "ISequencer.h"
-#include "SequencerSettings.h"
-#include "Sections/MovieSceneSpawnSection.h"
 #include "MovieScene.h"
 #include "Editor.h"
 #include "EditorViewportClient.h"
@@ -263,17 +261,17 @@ bool FControlRigEditMode::SetSequencer(TWeakPtr<ISequencer> InSequencer)
 	return false;
 }
 
-bool FControlRigEditMode::AddControlRigObject(UControlRig* ControlRig, TWeakPtr<ISequencer> InSequencer)
+bool FControlRigEditMode::AddControlRigObject(UControlRig* InControlRig, const TWeakPtr<ISequencer>& InSequencer)
 {
-	if (ControlRig)
+	if (InControlRig)
 	{
-		if (RuntimeControlRigs.Contains(ControlRig) == false)
+		if (RuntimeControlRigs.Contains(InControlRig) == false)
 		{
 			if (InSequencer.IsValid())
 			{
 				if (SetSequencer(InSequencer) == false) //was already there so just add it,otherwise this function will add everything in the active 
 				{
-					AddControlRigInternal(ControlRig);
+					AddControlRigInternal(InControlRig);
 					SetObjects_Internal();
 				}
 				return true;
@@ -386,7 +384,7 @@ void FControlRigEditMode::SetObjects_Internal()
 	else
 	{
 		// create default manipulation layer
-		RequestToRecreateControlShapeActors(nullptr);
+		RequestToRecreateControlShapeActors();
 	}
 }
 
@@ -440,6 +438,11 @@ void FControlRigEditMode::Enter()
 	if (const UControlRigEditModeSettings* Settings = GetDefault<UControlRigEditModeSettings>())
 	{
 		GetModeManager()->SetWidgetScale(Settings->GizmoScale);
+
+		if (!Settings->OnSettingsChange.IsBoundToObject(this))
+		{
+			Settings->OnSettingsChange.AddSP(this, &FControlRigEditMode::OnSettingsChanged);
+		}
 	}
 }
 
@@ -530,6 +533,11 @@ void FControlRigEditMode::Exit()
 	//make sure the widget is reset
 	ResetControlShapeSize();
 
+	if (const UControlRigEditModeSettings* Settings = GetDefault<UControlRigEditModeSettings>())
+	{
+		Settings->OnSettingsChange.RemoveAll(this);
+	}
+	
 	// Call parent implementation
 	FEdMode::Exit();
 }
@@ -578,21 +586,6 @@ void FControlRigEditMode::Tick(FEditorViewportClient* ViewportClient, float Delt
 		ViewportClient->Invalidate();
 	}
 
-	// check if the settings for xray rendering are different for any of the control shape actors
-	const UControlRigEditModeSettings* Settings = GetDefault<UControlRigEditModeSettings>();
-	if(bShowControlsAsOverlay != Settings->bShowControlsAsOverlay)
-	{
-		bShowControlsAsOverlay = Settings->bShowControlsAsOverlay;
-		for (TWeakObjectPtr<UControlRig>& RuntimeRigPtr : RuntimeControlRigs)
-		{
-			if (UControlRig* RuntimeControlRig = RuntimeRigPtr.Get())
-			{
-				UpdateSelectabilityOnSkeletalMeshes(RuntimeControlRig, !bShowControlsAsOverlay);
-			}
-		}
-		RequestToRecreateControlShapeActors();
-	}
-
 	// Defer creation of shapes if manipulating the viewport
 	if (RecreateControlShapesRequired != ERecreateControlRigShape::RecreateNone && !(FSlateApplication::Get().HasAnyMouseCaptor() || GUnrealEd->IsUserInteracting()))
 	{
@@ -3785,196 +3778,251 @@ static bool IsSupportedControlType(const ERigControlType ControlType)
 	return false;
 }
 
-void FControlRigEditMode::RecreateControlShapeActors(const TArray<FRigElementKey>& InSelectedElements)
+bool FControlRigEditMode::TryUpdatingControlsShapes(UControlRig* InControlRig)
 {
-	if (RecreateControlShapesRequired == ERecreateControlRigShape::RecreateAll)
+	using namespace ControlRigEditMode::Shapes;
+	
+	const URigHierarchy* Hierarchy = InControlRig ? InControlRig->GetHierarchy() : nullptr;
+	if (!Hierarchy)
 	{
-		for (TWeakObjectPtr<UControlRig>& RuntimeRigPtr : RuntimeControlRigs)
+		return false;
+	}
+
+	const auto* ShapeActors = ControlRigShapeActors.Find(InControlRig);
+	if (!ShapeActors)
+	{
+		// create the shapes if they don't already exist
+		CreateShapeActors(InControlRig);
+		return true;
+	}
+	
+	// get controls which need shapes
+	TArray<FRigControlElement*> Controls;
+	GetControlsEligibleForShapes(InControlRig, Controls);
+
+	if (Controls.IsEmpty())
+	{
+		// not control needing shape so clear the shape actors
+		DestroyShapesActors(InControlRig);
+		return true;
+	}
+	
+	const TArray<TObjectPtr<AControlRigShapeActor>>& Shapes = *ShapeActors;
+	const int32 NumShapes = Shapes.Num();
+
+	TArray<FRigControlElement*> ControlPerShapeActor;
+	ControlPerShapeActor.SetNumZeroed(NumShapes);
+
+	if (Controls.Num() == NumShapes)
+	{
+		//unfortunately n*n-ish but this should be very rare and much faster than recreating them
+		for (int32 ShapeActorIndex = 0; ShapeActorIndex < NumShapes; ShapeActorIndex++)
 		{
-			if (UControlRig* RuntimeControlRig = RuntimeRigPtr.Get())
+			if (const AControlRigShapeActor* Actor = Shapes[ShapeActorIndex].Get())
 			{
-				DestroyShapesActors(RuntimeControlRig);
-				CreateShapeActors(RuntimeControlRig);
+				const int32 ControlIndex = Controls.IndexOfByPredicate([Actor](const FRigControlElement* Control)
+				{
+					return Control && Control->GetFName() == Actor->ControlName;
+				});
+				if (ControlIndex != INDEX_NONE)
+				{
+					ControlPerShapeActor[ShapeActorIndex] = Controls[ControlIndex];
+					Controls.RemoveAtSwap(ControlIndex);
+				}
+			}
+			else //no actor just recreate
+			{
+				return false;
 			}
 		}
 	}
-	else if (ControlRigsToRecreate.Num() > 0)
+
+	// Some controls don't have associated shape so recreate them 
+	if (!Controls.IsEmpty())
+	{
+		return false;
+	}
+
+	// we have matching controls - we should at least sync their settings.
+	// PostPoseUpdate / TickControlShape is going to take care of color, visibility etc.
+	// MeshTransform has to be handled here.
+	const TArray<TSoftObjectPtr<UControlRigShapeLibrary>>& ShapeLibraries = InControlRig->GetShapeLibraries();
+	for (int32 ShapeActorIndex = 0; ShapeActorIndex < NumShapes; ShapeActorIndex++)
 	{
-		TArray < UControlRig*> ControlRigsCopy = ControlRigsToRecreate;
-		for (UControlRig* ControlRig : ControlRigsCopy)
+		const AControlRigShapeActor* ShapeActor = Shapes[ShapeActorIndex].Get();
+		FRigControlElement* ControlElement = ControlPerShapeActor[ShapeActorIndex];
+		if (ShapeActor && ControlElement)
 		{
-			//check to see if actors have really changed, if not don't do it
-			bool bRecreateThem = true;
-			if (auto* ShapeActors = ControlRigShapeActors.Find(ControlRig))
+			const FTransform ShapeTransform = Hierarchy->GetControlShapeTransform(ControlElement, ERigTransformType::CurrentLocal);
+			if (const FControlRigShapeDefinition* ShapeDef = UControlRigShapeLibrary::GetShapeByName(ControlElement->Settings.ShapeName, ShapeLibraries, InControlRig->ShapeLibraryNameMap))
 			{
-				TArray<FRigControlElement*> Controls = ControlRig->AvailableControls();
-				TArray<FRigControlElement*> ControlPerShapeActor;
-				ControlPerShapeActor.SetNumZeroed(ShapeActors->Num());
-				
-				if(Controls.Num() == ShapeActors->Num())
+				const FTransform& MeshTransform = ShapeDef->Transform;
+				if (UStaticMesh* ShapeMesh = ShapeDef->StaticMesh.LoadSynchronous())
 				{
-					for (int32 ControlIndex = Controls.Num() - 1; ControlIndex >= 0; --ControlIndex)
+					if(ShapeActor->StaticMeshComponent->GetStaticMesh() != ShapeMesh)
 					{
-						FRigControlElement* ControlElement = Controls[ControlIndex];
-						if (!ControlElement->Settings.SupportsShape() || !IsSupportedControlType(ControlElement->Settings.ControlType))
-						{
-							Controls.RemoveAtSwap(ControlIndex);
-						}
+						ShapeActor->StaticMeshComponent->SetStaticMesh(ShapeMesh);
 					}
-					//unfortunately n*n-ish but this should be very rare and much faster than recreating them
-					for (int32 ShapeActorIndex = 0; ShapeActorIndex < ShapeActors->Num(); ShapeActorIndex++)
-					{
-						const AControlRigShapeActor* Actor = ShapeActors->operator[](ShapeActorIndex).Get();
-						if (Actor)
-						{
-							for (int32 ControlIndex = 0; ControlIndex < Controls.Num(); ++ControlIndex)
-							{
-								FRigControlElement* Element = Controls[ControlIndex];
-								if (Element && Element->GetFName() == Actor->ControlName)
-								{
-									Controls.RemoveAtSwap(ControlIndex);
-									ControlPerShapeActor[ShapeActorIndex] = Element;
-									break;
-								}
-							}
-						}
-						else //no actor just recreate
-						{
-							break;
-						}
-					}
-				}
-				if (Controls.Num() == 0)
-				{
-					bRecreateThem = false;
-
-					// we have matching controls - we should at least sync their settings.
-					// PostPoseUpdate / TickControlShape is going to take care of color, visibility etc.
-					// MeshTransform has to be handled here.
-					for (int32 ShapeActorIndex = 0; ShapeActorIndex < ShapeActors->Num(); ShapeActorIndex++)
-					{
-						const AControlRigShapeActor* ShapeActor = ShapeActors->operator[](ShapeActorIndex).Get();
-						FRigControlElement* ControlElement = ControlPerShapeActor[ShapeActorIndex];
-						if (ShapeActor && ControlElement)
-						{
-							const FTransform ShapeTransform = ControlRig->GetHierarchy()->GetControlShapeTransform(ControlElement, ERigTransformType::CurrentLocal);
-							FTransform MeshTransform = FTransform::Identity;
-							if (const FControlRigShapeDefinition* ShapeDef = UControlRigShapeLibrary::GetShapeByName(ControlElement->Settings.ShapeName, ControlRig->GetShapeLibraries(), ControlRig->ShapeLibraryNameMap))
-							{
-								MeshTransform = ShapeDef->Transform;
-
-								if(UStaticMesh* ShapeMesh = ShapeDef->StaticMesh.LoadSynchronous())
-								{
-									if(ShapeActor->StaticMeshComponent->GetStaticMesh() != ShapeMesh)
-									{
-										ShapeActor->StaticMeshComponent->SetStaticMesh(ShapeMesh);
-									}
-								}
-							}
-							ShapeActor->StaticMeshComponent->SetRelativeTransform(MeshTransform * ShapeTransform);
-						}
-					}
-					
-					PostPoseUpdate();
 				}
+				ShapeActor->StaticMeshComponent->SetRelativeTransform(MeshTransform * ShapeTransform);
 			}
-			if (bRecreateThem)
+			else
 			{
-				DestroyShapesActors(ControlRig);
-				CreateShapeActors(ControlRig);
+				ShapeActor->StaticMeshComponent->SetRelativeTransform(ShapeTransform);
 			}
 		}
-		ControlRigsToRecreate.SetNum(0);
 	}
+
+	// equivalent to PostPoseUpdate for those shapes only
+	FTransform ComponentTransform = FTransform::Identity;
+	if (!AreEditingControlRigDirectly())
+	{
+		ComponentTransform = GetHostingSceneComponentTransform(InControlRig);
+	}
+		
+	const FShapeUpdateParams Params(InControlRig, ComponentTransform, IsControlRigSkelMeshVisible(InControlRig));
+	for (int32 ShapeActorIndex = 0; ShapeActorIndex < NumShapes; ShapeActorIndex++)
+	{
+		AControlRigShapeActor* ShapeActor = Shapes[ShapeActorIndex].Get();
+		FRigControlElement* ControlElement = ControlPerShapeActor[ShapeActorIndex];
+		if (ShapeActor && ControlElement)
+		{
+			UpdateControlShape(ShapeActor, ControlElement, Params);
+		}
+	}
+
+	return true;
 }
 
-void FControlRigEditMode::CreateShapeActors(UControlRig* ControlRig)
+void FControlRigEditMode::RecreateControlShapeActors()
 {
-	// create gizmo actors
-	FActorSpawnParameters ActorSpawnParameters;
-	ActorSpawnParameters.bTemporaryEditorActor = true;
+	if (RecreateControlShapesRequired == ERecreateControlRigShape::RecreateAll)
+	{
+		// recreate all control rigs shape actors
+		for (TWeakObjectPtr<UControlRig>& RuntimeRigPtr : RuntimeControlRigs)
+		{
+			if (UControlRig* RuntimeControlRig = RuntimeRigPtr.Get())
+			{
+				DestroyShapesActors(RuntimeControlRig);
+				CreateShapeActors(RuntimeControlRig);
+			}
+		}
+		RecreateControlShapesRequired = ERecreateControlRigShape::RecreateNone;
+		return;
+	}
+
+	if (ControlRigsToRecreate.IsEmpty())
+	{
+		// nothing to update
+		return;
+	}
+	
+	// update or recreate all control rigs in ControlRigsToRecreate
+	TArray<UControlRig*> ControlRigsCopy = ControlRigsToRecreate;
+	for (UControlRig* ControlRig : ControlRigsCopy)
+	{
+		const bool bUpdated = TryUpdatingControlsShapes(ControlRig);
+		if (!bUpdated)
+		{
+			DestroyShapesActors(ControlRig);
+			CreateShapeActors(ControlRig);
+		}
+	}
+	RecreateControlShapesRequired = ERecreateControlRigShape::RecreateNone;
+	ControlRigsToRecreate.SetNum(0);
+}
 
+void FControlRigEditMode::CreateShapeActors(UControlRig* InControlRig)
+{
+	using namespace ControlRigEditMode::Shapes;
+	
 	if(bShowControlsAsOverlay)
 	{
 		// enable translucent selection
 		GetMutableDefault<UEditorPerProjectUserSettings>()->bAllowSelectTranslucent = true;
 	}
 
-	TArray<FRigControlElement*> Controls = ControlRig->AvailableControls();
-	const TArray<TSoftObjectPtr<UControlRigShapeLibrary>> ShapeLibraries = ControlRig->GetShapeLibraries();
-	int32 ControlRigIndex = RuntimeControlRigs.Find(ControlRig);
-	for (FRigControlElement* ControlElement : Controls)
-	{
-		if (!ControlElement->Settings.SupportsShape())
-		{
-			continue;
-		}
-		if (IsSupportedControlType(ControlElement->Settings.ControlType))
-		{
-			FControlShapeActorCreationParam Param;
-			Param.ManipObj = ControlRig;
-			Param.ControlRigIndex = ControlRigIndex;
-			Param.ControlRig = ControlRig;
-			Param.ControlName = ControlElement->GetFName();
-			Param.ShapeName = ControlElement->Settings.ShapeName;
-			Param.SpawnTransform = ControlRig->GetControlGlobalTransform(ControlElement->GetFName());
-			Param.ShapeTransform = ControlRig->GetHierarchy()->GetControlShapeTransform(ControlElement, ERigTransformType::CurrentLocal);
-			Param.bSelectable = ControlElement->Settings.IsSelectable(false);
+	const TArray<TSoftObjectPtr<UControlRigShapeLibrary>> ShapeLibraries = InControlRig->GetShapeLibraries();
 
-			if (const FControlRigShapeDefinition* ShapeDef = UControlRigShapeLibrary::GetShapeByName(ControlElement->Settings.ShapeName, ShapeLibraries, ControlRig->ShapeLibraryNameMap))
-			{
-				Param.MeshTransform = ShapeDef->Transform;
-				Param.StaticMesh = ShapeDef->StaticMesh;
-				Param.Material = ShapeDef->Library->DefaultMaterial;
-				if (bShowControlsAsOverlay)
-				{
-					TSoftObjectPtr<UMaterial> XRayMaterial = ShapeDef->Library->XRayMaterial;
-					if (XRayMaterial.IsPending())
-					{
-						XRayMaterial.LoadSynchronous();
-					}
-					if (XRayMaterial.IsValid())
-					{
-						Param.Material = XRayMaterial;
-					}
-				}
-				Param.ColorParameterName = ShapeDef->Library->MaterialColorParameter;
-			}
+	const int32 ControlRigIndex = RuntimeControlRigs.Find(InControlRig);
+	const URigHierarchy* Hierarchy = InControlRig->GetHierarchy();
 
-			Param.Color = ControlElement->Settings.ShapeColor;
+	// get controls for which shapes are needed in the editor
+	TArray<FRigControlElement*> Controls;
+	GetControlsEligibleForShapes(InControlRig, Controls);
 
-			AControlRigShapeActor* ShapeActor = FControlRigShapeHelper::CreateDefaultShapeActor(WorldPtr, Param);
-			if (ShapeActor)
-			{
-				//not drawn in game or in game view.
-				ShapeActor->SetActorHiddenInGame(true);
-				auto* ShapeActors = ControlRigShapeActors.Find(ControlRig);
-				if (ShapeActors)
+	// new shape actors to be created
+	TArray<AControlRigShapeActor*> NewShapeActors;
+	NewShapeActors.Reserve(Controls.Num());
+
+	for (FRigControlElement* ControlElement : Controls)
+	{
+		const FRigControlSettings& ControlSettings = ControlElement->Settings;
+		
+		FControlShapeActorCreationParam Param;
+		Param.ManipObj = InControlRig;
+		Param.ControlRigIndex = ControlRigIndex;
+		Param.ControlRig = InControlRig;
+		Param.ControlName = ControlElement->GetFName();
+		Param.ShapeName = ControlSettings.ShapeName;
+		Param.SpawnTransform = InControlRig->GetControlGlobalTransform(ControlElement->GetFName());
+		Param.ShapeTransform = Hierarchy->GetControlShapeTransform(ControlElement, ERigTransformType::CurrentLocal);
+		Param.bSelectable = ControlSettings.IsSelectable(false);
+
+		if (const FControlRigShapeDefinition* ShapeDef = UControlRigShapeLibrary::GetShapeByName(ControlSettings.ShapeName, ShapeLibraries, InControlRig->ShapeLibraryNameMap))
+		{
+			Param.MeshTransform = ShapeDef->Transform;
+			Param.StaticMesh = ShapeDef->StaticMesh;
+			Param.Material = ShapeDef->Library->DefaultMaterial;
+			if (bShowControlsAsOverlay)
+			{
+				TSoftObjectPtr<UMaterial> XRayMaterial = ShapeDef->Library->XRayMaterial;
+				if (XRayMaterial.IsPending())
 				{
-					ShapeActors->Add(ShapeActor);
+					XRayMaterial.LoadSynchronous();
 				}
-				else
+				if (XRayMaterial.IsValid())
 				{
-					TArray<AControlRigShapeActor*> NewShapeActors;
-					NewShapeActors.Add(ShapeActor);
-					ControlRigShapeActors.Add(ControlRig, ObjectPtrWrap(NewShapeActors));
+					Param.Material = XRayMaterial;
 				}
 			}
+			Param.ColorParameterName = ShapeDef->Library->MaterialColorParameter;
 		}
-	}
 
+		Param.Color = ControlSettings.ShapeColor;
+
+		// create a new shape actor that will represent that control in the editor
+		AControlRigShapeActor* NewShapeActor = FControlRigShapeHelper::CreateDefaultShapeActor(WorldPtr, Param);
+		if (NewShapeActor)
+		{
+			//not drawn in game or in game view.
+			NewShapeActor->SetActorHiddenInGame(true);
+			NewShapeActors.Add(NewShapeActor);
+		}
+	}
 
-	USceneComponent* Component = GetHostingSceneComponent(ControlRig);
-	if (Component)
+	// add or replace shape actors
+	auto* ShapeActors = ControlRigShapeActors.Find(InControlRig);
+	if (ShapeActors)
+	{
+		// this shouldn't happen but make sure we destroy any existing shape
+		DestroyShapesActorsFromWorld(*ShapeActors);
+		*ShapeActors = NewShapeActors;
+	}
+	else
 	{
-		AActor* PreviewActor = Component->GetOwner();
+		ShapeActors = &ControlRigShapeActors.Emplace(InControlRig, ObjectPtrWrap(NewShapeActors));
+	}
 
-		const auto* ShapeActors = ControlRigShapeActors.Find(ControlRig);
-		if (ShapeActors)
+	// setup shape actors
+	if(ensure(ShapeActors))
+	{
+		const USceneComponent* Component = GetHostingSceneComponent(InControlRig);
+		if (AActor* PreviewActor = Component ? Component->GetOwner() : nullptr)
 		{
 			for (AControlRigShapeActor* ShapeActor : *ShapeActors)
 			{
-				// attach to preview actor, so that we can communicate via relative transfrom from the previewactor
+				// attach to preview actor, so that we can communicate via relative transform from the preview actor
 				ShapeActor->AttachToActor(PreviewActor, FAttachmentTransformRules::KeepWorldTransform);
 
 				TInlineComponentArray<UPrimitiveComponent*> PrimitiveComponents;
@@ -3987,24 +4035,14 @@ void FControlRigEditMode::CreateShapeActors(UControlRig* ControlRig)
 			}
 		}
 	}
+
 	if (!AreEditingControlRigDirectly())
 	{
-
 		if (ControlProxy)
 		{
-			ControlProxy->RecreateAllProxies(ControlRig);
-		}
-	}
-	/** MZ got rid of this make sure it's okay
-	for (const FRigElementKey& SelectedElement : InSelectedElements)
-	{
-		if(FRigControlElement* ControlElement = ControlRig->FindControl(SelectedElement.Name))
-		{
-			OnHierarchyModified(ERigHierarchyNotification::ElementSelected, ControlRig->GetHierarchy(), ControlElement);
+			ControlProxy->RecreateAllProxies(InControlRig);
 		}
 	}
-	*/
-	
 }
 
 FControlRigEditMode* FControlRigEditMode::GetEditModeFromWorldContext(UWorld* InWorldContext)
@@ -4525,6 +4563,28 @@ bool FControlRigEditMode::CanChangeControlShapeTransform()
 	return false;
 }
 
+void FControlRigEditMode::OnSettingsChanged(const UControlRigEditModeSettings* InSettings)
+{
+	if (!InSettings)
+	{
+		return;
+	}
+	
+	// check if the settings for xray rendering are different for any of the control shape actors
+	if(bShowControlsAsOverlay != InSettings->bShowControlsAsOverlay)
+	{
+		bShowControlsAsOverlay = InSettings->bShowControlsAsOverlay;
+		for (TWeakObjectPtr<UControlRig>& RuntimeRigPtr : RuntimeControlRigs)
+		{
+			if (UControlRig* RuntimeControlRig = RuntimeRigPtr.Get())
+			{
+				UpdateSelectabilityOnSkeletalMeshes(RuntimeControlRig, !bShowControlsAsOverlay);
+			}
+		}
+		RequestToRecreateControlShapeActors();
+	}
+}
+
 void FControlRigEditMode::SetControlShapeTransform(
 	const AControlRigShapeActor* InShapeActor,
 	const FTransform& InGlobalTransform,
@@ -4942,13 +5002,13 @@ bool FControlRigEditMode::ModeSupportedByShapeActor(const AControlRigShapeActor*
 	return false;
 }
 
-bool FControlRigEditMode::IsControlRigSkelMeshVisible(UControlRig* ControlRig) const
+bool FControlRigEditMode::IsControlRigSkelMeshVisible(const UControlRig* InControlRig) const
 {
 	if (IsInLevelEditor())
 	{
-		if (ControlRig)
+		if (InControlRig)
 		{
-			if (USceneComponent* SceneComponent = GetHostingSceneComponent(ControlRig))
+			if (const USceneComponent* SceneComponent = GetHostingSceneComponent(InControlRig))
 			{
 				const AActor* Actor = SceneComponent->GetTypedOuter<AActor>();
 				return Actor ? (Actor->IsHiddenEd() == false && SceneComponent->IsVisibleInEditor()) : SceneComponent->IsVisibleInEditor();
@@ -5284,53 +5344,35 @@ void FControlRigEditMode::FMarqueeDragTool::RenderDragTool(const FSceneView* Vie
 	}
 }
 
-void FControlRigEditMode::DestroyShapesActors(UControlRig* ControlRig)
+void FControlRigEditMode::DestroyShapesActors(UControlRig* InControlRig)
 {
-	if (ControlRig == nullptr)
+	using namespace ControlRigEditMode::Shapes;
+	
+	if (!InControlRig)
 	{
+		// destroy all control rigs shape actors
 		for(auto& ShapeActors: ControlRigShapeActors)
 		{
-			for (AControlRigShapeActor* ShapeActor : ShapeActors.Value)
-			{
-				UWorld* World = ShapeActor->GetWorld();
-				if (World)
-				{
-					ShapeActor->UnregisterAllComponents();
-					if (ShapeActor->GetAttachParentActor())
-					{
-						ShapeActor->DetachFromActor(FDetachmentTransformRules::KeepRelativeTransform);
-					}
-					World->EditorDestroyActor(ShapeActor, true);
-				}
-			}
+			DestroyShapesActorsFromWorld(ShapeActors.Value);
 		}
+		
 		ControlRigShapeActors.Reset();
 		ControlRigsToRecreate.Reset();
+		
 		if (OnWorldCleanupHandle.IsValid())
 		{
 			FWorldDelegates::OnWorldCleanup.Remove(OnWorldCleanupHandle);
 		}
+		
+		return;
 	}
-	else
+
+	// only destroy control rigs shape actors related to InControlRig
+	ControlRigsToRecreate.Remove(InControlRig);
+	if (const auto* ShapeActors = ControlRigShapeActors.Find(InControlRig))
 	{
-		ControlRigsToRecreate.Remove(ControlRig);
-		const auto* ShapeActors = ControlRigShapeActors.Find(ControlRig);
-		if (ShapeActors)
-		{
-			for (AControlRigShapeActor* ShapeActor : *ShapeActors)
-			{
-				UWorld* World = ShapeActor->GetWorld();
-				if (World)
-				{
-					if (ShapeActor->GetAttachParentActor())
-					{
-						ShapeActor->DetachFromActor(FDetachmentTransformRules::KeepRelativeTransform);
-					}
-					World->EditorDestroyActor(ShapeActor,true);
-				}
-			}
-			ControlRigShapeActors.Remove(ControlRig);
-		}
+		DestroyShapesActorsFromWorld(*ShapeActors);
+		ControlRigShapeActors.Remove(InControlRig);
 	}
 }
 
@@ -5381,7 +5423,7 @@ FTransform FControlRigEditMode::GetHostingSceneComponentTransform(const UControl
 	{
 		ControlRig = GetControlRigs()[0].Get();
 	}
-	USceneComponent* HostingComponent = GetHostingSceneComponent(ControlRig);
+	const USceneComponent* HostingComponent = GetHostingSceneComponent(ControlRig);
 	return HostingComponent ? HostingComponent->GetComponentTransform() : FTransform::Identity;
 }
 
@@ -5740,5 +5782,105 @@ void FDetailKeyFrameCacheAndHandler::UpdateIfDirty()
 	}
 }
 
+namespace ControlRigEditMode::Shapes
+{
+
+void GetControlsEligibleForShapes(UControlRig* InControlRig, TArray<FRigControlElement*>& OutControls)
+{
+	OutControls.Reset();
+
+	const URigHierarchy* Hierarchy = InControlRig ? InControlRig->GetHierarchy() : nullptr;
+	if (!Hierarchy)
+	{
+		return;
+	}
+
+	OutControls = Hierarchy->GetFilteredElements<FRigControlElement>([](const FRigControlElement* ControlElement)
+	{
+		const FRigControlSettings& ControlSettings = ControlElement->Settings;
+		return ControlSettings.SupportsShape() && IsSupportedControlType(ControlSettings.ControlType);
+	});
+}
+
+void DestroyShapesActorsFromWorld(const TArray<TObjectPtr<AControlRigShapeActor>>& InShapeActorsToDestroy)
+{
+	// NOTE: should UWorld::EditorDestroyActor really modify the level when removing the shapes?
+	// kept for legacy but I guess this should be set to false
+	static constexpr bool bShouldModifyLevel = true;
+
+	for (const TObjectPtr<AControlRigShapeActor>& ShapeActorPtr: InShapeActorsToDestroy)
+	{
+		if (AControlRigShapeActor* ShapeActor = ShapeActorPtr.Get())
+		{
+			if (UWorld* World = ShapeActor->GetWorld())
+			{
+				if (ShapeActor->GetAttachParentActor())
+				{
+					ShapeActor->DetachFromActor(FDetachmentTransformRules::KeepRelativeTransform);
+				}
+				World->EditorDestroyActor(ShapeActor,bShouldModifyLevel);
+			}			
+		}
+	}
+}
+
+FShapeUpdateParams::FShapeUpdateParams(const UControlRig* InControlRig, const FTransform& InComponentTransform, const bool InSkeletalMeshVisible)
+	: ControlRig(InControlRig)
+	, Hierarchy(InControlRig->GetHierarchy())
+	, Settings(GetDefault<UControlRigEditModeSettings>())
+	, ComponentTransform(InComponentTransform)
+	, bIsSkeletalMeshVisible(InSkeletalMeshVisible)
+{}
+
+bool FShapeUpdateParams::IsValid() const
+{
+	return ControlRig && Hierarchy && Settings;
+}
+
+void UpdateControlShape(AControlRigShapeActor* InShapeActor, FRigControlElement* InControlElement, const FShapeUpdateParams& InUpdateParams)
+{
+	if (!InShapeActor || !InControlElement || !InUpdateParams.IsValid())
+	{
+		return;
+	}
+
+	// update transform
+	const FTransform Transform = InUpdateParams.Hierarchy->GetTransform(InControlElement, ERigTransformType::CurrentGlobal);
+	InShapeActor->SetActorTransform(Transform * InUpdateParams.ComponentTransform);
+
+	const FRigControlSettings& ControlSettings = InControlElement->Settings;
+
+	// update visibility & color
+	bool bIsVisible = ControlSettings.IsVisible();
+	bool bRespectVisibilityForSelection = true; 
+
+	const bool bControlsHiddenInViewport =
+		InUpdateParams.Settings->bHideControlShapes ||
+		!InUpdateParams.ControlRig->GetControlsVisible() ||
+		!InUpdateParams.bIsSkeletalMeshVisible;
+	
+	if (!bControlsHiddenInViewport)
+	{
+		if (ControlSettings.AnimationType == ERigControlAnimationType::ProxyControl)
+		{					
+			bRespectVisibilityForSelection = false;
+			if (InUpdateParams.Settings->bShowAllProxyControls)
+			{
+				bIsVisible = true;
+			}
+		}
+	}
+
+	InShapeActor->SetIsTemporarilyHiddenInEditor(!bIsVisible || bControlsHiddenInViewport);
+
+	// update color
+	InShapeActor->SetShapeColor(InShapeActor->OverrideColor.A < SMALL_NUMBER ?
+		ControlSettings.ShapeColor : InShapeActor->OverrideColor);
+	
+	// update selectability
+	InShapeActor->SetSelectable( ControlSettings.IsSelectable(bRespectVisibilityForSelection) );
+}
+	
+}
 
 #undef LOCTEXT_NAMESPACE
diff --git a/Engine/Plugins/Animation/ControlRig/Source/ControlRigEditor/Private/EditMode/ControlRigEditModeSettings.cpp b/Engine/Plugins/Animation/ControlRig/Source/ControlRigEditor/Private/EditMode/ControlRigEditModeSettings.cpp
--- a/Engine/Plugins/Animation/ControlRig/Source/ControlRigEditor/Private/EditMode/ControlRigEditModeSettings.cpp
+++ b/Engine/Plugins/Animation/ControlRig/Source/ControlRigEditor/Private/EditMode/ControlRigEditModeSettings.cpp
@@ -6,6 +6,8 @@
 
 #include UE_INLINE_GENERATED_CPP_BY_NAME(ControlRigEditModeSettings)
 
+UControlRigEditModeSettings::FOnUpdateSettings UControlRigEditModeSettings::OnSettingsChange;
+
 void UControlRigEditModeSettings::PreEditChange(FProperty* PropertyAboutToChange)
 {
 	Super::PreEditChange(PropertyAboutToChange);
@@ -25,12 +27,14 @@ void UControlRigEditModeSettings::PostEditChangeProperty(struct FPropertyChanged
 		// Dragging spinboxes causes this to be called every frame so we wait until they've finished dragging before saving.
 		SaveConfig();
 	}
+	OnSettingsChange.Broadcast(this);
 }
 
 #if WITH_EDITOR
 void UControlRigEditModeSettings::PostEditUndo()
 {
 	GizmoScaleDelegate.Broadcast(GizmoScale);
+	OnSettingsChange.Broadcast(this);
 }
 #endif
 
diff --git a/Engine/Plugins/Animation/ControlRig/Source/ControlRigEditor/Private/EditMode/ControlRigEditModeSettings.h b/Engine/Plugins/Animation/ControlRig/Source/ControlRigEditor/Private/EditMode/ControlRigEditModeSettings.h
--- a/Engine/Plugins/Animation/ControlRig/Source/ControlRigEditor/Private/EditMode/ControlRigEditModeSettings.h
+++ b/Engine/Plugins/Animation/ControlRig/Source/ControlRigEditor/Private/EditMode/ControlRigEditModeSettings.h
@@ -98,4 +98,7 @@ public:
 
 	/** Delegate broadcasted whenever GizmoScale is modified */
 	FGizmoScaleSet GizmoScaleDelegate;
+
+	DECLARE_MULTICAST_DELEGATE_OneParam(FOnUpdateSettings, const UControlRigEditModeSettings*);
+	static FOnUpdateSettings OnSettingsChange;
 };
\ No newline at end of file
diff --git a/Engine/Plugins/Animation/ControlRig/Source/ControlRigEditor/Private/Editor/ControlRigEditor.cpp b/Engine/Plugins/Animation/ControlRig/Source/ControlRigEditor/Private/Editor/ControlRigEditor.cpp
--- a/Engine/Plugins/Animation/ControlRig/Source/ControlRigEditor/Private/Editor/ControlRigEditor.cpp
+++ b/Engine/Plugins/Animation/ControlRig/Source/ControlRigEditor/Private/Editor/ControlRigEditor.cpp
@@ -955,7 +955,7 @@ void FControlRigEditor::SetEventQueue(TArray<FName> InEventQueue, bool bCompile)
 
 	if (FControlRigEditMode* EditMode = GetEditMode())
 	{
-		EditMode->RecreateControlShapeActors(GetHierarchyBeingDebugged()->GetSelectedKeys());
+		EditMode->RecreateControlShapeActors();
 
 		UControlRigEditModeSettings* Settings = GetMutableDefault<UControlRigEditModeSettings>();
 		Settings->bDisplayNulls = IsConstructionModeEnabled();
@@ -1424,7 +1424,7 @@ void FControlRigEditor::Compile()
 
 		if (FControlRigEditMode* EditMode = GetEditMode())
 		{
-			EditMode->RecreateControlShapeActors(GetHierarchyBeingDebugged()->GetSelectedKeys());
+			EditMode->RecreateControlShapeActors();
 		}
 	}
 }
diff --git a/Engine/Plugins/Animation/ControlRig/Source/ControlRigEditor/Private/Sequencer/ControlRigParameterTrackEditor.cpp b/Engine/Plugins/Animation/ControlRig/Source/ControlRigEditor/Private/Sequencer/ControlRigParameterTrackEditor.cpp
--- a/Engine/Plugins/Animation/ControlRig/Source/ControlRigEditor/Private/Sequencer/ControlRigParameterTrackEditor.cpp
+++ b/Engine/Plugins/Animation/ControlRig/Source/ControlRigEditor/Private/Sequencer/ControlRigParameterTrackEditor.cpp
@@ -2165,8 +2165,11 @@ void FControlRigParameterTrackEditor::OnSelectionChanged(TArray<UMovieSceneTrack
 				{
 					if (ControlRigEditMode)
 					{
-						ControlRigEditMode->AddControlRigObject(TrackControlRig, GetSequencer());
-						ControlRigEditMode->RequestToRecreateControlShapeActors(TrackControlRig);
+						const bool bAdded = ControlRigEditMode->AddControlRigObject(TrackControlRig, GetSequencer());
+						if (bAdded)
+						{
+							ControlRigEditMode->RequestToRecreateControlShapeActors(TrackControlRig);
+						}
 						break;
 					}
 					else
@@ -2176,8 +2179,11 @@ void FControlRigParameterTrackEditor::OnSelectionChanged(TArray<UMovieSceneTrack
 						{
 							if (ControlRigEditMode)
 							{
-								ControlRigEditMode->AddControlRigObject(TrackControlRig, GetSequencer());
-								ControlRigEditMode->RequestToRecreateControlShapeActors(TrackControlRig);
+								const bool bAdded = ControlRigEditMode->AddControlRigObject(TrackControlRig, GetSequencer());
+								if (bAdded)
+								{
+									ControlRigEditMode->RequestToRecreateControlShapeActors(TrackControlRig);
+								}
 							}
 						}
 					}
diff --git a/Engine/Plugins/Animation/ControlRig/Source/ControlRigEditor/Public/EditMode/ControlRigEditMode.h b/Engine/Plugins/Animation/ControlRig/Source/ControlRigEditor/Public/EditMode/ControlRigEditMode.h
--- a/Engine/Plugins/Animation/ControlRig/Source/ControlRigEditor/Public/EditMode/ControlRigEditMode.h
+++ b/Engine/Plugins/Animation/ControlRig/Source/ControlRigEditor/Public/EditMode/ControlRigEditMode.h
@@ -46,6 +46,7 @@ struct FGizmoState;
 enum class EMovieSceneDataChangeType;
 struct FMovieSceneChannelMetaData;
 class UMovieSceneSection;
+class UControlRigEditModeSettings;
 
 DECLARE_DELEGATE_RetVal_ThreeParams(FTransform, FOnGetRigElementTransform, const FRigElementKey& /*RigElementKey*/, bool /*bLocal*/, bool /*bOnDebugInstance*/);
 DECLARE_DELEGATE_ThreeParams(FOnSetRigElementTransform, const FRigElementKey& /*RigElementKey*/, const FTransform& /*Transform*/, bool /*bLocal*/);
@@ -134,7 +135,7 @@ public:
 	void SetObjects(UControlRig* InControlRig, UObject* BindingObject, TWeakPtr<ISequencer> InSequencer);
 
 	/** Add a Control Rig object if it doesn't exist, will return true if it was added, false if it wasn't since it's already there. You can also set the Sequencer.*/
-	bool AddControlRigObject(UControlRig* InControlRig, TWeakPtr<ISequencer> InSequencer);
+	bool AddControlRigObject(UControlRig* InControlRig, const TWeakPtr<ISequencer>& InSequencer);
 
 	/* Remove control rig */
 	void RemoveControlRig(UControlRig* InControlRig);
@@ -328,7 +329,7 @@ private:
 	void BindCommands();
 
 	/** It creates if it doesn't have it */
-	void RecreateControlShapeActors(const TArray<FRigElementKey>& InSelectedElements = TArray<FRigElementKey>());
+	void RecreateControlShapeActors();
 
 	/** Let the preview scene know how we want to select components */
 	bool ShapeSelectionOverride(const UPrimitiveComponent* InComponent) const;
@@ -468,6 +469,9 @@ private:
 	bool bIsChangingCoordSystem;
 
 	bool CanChangeControlShapeTransform();
+
+	void OnSettingsChanged(const UControlRigEditModeSettings* InSettings);
+	
 public:
 	//Toolbar functions
 	void SetOnlySelectRigControls(bool val);
@@ -477,9 +481,10 @@ public:
 private:
 	TSet<FName> GetActiveControlsFromSequencer(UControlRig* ControlRig);
 
-	/** Create/Delete for the specified ControlRig*/
+	/** Create/Delete/Update shape actors for the specified ControlRig */
 	void CreateShapeActors(UControlRig* InControlRig);
 	void DestroyShapesActors(UControlRig* InControlRig);
+	bool TryUpdatingControlsShapes(UControlRig* InControlRig);
 
 	/*Internal function for adding ControlRig*/
 	void AddControlRigInternal(UControlRig* InControlRig);
@@ -510,7 +515,7 @@ protected:
 	FTransform	GetHostingSceneComponentTransform(const UControlRig* ControlRig =  nullptr) const;
 
 	//Get if the hosted component is visible
-	bool IsControlRigSkelMeshVisible(UControlRig* ControlRig) const;
+	bool IsControlRigSkelMeshVisible(const UControlRig* InControlRig) const;
 	
 public:  
 		TSharedPtr<FDetailKeyFrameCacheAndHandler> DetailKeyFrameCache;
@@ -589,3 +594,27 @@ private:
 	friend class UControlRigEditModeDelegateHelper;
 	friend class SControlRigEditModeTools;
 };
+
+namespace ControlRigEditMode::Shapes
+{
+	// Returns the list of controls for which a shape is expected
+	void GetControlsEligibleForShapes(UControlRig* InControlRig, TArray<FRigControlElement*>& OutControls);
+
+	// Destroys shape actors and removes them from their UWorld
+	void DestroyShapesActorsFromWorld(const TArray<TObjectPtr<AControlRigShapeActor>>& InShapeActorsToDestroy);
+
+	// Parameters used to update shape actors (transform, visibility, etc.)
+	struct FShapeUpdateParams
+	{
+		FShapeUpdateParams(const UControlRig* InControlRig, const FTransform& InComponentTransform, const bool InSkeletalMeshVisible);
+		const UControlRig* ControlRig = nullptr;
+		const URigHierarchy* Hierarchy = nullptr;
+		const UControlRigEditModeSettings* Settings = nullptr;
+		const FTransform& ComponentTransform = FTransform::Identity;
+		const bool bIsSkeletalMeshVisible = false;
+		bool IsValid() const;
+	};
+
+	// Updates shape actors transform, visibility, etc.
+	void UpdateControlShape(AControlRigShapeActor* InShapeActor, FRigControlElement* InControlElement, const FShapeUpdateParams& InUpdateParams);
+}
\ No newline at end of file
