Fixed an editor shutdown crash due to shared pointer cross dependency at ControlRigModeToos (the widget is not getting released due to children details holding a shared pointer to parent)

diff --git a/Engine/Plugins/Animation/ControlRig/Source/ControlRigEditor/Private/EditMode/ControlRigEditModeToolkit.cpp b/Engine/Plugins/Animation/ControlRig/Source/ControlRigEditor/Private/EditMode/ControlRigEditModeToolkit.cpp
--- a/Engine/Plugins/Animation/ControlRig/Source/ControlRigEditor/Private/EditMode/ControlRigEditModeToolkit.cpp
+++ b/Engine/Plugins/Animation/ControlRig/Source/ControlRigEditor/Private/EditMode/ControlRigEditModeToolkit.cpp
@@ -49,6 +49,14 @@ const FName FControlRigEditModeToolkit::SpacePickerTabName = FName(TEXT("Control
 TSharedPtr<SControlRigDetails> FControlRigEditModeToolkit::Details = nullptr;
 TSharedPtr<SControlRigOutliner> FControlRigEditModeToolkit::Outliner = nullptr;
 
+FControlRigEditModeToolkit::~FControlRigEditModeToolkit()
+{
+	if (ModeTools)
+	{
+		ModeTools->Cleanup();
+	}
+}
+
 void FControlRigEditModeToolkit::Init(const TSharedPtr<IToolkitHost>& InitToolkitHost)
 {
 	SAssignNew(ModeTools, SControlRigEditModeTools, SharedThis(this), EditMode);
diff --git a/Engine/Plugins/Animation/ControlRig/Source/ControlRigEditor/Private/EditMode/ControlRigEditModeToolkit.h b/Engine/Plugins/Animation/ControlRig/Source/ControlRigEditor/Private/EditMode/ControlRigEditModeToolkit.h
--- a/Engine/Plugins/Animation/ControlRig/Source/ControlRigEditor/Private/EditMode/ControlRigEditModeToolkit.h
+++ b/Engine/Plugins/Animation/ControlRig/Source/ControlRigEditor/Private/EditMode/ControlRigEditModeToolkit.h
@@ -26,6 +26,8 @@ public:
 		
 	}
 
+	~FControlRigEditModeToolkit();
+
 	/** IToolkit interface */
 	virtual FName GetToolkitFName() const override { return FName("AnimationMode"); }
 	virtual FText GetBaseToolkitName() const override { return NSLOCTEXT("AnimationModeToolkit", "DisplayName", "Animation"); }
diff --git a/Engine/Plugins/Animation/ControlRig/Source/ControlRigEditor/Private/EditMode/SControlRigEditModeTools.cpp b/Engine/Plugins/Animation/ControlRig/Source/ControlRigEditor/Private/EditMode/SControlRigEditModeTools.cpp
--- a/Engine/Plugins/Animation/ControlRig/Source/ControlRigEditor/Private/EditMode/SControlRigEditModeTools.cpp
+++ b/Engine/Plugins/Animation/ControlRig/Source/ControlRigEditor/Private/EditMode/SControlRigEditModeTools.cpp
@@ -47,6 +47,19 @@
 //statics to reuse in the UI
 FRigSpacePickerBakeSettings SControlRigEditModeTools::BakeSpaceSettings;
 
+void SControlRigEditModeTools::Cleanup()
+{
+	// This is required as these hold a shared pointer to THIS OBJECT and make this class not to be destroyed when the parent class releases the shared pointer of this object
+	if (SettingsDetailsView)
+	{
+		SettingsDetailsView->SetKeyframeHandler(nullptr);
+	}
+	if (RigOptionsDetailsView)
+	{
+		RigOptionsDetailsView->SetKeyframeHandler(nullptr);
+	}
+}
+
 void SControlRigEditModeTools::SetControlRigs(const TArrayView<TWeakObjectPtr<UControlRig>>& InControlRigs)
 {
 	for (TWeakObjectPtr<UControlRig>& ControlRig : ControlRigs)
@@ -66,15 +79,17 @@ void SControlRigEditModeTools::SetControlRigs(const TArrayView<TWeakObjectPtr<UC
 	}
 
 	//mz todo handle multiple rigs
-	UControlRig* Rig = ControlRigs.Num() > 0 ? ControlRigs[0].Get() : nullptr;
 	TArray<TWeakObjectPtr<>> Objects;
-	Objects.Add(Rig);
+	if (ControlRigs.Num() > 0)
+	{
+		UControlRig* Rig = ControlRigs[0].Get();
+		Objects.Add(Rig);
+	}
 	RigOptionsDetailsView->SetObjects(Objects);
 
 #if USE_LOCAL_DETAILS
 	HierarchyTreeView->RefreshTreeView(true);
 #endif
-
 }
 
 const URigHierarchy* SControlRigEditModeTools::GetHierarchy() const
@@ -482,7 +497,10 @@ void SControlRigEditModeTools::SetSettingsDetailsObject(const TWeakObjectPtr<>&
 	if (SettingsDetailsView)
 	{
 		TArray<TWeakObjectPtr<>> Objects;
-		Objects.Add(InObject);
+		if (InObject.IsValid())
+		{
+			Objects.Add(InObject);
+		}
 		SettingsDetailsView->SetObjects(Objects);
 	}
 }
diff --git a/Engine/Plugins/Animation/ControlRig/Source/ControlRigEditor/Private/EditMode/SControlRigEditModeTools.h b/Engine/Plugins/Animation/ControlRig/Source/ControlRigEditor/Private/EditMode/SControlRigEditModeTools.h
--- a/Engine/Plugins/Animation/ControlRig/Source/ControlRigEditor/Private/EditMode/SControlRigEditModeTools.h
+++ b/Engine/Plugins/Animation/ControlRig/Source/ControlRigEditor/Private/EditMode/SControlRigEditModeTools.h
@@ -32,6 +32,8 @@ public:
 	SLATE_END_ARGS();
 
 	void Construct(const FArguments& InArgs, TSharedPtr<FControlRigEditModeToolkit> InOwningToolkit, FControlRigEditMode& InEditMode);
+	void Cleanup();
+
 	/** Set the objects to be displayed in the details panel */
 	void SetSettingsDetailsObject(const TWeakObjectPtr<>& InObject);
 #if USE_LOCAL_DETAILS
