Control Rig: Fix crash when changing preview mesh
#rnx

diff --git a/Engine/Plugins/Animation/ControlRig/Source/ControlRigEditor/Private/Editor/ControlRigSkeletalMeshComponent.cpp b/Engine/Plugins/Animation/ControlRig/Source/ControlRigEditor/Private/Editor/ControlRigSkeletalMeshComponent.cpp
--- a/Engine/Plugins/Animation/ControlRig/Source/ControlRigEditor/Private/Editor/ControlRigSkeletalMeshComponent.cpp
+++ b/Engine/Plugins/Animation/ControlRig/Source/ControlRigEditor/Private/Editor/ControlRigSkeletalMeshComponent.cpp
@@ -97,6 +97,13 @@ void UControlRigSkeletalMeshComponent::RebuildDebugDrawSkeleton()
 
 			TMap<FName, int32> AddedBoneMap;
 			TArray<FRigBoneElement*> BoneElements = Hierarchy->GetBones(true);
+			for(FRigBoneElement* BoneElement : BoneElements)
+			{
+				AddedBoneMap.FindOrAdd(BoneElement->GetFName(), DebugDrawBones.Num());
+				DebugDrawBones.Add(DebugDrawBones.Num());
+				DebugDrawBoneIndexInHierarchy.Add(BoneElement->GetIndex());
+			}
+
 			for(FRigBoneElement* BoneElement : BoneElements)
 			{
 				const int32 Index = BoneElement->GetIndex();
@@ -134,10 +141,6 @@ void UControlRigSkeletalMeshComponent::RebuildDebugDrawSkeleton()
 				NewMeshBoneInfo.ParentIndex = ParentIndex; 
 				// give ref pose here
 				RefSkelModifier.Add(NewMeshBoneInfo, Hierarchy->GetInitialGlobalTransform(Index), true);
-
-				AddedBoneMap.FindOrAdd(BoneElement->GetFName(), DebugDrawBones.Num());
-				DebugDrawBones.Add(DebugDrawBones.Num());
-				DebugDrawBoneIndexInHierarchy.Add(BoneElement->GetIndex());
 			}
 		}
 	}
