Control Rig: Sockets not being created in the correct position
#rnx

diff --git a/Engine/Plugins/Animation/ControlRig/Source/ControlRigEditor/Private/Editor/SRigHierarchy.cpp b/Engine/Plugins/Animation/ControlRig/Source/ControlRigEditor/Private/Editor/SRigHierarchy.cpp
--- a/Engine/Plugins/Animation/ControlRig/Source/ControlRigEditor/Private/Editor/SRigHierarchy.cpp
+++ b/Engine/Plugins/Animation/ControlRig/Source/ControlRigEditor/Private/Editor/SRigHierarchy.cpp
@@ -1953,6 +1953,7 @@ void SRigHierarchy::HandleNewItem(ERigElementType InElementType, bool bIsAnimati
 
 	FRigElementKey NewItemKey;
 	URigHierarchy* Hierarchy = GetDefaultHierarchy();
+	URigHierarchy* DebugHierarchy = GetHierarchy();
 	if (Hierarchy)
 	{
 		// unselect current selected item
@@ -1985,7 +1986,8 @@ void SRigHierarchy::HandleNewItem(ERigElementType InElementType, bool bIsAnimati
 			if(SelectedKey.IsValid())
 			{
 				ParentKey = SelectedKey;
-				ParentTransform = Hierarchy->GetGlobalTransform(ParentKey);
+				// Use the transform of the debugged hierarchy rather than the default hierarchy
+				ParentTransform = DebugHierarchy->GetGlobalTransform(ParentKey);
 			}
 
 			// use bone's name as prefix if creating a control
