Control: Fix labels in control outliner and hide irrelevant items

#rnx

diff --git a/Engine/Plugins/Animation/ControlRig/Source/ControlRigEditor/Private/EditMode/SControlRigOutliner.cpp b/Engine/Plugins/Animation/ControlRig/Source/ControlRigEditor/Private/EditMode/SControlRigOutliner.cpp
--- a/Engine/Plugins/Animation/ControlRig/Source/ControlRigEditor/Private/EditMode/SControlRigOutliner.cpp
+++ b/Engine/Plugins/Animation/ControlRig/Source/ControlRigEditor/Private/EditMode/SControlRigOutliner.cpp
@@ -88,11 +88,20 @@ FText FMultiRigData::GetDisplayName() const
 	{
 		if(ControlRig.IsValid())
 		{
-			if(const FRigControlElement* ControlElement = ControlRig->GetHierarchy()->Find<FRigControlElement>(Key.GetValue()))
+			if(const URigHierarchy* Hierarchy = ControlRig->GetHierarchy())
 			{
-				if(!ControlElement->Settings.DisplayName.IsNone())
+				const FText DisplayNameForUI = Hierarchy->GetDisplayNameForUI(Key.GetValue(), false);
+				if(!DisplayNameForUI.IsEmpty())
 				{
-					return FText::FromName(ControlElement->Settings.DisplayName);
+					return DisplayNameForUI;
+				}
+				
+				if(const FRigControlElement* ControlElement = Hierarchy->Find<FRigControlElement>(Key.GetValue()))
+				{
+					if(!ControlElement->Settings.DisplayName.IsNone())
+					{
+						return FText::FromName(ControlElement->Settings.DisplayName);
+					}
 				}
 			}
 		}
@@ -538,6 +547,14 @@ bool SMultiRigHierarchyTreeView::AddElement(UControlRig* InControlRig, const FRi
 				}
 				break;
 			}
+			case ERigElementType::Connector:
+			{
+				if (!Settings.bShowConnectors)
+				{
+					return false;
+				}
+				break;
+			}
 			case ERigElementType::Curve:
 			{
 				return false;
@@ -964,6 +981,7 @@ void SControlRigOutliner::Construct(const FArguments& InArgs, FControlRigEditMod
 	DisplaySettings.bShowRigidBodies = false;
 	DisplaySettings.bHideParentsOnFilter = true;
 	DisplaySettings.bFlattenHierarchyOnFilter = true;
+	DisplaySettings.bShowConnectors = false;
 
 	FMultiRigTreeDelegates RigTreeDelegates;
 	RigTreeDelegates.OnGetDisplaySettings = FOnGetRigTreeDisplaySettings::CreateSP(this, &SControlRigOutliner::GetDisplaySettings);
diff --git a/Engine/Plugins/Animation/ControlRig/Source/ControlRigEditor/Private/Editor/SRigHierarchyTreeView.h b/Engine/Plugins/Animation/ControlRig/Source/ControlRigEditor/Private/Editor/SRigHierarchyTreeView.h
--- a/Engine/Plugins/Animation/ControlRig/Source/ControlRigEditor/Private/Editor/SRigHierarchyTreeView.h
+++ b/Engine/Plugins/Animation/ControlRig/Source/ControlRigEditor/Private/Editor/SRigHierarchyTreeView.h
@@ -28,6 +28,7 @@ struct CONTROLRIGEDITOR_API FRigTreeDisplaySettings
 		bShowRigidBodies = true;
 		bShowReferences = true;
 		bShowSockets = true;
+		bShowConnectors = true;
 		bShowIconColors = true;
 	}
 	
@@ -63,6 +64,9 @@ struct CONTROLRIGEDITOR_API FRigTreeDisplaySettings
 	/** Whether or not to show sockets in the hierarchy */
 	bool bShowSockets;
 
+	/** Whether or not to show connectors in the hierarchy */
+	bool bShowConnectors;
+	
 	/** Whether to tint the icons with the element color */
 	bool bShowIconColors;
 };
