Sequencer: Anim Details: Per animator feedback we only select curves on the section to key section when there are multiple sections.

diff --git a/Engine/Plugins/Animation/ControlRig/Source/ControlRigEditor/Private/EditMode/AnimDetailsProxy.cpp b/Engine/Plugins/Animation/ControlRig/Source/ControlRigEditor/Private/EditMode/AnimDetailsProxy.cpp
--- a/Engine/Plugins/Animation/ControlRig/Source/ControlRigEditor/Private/EditMode/AnimDetailsProxy.cpp
+++ b/Engine/Plugins/Animation/ControlRig/Source/ControlRigEditor/Private/EditMode/AnimDetailsProxy.cpp
@@ -126,6 +126,17 @@ static EAnimDetailSelectionState CachePropertySelection(TWeakPtr<FCurveEditor>&
 									{
 										continue;
 									}
+									if (UMovieSceneControlRigParameterTrack* Track = Section->GetTypedOuter< UMovieSceneControlRigParameterTrack>())
+									{
+										if (Track->GetSectionToKey() != Section)
+										{
+											continue;
+										}
+									}
+									else
+									{
+										continue;
+									}
 									UObject* ColorObject = nullptr;
 									FString Name;
 									Pair.Value->GetCurveColorObjectAndName(&ColorObject, Name);
diff --git a/Engine/Plugins/Animation/ControlRig/Source/ControlRigEditor/Private/EditMode/ControlsProxyDetailCustomization.cpp b/Engine/Plugins/Animation/ControlRig/Source/ControlRigEditor/Private/EditMode/ControlsProxyDetailCustomization.cpp
--- a/Engine/Plugins/Animation/ControlRig/Source/ControlRigEditor/Private/EditMode/ControlsProxyDetailCustomization.cpp
+++ b/Engine/Plugins/Animation/ControlRig/Source/ControlRigEditor/Private/EditMode/ControlsProxyDetailCustomization.cpp
@@ -19,6 +19,7 @@
 #include "PropertyEditorModule.h"
 #include "HAL/PlatformApplicationMisc.h"
 #include "Sequencer/MovieSceneControlRigParameterTrack.h"
+#include "Sequencer/MovieSceneControlRigParameterSection.h"
 
 #include "Widgets/SBoxPanel.h"
 #include "Widgets/Layout/SBorder.h"
@@ -34,6 +35,7 @@
 #include "MVVM/ViewModels/ViewModelIterators.h"
 #include "MVVM/Extensions/IOutlinerExtension.h"
 #include "MVVM/ViewModels/TrackModel.h"
+#include "MVVM/ViewModels/ChannelModel.h"
 #include "Rigs/RigControlHierarchy.h"
 #include "Tree/SCurveEditorTree.h"
 
@@ -566,12 +568,24 @@ void FAnimDetailValueCustomization::TogglePropertySelection(UControlRigControlsP
 						{
 							if (TSharedPtr<FTrackModel> TrackModel = OutlinerExtenstionIt.GetCurrentItem()->FindAncestorOfType<FTrackModel>())
 							{
-								if (UMovieSceneControlRigParameterTrack* Track =  Cast<UMovieSceneControlRigParameterTrack>(TrackModel->GetTrack()) )
+								if(UMovieSceneControlRigParameterTrack* Track = Cast<UMovieSceneControlRigParameterTrack>(TrackModel->GetTrack()))
 								{
 									if (Track->GetControlRig() != ControlRig)
 									{
 										continue;
 									}
+									if (TViewModelPtr<FChannelGroupOutlinerModel> ChannelModel = CastViewModel<FChannelGroupOutlinerModel>(OutlinerExtenstionIt.GetCurrentItem()))
+									{
+										if (ChannelModel->GetChannel(Track->GetSectionToKey()) == nullptr) //if not section to key we also don't select it.
+										{
+											continue;
+										}
+									}
+									else
+									{
+										continue;
+									}
+									
 									FName ID = OutlinerExtenstionIt->GetIdentifier();
 									FString Name = ID.ToString();
 									TArray<FString> StringArray;
