Control Rig: Avoid referencing game thread animinstance functions in non-game threads. Added some safety in case we end up creating a UControlRig object on a non-game thread.

#rnx

diff --git a/Engine/Plugins/Animation/ControlRig/Source/ControlRig/Private/AnimNode_ControlRig.cpp b/Engine/Plugins/Animation/ControlRig/Source/ControlRig/Private/AnimNode_ControlRig.cpp
--- a/Engine/Plugins/Animation/ControlRig/Source/ControlRig/Private/AnimNode_ControlRig.cpp
+++ b/Engine/Plugins/Animation/ControlRig/Source/ControlRig/Private/AnimNode_ControlRig.cpp
@@ -55,7 +55,7 @@ void FAnimNode_ControlRig::OnInitializeAnimInstance(const FAnimInstanceProxy* In
 		ControlRigClass = nullptr;
 	}
 	
-	if(UpdateControlRigIfNeeded(InAnimInstance))
+	if(UpdateControlRigIfNeeded(InAnimInstance, InAnimInstance->GetRequiredBones()))
 	{
 		UpdateControlRigRefPoseIfNeeded(InProxy);
 	}
@@ -113,7 +113,7 @@ void FAnimNode_ControlRig::Update_AnyThread(const FAnimationUpdateContext& Conte
 
 	if(const UAnimInstance* AnimInstance = Cast<UAnimInstance>(Context.GetAnimInstanceObject()))
 	{
-		(void)UpdateControlRigIfNeeded(AnimInstance);
+		(void)UpdateControlRigIfNeeded(AnimInstance, Context.AnimInstanceProxy->GetRequiredBones());
 	}
 
 	UpdateControlRigRefPoseIfNeeded(Context.AnimInstanceProxy);
@@ -158,7 +158,7 @@ void FAnimNode_ControlRig::CacheBones_AnyThread(const FAnimationCacheBonesContex
 	// update the control rig instance just in case the dynamic control rig class has changed
 	if(const UAnimInstance* AnimInstance = Cast<UAnimInstance>(Context.GetAnimInstanceObject()))
 	{
-		(void)UpdateControlRigIfNeeded(AnimInstance);
+		(void)UpdateControlRigIfNeeded(AnimInstance, Context.AnimInstanceProxy->GetRequiredBones());
 	}
 
 	FAnimNode_ControlRigBase::CacheBones_AnyThread(Context);
@@ -295,7 +295,7 @@ void FAnimNode_ControlRig::SetControlRigClass(TSubclassOf<UControlRig> InControl
 	ControlRigClass = InControlRigClass;
 }
 
-bool FAnimNode_ControlRig::UpdateControlRigIfNeeded(const UAnimInstance* InAnimInstance)
+bool FAnimNode_ControlRig::UpdateControlRigIfNeeded(const UAnimInstance* InAnimInstance, const FBoneContainer& InRequiredBones)
 {
 	if (UClass* ExpectedClass = GetTargetClass())
 	{
@@ -305,7 +305,7 @@ bool FAnimNode_ControlRig::UpdateControlRigIfNeeded(const UAnimInstance* InAnimI
 			{
 				UControlRig* NewControlRig = nullptr;
 
-				auto ReportErrorAndSwitchToDefaultRig = [this, InAnimInstance, ExpectedClass](const FString& InMessage) -> bool
+				auto ReportErrorAndSwitchToDefaultRig = [this, InAnimInstance, InRequiredBones, ExpectedClass](const FString& InMessage) -> bool
 				{
 					static constexpr TCHAR Format[] =  TEXT("[%s] Cannot switch to runtime rig class '%s' - reverting to default. %s");
 					UE_LOG(LogControlRig, Warning, Format, *InAnimInstance->GetPathName(), *ExpectedClass->GetName(), *InMessage);
@@ -316,7 +316,7 @@ bool FAnimNode_ControlRig::UpdateControlRigIfNeeded(const UAnimInstance* InAnimI
 					// fall back to the default control rig and switch to that
 					ControlRigClass = nullptr;
 					
-					return UpdateControlRigIfNeeded(InAnimInstance);
+					return UpdateControlRigIfNeeded(InAnimInstance, InRequiredBones);
 				};
 
 				// if we are reacting to a programmatic change
@@ -336,7 +336,7 @@ bool FAnimNode_ControlRig::UpdateControlRigIfNeeded(const UAnimInstance* InAnimI
 						{
 							// fall back to the default control rig and switch to that
 							ControlRigClass = nullptr;
-							return UpdateControlRigIfNeeded(InAnimInstance);
+							return UpdateControlRigIfNeeded(InAnimInstance, InRequiredBones);
 						}
 					}
 					else
@@ -375,11 +375,21 @@ bool FAnimNode_ControlRig::UpdateControlRigIfNeeded(const UAnimInstance* InAnimI
 						}
 						
 						// create a new control rig using the new class
-						NewControlRig = NewObject<UControlRig>(InAnimInstance->GetOwningComponent(), ExpectedClass);
-						ControlRig->SetObjectBinding(MakeShared<FControlRigObjectBinding>());
-						ControlRig->GetObjectBinding()->BindToObject(InAnimInstance->GetOwningComponent());
-						NewControlRig->Initialize(true);
-						NewControlRig->RequestInit();
+						{
+							// Let's make sure the GC isn't running when we try to create a new Control Rig.
+							FGCScopeGuard GCGuard;
+							
+							NewControlRig = NewObject<UControlRig>(InAnimInstance->GetOwningComponent(), ExpectedClass);
+							
+							// If the object was created on a non-game thread, clear the async flag immediately, so that it can be
+							// garbage collected in the future. 
+							(void)NewControlRig->AtomicallyClearInternalFlags(EInternalObjectFlags::Async);
+														
+							ControlRig->SetObjectBinding(MakeShared<FControlRigObjectBinding>());
+							ControlRig->GetObjectBinding()->BindToObject(InAnimInstance->GetOwningComponent());
+							NewControlRig->Initialize(true);
+							NewControlRig->RequestInit();
+						}
 
 						// temporarily set the new control rig to be the target instance
 						TGuardValue<TObjectPtr<UObject>> TargetInstanceGuard(TargetInstance, NewControlRig);
@@ -464,7 +474,12 @@ bool FAnimNode_ControlRig::UpdateControlRigIfNeeded(const UAnimInstance* InAnimI
 
 		if(ControlRig == nullptr)
 		{
+			// Let's make sure the GC isn't running when we try to create a new Control Rig.
+			FGCScopeGuard GCGuard;
+			
 			ControlRig = NewObject<UControlRig>(InAnimInstance->GetOwningComponent(), ExpectedClass);
+			(void)ControlRig->AtomicallyClearInternalFlags(EInternalObjectFlags::Async);
+			
 			ControlRig->SetObjectBinding(MakeShared<FControlRigObjectBinding>());
 			ControlRig->GetObjectBinding()->BindToObject(InAnimInstance->GetOwningComponent());
 			ControlRig->Initialize(true);
@@ -473,7 +488,7 @@ bool FAnimNode_ControlRig::UpdateControlRigIfNeeded(const UAnimInstance* InAnimI
 		RefPoseSetterHash.Reset();
 		ControlRig->OnInitialized_AnyThread().AddRaw(this, &FAnimNode_ControlRig::HandleOnInitialized_AnyThread);
 
-		UpdateInputOutputMappingIfRequired(ControlRig, InAnimInstance->GetRequiredBones());
+		UpdateInputOutputMappingIfRequired(ControlRig, InRequiredBones);
 
 		return true;
 	}
diff --git a/Engine/Plugins/Animation/ControlRig/Source/ControlRig/Public/AnimNode_ControlRig.h b/Engine/Plugins/Animation/ControlRig/Source/ControlRig/Public/AnimNode_ControlRig.h
--- a/Engine/Plugins/Animation/ControlRig/Source/ControlRig/Public/AnimNode_ControlRig.h
+++ b/Engine/Plugins/Animation/ControlRig/Source/ControlRig/Public/AnimNode_ControlRig.h
@@ -138,7 +138,7 @@ protected:
 	virtual void UpdateOutput(UControlRig* InControlRig, FPoseContext& InOutput) override;
 
 	void SetControlRigClass(TSubclassOf<UControlRig> InControlRigClass);
-	bool UpdateControlRigIfNeeded(const UAnimInstance* InAnimInstance);
+	bool UpdateControlRigIfNeeded(const UAnimInstance* InAnimInstance, const FBoneContainer& InRequiredBones);
 
 	// Helper function to update the initial ref pose within the Control Rig if needed
 	void UpdateControlRigRefPoseIfNeeded(const FAnimInstanceProxy* InProxy, bool bIncludePoseInHash = false);
