Control Rig: Fix crash when resetting the preview mesh
#rnx

diff --git a/Engine/Plugins/Animation/ControlRig/Source/ControlRig/Private/Rigs/RigHierarchyController.cpp b/Engine/Plugins/Animation/ControlRig/Source/ControlRig/Private/Rigs/RigHierarchyController.cpp
--- a/Engine/Plugins/Animation/ControlRig/Source/ControlRig/Private/Rigs/RigHierarchyController.cpp
+++ b/Engine/Plugins/Animation/ControlRig/Source/ControlRig/Private/Rigs/RigHierarchyController.cpp
@@ -1001,12 +1001,14 @@ TArray<FRigElementKey> URigHierarchyController::ImportBones(USkeleton* InSkeleto
                                                             bool bSelectBones, bool bSetupUndo,
                                                             bool bPrintPythonCommand)
 {
-	if (InSkeleton == nullptr)
+	FReferenceSkeleton EmptySkeleton;
+	FReferenceSkeleton& RefSkeleton = EmptySkeleton; 
+	if (InSkeleton != nullptr)
 	{
-		return TArray<FRigElementKey>();
+		RefSkeleton = InSkeleton->GetReferenceSkeleton();
 	}
 
-	const TArray<FRigElementKey> BoneKeys = ImportBones(InSkeleton->GetReferenceSkeleton(), InNameSpace, bReplaceExistingBones, bRemoveObsoleteBones,
+	const TArray<FRigElementKey> BoneKeys = ImportBones(RefSkeleton, InNameSpace, bReplaceExistingBones, bRemoveObsoleteBones,
 	                   bSelectBones, bSetupUndo);
 
 #if WITH_EDITOR
@@ -1017,7 +1019,7 @@ TArray<FRigElementKey> URigHierarchyController::ImportBones(USkeleton* InSkeleto
 		{
 			RigVMPythonUtils::Print(Blueprint->GetFName().ToString(),
 				FString::Printf(TEXT("hierarchy_controller.import_bones_from_asset('%s', '%s', %s, %s, %s)"),
-				*InSkeleton->GetPathName(),
+				(InSkeleton != nullptr) ? *InSkeleton->GetPathName() : TEXT(""),
 				*InNameSpace.ToString(),
 				(bReplaceExistingBones) ? TEXT("True") : TEXT("False"),
 				(bRemoveObsoleteBones) ? TEXT("True") : TEXT("False"),
diff --git a/Engine/Plugins/Animation/ControlRig/Source/ControlRigEditor/Private/Editor/ControlRigEditor.cpp b/Engine/Plugins/Animation/ControlRig/Source/ControlRigEditor/Private/Editor/ControlRigEditor.cpp
--- a/Engine/Plugins/Animation/ControlRig/Source/ControlRigEditor/Private/Editor/ControlRigEditor.cpp
+++ b/Engine/Plugins/Animation/ControlRig/Source/ControlRigEditor/Private/Editor/ControlRigEditor.cpp
@@ -2815,8 +2815,9 @@ void FControlRigEditor::HandlePreviewMeshChanged(USkeletalMesh* InOldSkeletalMes
 					TGuardValue<bool> SuspendBlueprintNotifs(ControlRigBP->bSuspendAllNotifications, true);
 					if(URigHierarchyController* Controller = ControlRigBP->GetHierarchyController())
 					{
-						Controller->ImportBones(InNewSkeletalMesh->GetSkeleton(), NAME_None, true, true, false, true, true);
-						Controller->ImportCurves(InNewSkeletalMesh->GetSkeleton(), NAME_None, false, true, true);
+						USkeleton* Skeleton = InNewSkeletalMesh ? InNewSkeletalMesh->GetSkeleton() : nullptr;
+						Controller->ImportBones(Skeleton, NAME_None, true, true, false, true, true);
+						Controller->ImportCurves(Skeleton, NAME_None, false, true, true);
 					}
 				}
 				ControlRigBP->PropagateHierarchyFromBPToInstances();
diff --git a/Engine/Plugins/Animation/ControlRig/Source/ControlRigEditor/Private/Editor/ControlRigSchematicModel.cpp b/Engine/Plugins/Animation/ControlRig/Source/ControlRigEditor/Private/Editor/ControlRigSchematicModel.cpp
--- a/Engine/Plugins/Animation/ControlRig/Source/ControlRigEditor/Private/Editor/ControlRigSchematicModel.cpp
+++ b/Engine/Plugins/Animation/ControlRig/Source/ControlRigEditor/Private/Editor/ControlRigSchematicModel.cpp
@@ -147,60 +147,62 @@ void FControlRigSchematicModel::Tick(float InDeltaTime)
 	if(ControlRigBlueprint.IsValid())
 	{
 		//UE_LOG(LogControlRig, Display, TEXT("Selected nodes %d"), GetSelectedNodes().Num());
-		
-		const URigHierarchy* Hierarchy = ControlRigBlueprint->GetDebuggedControlRig()->GetHierarchy();
-		check(Hierarchy);
-		const FModularRigConnections& Connections = ControlRigBlueprint->ModularRigModel.Connections;
-
-		TArray<FRigElementKey> KeysToRemove;
-		for(const TPair<FRigElementKey,FGuid>& Pair : RigElementKeyToGuid)
+		if (UControlRig* DebuggedRig = ControlRigBlueprint->GetDebuggedControlRig())
 		{
-			if((Pair.Key.Type == ERigElementType::Socket) ||
-				(Pair.Key.Type == ERigElementType::Connector))
+			const URigHierarchy* Hierarchy = DebuggedRig->GetHierarchy();
+			check(Hierarchy);
+			const FModularRigConnections& Connections = ControlRigBlueprint->ModularRigModel.Connections;
+
+			TArray<FRigElementKey> KeysToRemove;
+			for(const TPair<FRigElementKey,FGuid>& Pair : RigElementKeyToGuid)
 			{
-				continue;
+				if((Pair.Key.Type == ERigElementType::Socket) ||
+					(Pair.Key.Type == ERigElementType::Connector))
+				{
+					continue;
+				}
+				if(TemporaryNodeGuids.Contains(Pair.Value))
+				{
+					continue;
+				}
+
+				const TArray<FRigElementKey>& Connectors = Connections.FindConnectorsFromTarget(Pair.Key);
+				if(Connectors.Num() > 1)
+				{
+					continue;
+				}
+				if(Connectors.Num() == 1)
+				{
+					if(const FRigConnectorElement* Connector = Hierarchy->Find<FRigConnectorElement>(Connectors[0]))
+					{
+						if(!Connector->IsPrimary())
+						{
+							continue;
+						}
+					}
+				}
+			
+				KeysToRemove.Add(Pair.Key);
 			}
-			if(TemporaryNodeGuids.Contains(Pair.Value))
+			for(const FRigElementKey& KeyToRemove : KeysToRemove)
 			{
-				continue;
+				RemoveElementKeyNode(KeyToRemove);
 			}
 
-			const TArray<FRigElementKey>& Connectors = Connections.FindConnectorsFromTarget(Pair.Key);
-			if(Connectors.Num() > 1)
+			TArray<FRigElementKey> ConnectorKeys;
+			for(const TPair<FRigElementKey,FGuid>& Pair : RigElementKeyToGuid)
 			{
-				continue;
-			}
-			if(Connectors.Num() == 1)
-			{
-				if(const FRigConnectorElement* Connector = Hierarchy->Find<FRigConnectorElement>(Connectors[0]))
+				if(Pair.Key.Type == ERigElementType::Connector)
 				{
-					if(!Connector->IsPrimary())
-					{
-						continue;
-					}
+					ConnectorKeys.Add(Pair.Key);
 				}
 			}
-			
-			KeysToRemove.Add(Pair.Key);
-		}
-		for(const FRigElementKey& KeyToRemove : KeysToRemove)
-		{
-			RemoveElementKeyNode(KeyToRemove);
-		}
 
-		TArray<FRigElementKey> ConnectorKeys;
-		for(const TPair<FRigElementKey,FGuid>& Pair : RigElementKeyToGuid)
-		{
-			if(Pair.Key.Type == ERigElementType::Connector)
+			for(const FRigElementKey& ConnectorKey : ConnectorKeys)
 			{
-				ConnectorKeys.Add(Pair.Key);
+				UpdateConnector(ConnectorKey);
 			}
 		}
-
-		for(const FRigElementKey& ConnectorKey : ConnectorKeys)
-		{
-			UpdateConnector(ConnectorKey);
-		}
 	}
 }
 
