Control Rig: Avoid intefering with external instances of the rig during debugging

#rnx

diff --git a/Engine/Plugins/Animation/ControlRig/Source/ControlRig/Private/AnimNode_ControlRigBase.cpp b/Engine/Plugins/Animation/ControlRig/Source/ControlRig/Private/AnimNode_ControlRigBase.cpp
--- a/Engine/Plugins/Animation/ControlRig/Source/ControlRig/Private/AnimNode_ControlRigBase.cpp
+++ b/Engine/Plugins/Animation/ControlRig/Source/ControlRig/Private/AnimNode_ControlRigBase.cpp
@@ -136,7 +136,7 @@ void FAnimNode_ControlRigBase::UpdateInput(UControlRig* ControlRig, const FPoseC
 
 #if WITH_EDITOR
 	// if we are recording any change - let's clear the undo stack
-	if(Hierarchy->IsTracingChanges())
+	if(bExecute && Hierarchy->IsTracingChanges())
 	{
 		Hierarchy->ResetTransformStack();
 	}
@@ -249,7 +249,7 @@ void FAnimNode_ControlRigBase::UpdateInput(UControlRig* ControlRig, const FPoseC
 	}
 
 #if WITH_EDITOR
-	if(Hierarchy->IsTracingChanges())
+	if(bExecute && Hierarchy->IsTracingChanges())
 	{
 		Hierarchy->StorePoseForTrace(TEXT("FAnimNode_ControlRigBase::UpdateInput"));
 	}
@@ -375,7 +375,7 @@ void FAnimNode_ControlRigBase::UpdateOutput(UControlRig* ControlRig, FPoseContex
 	}
 
 #if WITH_EDITOR
-	if(Hierarchy->IsTracingChanges())
+	if(bExecute && Hierarchy->IsTracingChanges())
 	{
 		Hierarchy->StorePoseForTrace(TEXT("FAnimNode_ControlRigBase::UpdateOutput"));
 		Hierarchy->DumpTransformStackToFile();
@@ -449,7 +449,6 @@ void FAnimNode_ControlRigBase::ExecuteControlRig(FPoseContext& InOutput)
 		
 		// first update input to the system
 		UpdateInput(ControlRig, InOutput);
-
 		
 		if (bExecute)
 		{
@@ -562,28 +561,35 @@ void FAnimNode_ControlRigBase::CacheBones_AnyThread(const FAnimationCacheBonesCo
 	if (UControlRig* ControlRig = GetControlRig())
 	{
 		// fill up node names
-		FBoneContainer& RequiredBones = Context.AnimInstanceProxy->GetRequiredBones();
+		const FBoneContainer& RequiredBones = Context.AnimInstanceProxy->GetRequiredBones();
 
 		const uint16 BonesSerialNumber = RequiredBones.GetSerialNumber();
-		const bool bIsLODChange = !bControlRigRequiresInitialization && (BonesSerialNumber != LastBonesSerialNumberForCacheBones);
 
 		// the construction event may create a set of bones that we can map to. let's run construction now.
-		if(ControlRig->IsConstructionModeEnabled() ||
-			(ControlRig->IsConstructionRequired() && (bControlRigRequiresInitialization || bIsLODChange)))
+		if (bExecute)
 		{
-			UpdateGetAssetUserDataDelegate(ControlRig);
-			ControlRig->Execute(FRigUnit_PrepareForExecution::EventName);
-			bControlRigRequiresInitialization = false;
+			const bool bIsLODChange = !bControlRigRequiresInitialization && (BonesSerialNumber != LastBonesSerialNumberForCacheBones);
+		
+			if(ControlRig->IsConstructionModeEnabled() ||
+				(ControlRig->IsConstructionRequired() && (bControlRigRequiresInitialization || bIsLODChange)))
+			{
+				UpdateGetAssetUserDataDelegate(ControlRig);
+				ControlRig->Execute(FRigUnit_PrepareForExecution::EventName);
+				bControlRigRequiresInitialization = false;
+			}
 		}
 
 		UpdateInputOutputMappingIfRequired(ControlRig, RequiredBones);
 
 		if(bControlRigRequiresInitialization)
 		{
-			// re-init only if this is the first run
-			// and restore control values
-			ControlRig->RequestInit();
-			bControlRigRequiresInitialization = false;
+			if(bExecute)
+			{
+				// re-init only if this is the first run
+				// and restore control values
+				ControlRig->RequestInit();
+				bControlRigRequiresInitialization = false;
+			}
 		}
 		
 		LastBonesSerialNumberForCacheBones = BonesSerialNumber;
diff --git a/Engine/Plugins/Animation/ControlRig/Source/ControlRigEditor/Private/Editor/ControlRigEditor.cpp b/Engine/Plugins/Animation/ControlRig/Source/ControlRigEditor/Private/Editor/ControlRigEditor.cpp
--- a/Engine/Plugins/Animation/ControlRig/Source/ControlRigEditor/Private/Editor/ControlRigEditor.cpp
+++ b/Engine/Plugins/Animation/ControlRig/Source/ControlRigEditor/Private/Editor/ControlRigEditor.cpp
@@ -1073,8 +1073,7 @@ void FControlRigEditor::HandleSetObjectBeingDebugged(UObject* InObject)
 
 	if (DebuggedControlRig)
 	{
-		bool bIsExternalControlRig = DebuggedControlRig != GetControlRig();
-		bool bShouldExecute = (!bIsExternalControlRig) && bExecutionControlRig;
+		const bool bShouldExecute = ShouldExecuteControlRig(DebuggedControlRig);
 		GetControlRigBlueprint()->Hierarchy->HierarchyForSelectionPtr = DebuggedControlRig->DynamicHierarchy;
 
 		UControlRigSkeletalMeshComponent* EditorSkelComp = Cast<UControlRigSkeletalMeshComponent>(GetPersonaToolkit()->GetPreviewScene()->GetPreviewMeshComponent());
@@ -1104,12 +1103,15 @@ void FControlRigEditor::HandleSetObjectBeingDebugged(UObject* InObject)
 			}
 			
 			// get the bone intial transforms from the preview skeletal mesh
-			DebuggedControlRig->SetBoneInitialTransformsFromSkeletalMeshComponent(EditorSkelComp);
-			if(UControlRigBlueprint* RigBlueprint = Cast<UControlRigBlueprint>(GetBlueprintObj()))
+			if(bShouldExecute)
 			{
-				// copy the initial transforms back to the blueprint
-				// no need to call modify here since this code only modifies the bp if the preview mesh changed
-				RigBlueprint->Hierarchy->CopyPose(DebuggedControlRig->GetHierarchy(), false, true, false);
+				DebuggedControlRig->SetBoneInitialTransformsFromSkeletalMeshComponent(EditorSkelComp);
+				if(UControlRigBlueprint* RigBlueprint = Cast<UControlRigBlueprint>(GetBlueprintObj()))
+				{
+					// copy the initial transforms back to the blueprint
+					// no need to call modify here since this code only modifies the bp if the preview mesh changed
+					RigBlueprint->Hierarchy->CopyPose(DebuggedControlRig->GetHierarchy(), false, true, false);
+				}
 			}
 		}
 
@@ -1133,6 +1135,14 @@ void FControlRigEditor::HandleSetObjectBeingDebugged(UObject* InObject)
 		{
 			EditorSkelComp->SetComponentToWorld(FTransform::Identity);
 		}
+
+		if(!bShouldExecute)
+		{
+			if (FControlRigEditMode* EditMode = GetEditMode())
+			{
+				EditMode->RequestToRecreateControlShapeActors();
+			}
+		}
 	}
 	else
 	{
@@ -2586,6 +2596,23 @@ bool FControlRigEditor::IsConstructionModeEnabled() const
 	return GetEventQueue() == ConstructionEventQueue;
 }
 
+bool FControlRigEditor::IsDebuggingExternalControlRig(const UControlRig* InControlRig) const
+{
+	if(InControlRig == nullptr)
+	{
+		if(const UControlRigBlueprint* ControlRigBlueprint = GetControlRigBlueprint())
+		{
+			InControlRig = Cast<UControlRig>(ControlRigBlueprint->GetObjectBeingDebugged());
+		}
+	}
+	return InControlRig != GetControlRig();
+}
+
+bool FControlRigEditor::ShouldExecuteControlRig(const UControlRig* InControlRig) const
+{
+	return (!IsDebuggingExternalControlRig(InControlRig)) && bExecutionControlRig;
+}
+
 void FControlRigEditor::HandlePreviewSceneCreated(const TSharedRef<IPersonaPreviewScene>& InPersonaPreviewScene)
 {
 	DECLARE_SCOPE_HIERARCHICAL_COUNTER_FUNC()
@@ -5260,15 +5287,19 @@ void FControlRigEditor::OnPreForwardsSolve_AnyThread(UControlRig* InRig, const F
 void FControlRigEditor::OnPreConstructionForUI_AnyThread(UControlRig* InRig, const FName& InEventName)
 {
 	bIsConstructionEventRunning = true;
-	const TArrayView<const FRigElementKey> Elements;
-	PreConstructionPose = InRig->GetHierarchy()->GetPose(false, ERigElementType::ToResetAfterConstructionEvent, Elements);
 
-	if(const UControlRigBlueprint* RigBlueprint = GetControlRigBlueprint())
+	if(ShouldExecuteControlRig(InRig))
 	{
-		if(RigBlueprint->IsControlRigModule())
+		const TArrayView<const FRigElementKey> Elements;
+		PreConstructionPose = InRig->GetHierarchy()->GetPose(false, ERigElementType::ToResetAfterConstructionEvent, Elements);
+
+		if(const UControlRigBlueprint* RigBlueprint = GetControlRigBlueprint())
 		{
-			SocketStates = InRig->GetHierarchy()->GetSocketStates();
-			ConnectorStates = RigBlueprint->Hierarchy->GetConnectorStates();
+			if(RigBlueprint->IsControlRigModule())
+			{
+				SocketStates = InRig->GetHierarchy()->GetSocketStates();
+				ConnectorStates = RigBlueprint->Hierarchy->GetConnectorStates();
+			}
 		}
 	}
 }
@@ -5320,7 +5351,7 @@ void FControlRigEditor::OnPreConstruction_AnyThread(UControlRig* InRig, const FN
 					}
 				}
 
-				if(RigBlueprint->IsControlRigModule())
+				if(ShouldExecuteControlRig(InRig))
 				{
 					RigBlueprint->Hierarchy->RestoreSocketsFromStates(SocketStates);
 					InRig->GetHierarchy()->RestoreSocketsFromStates(SocketStates);
@@ -5333,15 +5364,16 @@ void FControlRigEditor::OnPreConstruction_AnyThread(UControlRig* InRig, const FN
 void FControlRigEditor::OnPostConstruction_AnyThread(UControlRig* InRig, const FName& InEventName)
 {
 	bIsConstructionEventRunning = false;
+	const bool bShouldExecute = ShouldExecuteControlRig(InRig);
 
 	if(UControlRigBlueprint* RigBlueprint = GetControlRigBlueprint())
 	{
-		if(RigBlueprint->IsControlRigModule())
+		if(bShouldExecute && RigBlueprint->IsControlRigModule())
 		{
 			RigBlueprint->Hierarchy->RestoreConnectorsFromStates(ConnectorStates);
 		}
 
-		if(RigBlueprint->IsModularRig())
+		if(bShouldExecute && RigBlueprint->IsModularRig())
 		{
 			// auto resolve the root module's primary connector
 			if(RigBlueprint->ModularRigModel.Connections.IsEmpty() && RigBlueprint->ModularRigModel.Modules.Num() == 1 && RigBlueprint->Hierarchy->Num(ERigElementType::Bone) > 0)
@@ -5406,7 +5438,7 @@ void FControlRigEditor::OnPostConstruction_AnyThread(UControlRig* InRig, const F
 			}, TStatId(), NULL, ENamedThreads::GameThread);
 		}
 	}
-	else
+	else if(bShouldExecute)
 	{
 		InRig->GetHierarchy()->SetPose(PreConstructionPose, ERigTransformType::CurrentGlobal);
 	}
diff --git a/Engine/Plugins/Animation/ControlRig/Source/ControlRigEditor/Private/Editor/ControlRigEditor.h b/Engine/Plugins/Animation/ControlRig/Source/ControlRigEditor/Private/Editor/ControlRigEditor.h
--- a/Engine/Plugins/Animation/ControlRig/Source/ControlRigEditor/Private/Editor/ControlRigEditor.h
+++ b/Engine/Plugins/Animation/ControlRig/Source/ControlRigEditor/Private/Editor/ControlRigEditor.h
@@ -331,6 +331,8 @@ protected:
 	void OnAnimInitialized();
 
 	bool IsConstructionModeEnabled() const;
+	bool IsDebuggingExternalControlRig(const UControlRig* InControlRig = nullptr) const;
+	bool ShouldExecuteControlRig(const UControlRig* InControlRig = nullptr) const;
 
 	int32 RigHierarchyTabCount;
 	int32 ModularRigHierarchyTabCount;
