Control Rig: Undoing changes to the RigHierarchy provokes references to an invalid debug object in the editor
#rnx

diff --git a/Engine/Plugins/Animation/ControlRig/Source/ControlRigDeveloper/Private/ControlRigBlueprint.cpp b/Engine/Plugins/Animation/ControlRig/Source/ControlRigDeveloper/Private/ControlRigBlueprint.cpp
--- a/Engine/Plugins/Animation/ControlRig/Source/ControlRigDeveloper/Private/ControlRigBlueprint.cpp
+++ b/Engine/Plugins/Animation/ControlRig/Source/ControlRigDeveloper/Private/ControlRigBlueprint.cpp
@@ -1359,11 +1359,10 @@ void UControlRigBlueprint::PostTransacted(const FTransactionObjectEvent& Transac
 	if (TransactionEvent.GetEventType() == ETransactionObjectEventType::UndoRedo)
 	{
 		TArray<FName> PropertiesChanged = TransactionEvent.GetChangedProperties();
-		if (PropertiesChanged.Contains(GET_MEMBER_NAME_CHECKED(UControlRigBlueprint, Hierarchy)))
+		int32 TransactionIndex = GEditor->Trans->FindTransactionIndex(TransactionEvent.GetTransactionId());
+		const FTransaction* Transaction = GEditor->Trans->GetTransaction(TransactionIndex);
+		if (Transaction && Transaction->ContainsObject(Hierarchy))
 		{
-			int32 TransactionIndex = GEditor->Trans->FindTransactionIndex(TransactionEvent.GetTransactionId());
-			const FTransaction* Transaction = GEditor->Trans->GetTransaction(TransactionIndex);
-
 			if (Transaction->GetTitle().BuildSourceString() == TEXT("Transform Gizmo"))
 			{
 				PropagatePoseFromBPToInstances();
diff --git a/Engine/Plugins/Animation/ControlRig/Source/ControlRigEditor/Private/Editor/ControlRigEditor.cpp b/Engine/Plugins/Animation/ControlRig/Source/ControlRigEditor/Private/Editor/ControlRigEditor.cpp
--- a/Engine/Plugins/Animation/ControlRig/Source/ControlRigEditor/Private/Editor/ControlRigEditor.cpp
+++ b/Engine/Plugins/Animation/ControlRig/Source/ControlRigEditor/Private/Editor/ControlRigEditor.cpp
@@ -1640,6 +1640,8 @@ void FControlRigEditor::PostTransaction(bool bSuccess, const FTransaction* Trans
 		// Do not compile here. ControlRigBlueprint::PostTransacted decides when it is necessary to compile depending
 		// on the properties that are affected.
 		//Compile();
+
+		UpdateRigVMHost();
 		
 		USkeletalMesh* PreviewMesh = GetPersonaToolkit()->GetPreviewScene()->GetPreviewMesh();
 		if (PreviewMesh != RigBlueprint->GetPreviewMesh())
