Control Rig: Use correct component transform for direct manipulation in asset editor

#rnx

diff --git a/Engine/Plugins/Animation/ControlRig/Source/ControlRigEditor/Private/EditMode/ControlRigEditMode.cpp b/Engine/Plugins/Animation/ControlRig/Source/ControlRigEditor/Private/EditMode/ControlRigEditMode.cpp
--- a/Engine/Plugins/Animation/ControlRig/Source/ControlRigEditor/Private/EditMode/ControlRigEditMode.cpp
+++ b/Engine/Plugins/Animation/ControlRig/Source/ControlRigEditor/Private/EditMode/ControlRigEditMode.cpp
@@ -5419,12 +5419,20 @@ USceneComponent* FControlRigEditMode::GetHostingSceneComponent(const UControlRig
 
 FTransform FControlRigEditMode::GetHostingSceneComponentTransform(const UControlRig* ControlRig) const
 {
-	if (ControlRig == nullptr && GetControlRigs().Num() > 0)
+	// we care about this transform only in the level,
+	// since in the control rig editor the debug skeletal mesh component
+	// is set at identity anyway.
+	if(IsInLevelEditor())
 	{
-		ControlRig = GetControlRigs()[0].Get();
+		if (ControlRig == nullptr && GetControlRigs().Num() > 0)
+		{
+			ControlRig = GetControlRigs()[0].Get();
+		}
+
+		const USceneComponent* HostingComponent = GetHostingSceneComponent(ControlRig);
+		return HostingComponent ? HostingComponent->GetComponentTransform() : FTransform::Identity;
 	}
-	const USceneComponent* HostingComponent = GetHostingSceneComponent(ControlRig);
-	return HostingComponent ? HostingComponent->GetComponentTransform() : FTransform::Identity;
+	return FTransform::Identity;
 }
 
 void FControlRigEditMode::OnPoseInitialized()
