Control Rig: Fix modify transform direct manipulation

#rnx

diff --git a/Engine/Plugins/Animation/ControlRig/Source/ControlRig/Private/Units/Highlevel/Hierarchy/RigUnit_ModifyTransforms.cpp b/Engine/Plugins/Animation/ControlRig/Source/ControlRig/Private/Units/Highlevel/Hierarchy/RigUnit_ModifyTransforms.cpp
--- a/Engine/Plugins/Animation/ControlRig/Source/ControlRig/Private/Units/Highlevel/Hierarchy/RigUnit_ModifyTransforms.cpp
+++ b/Engine/Plugins/Animation/ControlRig/Source/ControlRig/Private/Units/Highlevel/Hierarchy/RigUnit_ModifyTransforms.cpp
@@ -131,19 +131,23 @@ bool FRigUnit_ModifyTransforms::UpdateHierarchyForDirectManipulation(const URigV
 			switch(Mode)
 			{
 				case EControlRigModifyBoneMode::AdditiveLocal:
-				case EControlRigModifyBoneMode::AdditiveGlobal:
 				{
-					InInfo->OffsetTransform = Hierarchy->GetGlobalTransform(ItemToModify[Index].Item);
+					// place the control offset transform where the target is now "without" the change provided by the pin
+					InInfo->OffsetTransform = ItemToModify[Index].Transform.Inverse() * Hierarchy->GetGlobalTransform(ItemToModify[Index].Item);
 					break;
 				}
 				case EControlRigModifyBoneMode::OverrideLocal:
 				{
+					// changing local means let's place the control offset transform where the parent is
 					InInfo->OffsetTransform = Hierarchy->GetGlobalTransform(FirstParent);
 					break;
 				}
+				case EControlRigModifyBoneMode::AdditiveGlobal:
 				case EControlRigModifyBoneMode::OverrideGlobal:
 				default:
 				{
+					// in this case the value input is a global transform
+					// so we'll leave the control offset transform at identity
 					break;
 				}
 			}
