Control Rig: Fix hierarchy order whens switching preview meshes

#rnx

diff --git a/Engine/Plugins/Animation/ControlRig/Source/ControlRigEditor/Private/Editor/ControlRigEditor.cpp b/Engine/Plugins/Animation/ControlRig/Source/ControlRigEditor/Private/Editor/ControlRigEditor.cpp
--- a/Engine/Plugins/Animation/ControlRig/Source/ControlRigEditor/Private/Editor/ControlRigEditor.cpp
+++ b/Engine/Plugins/Animation/ControlRig/Source/ControlRigEditor/Private/Editor/ControlRigEditor.cpp
@@ -2844,6 +2844,15 @@ void FControlRigEditor::HandlePreviewMeshChanged(USkeletalMesh* InOldSkeletalMes
 					TGuardValue<bool> SuspendBlueprintNotifs(ControlRigBP->bSuspendAllNotifications, true);
 					if(URigHierarchyController* Controller = ControlRigBP->GetHierarchyController())
 					{
+						// remove all connectors / sockets. keeping them around may mess up the order of the elements
+						// in the hierarchy, such as [bone,bone,bone,connector,connector,bone,bone,bone].
+						TArray<FRigElementKey> ConnectorsAndSockets = Controller->GetHierarchy()->GetConnectorKeys();
+						ConnectorsAndSockets.Append(Controller->GetHierarchy()->GetSocketKeys());
+						for(const FRigElementKey& Key : ConnectorsAndSockets)
+						{
+							(void)Controller->RemoveElement(Key, true, true);
+						}
+						
 						USkeleton* Skeleton = InNewSkeletalMesh ? InNewSkeletalMesh->GetSkeleton() : nullptr;
 						Controller->ImportBones(Skeleton, NAME_None, true, true, false, true, true);
 						Controller->ImportCurves(Skeleton, NAME_None, false, true, true);
