Avoid restoring skeletal mesh component state if it has been removed from the world

diff --git a/Engine/Plugins/Animation/ControlRig/Source/ControlRig/Private/Sequencer/MovieSceneControlRigParameterTemplate.cpp b/Engine/Plugins/Animation/ControlRig/Source/ControlRig/Private/Sequencer/MovieSceneControlRigParameterTemplate.cpp
--- a/Engine/Plugins/Animation/ControlRig/Source/ControlRig/Private/Sequencer/MovieSceneControlRigParameterTemplate.cpp
+++ b/Engine/Plugins/Animation/ControlRig/Source/ControlRig/Private/Sequencer/MovieSceneControlRigParameterTemplate.cpp
@@ -906,18 +906,22 @@ struct FControlRigParameterPreAnimatedTokenProducer : IMovieScenePreAnimatedToke
 						//do a tick and restore skel mesh
 						if (USkeletalMeshComponent* SkeletalMeshComponent = Cast<USkeletalMeshComponent>(ControlRig->GetObjectBinding()->GetBoundObject()))
 						{
-							// Restore pose after unbinding to force the restored pose
-							SkeletalMeshComponent->SetUpdateAnimationInEditor(true);
-							SkeletalMeshComponent->SetUpdateClothInEditor(true);
-							SkeletalMeshComponent->TickAnimation(0.f, false);
-
-							SkeletalMeshComponent->RefreshBoneTransforms();
-							SkeletalMeshComponent->RefreshFollowerComponents();
-							SkeletalMeshComponent->UpdateComponentToWorld();
-							SkeletalMeshComponent->FinalizeBoneTransform();
-							SkeletalMeshComponent->MarkRenderTransformDirty();
-							SkeletalMeshComponent->MarkRenderDynamicDataDirty();
-							SkeletalMeshRestoreState.RestoreState(SkeletalMeshComponent);
+							// If the skel mesh comp owner has been removed from the world, no need to restore anything
+							if (SkeletalMeshComponent->IsRegistered())
+							{
+								// Restore pose after unbinding to force the restored pose
+								SkeletalMeshComponent->SetUpdateAnimationInEditor(true);
+								SkeletalMeshComponent->SetUpdateClothInEditor(true);
+								SkeletalMeshComponent->TickAnimation(0.f, false);
+
+								SkeletalMeshComponent->RefreshBoneTransforms();
+								SkeletalMeshComponent->RefreshFollowerComponents();
+								SkeletalMeshComponent->UpdateComponentToWorld();
+								SkeletalMeshComponent->FinalizeBoneTransform();
+								SkeletalMeshComponent->MarkRenderTransformDirty();
+								SkeletalMeshComponent->MarkRenderDynamicDataDirty();
+								SkeletalMeshRestoreState.RestoreState(SkeletalMeshComponent);
+							}
 						}
 						//only unbind if not a component
 						if (Cast<UControlRigComponent>(ControlRig->GetObjectBinding()->GetBoundObject()) == nullptr)
