Control Rig: Enum widget so that users are able to select native enum types for integer controls
#rnx

diff --git a/Engine/Plugins/Animation/ControlRig/Source/ControlRig/Public/Rigs/RigHierarchyDefines.h b/Engine/Plugins/Animation/ControlRig/Source/ControlRig/Public/Rigs/RigHierarchyDefines.h
--- a/Engine/Plugins/Animation/ControlRig/Source/ControlRig/Public/Rigs/RigHierarchyDefines.h
+++ b/Engine/Plugins/Animation/ControlRig/Source/ControlRig/Public/Rigs/RigHierarchyDefines.h
@@ -17,7 +17,7 @@ class URigHierarchy;
  * This is rig element types that we support
  * This can be used as a mask so supported as a bitfield
  */
-UENUM(BlueprintType)
+UENUM(BlueprintType, meta = (RigVMTypeAllowed))
 enum class ERigElementType : uint8
 {
 	None = 0,
@@ -47,7 +47,7 @@ enum class ERigBoneType : uint8
 /* 
  * The type of meta data stored on an element
  */
-UENUM(BlueprintType)
+UENUM(BlueprintType, meta = (RigVMTypeAllowed))
 enum class ERigMetadataType : uint8
 {
 	Bool,
@@ -103,7 +103,7 @@ enum class ERigHierarchyNotification : uint8
 	Max UMETA(Hidden),
 };
 
-UENUM()
+UENUM(meta = (RigVMTypeAllowed))
 enum class ERigEvent : uint8
 {
 	/** Invalid event */
@@ -131,7 +131,7 @@ enum class EControlRigSetKey : uint8
 	Never				//Never set a key here.
 };
 
-UENUM(BlueprintType)
+UENUM(BlueprintType, meta = (RigVMTypeAllowed))
 enum class ERigControlType : uint8
 {
 	Bool,
@@ -171,7 +171,7 @@ enum class ERigControlValueType : uint8
     Maximum
 };
 
-UENUM(BlueprintType)
+UENUM(BlueprintType, meta = (RigVMTypeAllowed))
 enum class ERigControlVisibility : uint8
 {
 	// Visibility controlled by the graph
@@ -180,7 +180,7 @@ enum class ERigControlVisibility : uint8
 	BasedOnSelection 
 };
 
-UENUM(BlueprintType)
+UENUM(BlueprintType, meta = (RigVMTypeAllowed))
 enum class ERigControlAxis : uint8
 {
 	X,
diff --git a/Engine/Plugins/Animation/ControlRig/Source/ControlRig/Public/Units/Execution/RigUnit_DynamicHierarchy.h b/Engine/Plugins/Animation/ControlRig/Source/ControlRig/Public/Units/Execution/RigUnit_DynamicHierarchy.h
--- a/Engine/Plugins/Animation/ControlRig/Source/ControlRig/Public/Units/Execution/RigUnit_DynamicHierarchy.h
+++ b/Engine/Plugins/Animation/ControlRig/Source/ControlRig/Public/Units/Execution/RigUnit_DynamicHierarchy.h
@@ -82,7 +82,7 @@ struct CONTROLRIG_API FRigUnit_SetDefaultParent : public FRigUnit_DynamicHierarc
 	FRigElementKey Parent;
 };
 
-UENUM()
+UENUM(meta = (RigVMTypeAllowed))
 enum class ERigSwitchParentMode : uint8
 {
 	/** Switches the element to be parented to the world */
diff --git a/Engine/Plugins/Animation/ControlRig/Source/ControlRig/Public/Units/RigUnitContext.h b/Engine/Plugins/Animation/ControlRig/Source/ControlRig/Public/Units/RigUnitContext.h
--- a/Engine/Plugins/Animation/ControlRig/Source/ControlRig/Public/Units/RigUnitContext.h
+++ b/Engine/Plugins/Animation/ControlRig/Source/ControlRig/Public/Units/RigUnitContext.h
@@ -28,7 +28,7 @@ enum class EControlRigInteractionType : uint8
 	All = Translate | Rotate | Scale
 };
 
-UENUM(BlueprintType)
+UENUM(BlueprintType, meta = (RigVMTypeAllowed))
 enum class ERigMetaDataNameSpace : uint8
 {
 	// Use no namespace - store the metadata directly on the item
diff --git a/Engine/Plugins/Animation/ControlRig/Source/ControlRigEditor/Private/ControlRigElementDetails.cpp b/Engine/Plugins/Animation/ControlRig/Source/ControlRigEditor/Private/ControlRigElementDetails.cpp
--- a/Engine/Plugins/Animation/ControlRig/Source/ControlRigEditor/Private/ControlRigElementDetails.cpp
+++ b/Engine/Plugins/Animation/ControlRig/Source/ControlRigEditor/Private/ControlRigElementDetails.cpp
@@ -25,6 +25,7 @@
 #include "Editor/SRigHierarchyTreeView.h"
 #include "StructViewerFilter.h"
 #include "StructViewerModule.h"
+#include "Widgets/SRigVMGraphPinEnumPicker.h"
 
 #define LOCTEXT_NAMESPACE "ControlRigElementDetails"
 
@@ -3186,60 +3187,43 @@ void FRigControlElementDetails::CustomizeControl(IDetailLayoutBuilder& DetailBui
 	}
 
 	if(IsAnyControlOfValueType(ERigControlType::Integer))
-	{
-		const TSharedPtr<IPropertyHandle> ControlEnumHandle = SettingsHandle->GetChildHandle(TEXT("ControlEnum"));
-		ControlCategory.AddProperty(ControlEnumHandle.ToSharedRef()).DisplayName(FText::FromString(TEXT("Control Enum")))
-		.IsEnabled(bIsEnabled);
-
-		ControlEnumHandle->SetOnPropertyValueChanged(FSimpleDelegate::CreateLambda(
-			[this, PropertyUtilities]()
+	{		
+		FDetailWidgetRow* EnumWidgetRow = &ControlCategory.AddCustomRow(FText::FromString(TEXT("ControlEnum")))
+		.NameContent()
+		[
+			SNew(STextBlock)
+			.Text(FText::FromString(TEXT("Control Enum")))
+			.Font(IDetailLayoutBuilder::GetDetailFont())
+			.IsEnabled(bIsEnabled)
+		]
+		.ValueContent()
+		[
+			SNew(SRigVMEnumPicker)
+			.OnEnumChanged(this, &FRigControlElementDetails::HandleControlEnumChanged, PropertyUtilities)
+			.IsEnabled(bIsEnabled)
+			.GetCurrentEnum_Lambda([this]()
 			{
-				PropertyUtilities->ForceRefresh();
-
-				for(int32 ControlIndex = 0; ControlIndex < PerElementInfos.Num(); ControlIndex++)
+				UEnum* CommonControlEnum = nullptr;
+				for (int32 ControlIndex=0; ControlIndex < PerElementInfos.Num(); ++ControlIndex)
 				{
 					FPerElementInfo& Info = PerElementInfos[ControlIndex];
 					const FRigControlElement ControlInView = Info.WrapperObject->GetContent<FRigControlElement>();
 					FRigControlElement* ControlBeingCustomized = Info.GetDefaultElement<FRigControlElement>();
-					
-					const UEnum* ControlEnum = ControlBeingCustomized->Settings.ControlEnum;
-					if (ControlEnum != nullptr)
+						
+					UEnum* ControlEnum = ControlBeingCustomized->Settings.ControlEnum;
+					if (ControlIndex == 0)
 					{
-						int32 Maximum = (int32)ControlEnum->GetMaxEnumValue() - 1;
-						ControlBeingCustomized->Settings.MinimumValue.Set<int32>(0);
-						ControlBeingCustomized->Settings.MaximumValue.Set<int32>(Maximum);
-						ControlBeingCustomized->Settings.LimitEnabled.Reset();
-						ControlBeingCustomized->Settings.LimitEnabled.Add(true);
-						Info.GetDefaultHierarchy()->SetControlSettings(ControlBeingCustomized, ControlBeingCustomized->Settings, true, true, true);
-
-						FRigControlValue InitialValue = Info.GetDefaultHierarchy()->GetControlValue(ControlBeingCustomized, ERigControlValueType::Initial);
-						FRigControlValue CurrentValue = Info.GetDefaultHierarchy()->GetControlValue(ControlBeingCustomized, ERigControlValueType::Current);
-
-						ControlBeingCustomized->Settings.ApplyLimits(InitialValue);
-						ControlBeingCustomized->Settings.ApplyLimits(CurrentValue);
-						Info.GetDefaultHierarchy()->SetControlValue(ControlBeingCustomized, InitialValue, ERigControlValueType::Initial, false, false, true);
-						Info.GetDefaultHierarchy()->SetControlValue(ControlBeingCustomized, CurrentValue, ERigControlValueType::Current, false, false, true);
-
-						if (UControlRig* DebuggedRig = Cast<UControlRig>(Info.GetBlueprint()->GetObjectBeingDebugged()))
-						{
-							URigHierarchy* DebuggedHierarchy = DebuggedRig->GetHierarchy();
-							if(FRigControlElement* DebuggedControlElement = DebuggedHierarchy->Find<FRigControlElement>(ControlBeingCustomized->GetKey()))
-							{
-								DebuggedControlElement->Settings.MinimumValue.Set<int32>(0);
-                                DebuggedControlElement->Settings.MaximumValue.Set<int32>(Maximum);
-								DebuggedHierarchy->SetControlSettings(DebuggedControlElement, DebuggedControlElement->Settings, true, true, true);
-
-                                DebuggedHierarchy->SetControlValue(DebuggedControlElement, InitialValue, ERigControlValueType::Initial);
-                                DebuggedHierarchy->SetControlValue(DebuggedControlElement, CurrentValue, ERigControlValueType::Current);
-							}
-						}
+						CommonControlEnum = ControlEnum;
+					}
+					else if(ControlEnum != CommonControlEnum)
+					{
+						CommonControlEnum = nullptr;
+						break;
 					}
-
-					Info.WrapperObject->SetContent<FRigControlElement>(*ControlBeingCustomized);
 				}
-			}
-		));
-		
+				return CommonControlEnum;
+			})
+		];
 	}
 
 	const TSharedPtr<IPropertyHandle> CustomizationHandle = SettingsHandle->GetChildHandle(TEXT("Customization"));
@@ -3304,6 +3288,54 @@ void FRigControlElementDetails::CustomizeControl(IDetailLayoutBuilder& DetailBui
 	}
 }
 
+void FRigControlElementDetails::HandleControlEnumChanged(TSharedPtr<FString> InItem, ESelectInfo::Type InSelectionInfo, const TSharedRef<IPropertyUtilities> PropertyUtilities)
+{
+	PropertyUtilities->ForceRefresh();
+	UEnum* ControlEnum = FindObject<UEnum>(nullptr, **InItem.Get(), false);
+
+	for(int32 ControlIndex = 0; ControlIndex < PerElementInfos.Num(); ControlIndex++)
+	{
+		FPerElementInfo& Info = PerElementInfos[ControlIndex];
+		const FRigControlElement ControlInView = Info.WrapperObject->GetContent<FRigControlElement>();
+		FRigControlElement* ControlBeingCustomized = Info.GetDefaultElement<FRigControlElement>();
+		
+		ControlBeingCustomized->Settings.ControlEnum = ControlEnum;
+		if (ControlEnum != nullptr)
+		{
+			int32 Maximum = (int32)ControlEnum->GetMaxEnumValue() - 1;
+			ControlBeingCustomized->Settings.MinimumValue.Set<int32>(0);
+			ControlBeingCustomized->Settings.MaximumValue.Set<int32>(Maximum);
+			ControlBeingCustomized->Settings.LimitEnabled.Reset();
+			ControlBeingCustomized->Settings.LimitEnabled.Add(true);
+			Info.GetDefaultHierarchy()->SetControlSettings(ControlBeingCustomized, ControlBeingCustomized->Settings, true, true, true);
+
+			FRigControlValue InitialValue = Info.GetDefaultHierarchy()->GetControlValue(ControlBeingCustomized, ERigControlValueType::Initial);
+			FRigControlValue CurrentValue = Info.GetDefaultHierarchy()->GetControlValue(ControlBeingCustomized, ERigControlValueType::Current);
+
+			ControlBeingCustomized->Settings.ApplyLimits(InitialValue);
+			ControlBeingCustomized->Settings.ApplyLimits(CurrentValue);
+			Info.GetDefaultHierarchy()->SetControlValue(ControlBeingCustomized, InitialValue, ERigControlValueType::Initial, false, false, true);
+			Info.GetDefaultHierarchy()->SetControlValue(ControlBeingCustomized, CurrentValue, ERigControlValueType::Current, false, false, true);
+
+			if (UControlRig* DebuggedRig = Cast<UControlRig>(Info.GetBlueprint()->GetObjectBeingDebugged()))
+			{
+				URigHierarchy* DebuggedHierarchy = DebuggedRig->GetHierarchy();
+				if(FRigControlElement* DebuggedControlElement = DebuggedHierarchy->Find<FRigControlElement>(ControlBeingCustomized->GetKey()))
+				{
+					DebuggedControlElement->Settings.MinimumValue.Set<int32>(0);
+					DebuggedControlElement->Settings.MaximumValue.Set<int32>(Maximum);
+					DebuggedHierarchy->SetControlSettings(DebuggedControlElement, DebuggedControlElement->Settings, true, true, true);
+
+					DebuggedHierarchy->SetControlValue(DebuggedControlElement, InitialValue, ERigControlValueType::Initial);
+					DebuggedHierarchy->SetControlValue(DebuggedControlElement, CurrentValue, ERigControlValueType::Current);
+				}
+			}
+		}
+
+		Info.WrapperObject->SetContent<FRigControlElement>(*ControlBeingCustomized);
+	}
+}
+
 void FRigControlElementDetails::CustomizeAnimationChannels(IDetailLayoutBuilder& DetailBuilder)
 {
 	// We only show this section for parents of animation channels
diff --git a/Engine/Plugins/Animation/ControlRig/Source/ControlRigEditor/Private/ControlRigElementDetails.h b/Engine/Plugins/Animation/ControlRig/Source/ControlRigEditor/Private/ControlRigElementDetails.h
--- a/Engine/Plugins/Animation/ControlRig/Source/ControlRigEditor/Private/ControlRigElementDetails.h
+++ b/Engine/Plugins/Animation/ControlRig/Source/ControlRigEditor/Private/ControlRigElementDetails.h
@@ -743,6 +743,7 @@ private:
 	TSharedRef<ITableRow> HandleGenerateAnimationChannelTypeRow(TSharedPtr<ERigControlType> ControlType, const TSharedRef<STableViewBase>& OwnerTable, FRigElementKey ControlKey);
 	void HandleControlTypeChanged(TSharedPtr<ERigControlType> ControlType, ESelectInfo::Type SelectInfo, FRigElementKey ControlKey, const TSharedRef<IPropertyUtilities> PropertyUtilities);
 	void HandleControlTypeChanged(ERigControlType ControlType, TArray<FRigElementKey> ControlKeys, const TSharedRef<IPropertyUtilities> PropertyUtilities);
+	void HandleControlEnumChanged(TSharedPtr<FString> InItem, ESelectInfo::Type InSelectionInfo, const TSharedRef<IPropertyUtilities> PropertyUtilities);
 
 	TArray<TSharedPtr<FRigVMStringWithTag>> ShapeNameList;
 	TSharedPtr<FRigInfluenceEntryModifier> InfluenceModifier;
diff --git a/Engine/Plugins/Runtime/RigVM/Source/RigVM/Public/RigVMFunctions/Debug/RigVMFunction_DebugPoint.h b/Engine/Plugins/Runtime/RigVM/Source/RigVM/Public/RigVMFunctions/Debug/RigVMFunction_DebugPoint.h
--- a/Engine/Plugins/Runtime/RigVM/Source/RigVM/Public/RigVMFunctions/Debug/RigVMFunction_DebugPoint.h
+++ b/Engine/Plugins/Runtime/RigVM/Source/RigVM/Public/RigVMFunctions/Debug/RigVMFunction_DebugPoint.h
@@ -5,7 +5,7 @@
 #include "RigVMFunction_DebugBase.h"
 #include "RigVMFunction_DebugPoint.generated.h"
 
-UENUM()
+UENUM(meta = (RigVMTypeAllowed))
 enum class ERigUnitDebugPointMode : uint8
 {
 	/** Draw as point */
diff --git a/Engine/Plugins/Runtime/RigVM/Source/RigVM/Public/RigVMFunctions/Debug/RigVMFunction_DebugTransform.h b/Engine/Plugins/Runtime/RigVM/Source/RigVM/Public/RigVMFunctions/Debug/RigVMFunction_DebugTransform.h
--- a/Engine/Plugins/Runtime/RigVM/Source/RigVM/Public/RigVMFunctions/Debug/RigVMFunction_DebugTransform.h
+++ b/Engine/Plugins/Runtime/RigVM/Source/RigVM/Public/RigVMFunctions/Debug/RigVMFunction_DebugTransform.h
@@ -5,7 +5,7 @@
 #include "RigVMFunction_DebugBase.h"
 #include "RigVMFunction_DebugTransform.generated.h"
 
-UENUM()
+UENUM(meta = (RigVMTypeAllowed))
 enum class ERigUnitDebugTransformMode : uint8
 {
 	/** Draw as point */
diff --git a/Engine/Plugins/Runtime/RigVM/Source/RigVM/Public/RigVMFunctions/Debug/RigVMFunction_VisualDebug.h b/Engine/Plugins/Runtime/RigVM/Source/RigVM/Public/RigVMFunctions/Debug/RigVMFunction_VisualDebug.h
--- a/Engine/Plugins/Runtime/RigVM/Source/RigVM/Public/RigVMFunctions/Debug/RigVMFunction_VisualDebug.h
+++ b/Engine/Plugins/Runtime/RigVM/Source/RigVM/Public/RigVMFunctions/Debug/RigVMFunction_VisualDebug.h
@@ -5,7 +5,7 @@
 #include "RigVMFunction_DebugBase.h"
 #include "RigVMFunction_VisualDebug.generated.h"
 
-UENUM()
+UENUM(meta = (RigVMTypeAllowed))
 enum class ERigUnitVisualDebugPointMode : uint8
 {
 	/** Draw as point */
diff --git a/Engine/Plugins/Runtime/RigVM/Source/RigVM/Public/RigVMFunctions/Math/RigVMMathLibrary.h b/Engine/Plugins/Runtime/RigVM/Source/RigVM/Public/RigVMFunctions/Math/RigVMMathLibrary.h
--- a/Engine/Plugins/Runtime/RigVM/Source/RigVM/Public/RigVMFunctions/Math/RigVMMathLibrary.h
+++ b/Engine/Plugins/Runtime/RigVM/Source/RigVM/Public/RigVMFunctions/Math/RigVMMathLibrary.h
@@ -8,7 +8,7 @@
 #include "RigVMFunctions/RigVMFunctionDefines.h"
 #include "RigVMMathLibrary.generated.h"
 
-UENUM()
+UENUM(meta = (RigVMTypeAllowed))
 enum class ERigVMAnimEasingType : uint8
 {
 	Linear,
@@ -98,7 +98,7 @@ struct RIGVM_API FRigVMMirrorSettings
 	FVector MirrorVector(const FVector& InVector) const;
 };
 
-UENUM()
+UENUM(meta = (RigVMTypeAllowed))
 enum class ERigVMSimPointIntegrateType : uint8
 {
 	Verlet,
diff --git a/Engine/Plugins/Runtime/RigVM/Source/RigVM/Public/RigVMFunctions/RigVMFunctionDefines.h b/Engine/Plugins/Runtime/RigVM/Source/RigVM/Public/RigVMFunctions/RigVMFunctionDefines.h
--- a/Engine/Plugins/Runtime/RigVM/Source/RigVM/Public/RigVMFunctions/RigVMFunctionDefines.h
+++ b/Engine/Plugins/Runtime/RigVM/Source/RigVM/Public/RigVMFunctions/RigVMFunctionDefines.h
@@ -7,7 +7,7 @@
 
 #include "RigVMFunctionDefines.generated.h"
 
-UENUM()
+UENUM(meta = (RigVMTypeAllowed))
 enum class ERigVMTransformSpace : uint8
 {
 	/** Apply in parent space */
@@ -20,7 +20,7 @@ enum class ERigVMTransformSpace : uint8
 	Max UMETA(Hidden),
 };
 
-UENUM()
+UENUM(meta = (RigVMTypeAllowed))
 namespace ERigVMClampSpatialMode
 {
 	enum Type : int
diff --git a/Engine/Plugins/Runtime/RigVM/Source/RigVMEditor/Private/EdGraph/RigVMEdGraphPanelPinFactory.cpp b/Engine/Plugins/Runtime/RigVM/Source/RigVMEditor/Private/EdGraph/RigVMEdGraphPanelPinFactory.cpp
--- a/Engine/Plugins/Runtime/RigVM/Source/RigVMEditor/Private/EdGraph/RigVMEdGraphPanelPinFactory.cpp
+++ b/Engine/Plugins/Runtime/RigVM/Source/RigVMEditor/Private/EdGraph/RigVMEdGraphPanelPinFactory.cpp
@@ -19,6 +19,7 @@
 #include "Curves/CurveFloat.h"
 #include "RigVMModel/Nodes/RigVMUnitNode.h"
 #include "RigVMCore/RigVMExecuteContext.h"
+#include "Widgets/SRigVMGraphPinEnumPicker.h"
 
 FName FRigVMEdGraphPanelPinFactory::GetFactoryName() const
 {
@@ -113,6 +114,12 @@ TSharedPtr<SGraphPin> FRigVMEdGraphPanelPinFactory::CreatePin_Internal(UEdGraphP
 					return SNew(SRigVMGraphPinUserDataPath, InPin)
 						.ModelPins({ModelPin});
 				}
+
+				if (ModelPin->GetCPPTypeObject() == UEnum::StaticClass())
+				{
+					return SNew(SRigVMGraphPinEnumPicker, InPin)
+						.ModelPin(ModelPin);
+				}
 			}
 
 			if (InPin->PinType.PinCategory == UEdGraphSchema_K2::PC_Struct)
diff --git a/Engine/Plugins/Runtime/RigVM/Source/RigVMEditor/Private/Editor/RigVMEditor.cpp b/Engine/Plugins/Runtime/RigVM/Source/RigVMEditor/Private/Editor/RigVMEditor.cpp
--- a/Engine/Plugins/Runtime/RigVM/Source/RigVMEditor/Private/Editor/RigVMEditor.cpp
+++ b/Engine/Plugins/Runtime/RigVM/Source/RigVMEditor/Private/Editor/RigVMEditor.cpp
@@ -372,6 +372,12 @@ void FRigVMEditor::InitRigVMEditor(const EToolkitMode::Type Mode, const TSharedP
 			}));
 	}
 
+	Inspector->GetPropertyView()->RegisterInstancedCustomPropertyTypeLayout(UEnum::StaticClass()->GetFName(),
+		FOnGetPropertyTypeCustomizationInstance::CreateLambda([=]()
+		{
+			return FRigVMGraphEnumDetailCustomization::MakeInstance();
+		}));
+
 	PropertyChangedHandle = FCoreUObjectDelegates::OnObjectPropertyChanged.AddSP(this, &FRigVMEditor::OnPropertyChanged);
 }
 
diff --git a/Engine/Plugins/Runtime/RigVM/Source/RigVMEditor/Private/Editor/RigVMGraphDetailCustomization.cpp b/Engine/Plugins/Runtime/RigVM/Source/RigVMEditor/Private/Editor/RigVMGraphDetailCustomization.cpp
--- a/Engine/Plugins/Runtime/RigVM/Source/RigVMEditor/Private/Editor/RigVMGraphDetailCustomization.cpp
+++ b/Engine/Plugins/Runtime/RigVM/Source/RigVMEditor/Private/Editor/RigVMGraphDetailCustomization.cpp
@@ -26,6 +26,7 @@
 #include "RigVMModel/Nodes/RigVMAggregateNode.h"
 #include "Widgets/SRigVMGraphPinVariableBinding.h"
 #include "InstancedPropertyBagStructureDataProvider.h"
+#include "Widgets/SRigVMGraphPinEnumPicker.h"
 
 #define LOCTEXT_NAMESPACE "RigVMGraphDetailCustomization"
 
@@ -1834,6 +1835,148 @@ void FRigVMWrappedNodeDetailCustomization::CustomizeLiveValues(IDetailLayoutBuil
 	*/
 }
 
+FRigVMGraphEnumDetailCustomization::FRigVMGraphEnumDetailCustomization()
+: BlueprintBeingCustomized(nullptr)
+, GraphBeingCustomized(nullptr)
+{
+}
+
+void FRigVMGraphEnumDetailCustomization::CustomizeHeader(TSharedRef<IPropertyHandle> InPropertyHandle, FDetailWidgetRow& HeaderRow, IPropertyTypeCustomizationUtils& StructCustomizationUtils)
+{
+	TArray<UObject*> Objects;
+	InPropertyHandle->GetOuterObjects(Objects);
+
+	StructsBeingCustomized.Reset();
+	InPropertyHandle->GetOuterStructs(StructsBeingCustomized);
+
+	for (UObject* Object : Objects)
+	{
+		ObjectsBeingCustomized.Add(Object);
+
+		if(BlueprintBeingCustomized == nullptr)
+		{
+			BlueprintBeingCustomized = Object->GetTypedOuter<URigVMBlueprint>();
+		}
+
+		if(GraphBeingCustomized == nullptr)
+		{
+			GraphBeingCustomized = Object->GetTypedOuter<URigVMGraph>();
+		}
+	}
+
+	FProperty* Property = InPropertyHandle->GetProperty();
+	const FObjectProperty* ObjectProperty = CastField<FObjectProperty>(Property);
+
+	HeaderRow
+		.NameContent()
+		[
+			InPropertyHandle->CreatePropertyNameWidget()
+		]
+		.ValueContent()
+			.MinDesiredWidth(375.f)
+			.MaxDesiredWidth(375.f)
+			.HAlign(HAlign_Left)
+			[
+				SNew(SBox)
+					.MinDesiredWidth(150)
+					.MaxDesiredWidth(400)
+					[
+						SNew(SRigVMEnumPicker)
+						.IsEnabled(true)
+						.OnEnumChanged(this, &FRigVMGraphEnumDetailCustomization::HandleControlEnumChanged, InPropertyHandle)
+						.GetCurrentEnum_Lambda([this, InPropertyHandle]()
+						{
+							UEnum* Enum = nullptr;
+							FEditPropertyChain PropertyChain;
+							TArray<int32> PropertyArrayIndices;
+							bool bEnabled;
+							if (!GetPropertyChain(InPropertyHandle, PropertyChain, PropertyArrayIndices, bEnabled))
+							{
+								return Enum;
+							}
+
+							const TArray<uint8*> MemoryBlocks = GetMemoryBeingCustomized();
+							for(uint8* MemoryBlock: MemoryBlocks)
+							{
+								if(MemoryBlock)
+								{
+									if (UEnum** CurrentEnum = ContainerMemoryBlockToEnumPtr(MemoryBlock, PropertyChain, PropertyArrayIndices))
+									{
+										Enum = *CurrentEnum;
+									}
+								}
+							}
+							return Enum;
+						})
+					]
+				];
+}
+
+void FRigVMGraphEnumDetailCustomization::CustomizeChildren(TSharedRef<IPropertyHandle> InPropertyHandle, IDetailChildrenBuilder& StructBuilder, IPropertyTypeCustomizationUtils& StructCustomizationUtils)
+{
+	// nothing to do here
+}
+
+void FRigVMGraphEnumDetailCustomization::HandleControlEnumChanged(TSharedPtr<FString> InEnumPath, ESelectInfo::Type InSelectType, TSharedRef<IPropertyHandle> InPropertyHandle)
+{
+	if (ObjectsBeingCustomized.IsEmpty())
+	{
+		return;
+	}
+		
+	FEditPropertyChain PropertyChain;
+	TArray<int32> PropertyArrayIndices;
+	bool bEnabled;
+	if (!GetPropertyChain(InPropertyHandle, PropertyChain, PropertyArrayIndices, bEnabled))
+	{
+		return;
+	}
+	
+	TArray<UObject*> ObjectsView;
+	for(int32 Index = 0; Index < ObjectsBeingCustomized.Num(); Index++)
+	{
+		const TWeakObjectPtr<UObject>& Object = ObjectsBeingCustomized[Index];
+		if (Object.Get())
+		{
+			ObjectsView.Add(Object.Get());
+		}
+	}
+	FPropertyChangedEvent PropertyChangedEvent(InPropertyHandle->GetProperty(), EPropertyChangeType::ValueSet, ObjectsView);
+	FPropertyChangedChainEvent PropertyChangedChainEvent(PropertyChain, PropertyChangedEvent);
+
+	URigVMController* Controller = nullptr;
+	if(BlueprintBeingCustomized && GraphBeingCustomized)
+	{
+		Controller = BlueprintBeingCustomized->GetController(GraphBeingCustomized);
+		Controller->OpenUndoBracket(FString::Printf(TEXT("Set %s"), *InPropertyHandle->GetProperty()->GetName()));
+	}
+
+	for(int32 Index = 0; Index < ObjectsBeingCustomized.Num(); Index++)
+	{
+		const TWeakObjectPtr<UObject>& Object = ObjectsBeingCustomized[Index];
+		if(Object.Get() && InPropertyHandle->IsValidHandle())
+		{
+			UEnum** CurrentEnum = ContainerMemoryBlockToEnumPtr((uint8*)Object.Get(), PropertyChain, PropertyArrayIndices);
+			if (CurrentEnum)
+			{
+				const UEnum* PreviousEnum = *CurrentEnum;
+				*CurrentEnum = FindObject<UEnum>(nullptr, **InEnumPath.Get(), false);
+
+				if (PreviousEnum != *CurrentEnum)
+				{
+					Object->PostEditChangeChainProperty(PropertyChangedChainEvent);
+					InPropertyHandle->NotifyPostChange(PropertyChangedEvent.ChangeType);
+				}
+			}
+		}
+	}
+
+	if(Controller)
+	{
+		Controller->CloseUndoBracket();
+	}
+}
+
 template<typename VectorType, int32 NumberOfComponents>
 void FRigVMGraphMathTypeDetailCustomization::MakeVectorHeaderRow(TSharedRef<class IPropertyHandle> InPropertyHandle, FDetailWidgetRow& HeaderRow, IPropertyTypeCustomizationUtils& StructCustomizationUtils)
 {
diff --git a/Engine/Plugins/Runtime/RigVM/Source/RigVMEditor/Private/Widgets/SRigVMGraphPinEnumPicker.cpp b/Engine/Plugins/Runtime/RigVM/Source/RigVMEditor/Private/Widgets/SRigVMGraphPinEnumPicker.cpp
new file mode 100644
--- /dev/null
+++ b/Engine/Plugins/Runtime/RigVM/Source/RigVMEditor/Private/Widgets/SRigVMGraphPinEnumPicker.cpp
@@ -0,0 +1,133 @@
+// Copyright Epic Games, Inc. All Rights Reserved.
+
+
+#include "Widgets/SRigVMGraphPinEnumPicker.h"
+
+#include "DetailLayoutBuilder.h"
+#include "EdGraph/RigVMEdGraph.h"
+#include "RigVMBlueprint.h"
+#include "Widgets/Layout/SBox.h"
+#include "ScopedTransaction.h"
+#include "SSearchableComboBox.h"
+#include "UObject/UObjectIterator.h"
+
+void SRigVMEnumPicker::Construct(const FArguments& InArgs)
+{
+	OnEnumChangedDelegate = InArgs._OnEnumChanged;
+	bIsEnabled = InArgs._IsEnabled;
+	GetCurrentEnumDelegate = InArgs._GetCurrentEnum;
+
+	PopulateEnumOptions();
+
+	ChildSlot
+	[
+		SNew(SBox)
+		.MinDesiredWidth(150)
+		.MaxDesiredWidth(400)
+		[
+			SNew(SSearchableComboBox)
+			.OptionsSource(&EnumOptions)
+			.OnSelectionChanged(this, &SRigVMEnumPicker::HandleControlEnumChanged)
+			.OnGenerateWidget(this, &SRigVMEnumPicker::OnGetEnumNameWidget)
+			.Content()
+			[
+				SNew(STextBlock)
+				.Text_Lambda([this]() -> FText
+				{
+					const UEnum* ControlEnum = GetCurrentEnum();
+					if (ControlEnum)
+					{
+						return FText::FromName(ControlEnum->GetFName());
+					}
+					return FText();
+				})
+				.Font(IDetailLayoutBuilder::GetDetailFont())
+			]
+		]
+	];
+}
+
+TSharedRef<SWidget> SRigVMEnumPicker::OnGetEnumNameWidget(TSharedPtr<FString> InItem)
+{
+	if (InItem.IsValid())
+	{
+		UEnum* Enum = FindObject<UEnum>(nullptr, **InItem.Get(), false);
+		if (Enum)
+		{
+			return SNew(STextBlock)
+			.Text(FText::FromString(Enum->GetName()))
+			.ToolTip(FSlateApplication::Get().MakeToolTip(FText::FromString(Enum->GetPathName())))
+			.Font(IDetailLayoutBuilder::GetDetailFont());
+		}
+	}
+	return SNew(STextBlock)
+	.Text(FText::FromString(TEXT("None")))
+	.Font(IDetailLayoutBuilder::GetDetailFont());
+}
+
+void SRigVMEnumPicker::PopulateEnumOptions()
+{
+	EnumOptions.Reset();
+	EnumOptions.Add(MakeShareable(new FString(TEXT("None"))));
+	for (TObjectIterator<UEnum> EnumIt; EnumIt; ++EnumIt)
+	{
+		UEnum* Enum = *EnumIt;
+
+		if (Enum->HasAnyFlags(RF_BeginDestroyed | RF_FinishDestroyed) || !Enum->HasAllFlags(RF_Public))
+		{
+			continue;
+		}
+
+		// Any asset based enum is valid
+		if (!Enum->IsAsset())
+		{
+			// Native enums only allowed if contain RigVMTypeAllowed metadata
+			if (!Enum->HasMetaData(TEXT("RigVMTypeAllowed")))
+			{
+				continue;
+			}
+		}
+
+		EnumOptions.Add(MakeShareable(new FString(Enum->GetPathName())));
+	}
+}
+
+void SRigVMGraphPinEnumPicker::Construct(const FArguments& InArgs, UEdGraphPin* InGraphPinObj)
+{
+	ModelPin = InArgs._ModelPin;
+	SGraphPin::Construct(SGraphPin::FArguments(), InGraphPinObj);
+}
+
+TSharedRef<SWidget>	SRigVMGraphPinEnumPicker::GetDefaultValueWidget()
+{
+	return SNew(SRigVMEnumPicker)
+	.GetCurrentEnum_Lambda([this]()
+	{
+		UEnum* ControlEnum = nullptr;
+		if (ModelPin)
+		{
+			const FString& Path = ModelPin->GetDefaultValue();
+			ControlEnum = FindObject<UEnum>(nullptr, *Path, false);
+		}
+		return ControlEnum;
+	})
+	.OnEnumChanged(this, &SRigVMGraphPinEnumPicker::HandleControlEnumChanged);
+}
+
+void SRigVMGraphPinEnumPicker::HandleControlEnumChanged(TSharedPtr<FString> InItem, ESelectInfo::Type InSelectionInfo)
+{
+	if(ModelPin)
+	{
+		if(URigVMBlueprint* Blueprint = ModelPin->GetTypedOuter<URigVMBlueprint>())
+		{
+			if(URigVMController* Controller = Blueprint->GetOrCreateController(ModelPin->GetGraph()))
+			{
+				Controller->SetPinDefaultValue(ModelPin->GetPinPath(), *InItem.Get());
+			}
+		}
+	}
+}
+
+
+
+
diff --git a/Engine/Plugins/Runtime/RigVM/Source/RigVMEditor/Public/Editor/RigVMGraphDetailCustomization.h b/Engine/Plugins/Runtime/RigVM/Source/RigVMEditor/Public/Editor/RigVMGraphDetailCustomization.h
--- a/Engine/Plugins/Runtime/RigVM/Source/RigVMEditor/Public/Editor/RigVMGraphDetailCustomization.h
+++ b/Engine/Plugins/Runtime/RigVM/Source/RigVMEditor/Public/Editor/RigVMGraphDetailCustomization.h
@@ -240,6 +240,135 @@ public:
 	TMap<FName, TSharedPtr<SRigVMGraphPinNameListValueWidget>> NameListWidgets;
 };
 
+/** Customization for editing a rig vm integer control enum class */
+class RIGVMEDITOR_API FRigVMGraphEnumDetailCustomization : public IPropertyTypeCustomization
+{
+public:
+
+	FRigVMGraphEnumDetailCustomization();
+	
+	static TSharedRef<IPropertyTypeCustomization> MakeInstance()
+	{
+		return MakeShareable(new FRigVMGraphEnumDetailCustomization);
+	}
+
+	/** IPropertyTypeCustomization interface */
+	virtual void CustomizeHeader(TSharedRef<class IPropertyHandle> InPropertyHandle, class FDetailWidgetRow& HeaderRow, IPropertyTypeCustomizationUtils& StructCustomizationUtils) override;
+	virtual void CustomizeChildren(TSharedRef<class IPropertyHandle> InPropertyHandle, class IDetailChildrenBuilder& StructBuilder, IPropertyTypeCustomizationUtils& StructCustomizationUtils) override;
+
+protected:
+
+	TArray<uint8*> GetMemoryBeingCustomized()
+	{
+		TArray<uint8*> MemoryPtr;
+		MemoryPtr.Reserve(ObjectsBeingCustomized.Num() + StructsBeingCustomized.Num());
+
+		for(const TWeakObjectPtr<UObject>& Object : ObjectsBeingCustomized)
+		{
+			if(Object.IsValid())
+			{
+				MemoryPtr.Add((uint8*)Object.Get());
+			}
+		}
+
+		for(const TSharedPtr<FStructOnScope>& StructPtr: StructsBeingCustomized)
+		{
+			if(StructPtr.IsValid())
+			{
+				MemoryPtr.Add(StructPtr->GetStructMemory());
+			}
+		}
+
+		return MemoryPtr;
+	}
+	
+	bool GetPropertyChain(TSharedRef<class IPropertyHandle> InPropertyHandle, FEditPropertyChain& OutPropertyChain, TArray<int32> &OutPropertyArrayIndices, bool& bOutEnabled)
+	{
+		if (!InPropertyHandle->IsValidHandle())
+		{
+			return false;
+		}
+		
+		OutPropertyChain.Empty();
+		OutPropertyArrayIndices.Reset();
+		bOutEnabled = false;
+
+		const bool bHasObject = !ObjectsBeingCustomized.IsEmpty() && ObjectsBeingCustomized[0].Get();
+		const bool bHasStruct = !StructsBeingCustomized.IsEmpty() && StructsBeingCustomized[0].Get();
+		
+		if (bHasStruct || bHasObject)
+		{
+			TSharedPtr<class IPropertyHandle> ChainHandle = InPropertyHandle;
+			while (ChainHandle.IsValid() && ChainHandle->GetProperty() != nullptr)
+			{
+				OutPropertyChain.AddHead(ChainHandle->GetProperty());
+				OutPropertyArrayIndices.Insert(ChainHandle->GetIndexInArray(), 0);
+				ChainHandle = ChainHandle->GetParentHandle();
+			}
+
+			if (OutPropertyChain.GetHead() != nullptr)
+			{
+				OutPropertyChain.SetActiveMemberPropertyNode(OutPropertyChain.GetTail()->GetValue());
+				bOutEnabled = !OutPropertyChain.GetHead()->GetValue()->HasAnyPropertyFlags(CPF_EditConst);
+				return true;
+			}
+		}
+		return false;
+	}
+
+	// extracts the value for a nested property from an outer owner
+	static UEnum** ContainerMemoryBlockToEnumPtr(uint8* InMemoryBlock, FEditPropertyChain& InPropertyChain, TArray<int32> &InPropertyArrayIndices)
+	{
+		if (InPropertyChain.GetHead() == nullptr)
+		{
+			return nullptr;
+		}
+		
+		FEditPropertyChain::TDoubleLinkedListNode* PropertyNode = InPropertyChain.GetHead();
+		uint8* MemoryPtr = InMemoryBlock;
+		int32 ChainIndex = 0;
+		do
+		{
+			const FProperty* Property = PropertyNode->GetValue();
+			MemoryPtr = Property->ContainerPtrToValuePtr<uint8>(MemoryPtr);
+
+			PropertyNode = PropertyNode->GetNextNode();
+			ChainIndex++;
+			
+			if(InPropertyArrayIndices.IsValidIndex(ChainIndex))
+			{
+				const int32 ArrayIndex = InPropertyArrayIndices[ChainIndex];
+				if(ArrayIndex != INDEX_NONE)
+				{
+					const FArrayProperty* ArrayProperty = CastField<FArrayProperty>(Property->GetOwnerProperty());
+					check(ArrayProperty);
+					
+					FScriptArrayHelper ArrayHelper(ArrayProperty, MemoryPtr);
+					if(!ArrayHelper.IsValidIndex(ArrayIndex))
+					{
+						return nullptr;
+					}
+					MemoryPtr = ArrayHelper.GetRawPtr(ArrayIndex);
+
+					// skip to the next property node already
+					PropertyNode = PropertyNode->GetNextNode();
+					ChainIndex++;
+				}
+			}
+		}
+		while (PropertyNode);
+
+		return (UEnum**)MemoryPtr;
+	}
+
+	void HandleControlEnumChanged(TSharedPtr<FString> InEnumPath, ESelectInfo::Type InSelectType, TSharedRef<IPropertyHandle> InPropertyHandle);
+
+	URigVMBlueprint* BlueprintBeingCustomized;
+	URigVMGraph* GraphBeingCustomized;
+	TArray<TWeakObjectPtr<UObject>> ObjectsBeingCustomized;
+	TArray<TSharedPtr<FStructOnScope>> StructsBeingCustomized;
+};
+
 /** Customization for editing a rig vm node */
 class RIGVMEDITOR_API FRigVMGraphMathTypeDetailCustomization : public IPropertyTypeCustomization
 {
diff --git a/Engine/Plugins/Runtime/RigVM/Source/RigVMEditor/Public/Widgets/SRigVMGraphPinEnumPicker.h b/Engine/Plugins/Runtime/RigVM/Source/RigVMEditor/Public/Widgets/SRigVMGraphPinEnumPicker.h
new file mode 100644
--- /dev/null
+++ b/Engine/Plugins/Runtime/RigVM/Source/RigVMEditor/Public/Widgets/SRigVMGraphPinEnumPicker.h
@@ -0,0 +1,72 @@
+// Copyright Epic Games, Inc. All Rights Reserved.
+
+#pragma once
+
+#include "CoreMinimal.h"
+#include "Widgets/DeclarativeSyntaxSupport.h"
+#include "Widgets/SWidget.h"
+#include "SGraphPin.h"
+#include "SSearchableComboBox.h"
+#include "RigVMModel/RigVMPin.h"
+
+class RIGVMEDITOR_API SRigVMEnumPicker : public SCompoundWidget
+{
+public:
+	DECLARE_DELEGATE_RetVal(TObjectPtr<UEnum>, FGetCurrentEnum);
+	
+	SLATE_BEGIN_ARGS(SRigVMEnumPicker){}
+		SLATE_EVENT(SSearchableComboBox::FOnSelectionChanged, OnEnumChanged)
+		SLATE_EVENT(FGetCurrentEnum, GetCurrentEnum)
+		SLATE_ARGUMENT(bool, IsEnabled)
+	SLATE_END_ARGS()
+
+	void Construct(const FArguments& InArgs);
+
+protected:
+
+	void HandleControlEnumChanged(TSharedPtr<FString> InItem, ESelectInfo::Type InSelectionInfo)
+	{
+		if (OnEnumChangedDelegate.IsBound())
+		{
+			return OnEnumChangedDelegate.Execute(InItem, InSelectionInfo);
+		}
+	}
+
+	UEnum* GetCurrentEnum()
+	{
+		if (GetCurrentEnumDelegate.IsBound())
+		{
+			return GetCurrentEnumDelegate.Execute();
+		}
+		return nullptr;
+	}
+
+	TSharedRef<SWidget> OnGetEnumNameWidget(TSharedPtr<FString> InItem);
+	void PopulateEnumOptions();
+
+	TArray<TSharedPtr<FString>> EnumOptions;
+	SSearchableComboBox::FOnSelectionChanged OnEnumChangedDelegate;
+	FGetCurrentEnum GetCurrentEnumDelegate;
+	bool bIsEnabled;
+};
+
+class RIGVMEDITOR_API SRigVMGraphPinEnumPicker : public SGraphPin
+{
+public:
+	SLATE_BEGIN_ARGS(SRigVMGraphPinEnumPicker)
+	: _ModelPin(nullptr) {}
+		SLATE_ARGUMENT(URigVMPin*, ModelPin)
+	SLATE_END_ARGS()
+
+	void Construct(const FArguments& InArgs, UEdGraphPin* InGraphPinObj);
+
+protected:
+	//~ Begin SGraphPin Interface
+	virtual TSharedRef<SWidget>	GetDefaultValueWidget() override;
+	//~ End SGraphPin Interface
+
+	void HandleControlEnumChanged(TSharedPtr<FString> InItem, ESelectInfo::Type InSelectionInfo);
+	
+	URigVMPin* ModelPin;
+
+};
