Sequencer: Anim Details:  Use Relative values directly to avoid quat/rot conversions.

diff --git a/Engine/Plugins/Animation/ControlRig/Source/ControlRigEditor/Private/EditMode/AnimDetailsProxy.cpp b/Engine/Plugins/Animation/ControlRig/Source/ControlRigEditor/Private/EditMode/AnimDetailsProxy.cpp
--- a/Engine/Plugins/Animation/ControlRig/Source/ControlRigEditor/Private/EditMode/AnimDetailsProxy.cpp
+++ b/Engine/Plugins/Animation/ControlRig/Source/ControlRigEditor/Private/EditMode/AnimDetailsProxy.cpp
@@ -910,8 +910,9 @@ static FEulerTransform GetCurrentValue(UObject* InObject, TSharedPtr<FTrackInsta
 	FEulerTransform EulerTransform = FEulerTransform::Identity;
 	if (SceneComponentThatChanged)
 	{
-		FTransform Transform(SceneComponentThatChanged->GetRelativeTransform());
-		EulerTransform = FEulerTransform(Transform);
+		EulerTransform.Location = SceneComponentThatChanged->GetRelativeLocation();
+		EulerTransform.Rotation = SceneComponentThatChanged->GetRelativeRotation();
+		EulerTransform.Scale = SceneComponentThatChanged->GetRelativeScale3D();
 	}
 	return EulerTransform;
 }
@@ -1017,7 +1018,7 @@ void UAnimDetailControlsProxyTransform::SetBindingValueFromCurrent(UObject* InOb
 			SceneComponentThatChanged->SetRelativeTransform(RealTransform, false, nullptr, ETeleportType::None);
 			// Force the location and rotation values to avoid Rot->Quat->Rot conversions
 			SceneComponentThatChanged->SetRelativeLocation_Direct(TLocation);
-			SceneComponentThatChanged->SetRelativeRotation_Direct(TRotation);
+			SceneComponentThatChanged->SetRelativeRotationExact(TRotation);
 
 			if (ValueProperty)
 			{
