Control Rig: mark the hierarchy as modified when switching from rig to module
    + added a safety guard at the details level

#rnx

diff --git a/Engine/Plugins/Animation/ControlRig/Source/ControlRigDeveloper/Private/ControlRigBlueprint.cpp b/Engine/Plugins/Animation/ControlRig/Source/ControlRigDeveloper/Private/ControlRigBlueprint.cpp
--- a/Engine/Plugins/Animation/ControlRig/Source/ControlRigDeveloper/Private/ControlRigBlueprint.cpp
+++ b/Engine/Plugins/Animation/ControlRig/Source/ControlRigDeveloper/Private/ControlRigBlueprint.cpp
@@ -274,6 +274,7 @@ bool UControlRigBlueprint::TurnIntoControlRigModule(bool InAutoConvertHierarchy,
 
 	if(Hierarchy)
 	{
+		Hierarchy->Modify();
 		URigHierarchyController* Controller = Hierarchy->GetController(true);
 
 		// create a copy of this hierarchy
@@ -342,7 +343,7 @@ bool UControlRigBlueprint::TurnIntoStandaloneRig(FString* OutErrorMessage)
 		return false;
 	}
 
-	FScopedTransaction Transaction(LOCTEXT("TurnIntoControlRigModule", "Turn Rig into Module"));
+	FScopedTransaction Transaction(LOCTEXT("TurnIntoStandaloneRig", "Turn Module into Rig"));
 
 	Modify();
 	RigModuleSettings = FRigModuleSettings();
diff --git a/Engine/Plugins/Animation/ControlRig/Source/ControlRigEditor/Private/ControlRigElementDetails.cpp b/Engine/Plugins/Animation/ControlRig/Source/ControlRigEditor/Private/ControlRigElementDetails.cpp
--- a/Engine/Plugins/Animation/ControlRig/Source/ControlRigEditor/Private/ControlRigElementDetails.cpp
+++ b/Engine/Plugins/Animation/ControlRig/Source/ControlRigEditor/Private/ControlRigElementDetails.cpp
@@ -784,7 +784,10 @@ void FRigBaseElementDetails::CustomizeDetails(IDetailLayoutBuilder& DetailBuilde
 
 		FPerElementInfo Info;
 		Info.WrapperObject = WrapperObject;
-		Info.Element = Cast<URigHierarchy>(WrapperObject->GetSubject())->GetHandle(Key);
+		if (const URigHierarchy* Hierarchy = Cast<URigHierarchy>(WrapperObject->GetSubject()))
+		{
+			Info.Element = Hierarchy->GetHandle(Key);
+		}
 
 		if(!Info.Element.IsValid())
 		{
