Control Rig: Use a soft object ptr to refer to external assets from asset user data objects.

This allows the referring objects to cook even if the assets cannot be referenced.

diff --git a/Engine/Plugins/Animation/ControlRig/Source/ControlRig/Private/ControlRigAssetUserData.cpp b/Engine/Plugins/Animation/ControlRig/Source/ControlRig/Private/ControlRigAssetUserData.cpp
--- a/Engine/Plugins/Animation/ControlRig/Source/ControlRig/Private/ControlRigAssetUserData.cpp
+++ b/Engine/Plugins/Animation/ControlRig/Source/ControlRig/Private/ControlRigAssetUserData.cpp
@@ -2,14 +2,15 @@
 
 #include "ControlRigAssetUserData.h"
 
-void UControlRigShapeLibraryLink::SetShapeLibrary(UControlRigShapeLibrary* InShapeLibrary)
+void UControlRigShapeLibraryLink::SetShapeLibrary(TSoftObjectPtr<UControlRigShapeLibrary> InShapeLibrary)
 {
 	InvalidateCache();
 	ShapeLibrary = InShapeLibrary;
+	ShapeLibraryCached = InShapeLibrary.Get();
 	ShapeNames.Reset();
-	if(ShapeLibrary)
+	if(ShapeLibraryCached)
 	{
-		for(const FControlRigShapeDefinition& Shape : ShapeLibrary->Shapes)
+		for(const FControlRigShapeDefinition& Shape : ShapeLibraryCached->Shapes)
 		{
 			ShapeNames.Add(Shape.ShapeName);
 		}
@@ -28,18 +29,18 @@ const UNameSpacedUserData::FUserData* UControlRigShapeLibraryLink::GetUserData(c
 	{
 		if(InPath.Equals(GET_MEMBER_NAME_STRING_CHECKED(UControlRigShapeLibraryLink, ShapeLibrary), ESearchCase::CaseSensitive))
 		{
-			static const FProperty* ShapeLibraryProperty = FindPropertyByName(StaticClass(), GET_MEMBER_NAME_CHECKED(UControlRigShapeLibraryLink, ShapeLibrary));
-			return StoreCacheForUserData({InPath, ShapeLibraryProperty, (const uint8*)ShapeLibrary});
+			static const FProperty* ShapeLibraryProperty = FindPropertyByName(StaticClass(), GET_MEMBER_NAME_CHECKED(UControlRigShapeLibraryLink, ShapeLibraryCached));
+			return StoreCacheForUserData({InPath, ShapeLibraryProperty, static_cast<const uint8*>(ShapeLibraryCached)});
 		}
 		if(InPath.Equals(DefaultShapePath, ESearchCase::CaseSensitive))
 		{
 			static const FProperty* DefaultShapeProperty = FindPropertyByName(FControlRigShapeDefinition::StaticStruct(), GET_MEMBER_NAME_CHECKED(FControlRigShapeDefinition, ShapeName));
-			return StoreCacheForUserData({InPath, DefaultShapeProperty, (const uint8*)&ShapeLibrary->DefaultShape.ShapeName});
+			return StoreCacheForUserData({InPath, DefaultShapeProperty, reinterpret_cast<const uint8*>(&ShapeLibrary->DefaultShape.ShapeName)});
 		}
 		if(InPath.Equals(ShapeNamesPath, ESearchCase::CaseSensitive))
 		{
 			static const FProperty* ShapeNamesProperty = FindPropertyByName(StaticClass(), GET_MEMBER_NAME_CHECKED(UControlRigShapeLibraryLink, ShapeNames));
-			return StoreCacheForUserData({InPath, ShapeNamesProperty, (const uint8*)&ShapeNames});
+			return StoreCacheForUserData({InPath, ShapeNamesProperty, reinterpret_cast<const uint8*>(&ShapeNames)});
 		}
 		if(OutErrorMessage && OutErrorMessage->IsEmpty())
 		{
@@ -86,6 +87,13 @@ const TArray<const UNameSpacedUserData::FUserData*>& UControlRigShapeLibraryLink
 	return EmptyUserDatas;
 }
 
+void UControlRigShapeLibraryLink::PostLoad()
+{
+	Super::PostLoad();
+
+	ShapeLibraryCached = ShapeLibrary.Get();
+}
+
 #if WITH_EDITOR
 
 void UControlRigShapeLibraryLink::PostEditChangeProperty(FPropertyChangedEvent& PropertyChangedEvent)
diff --git a/Engine/Plugins/Animation/ControlRig/Source/ControlRig/Public/ControlRigAssetUserData.h b/Engine/Plugins/Animation/ControlRig/Source/ControlRig/Public/ControlRigAssetUserData.h
--- a/Engine/Plugins/Animation/ControlRig/Source/ControlRig/Public/ControlRigAssetUserData.h
+++ b/Engine/Plugins/Animation/ControlRig/Source/ControlRig/Public/ControlRigAssetUserData.h
@@ -18,17 +18,18 @@ public:
 
 	/** If assigned, the data asset link will provide access to the data asset's content. */
 	UPROPERTY(EditAnywhere, BlueprintReadWrite, Category = General, BlueprintGetter = GetShapeLibrary, BlueprintSetter = SetShapeLibrary, Meta = (DisplayAfter="NameSpace"))
-	TObjectPtr<UControlRigShapeLibrary> ShapeLibrary;
+	TSoftObjectPtr<UControlRigShapeLibrary> ShapeLibrary;
 
 	UFUNCTION(BlueprintGetter)
-	UControlRigShapeLibrary* GetShapeLibrary() const { return ShapeLibrary; }
+	TSoftObjectPtr<UControlRigShapeLibrary> GetShapeLibrary() const { return ShapeLibrary; }
 
 	UFUNCTION(BlueprintSetter)
-	void SetShapeLibrary(UControlRigShapeLibrary* InShapeLibrary);
+	void SetShapeLibrary(TSoftObjectPtr<UControlRigShapeLibrary> InShapeLibrary);
 
 	virtual const FUserData* GetUserData(const FString& InPath, FString* OutErrorMessage = nullptr) const override;
 	virtual const TArray<const FUserData*>& GetUserDataArray(const FString& InParentPath = FString(), FString* OutErrorMessage = nullptr) const override;
 
+	virtual void PostLoad() override;
 #if WITH_EDITOR
 	virtual void PostEditChangeProperty(FPropertyChangedEvent& PropertyChangedEvent) override;
 #endif
@@ -38,6 +39,10 @@ protected:
 	UPROPERTY(transient)
 	TArray<FName> ShapeNames;
 
+	UPROPERTY(transient)
+	TObjectPtr<UControlRigShapeLibrary> ShapeLibraryCached;
+	
+
 	static inline constexpr TCHAR DefaultShapePath[] = TEXT("DefaultShape");
 	static inline constexpr TCHAR ShapeNamesPath[] = TEXT("ShapeNames");
 	static inline constexpr TCHAR ShapeLibraryNullFormat[] = TEXT("User data path '%s' could not be found (ShapeLibrary not provided)");
diff --git a/Engine/Plugins/Runtime/RigVM/Source/RigVM/Private/RigVMCore/RigVMAssetUserData.cpp b/Engine/Plugins/Runtime/RigVM/Source/RigVM/Private/RigVMCore/RigVMAssetUserData.cpp
--- a/Engine/Plugins/Runtime/RigVM/Source/RigVM/Private/RigVMCore/RigVMAssetUserData.cpp
+++ b/Engine/Plugins/Runtime/RigVM/Source/RigVM/Private/RigVMCore/RigVMAssetUserData.cpp
@@ -372,12 +372,13 @@ bool UNameSpacedUserData::IsPropertySupported(const FProperty* InProperty, const
 	return true;
 }
 
-void UDataAssetLink::SetDataAsset(UDataAsset* InDataAsset)
+void UDataAssetLink::SetDataAsset(TSoftObjectPtr<UDataAsset> InDataAsset)
 {
 	InvalidateCache();
 	DataAsset = InDataAsset;
+	DataAssetCached = InDataAsset.Get();
 
-	if(NameSpace.IsEmpty() && DataAsset)
+	if(NameSpace.IsEmpty() && DataAssetCached)
 	{
 		NameSpace = DataAsset->GetName();
 	}
@@ -391,10 +392,10 @@ const UNameSpacedUserData::FUserData* UDataAssetLink::GetUserData(const FString&
 		return ResultFromSuper;
 	}
 
-	if(DataAsset)
+	if(DataAssetCached)
 	{
 		// this method caches as well - so the next time around Super::GetUserData should return the cache 
-		return GetUserDataWithinStruct(DataAsset->GetClass(), (const uint8*)DataAsset, InPath, FString(), OutErrorMessage);
+		return GetUserDataWithinStruct(DataAssetCached->GetClass(), static_cast<const uint8*>(DataAssetCached), InPath, FString(), OutErrorMessage);
 	}
 
 	if(OutErrorMessage && OutErrorMessage->IsEmpty())
@@ -413,13 +414,13 @@ const TArray<const UNameSpacedUserData::FUserData*>& UDataAssetLink::GetUserData
 		return ResultFromSuper;
 	}
 
-	if(DataAsset)
+	if(DataAssetCached)
 	{
 		// we should only get here if we haven't cached this user data array before.
 		if(InParentPath.IsEmpty())
 		{
 			// this method caches as well - so the next time around Super::GetUserDataArray should return the cache 
-			return GetUserDataArrayWithinStruct(DataAsset->GetClass(), (const uint8*)DataAsset, InParentPath, OutErrorMessage);
+			return GetUserDataArrayWithinStruct(DataAssetCached->GetClass(), static_cast<const uint8*>(DataAssetCached), InParentPath, OutErrorMessage);
 		}
 
 		if(const FUserData* ParentUserData = GetUserData(InParentPath, OutErrorMessage))
@@ -455,6 +456,13 @@ const TArray<const UNameSpacedUserData::FUserData*>& UDataAssetLink::GetUserData
 	return EmptyUserDatas;
 }
 
+void UDataAssetLink::PostLoad()
+{
+	Super::PostLoad();
+
+	DataAssetCached = DataAsset.Get();
+}
+
 #if WITH_EDITOR
 
 void UDataAssetLink::PostEditChangeProperty(FPropertyChangedEvent& PropertyChangedEvent)
diff --git a/Engine/Plugins/Runtime/RigVM/Source/RigVM/Public/RigVMCore/RigVMAssetUserData.h b/Engine/Plugins/Runtime/RigVM/Source/RigVM/Public/RigVMCore/RigVMAssetUserData.h
--- a/Engine/Plugins/Runtime/RigVM/Source/RigVM/Public/RigVMCore/RigVMAssetUserData.h
+++ b/Engine/Plugins/Runtime/RigVM/Source/RigVM/Public/RigVMCore/RigVMAssetUserData.h
@@ -192,22 +192,25 @@ public:
 
 	/** If assigned, the data asset link will provide access to the data asset's content. */
 	UPROPERTY(EditAnywhere, BlueprintReadWrite, Category = General, BlueprintGetter = GetDataAsset, BlueprintSetter = SetDataAsset, Meta = (DisplayAfter="NameSpace"))
-	TObjectPtr<UDataAsset> DataAsset;
+	TSoftObjectPtr<UDataAsset> DataAsset;
 
 	UFUNCTION(BlueprintGetter)
-	UDataAsset* GetDataAsset() const { return DataAsset; }
+	TSoftObjectPtr<UDataAsset> GetDataAsset() const { return DataAsset; }
 
 	UFUNCTION(BlueprintSetter)
-	void SetDataAsset(UDataAsset* InDataAsset);
+	void SetDataAsset(TSoftObjectPtr<UDataAsset> InDataAsset);
 
 	virtual const FUserData* GetUserData(const FString& InPath, FString* OutErrorMessage = nullptr) const override;
 	virtual const TArray<const FUserData*>& GetUserDataArray(const FString& InParentPath = FString(), FString* OutErrorMessage = nullptr) const override;
 
+	virtual void PostLoad() override;
 #if WITH_EDITOR
 	virtual void PostEditChangeProperty(FPropertyChangedEvent& PropertyChangedEvent) override;
 #endif
 
 protected:
-
+	UPROPERTY(transient)
+	TObjectPtr<UDataAsset> DataAssetCached;
+	
 	static inline constexpr TCHAR DataAssetNullFormat[] = TEXT("User data path '%s' could not be found (DataAsset not provided)");
 };
