Sequencer: Add store/restore state for animation mode for Control Rig Template

This is the same as what FPreAnimatedSkeletalAnimationState does in UMovieSceneSkeletalAnimationSystem.

Tested setting a skeletal mesh with a default animation rig to Use Animation Blueprint, enter Sequencer, exit Sequencer, and that it returns to Use Animation Blueprint (and is not left on Use Custom Mode)

diff --git a/Engine/Plugins/Animation/ControlRig/Source/ControlRig/Private/Sequencer/MovieSceneControlRigParameterTemplate.cpp b/Engine/Plugins/Animation/ControlRig/Source/ControlRig/Private/Sequencer/MovieSceneControlRigParameterTemplate.cpp
--- a/Engine/Plugins/Animation/ControlRig/Source/ControlRig/Private/Sequencer/MovieSceneControlRigParameterTemplate.cpp
+++ b/Engine/Plugins/Animation/ControlRig/Source/ControlRig/Private/Sequencer/MovieSceneControlRigParameterTemplate.cpp
@@ -758,6 +758,7 @@ struct FControlRigParameterPreAnimatedTokenProducer : IMovieScenePreAnimatedToke
 			void SetSkelMesh(USkeletalMeshComponent* InComponent)
 			{
 				SkeletalMeshRestoreState.SaveState(InComponent);
+				AnimationMode = InComponent->GetAnimationMode();
 			}
 
 			virtual void RestoreState(UObject& InObject, const UE::MovieScene::FRestoreStateParams& Params) override
@@ -921,6 +922,11 @@ struct FControlRigParameterPreAnimatedTokenProducer : IMovieScenePreAnimatedToke
 								SkeletalMeshComponent->MarkRenderTransformDirty();
 								SkeletalMeshComponent->MarkRenderDynamicDataDirty();
 								SkeletalMeshRestoreState.RestoreState(SkeletalMeshComponent);
+
+								if (SkeletalMeshComponent->GetAnimationMode() != AnimationMode)
+								{
+									SkeletalMeshComponent->SetAnimationMode(AnimationMode);
+								}
 							}
 						}
 						//only unbind if not a component
@@ -942,7 +948,7 @@ struct FControlRigParameterPreAnimatedTokenProducer : IMovieScenePreAnimatedToke
 			TArray< TNameAndValue<FVector2D> > Vector2DValues;
 			TArray< TNameAndValue<FEulerTransform> > TransformValues;
 			FSkeletalMeshRestoreState SkeletalMeshRestoreState;
-
+			EAnimationMode::Type AnimationMode;
 		};
 
 
