Control Rig: Apply location, rotation, scale property metadata to control rig transforms

This matches a change a while ago to read the property metadata for the spinboxes (ie. scene component has linear sensitivity)

diff --git a/Engine/Plugins/Animation/ControlRig/Source/ControlRig/Private/Sequencer/MovieSceneControlRigParameterSection.cpp b/Engine/Plugins/Animation/ControlRig/Source/ControlRig/Private/Sequencer/MovieSceneControlRigParameterSection.cpp
--- a/Engine/Plugins/Animation/ControlRig/Source/ControlRig/Private/Sequencer/MovieSceneControlRigParameterSection.cpp
+++ b/Engine/Plugins/Animation/ControlRig/Source/ControlRig/Private/Sequencer/MovieSceneControlRigParameterSection.cpp
@@ -344,6 +344,12 @@ struct FParameterTransformChannelEditorData
 		ParameterName = InName;
 		static FText LongIntentFormatStr = NSLOCTEXT("MovieSceneControlParameterRigSection", "LongIntentFormatString", "{GroupName}.{IntentName}");
 
+		static const TSet<FName> PropertyMetaDataKeys = { "UIMin", "UIMax", "SliderExponent", "LinearDeltaSensitivity", "Delta", "ClampMin", "ClampMax", "ForceUnits", "WheelStep" };
+
+		const FProperty* RelativeLocationProperty = USceneComponent::StaticClass()->FindPropertyByName(USceneComponent::GetRelativeLocationPropertyName());
+		const FProperty* RelativeRotationProperty = USceneComponent::StaticClass()->FindPropertyByName(USceneComponent::GetRelativeRotationPropertyName());
+		const FProperty* RelativeScale3DProperty = USceneComponent::StaticClass()->FindPropertyByName(USceneComponent::GetRelativeScale3DPropertyName());
+
 		//FText LocationGroup = NSLOCTEXT("MovieSceneControlParameterRigSection", "Location", "Location");
 		//FText RotationGroup = NSLOCTEXT("MovieSceneControlParameterRigSection", "Rotation", "Rotation");
 		//FText ScaleGroup = NSLOCTEXT("MovieSceneControlParameterRigSection", "Scale", "Scale");
@@ -364,6 +370,10 @@ struct FParameterTransformChannelEditorData
 			MetaData[0].Color = FCommonChannelData::RedChannelColor;
 			MetaData[0].SortOrder = SortStartIndex++;
 			MetaData[0].bCanCollapseToTrack = false;
+			for (const FName& PropertyMetaDataKey : PropertyMetaDataKeys)
+			{
+				MetaData[0].PropertyMetaData.Add(PropertyMetaDataKey, RelativeLocationProperty->GetMetaData(PropertyMetaDataKey));
+			}
 
 			//MetaData[1].SetIdentifiers("Location.Y", FCommonChannelData::ChannelY, LocationGroup);
 			TotalName += ".Location.Y";
@@ -376,6 +386,10 @@ struct FParameterTransformChannelEditorData
 			MetaData[1].Color = FCommonChannelData::GreenChannelColor;
 			MetaData[1].SortOrder = SortStartIndex++;
 			MetaData[1].bCanCollapseToTrack = false;
+			for (const FName& PropertyMetaDataKey : PropertyMetaDataKeys)
+			{
+				MetaData[1].PropertyMetaData.Add(PropertyMetaDataKey, RelativeLocationProperty->GetMetaData(PropertyMetaDataKey));
+			}
 
 			//MetaData[2].SetIdentifiers("Location.Z", FCommonChannelData::ChannelZ, LocationGroup);
 			TotalName += ".Location.Z";
@@ -388,6 +402,10 @@ struct FParameterTransformChannelEditorData
 			MetaData[2].Color = FCommonChannelData::BlueChannelColor;
 			MetaData[2].SortOrder = SortStartIndex++;
 			MetaData[2].bCanCollapseToTrack = false;
+			for (const FName& PropertyMetaDataKey : PropertyMetaDataKeys)
+			{
+				MetaData[2].PropertyMetaData.Add(PropertyMetaDataKey, RelativeLocationProperty->GetMetaData(PropertyMetaDataKey));
+			}
 		}
 		{
 			//MetaData[3].SetIdentifiers("Rotation.X", NSLOCTEXT("MovieSceneTransformSection", "RotationX", "Roll"), RotationGroup);
@@ -401,6 +419,10 @@ struct FParameterTransformChannelEditorData
 			MetaData[3].Color = FCommonChannelData::RedChannelColor;
 			MetaData[3].SortOrder = SortStartIndex++;
 			MetaData[3].bCanCollapseToTrack = false;
+			for (const FName& PropertyMetaDataKey : PropertyMetaDataKeys)
+			{
+				MetaData[3].PropertyMetaData.Add(PropertyMetaDataKey, RelativeRotationProperty->GetMetaData(PropertyMetaDataKey));
+			}
 
 			//MetaData[4].SetIdentifiers("Rotation.Y", NSLOCTEXT("MovieSceneTransformSection", "RotationY", "Pitch"), RotationGroup);
 			TotalName += ".Rotation.Y";
@@ -413,6 +435,10 @@ struct FParameterTransformChannelEditorData
 			MetaData[4].Color = FCommonChannelData::GreenChannelColor;
 			MetaData[4].SortOrder = SortStartIndex++;
 			MetaData[4].bCanCollapseToTrack = false;
+			for (const FName& PropertyMetaDataKey : PropertyMetaDataKeys)
+			{
+				MetaData[4].PropertyMetaData.Add(PropertyMetaDataKey, RelativeRotationProperty->GetMetaData(PropertyMetaDataKey));
+			}
 
 			//MetaData[5].SetIdentifiers("Rotation.Z", NSLOCTEXT("MovieSceneTransformSection", "RotationZ", "Yaw"), RotationGroup);
 			TotalName += ".Rotation.Z";
@@ -425,6 +451,10 @@ struct FParameterTransformChannelEditorData
 			MetaData[5].Color = FCommonChannelData::BlueChannelColor;
 			MetaData[5].SortOrder = SortStartIndex++;
 			MetaData[5].bCanCollapseToTrack = false;
+			for (const FName& PropertyMetaDataKey : PropertyMetaDataKeys)
+			{
+				MetaData[5].PropertyMetaData.Add(PropertyMetaDataKey, RelativeRotationProperty->GetMetaData(PropertyMetaDataKey));
+			}
 		}
 		{
 			//MetaData[6].SetIdentifiers("Scale.X", FCommonChannelData::ChannelX, ScaleGroup);
@@ -438,6 +468,10 @@ struct FParameterTransformChannelEditorData
 			MetaData[6].Color = FCommonChannelData::RedChannelColor;
 			MetaData[6].SortOrder = SortStartIndex++;
 			MetaData[6].bCanCollapseToTrack = false;
+			for (const FName& PropertyMetaDataKey : PropertyMetaDataKeys)
+			{
+				MetaData[6].PropertyMetaData.Add(PropertyMetaDataKey, RelativeScale3DProperty->GetMetaData(PropertyMetaDataKey));
+			}
 
 			//MetaData[7].SetIdentifiers("Scale.Y", FCommonChannelData::ChannelY, ScaleGroup);
 			TotalName += ".Scale.Y";
@@ -450,6 +484,10 @@ struct FParameterTransformChannelEditorData
 			MetaData[7].Color = FCommonChannelData::GreenChannelColor;
 			MetaData[7].SortOrder = SortStartIndex++;
 			MetaData[7].bCanCollapseToTrack = false;
+			for (const FName& PropertyMetaDataKey : PropertyMetaDataKeys)
+			{
+				MetaData[7].PropertyMetaData.Add(PropertyMetaDataKey, RelativeScale3DProperty->GetMetaData(PropertyMetaDataKey));
+			}
 
 			//MetaData[8].SetIdentifiers("Scale.Z", FCommonChannelData::ChannelZ, ScaleGroup);
 			TotalName += ".Scale.Z";
@@ -462,6 +500,10 @@ struct FParameterTransformChannelEditorData
 			MetaData[8].Color = FCommonChannelData::BlueChannelColor;
 			MetaData[8].SortOrder = SortStartIndex++;
 			MetaData[8].bCanCollapseToTrack = false;
+			for (const FName& PropertyMetaDataKey : PropertyMetaDataKeys)
+			{
+				MetaData[8].PropertyMetaData.Add(PropertyMetaDataKey, RelativeScale3DProperty->GetMetaData(PropertyMetaDataKey));
+			}
 		}
 		{
 			//MetaData[9].SetIdentifiers("Weight", NSLOCTEXT("MovieSceneTransformSection", "Weight", "Weight"));
