Sequencer: Control Rig: Avoid issues with running PIE with the editor with Control Rig.  Instead of sharing the same Control Rig in all worlds if there is a new game world we use a new instance of the control rig and set up/evalute that one in evaluation.

diff --git a/Engine/Plugins/Animation/ControlRig/Source/ControlRig/Private/Sequencer/MovieSceneControlRigParameterSection.cpp b/Engine/Plugins/Animation/ControlRig/Source/ControlRig/Private/Sequencer/MovieSceneControlRigParameterSection.cpp
--- a/Engine/Plugins/Animation/ControlRig/Source/ControlRig/Private/Sequencer/MovieSceneControlRigParameterSection.cpp
+++ b/Engine/Plugins/Animation/ControlRig/Source/ControlRig/Private/Sequencer/MovieSceneControlRigParameterSection.cpp
@@ -3470,6 +3470,19 @@ void UMovieSceneControlRigParameterSection::RemoveAllKeys(bool bIncludeSpaceKeys
 	}
 }
 
+UControlRig* UMovieSceneControlRigParameterSection::GetControlRig(UWorld* InGameWorld) const
+{
+	if (InGameWorld == nullptr)
+	{
+		return ControlRig;
+	}
+	else if (UMovieSceneControlRigParameterTrack* Track = GetTypedOuter<UMovieSceneControlRigParameterTrack>())
+	{
+		return Track->GetGameWorldControlRig(InGameWorld);
+	}
+	return nullptr;
+}
+
 
 int32 UMovieSceneControlRigParameterSection::GetActiveCategoryIndex(FName ControlName) const
 {
diff --git a/Engine/Plugins/Animation/ControlRig/Source/ControlRig/Private/Sequencer/MovieSceneControlRigParameterTemplate.cpp b/Engine/Plugins/Animation/ControlRig/Source/ControlRig/Private/Sequencer/MovieSceneControlRigParameterTemplate.cpp
--- a/Engine/Plugins/Animation/ControlRig/Source/ControlRig/Private/Sequencer/MovieSceneControlRigParameterTemplate.cpp
+++ b/Engine/Plugins/Animation/ControlRig/Source/ControlRig/Private/Sequencer/MovieSceneControlRigParameterTemplate.cpp
@@ -1106,7 +1106,8 @@ struct FControlRigParameterPreAnimatedTokenProducer : IMovieScenePreAnimatedToke
 
 static UControlRig* GetControlRig(const UMovieSceneControlRigParameterSection* Section,UObject* BoundObject)
 {
-	UControlRig* ControlRig = Section->GetControlRig();
+	UWorld* GameWorld = (BoundObject && BoundObject->GetWorld() && BoundObject->GetWorld()->IsGameWorld()) ? BoundObject->GetWorld() : nullptr;
+	UControlRig* ControlRig =  Section->GetControlRig(GameWorld);
 	if (ControlRig->GetObjectBinding())
 	{
 		if (UControlRigComponent* ControlRigComponent = Cast<UControlRigComponent>(ControlRig->GetObjectBinding()->GetBoundObject()))
@@ -1230,93 +1231,89 @@ struct FControlRigParameterExecutionToken : IMovieSceneExecutionToken
 		MOVIESCENE_DETAILED_SCOPE_CYCLE_COUNTER(MovieSceneEval_ControlRigParameterTrack_TokenExecute)
 		
 		FMovieSceneSequenceID SequenceID = Operand.SequenceID;
-		UControlRig* ControlRig = Section->GetControlRig();
-
-		// Update the animation's state
 		TArrayView<TWeakObjectPtr<>> BoundObjects = Player.FindBoundObjects(Operand);
-		if (ControlRig)
-		{
 			const UMovieSceneSequence* Sequence = Player.State.FindSequence(Operand.SequenceID);
+		UControlRig* ControlRig = nullptr;
+		UObject* BoundObject = BoundObjects.Num() > 0 ? BoundObjects[0].Get() : nullptr;
+		if (Sequence && BoundObject)
+		{
+			UWorld* GameWorld = (BoundObject->GetWorld() && BoundObject->GetWorld()->IsGameWorld()) ? BoundObject->GetWorld() : nullptr;
+			ControlRig = Section->GetControlRig(GameWorld);
+			if (!ControlRig->GetObjectBinding())
+			{
+				ControlRig->SetObjectBinding(MakeShared<FControlRigObjectBinding>());
+			}
 
-			UObject* BoundObject = BoundObjects.Num() > 0 ? BoundObjects[0].Get() : nullptr;
-			if (Sequence && BoundObject)
+			if (ControlRig->GetObjectBinding()->GetBoundObject() != FControlRigObjectBinding::GetBindableObject(BoundObject))
 			{
-				if (!ControlRig->GetObjectBinding())
+				ControlRig->GetObjectBinding()->BindToObject(BoundObject);
+				TArray<FName> SelectedControls = ControlRig->CurrentControlSelection();
+				ControlRig->Initialize();
+				if (USkeletalMeshComponent* SkeletalMeshComponent = Cast<USkeletalMeshComponent>(FControlRigObjectBinding::GetBindableObject(BoundObject)))
 				{
-					ControlRig->SetObjectBinding(MakeShared<FControlRigObjectBinding>());
-				}
-
-				if (ControlRig->GetObjectBinding()->GetBoundObject() != FControlRigObjectBinding::GetBindableObject(BoundObject))
+					ControlRig->RequestInit();
+					ControlRig->SetBoneInitialTransformsFromSkeletalMeshComponent(SkeletalMeshComponent, true);
+					ControlRig->Evaluate_AnyThread();
+				};
+				if (ControlRig->IsA<UFKControlRig>())
 				{
-					ControlRig->GetObjectBinding()->BindToObject(BoundObject);
-					TArray<FName> SelectedControls = ControlRig->CurrentControlSelection();
-					ControlRig->Initialize();
-					if (USkeletalMeshComponent* SkeletalMeshComponent = Cast<USkeletalMeshComponent>(FControlRigObjectBinding::GetBindableObject(BoundObject)))
-					{
-						ControlRig->RequestInit();
-						ControlRig->SetBoneInitialTransformsFromSkeletalMeshComponent(SkeletalMeshComponent, true);
-						ControlRig->Evaluate_AnyThread();
-					};
-					if (ControlRig->IsA<UFKControlRig>())
-					{
-						UMovieSceneControlRigParameterTrack* Track = Section->GetTypedOuter<UMovieSceneControlRigParameterTrack>();
-						if (Track)
-						{
-							Track->ReplaceControlRig(ControlRig, true);
-						}
-					}
-					TArray<FName> NewSelectedControls = ControlRig->CurrentControlSelection();
-					if (SelectedControls != NewSelectedControls)
+					UMovieSceneControlRigParameterTrack* Track = Section->GetTypedOuter<UMovieSceneControlRigParameterTrack>();
+					if (Track)
 					{
-						SelectControls(ControlRig, SelectedControls);
+						Track->ReplaceControlRig(ControlRig, true);
 					}
 				}
-
-				// make sure to pick the correct CR instance for the  Components to bind.
-				// In case of PIE + Spawnable Actor + CR component, sequencer should grab
-				// CR component's CR instance for evaluation, see comment in BindToSequencerInstance
-				// i.e. CR component should bind to the instance that it owns itself.
-				ControlRig = GetControlRig(Section, BoundObjects[0].Get());
-				
-				// ensure that pre animated state is saved, must be done before bind
-				Player.SavePreAnimatedState(*ControlRig, FMovieSceneControlRigParameterTemplate::GetAnimTypeID(), FControlRigParameterPreAnimatedTokenProducer(Operand.SequenceID));
-				if (USkeletalMeshComponent* SkeletalMeshComponent = Cast<USkeletalMeshComponent>(FControlRigObjectBinding::GetBindableObject(BoundObject)))
+				TArray<FName> NewSelectedControls = ControlRig->CurrentControlSelection();
+				if (SelectedControls != NewSelectedControls)
 				{
-					Player.SavePreAnimatedState(*SkeletalMeshComponent, FControlRigSkeletalMeshComponentBindingTokenProducer::GetAnimTypeID(), FControlRigSkeletalMeshComponentBindingTokenProducer(Operand.SequenceID, ControlRig));
+					SelectControls(ControlRig, SelectedControls);
 				}
+			}
 
-				FControlRigBindingHelper::BindToSequencerInstance(ControlRig);
+			// make sure to pick the correct CR instance for the  Components to bind.
+			// In case of PIE + Spawnable Actor + CR component, sequencer should grab
+			// CR component's CR instance for evaluation, see comment in BindToSequencerInstance
+			// i.e. CR component should bind to the instance that it owns itself.
+			ControlRig = GetControlRig(Section, BoundObjects[0].Get());
+				
+			// ensure that pre animated state is saved, must be done before bind
+			Player.SavePreAnimatedState(*ControlRig, FMovieSceneControlRigParameterTemplate::GetAnimTypeID(), FControlRigParameterPreAnimatedTokenProducer(Operand.SequenceID));
+			if (USkeletalMeshComponent* SkeletalMeshComponent = Cast<USkeletalMeshComponent>(FControlRigObjectBinding::GetBindableObject(BoundObject)))
+			{
+				Player.SavePreAnimatedState(*SkeletalMeshComponent, FControlRigSkeletalMeshComponentBindingTokenProducer::GetAnimTypeID(), FControlRigSkeletalMeshComponentBindingTokenProducer(Operand.SequenceID, ControlRig));
+			}
 
-				if (ControlRig->GetObjectBinding())
+			FControlRigBindingHelper::BindToSequencerInstance(ControlRig);
+
+			if (ControlRig->GetObjectBinding())
+			{
+				if (USkeletalMeshComponent* SkeletalMeshComponent = Cast<USkeletalMeshComponent>(ControlRig->GetObjectBinding()->GetBoundObject()))
 				{
-					if (USkeletalMeshComponent* SkeletalMeshComponent = Cast<USkeletalMeshComponent>(ControlRig->GetObjectBinding()->GetBoundObject()))
+					if (UControlRigLayerInstance* AnimInstance = Cast<UControlRigLayerInstance>(SkeletalMeshComponent->GetAnimInstance()))
 					{
-						if (UControlRigLayerInstance* AnimInstance = Cast<UControlRigLayerInstance>(SkeletalMeshComponent->GetAnimInstance()))
+						float Weight = Section->EvaluateEasing(Context.GetTime());
+						if (EnumHasAllFlags(Section->TransformMask.GetChannels(), EMovieSceneTransformChannel::Weight))
 						{
-							float Weight = Section->EvaluateEasing(Context.GetTime());
-							if (EnumHasAllFlags(Section->TransformMask.GetChannels(), EMovieSceneTransformChannel::Weight))
-							{
-								float ManualWeight = 1.f;
-								Section->Weight.Evaluate(Context.GetTime(), ManualWeight);
-								Weight *= ManualWeight;
-							}
-							FControlRigIOSettings InputSettings;
-							InputSettings.bUpdateCurves = true;
-							InputSettings.bUpdatePose = true;
-							//this is not great but assumes we have 1 absolute track that will be used for weighting
-							if (Section->GetBlendType() == EMovieSceneBlendType::Absolute)
-							{
-								AnimInstance->UpdateControlRigTrack(ControlRig->GetUniqueID(), Weight, InputSettings, true);
-							}
+							float ManualWeight = 1.f;
+							Section->Weight.Evaluate(Context.GetTime(), ManualWeight);
+							Weight *= ManualWeight;
+						}
+						FControlRigIOSettings InputSettings;
+						InputSettings.bUpdateCurves = true;
+						InputSettings.bUpdatePose = true;
+						//this is not great but assumes we have 1 absolute track that will be used for weighting
+						if (Section->GetBlendType() == EMovieSceneBlendType::Absolute)
+						{
+							AnimInstance->UpdateControlRigTrack(ControlRig->GetUniqueID(), Weight, InputSettings, true);
 						}
-					}
-					else
-					{
-						ControlRig = GetControlRig(Section, BoundObject);
 					}
 				}
-			}		
-		}
+				else
+				{
+					ControlRig = GetControlRig(Section, BoundObject);
+				}
+			}
+		}		
 
 		//Do Bool straight up no blending
 		if (Section->GetBlendType().Get() != EMovieSceneBlendType::Additive)
@@ -1398,7 +1395,6 @@ struct FControlRigParameterExecutionToken : IMovieSceneExecutionToken
 						}
 					}
 				}
-				UObject* BoundObject = BoundObjects.Num() > 0 ? BoundObjects[0].Get() : nullptr;
 				if (BoundObject)
 				{
 					const FConstraintsManagerController& Controller = FConstraintsManagerController::Get(BoundObject->GetWorld());
@@ -1827,6 +1823,17 @@ void FMovieSceneControlRigParameterTemplate::Evaluate(const FMovieSceneEvaluatio
 	const UMovieSceneControlRigParameterSection* Section = Cast<UMovieSceneControlRigParameterSection>(GetSourceSection());
 	if (Section && Section->GetControlRig())
 	{
+		UControlRig* ControlRig = Section->GetControlRig();
+		if (Operand.ObjectBindingID.IsValid())
+		{
+			TArrayView<TWeakObjectPtr<>> BoundObjects = PersistentData.GetMovieScenePlayer().FindBoundObjects(Operand);
+			if(BoundObjects.Num() > 0 && BoundObjects[0].IsValid() && BoundObjects[0].Get()->GetWorld()) //just support one bound object per control rig
+			{
+				UWorld* GameWorld = BoundObjects[0].Get()->GetWorld()->IsGameWorld() ? BoundObjects[0].Get()->GetWorld() : nullptr;
+				ControlRig = Section->GetControlRig(BoundObjects[0].Get()->GetWorld());
+			}
+		}
+
 		FEvaluatedControlRigParameterSectionChannelMasks* ChannelMasks = PersistentData.FindSectionData<FEvaluatedControlRigParameterSectionChannelMasks>();
 		if (!ChannelMasks)
 		{
@@ -1861,7 +1868,8 @@ void FMovieSceneControlRigParameterTemplate::Evaluate(const FMovieSceneEvaluatio
 		FControlRigParameterExecutionToken ExecutionToken(Section,Values);
 		ExecutionTokens.Add(MoveTemp(ExecutionToken));
 
-		FControlRigAnimTypeIDsPtr TypeIDs = FControlRigAnimTypeIDs::Get(Section->GetControlRig());
+
+		FControlRigAnimTypeIDsPtr TypeIDs = FControlRigAnimTypeIDs::Get(ControlRig);
 
 		for (const FScalarParameterStringAndValue& ScalarNameAndValue : Values.ScalarValues)
 		{
diff --git a/Engine/Plugins/Animation/ControlRig/Source/ControlRig/Private/Sequencer/MovieSceneControlRigParameterTrack.cpp b/Engine/Plugins/Animation/ControlRig/Source/ControlRig/Private/Sequencer/MovieSceneControlRigParameterTrack.cpp
--- a/Engine/Plugins/Animation/ControlRig/Source/ControlRig/Private/Sequencer/MovieSceneControlRigParameterTrack.cpp
+++ b/Engine/Plugins/Animation/ControlRig/Source/ControlRig/Private/Sequencer/MovieSceneControlRigParameterTrack.cpp
@@ -874,6 +874,21 @@ bool UMovieSceneControlRigParameterTrack::GetFbxCurveDataFromChannelMetadata(con
 }
 #endif
 
+UControlRig* UMovieSceneControlRigParameterTrack::GetGameWorldControlRig(UWorld* InWorld) 
+{
+	if (GameWorldControlRigs.Find(InWorld) == nullptr && ControlRig)
+	{
+		UControlRig* GameWorldControlRig = NewObject<UControlRig>(this, ControlRig->GetClass(), NAME_None, RF_Transient);
+		GameWorldControlRigs.Add(InWorld, GameWorldControlRig);
+	}
+	TObjectPtr<UControlRig> * GameWorldControlRig = GameWorldControlRigs.Find(InWorld);
+	if (GameWorldControlRig != nullptr)
+	{
+		return GameWorldControlRig->Get();
+	}
+	return nullptr;
+}
+
 TArray<FRigControlFBXNodeAndChannels>* UMovieSceneControlRigParameterTrack::GetNodeAndChannelMappings(UMovieSceneSection* InSection )
 {
 #if WITH_EDITOR
diff --git a/Engine/Plugins/Animation/ControlRig/Source/ControlRig/Public/Sequencer/MovieSceneControlRigParameterSection.h b/Engine/Plugins/Animation/ControlRig/Source/ControlRig/Public/Sequencer/MovieSceneControlRigParameterSection.h
--- a/Engine/Plugins/Animation/ControlRig/Source/ControlRig/Public/Sequencer/MovieSceneControlRigParameterSection.h
+++ b/Engine/Plugins/Animation/ControlRig/Source/ControlRig/Public/Sequencer/MovieSceneControlRigParameterSection.h
@@ -388,8 +388,8 @@ public:
 
 	/* Set the control rig for this section */
 	void SetControlRig(UControlRig* InControlRig);
-	/* Get the control rig for this section */
-	UControlRig* GetControlRig() const { return ControlRig; }
+	/* Get the control rig for this section, by default in non-game world */
+	UControlRig* GetControlRig(UWorld* InGameWorld = nullptr) const;
 
 	/** Whether or not to key currently, maybe evaluating so don't*/
 	void  SetDoNotKey(bool bIn) const { bDoNotKey = bIn; }
diff --git a/Engine/Plugins/Animation/ControlRig/Source/ControlRig/Public/Sequencer/MovieSceneControlRigParameterTrack.h b/Engine/Plugins/Animation/ControlRig/Source/ControlRig/Public/Sequencer/MovieSceneControlRigParameterTrack.h
--- a/Engine/Plugins/Animation/ControlRig/Source/ControlRig/Public/Sequencer/MovieSceneControlRigParameterTrack.h
+++ b/Engine/Plugins/Animation/ControlRig/Source/ControlRig/Public/Sequencer/MovieSceneControlRigParameterTrack.h
@@ -180,6 +180,7 @@ private:
 
 private:
 
+
 	/** Control Rig we control*/
 	UPROPERTY()
 	TObjectPtr<UControlRig> ControlRig;
@@ -206,6 +207,13 @@ private:
 public:
 	static CONTROLRIG_API FColor AbsoluteRigTrackColor;
 	static CONTROLRIG_API FColor LayeredRigTrackColor;
+
+public:
+	UControlRig* GetGameWorldControlRig(UWorld* InWorld);
+private:
+	/** copy of the controlled control rig that we use in the game world so editor control rig doesn't conflict*/
+	UPROPERTY(transient)
+	TMap<TWeakObjectPtr<UWorld>,TObjectPtr<UControlRig>> GameWorldControlRigs;
 };
 
 
diff --git a/Engine/Source/Runtime/MovieScene/Public/Evaluation/PersistentEvaluationData.h b/Engine/Source/Runtime/MovieScene/Public/Evaluation/PersistentEvaluationData.h
--- a/Engine/Source/Runtime/MovieScene/Public/Evaluation/PersistentEvaluationData.h
+++ b/Engine/Source/Runtime/MovieScene/Public/Evaluation/PersistentEvaluationData.h
@@ -126,6 +126,11 @@ public:
 	 */
 	MOVIESCENE_API const FMovieSceneSequenceInstanceData* GetInstanceData() const;
 
+	/**
+	 * Get the player
+	 */
+	MOVIESCENE_API IMovieScenePlayer& GetMovieScenePlayer() const { return Player; }
+
 	/**
 	 * Find the current sequence's instance data as the templated type, provided its type matches
 	 */
