Control Rig: Fix PreForward solve and PostForward solve not profiling correctly
#rnx

diff --git a/Engine/Plugins/Animation/ControlRig/Source/ControlRig/Private/ControlRig.cpp b/Engine/Plugins/Animation/ControlRig/Source/ControlRig/Private/ControlRig.cpp
--- a/Engine/Plugins/Animation/ControlRig/Source/ControlRig/Private/ControlRig.cpp
+++ b/Engine/Plugins/Animation/ControlRig/Source/ControlRig/Private/ControlRig.cpp
@@ -691,8 +691,8 @@ bool UControlRig::Execute(const FName& InEventName)
 	{
 		ensureMsgf(GetTypedOuter<UControlRig>() == nullptr, TEXT("UControlRig::Execute running from a nested rig in %s"), *GetPackage()->GetPathName());
 	}
-	ensureMsgf(InEventName != FRigUnit_PreBeginExecution::EventName &&
-						InEventName != FRigUnit_PostBeginExecution::EventName, TEXT("Requested execution of invalid event %s on top level rig in %s"), *InEventName.ToString(), *GetPackage()->GetPathName());
+	// ensureMsgf(InEventName != FRigUnit_PreBeginExecution::EventName &&
+	// 					InEventName != FRigUnit_PostBeginExecution::EventName, TEXT("Requested execution of invalid event %s on top level rig in %s"), *InEventName.ToString(), *GetPackage()->GetPathName());
 
 	bool bJustRanInit = false;
 	if(bRequiresInitExecution)
@@ -718,7 +718,11 @@ bool UControlRig::Execute(const FName& InEventName)
 	const bool bIsEventFirstInQueue = !LocalEventQueueToRun.IsEmpty() && LocalEventQueueToRun[0] == InEventName; 
 	const bool bIsEventLastInQueue = !LocalEventQueueToRun.IsEmpty() && LocalEventQueueToRun.Last() == InEventName;
 	const bool bIsConstructionEvent = InEventName == FRigUnit_PrepareForExecution::EventName;
+	const bool bPreForwardSolveInQueue = LocalEventQueueToRun.Contains(FRigUnit_PreBeginExecution::EventName);
+	const bool bPostForwardSolveInQueue = LocalEventQueueToRun.Contains(FRigUnit_PostBeginExecution::EventName);
+	const bool bIsPreForwardSolve = InEventName == FRigUnit_PreBeginExecution::EventName;
 	const bool bIsForwardSolve = InEventName == FRigUnit_BeginExecution::EventName;
+	const bool bIsPostForwardSolve = InEventName == FRigUnit_PostBeginExecution::EventName;
 	const bool bIsInteractionEvent = InEventName == FRigUnit_InteractionExecution::EventName;
 
 	ensure(!HasAnyFlags(RF_ClassDefaultObject));
@@ -1084,7 +1088,7 @@ bool UControlRig::Execute(const FName& InEventName)
 		}
 #endif
 		
-		if (bIsForwardSolve)
+		if (bIsPreForwardSolve || (bIsForwardSolve && !bPreForwardSolveInQueue))
 		{
 			if (PreForwardsSolveEvent.IsBound())
 			{
@@ -1092,16 +1096,10 @@ bool UControlRig::Execute(const FName& InEventName)
 				PreForwardsSolveEvent.Broadcast(this, FRigUnit_BeginExecution::EventName);
 			}
 
-			Execute_Internal(FRigUnit_PreBeginExecution::EventName);
 		}
 
 		bSuccess = Execute_Internal(InEventName);
 
-		if (bIsForwardSolve)
-		{
-			Execute_Internal(FRigUnit_PostBeginExecution::EventName);
-		}
-
 #if WITH_EDITOR
 		if (bEnableAnimAttributeTrace && ExternalAnimAttributeContainer != nullptr)
 		{
@@ -1109,7 +1107,7 @@ bool UControlRig::Execute(const FName& InEventName)
 		}
 #endif
 
-		if (bIsForwardSolve)
+		if (bIsPostForwardSolve || (bIsForwardSolve && !bPostForwardSolveInQueue))
 		{
 			if (PostForwardsSolveEvent.IsBound())
 			{
@@ -1573,6 +1571,23 @@ void UControlRig::AdaptEventQueueForEvaluate(TArray<FName>& InOutEventQueueToRun
 {
 	Super::AdaptEventQueueForEvaluate(InOutEventQueueToRun);
 
+	for (int32 i=0; i<InOutEventQueueToRun.Num(); ++i)
+	{
+		if (InOutEventQueueToRun[i] == FRigUnit_BeginExecution::EventName)
+		{
+			if (SupportsEvent(FRigUnit_PreBeginExecution::EventName))
+			{
+				InOutEventQueueToRun.Insert(FRigUnit_PreBeginExecution::EventName, i);
+				i++; // skip preforward 
+			}
+			if (SupportsEvent(FRigUnit_PostBeginExecution::EventName))
+			{
+				i++; // skip forward
+				InOutEventQueueToRun.Insert(FRigUnit_PostBeginExecution::EventName, i);
+			}
+		}
+	}
+
 	if(InteractionType != (uint8)EControlRigInteractionType::None)
 	{
 		static const FName& InteractionEvent = FRigUnit_InteractionExecution::EventName;
