Constraints: do not directly modify the CDO to avoid propagating these changes to other instances when creating configurable constraints

#rnx

diff --git a/Engine/Plugins/Animation/ControlRig/Source/ControlRigEditor/Private/Editor/Constraints/SConstraintsWidget.cpp b/Engine/Plugins/Animation/ControlRig/Source/ControlRigEditor/Private/Editor/Constraints/SConstraintsWidget.cpp
--- a/Engine/Plugins/Animation/ControlRig/Source/ControlRigEditor/Private/Editor/Constraints/SConstraintsWidget.cpp
+++ b/Engine/Plugins/Animation/ControlRig/Source/ControlRigEditor/Private/Editor/Constraints/SConstraintsWidget.cpp
@@ -69,6 +69,18 @@ const TArray< UTickableTransformConstraint* >& FConstraintInfo::GetMutableDefaul
 	return MutableDefaults;
 }
 
+const TArray< TFunction<UTickableTransformConstraint*()> >& FConstraintInfo::GetConfigurableConstraints()
+{
+	static TArray< TFunction<UTickableTransformConstraint*()> > ConfigurableConstraints({
+		[](){ return NewObject<UTickableTranslationConstraint>(GetTransientPackage(), NAME_None, RF_Transactional, GetMutableDefault<UTickableTranslationConstraint>(), true); },
+		[](){ return NewObject<UTickableRotationConstraint>(GetTransientPackage(), NAME_None, RF_Transactional, GetMutableDefault<UTickableRotationConstraint>(), true); },
+		[](){ return NewObject<UTickableScaleConstraint>(GetTransientPackage(), NAME_None, RF_Transactional, GetMutableDefault<UTickableScaleConstraint>(), true); },
+		[](){ return NewObject<UTickableParentConstraint>(GetTransientPackage(), NAME_None, RF_Transactional, GetMutableDefault<UTickableParentConstraint>(), true); },
+		[](){ return NewObject<UTickableLookAtConstraint>(GetTransientPackage(), NAME_None, RF_Transactional, GetMutableDefault<UTickableLookAtConstraint>(), true); },
+	});
+	return ConfigurableConstraints;
+}
+
 const FSlateBrush* FConstraintInfo::GetBrush(uint8 InType)
 {
 	static const UEnum* ETransformConstraintTypeEnum = StaticEnum<ETransformConstraintType>();
@@ -100,6 +112,17 @@ UTickableTransformConstraint* FConstraintInfo::GetMutable(ETransformConstraintTy
 	return nullptr;
 }
 
+UTickableTransformConstraint* FConstraintInfo::GetConfigurable(ETransformConstraintType InType)
+{
+	static const UEnum* ETransformConstraintTypeEnum = StaticEnum<ETransformConstraintType>();
+	const uint8 ConstraintType = static_cast<uint8>(InType);
+	if (ETransformConstraintTypeEnum->IsValidEnumValue(ConstraintType))
+	{
+		return GetConfigurableConstraints()[ConstraintType]();
+	}
+	return nullptr;
+}
+
 namespace
 {
 
@@ -245,7 +268,7 @@ FReply SConstraintMenuEntry::OnMouseButtonDown(const FGeometry& MyGeometry, cons
 	return FReply::Unhandled();
 }
 
-TSharedRef<SWidget> SConstraintMenuEntry::GenerateConstraintDefaultWidget() const
+TSharedRef<SWidget> SConstraintMenuEntry::GenerateConstraintDefaultWidget()
 {
 	// static constexpr bool CloseAfterSelection = true;
 	// FMenuBuilder MenuBuilder(CloseAfterSelection, nullptr);
@@ -262,14 +285,15 @@ TSharedRef<SWidget> SConstraintMenuEntry::GenerateConstraintDefaultWidget() cons
 		DetailsViewArgs.bShowOptions = false;
 		DetailsViewArgs.bShowModifiedPropertiesOption = false;
 		DetailsViewArgs.ColumnWidth = 0.45f;
+		DetailsViewArgs.NotifyHook = this;
 	}
 	
 	TSharedRef<IDetailsView> DetailsView = FModuleManager::GetModuleChecked<FPropertyEditorModule>("PropertyEditor").CreateDetailView(DetailsViewArgs);
 
-	UTickableTransformConstraint* Constraint = FConstraintInfo::GetMutable(ConstraintType);
+	ConfigurableConstraint = FConstraintInfo::GetConfigurable(ConstraintType);
 	
 	// manage properties visibility
-	const bool bIsLookAtConstraint = Constraint->GetClass() == UTickableLookAtConstraint::StaticClass();
+	const bool bIsLookAtConstraint = ConfigurableConstraint->GetClass() == UTickableLookAtConstraint::StaticClass();
 	const auto PropertyVisibility = FIsPropertyVisible::CreateLambda(
 		[bIsLookAtConstraint](const FPropertyAndParent& InPropertyAndParent)
 		{
@@ -297,7 +321,7 @@ TSharedRef<SWidget> SConstraintMenuEntry::GenerateConstraintDefaultWidget() cons
 		});
 	DetailsView->SetIsPropertyVisibleDelegate(PropertyVisibility);
 		
-	DetailsView->SetObject(Constraint);
+	DetailsView->SetObject(ConfigurableConstraint);
 
 	return SNew(SBorder)
 		.Visibility(EVisibility::Visible)
@@ -360,6 +384,36 @@ FReply SConstraintMenuEntry::OnDragDetected(const FGeometry& MyGeometry, const F
 	return FReply::Handled();
 }
 
+void SConstraintMenuEntry::NotifyPostChange(const FPropertyChangedEvent& InPropertyChangedEvent, FProperty* InPropertyThatChanged)
+{
+	if (!InPropertyThatChanged || InPropertyChangedEvent.ChangeType != EPropertyChangeType::ValueSet)
+	{
+		return;
+	}
+
+	UTickableTransformConstraint* MutableConstraint = FConstraintInfo::GetMutable(ConstraintType);
+	if (!IsValid(ConfigurableConstraint) || !IsValid(MutableConstraint))
+	{
+		return;
+	}
+
+	const UClass* ConstraintClass = MutableConstraint->GetClass();
+	
+	if (const FProperty* CDOProperty = ConstraintClass->FindPropertyByName(InPropertyChangedEvent.GetPropertyName()))
+	{
+		// copy the property
+		CDOProperty->CopyCompleteValue_InContainer(MutableConstraint, ConfigurableConstraint.Get());
+	}
+	else if (InPropertyChangedEvent.MemberProperty)
+	{
+		// copy the member property instead
+		if (const FProperty* CDOMemberProperty = ConstraintClass->FindPropertyByName(InPropertyChangedEvent.GetMemberPropertyName()))
+		{
+			CDOMemberProperty->CopyCompleteValue_InContainer(MutableConstraint, ConfigurableConstraint.Get());
+		}
+	}
+}
+
 FReply SConstraintMenuEntry::CreateSelectionPicker(const bool bUseDefault) const
 {
 	// FIXME temp approach for selecting the parent
diff --git a/Engine/Plugins/Animation/ControlRig/Source/ControlRigEditor/Private/Editor/Constraints/SConstraintsWidget.h b/Engine/Plugins/Animation/ControlRig/Source/ControlRigEditor/Private/Editor/Constraints/SConstraintsWidget.h
--- a/Engine/Plugins/Animation/ControlRig/Source/ControlRigEditor/Private/Editor/Constraints/SConstraintsWidget.h
+++ b/Engine/Plugins/Animation/ControlRig/Source/ControlRigEditor/Private/Editor/Constraints/SConstraintsWidget.h
@@ -10,6 +10,7 @@
 #include "BakingAnimationKeySettings.h"
 #include "Misc/QualifiedFrameTime.h"
 #include "ConstraintsManager.h"
+#include "Misc/NotifyHook.h"
 
 class UTickableTransformConstraint;
 class AActor;
@@ -28,10 +29,13 @@ public:
 	static const FSlateBrush* GetBrush(uint8 InType);
 	static int8 GetType(UClass* InClass);
 	static UTickableTransformConstraint* GetMutable(ETransformConstraintType InType);
+	static UTickableTransformConstraint* GetConfigurable(ETransformConstraintType InType);
+	
 private:
 	static const TArray< const FSlateBrush* >& GetBrushes();
 	static const TMap< UClass*, ETransformConstraintType >& GetConstraintToType();
 	static const TArray< UTickableTransformConstraint* >& GetMutableDefaults();
+	static const TArray< TFunction<UTickableTransformConstraint*()> >& GetConfigurableConstraints();
 };
 
 /**
@@ -67,7 +71,7 @@ private:
  * SConstraintMenuEntry
  */
 
-class SConstraintMenuEntry : public SCompoundWidget
+class SConstraintMenuEntry : public SCompoundWidget, public FNotifyHook
 {
 public:
 	SLATE_BEGIN_ARGS(SConstraintMenuEntry){}
@@ -86,6 +90,10 @@ public:
 	virtual FReply OnMouseButtonUp(const FGeometry& MyGeometry, const FPointerEvent& MouseEvent) override;
 	virtual FReply OnDragDetected(const FGeometry& MyGeometry, const FPointerEvent& MouseEvent) override;
 	// End of SWidget interface
+
+	// FNotifyHook interface
+	/** NotifyPostChange is used to copy the property value from the configurable constraint to the cdo. */
+	virtual void NotifyPostChange( const FPropertyChangedEvent& InPropertyChangedEvent, FProperty* InPropertyThatChanged ) override;
 	
 private:
 
@@ -96,7 +104,7 @@ private:
 	FOnConstraintCreated OnConstraintCreated;
 
 	/** Creates a widget to edit the constraint default properties. */
-	TSharedRef<SWidget> GenerateConstraintDefaultWidget() const;
+	TSharedRef<SWidget> GenerateConstraintDefaultWidget();
 	
 	/** Creates the constraint between the current selection and the picked actor. */
 	static void CreateConstraint(
@@ -113,7 +121,8 @@ private:
 	/** State on the item. */
 	bool bIsPressed = false;
 
-	TWeakPtr<SConstraintsCreationWidget> ConstraintsWidget;
+	/** object ptr to the configurable constraint. */
+	TObjectPtr<UTickableConstraint> ConfigurableConstraint = nullptr;
 };
 
 /**
