Sequencer: Control Rig: Keep selection on IK/FK controls if selected and they lose visibility

diff --git a/Engine/Plugins/Animation/ControlRig/Source/ControlRigEditor/Private/EditMode/ControlRigEditMode.cpp b/Engine/Plugins/Animation/ControlRig/Source/ControlRigEditor/Private/EditMode/ControlRigEditMode.cpp
--- a/Engine/Plugins/Animation/ControlRig/Source/ControlRigEditor/Private/EditMode/ControlRigEditMode.cpp
+++ b/Engine/Plugins/Animation/ControlRig/Source/ControlRigEditor/Private/EditMode/ControlRigEditMode.cpp
@@ -5035,7 +5035,7 @@ void FControlRigEditMode::TickControlShape(AControlRigShapeActor* ShapeActor, co
 					|| (IsControlRigSkelMeshVisible(ControlRig) == false);
 
 				bool bIsVisible = ControlElement->Settings.IsVisible();
-				bool bRespectVisibilityForSelection = true; 
+				bool bRespectVisibilityForSelection = IsInLevelEditor() ? false: true;
 
 				if(!bControlsHiddenInViewport)
 				{
