Control Rig: Fix UserDataAssets not loading when needed
#rnx

diff --git a/Engine/Plugins/Animation/ControlRig/Source/ControlRig/Private/ControlRigAssetUserData.cpp b/Engine/Plugins/Animation/ControlRig/Source/ControlRig/Private/ControlRigAssetUserData.cpp
--- a/Engine/Plugins/Animation/ControlRig/Source/ControlRig/Private/ControlRigAssetUserData.cpp
+++ b/Engine/Plugins/Animation/ControlRig/Source/ControlRig/Private/ControlRigAssetUserData.cpp
@@ -6,6 +6,7 @@ void UControlRigShapeLibraryLink::SetShapeLibrary(TSoftObjectPtr<UControlRigShap
 {
 	InvalidateCache();
 	ShapeLibrary = InShapeLibrary;
+	ShapeLibrary.LoadSynchronous();
 	ShapeLibraryCached = InShapeLibrary.Get();
 	ShapeNames.Reset();
 	if(ShapeLibraryCached)
@@ -91,6 +92,7 @@ void UControlRigShapeLibraryLink::PostLoad()
 {
 	Super::PostLoad();
 
+	ShapeLibrary.LoadSynchronous();
 	ShapeLibraryCached = ShapeLibrary.Get();
 }
 
diff --git a/Engine/Plugins/Runtime/RigVM/Source/RigVM/Private/RigVMCore/RigVMAssetUserData.cpp b/Engine/Plugins/Runtime/RigVM/Source/RigVM/Private/RigVMCore/RigVMAssetUserData.cpp
--- a/Engine/Plugins/Runtime/RigVM/Source/RigVM/Private/RigVMCore/RigVMAssetUserData.cpp
+++ b/Engine/Plugins/Runtime/RigVM/Source/RigVM/Private/RigVMCore/RigVMAssetUserData.cpp
@@ -376,6 +376,7 @@ void UDataAssetLink::SetDataAsset(TSoftObjectPtr<UDataAsset> InDataAsset)
 {
 	InvalidateCache();
 	DataAsset = InDataAsset;
+	DataAsset.LoadSynchronous();
 	DataAssetCached = InDataAsset.Get();
 
 	if(NameSpace.IsEmpty() && DataAssetCached)
@@ -460,6 +461,7 @@ void UDataAssetLink::PostLoad()
 {
 	Super::PostLoad();
 
+	DataAsset.LoadSynchronous();
 	DataAssetCached = DataAsset.Get();
 }
 
