Control Rig: Fix crash when connecting module to socket
#rnx

diff --git a/Engine/Plugins/Animation/ControlRig/Source/ControlRig/Private/ModularRigController.cpp b/Engine/Plugins/Animation/ControlRig/Source/ControlRig/Private/ModularRigController.cpp
--- a/Engine/Plugins/Animation/ControlRig/Source/ControlRig/Private/ModularRigController.cpp
+++ b/Engine/Plugins/Animation/ControlRig/Source/ControlRig/Private/ModularRigController.cpp
@@ -301,7 +301,7 @@ bool UModularRigController::ConnectConnectorToElement(const FRigElementKey& InCo
 					}
 				}
 
-				if (bAutoResolveOtherConnectors && bResolvedPrimaryConnector)
+				if (Module && bAutoResolveOtherConnectors && bResolvedPrimaryConnector)
 				{
 					(void)AutoConnectModules( {Module->GetPath()}, false, bSetupUndo);
 				}
