Control Rig: Fix UserDataAssets not loading when needed
#rnx

diff --git a/Engine/Plugins/Animation/ControlRig/Source/ControlRig/Private/ControlRigAssetUserData.cpp b/Engine/Plugins/Animation/ControlRig/Source/ControlRig/Private/ControlRigAssetUserData.cpp
--- a/Engine/Plugins/Animation/ControlRig/Source/ControlRig/Private/ControlRigAssetUserData.cpp
+++ b/Engine/Plugins/Animation/ControlRig/Source/ControlRig/Private/ControlRigAssetUserData.cpp
@@ -1,6 +1,7 @@
 // Copyright Epic Games, Inc. All Rights Reserved.
 
 #include "ControlRigAssetUserData.h"
+#include "Misc/PackageName.h"
 
 void UControlRigShapeLibraryLink::SetShapeLibrary(TSoftObjectPtr<UControlRigShapeLibrary> InShapeLibrary)
 {
@@ -25,7 +26,12 @@ const UNameSpacedUserData::FUserData* UControlRigShapeLibraryLink::GetUserData(c
 		return ResultFromSuper;
 	}
 
-	if(ShapeLibrary)
+	if(!ShapeLibraryCached)
+	{
+		UpdateShapeLibraryCache();
+	}
+
+	if(ShapeLibraryCached)
 	{
 		if(InPath.Equals(GET_MEMBER_NAME_STRING_CHECKED(UControlRigShapeLibraryLink, ShapeLibrary), ESearchCase::CaseSensitive))
 		{
@@ -64,7 +70,12 @@ const TArray<const UNameSpacedUserData::FUserData*>& UControlRigShapeLibraryLink
 		return ResultFromSuper;
 	}
 
-	if(ShapeLibrary)
+	if (!ShapeLibraryCached)
+	{
+		UpdateShapeLibraryCache();
+	}
+
+	if(ShapeLibraryCached)
 	{
 		// UControlRigShapeLibraryLink doesn't offer any arrays other than the top level 
 		if(InParentPath.IsEmpty())
@@ -87,6 +98,26 @@ const TArray<const UNameSpacedUserData::FUserData*>& UControlRigShapeLibraryLink
 	return EmptyUserDatas;
 }
 
+void UControlRigShapeLibraryLink::UpdateShapeLibraryCache() const
+{
+	const FString PackagePath = ShapeLibrary.ToSoftObjectPath().GetLongPackageName();
+	const FName PluginMountPoint = FPackageName::GetPackageMountPoint(PackagePath, false);
+	if (FPackageName::MountPointExists(PluginMountPoint.ToString()))
+	{
+		ShapeLibrary.LoadSynchronous();
+		ShapeLibraryCached = ShapeLibrary.Get();
+
+		if (ShapeLibraryCached)
+		{
+			ShapeNames.Reset();
+			for(const FControlRigShapeDefinition& Shape : ShapeLibraryCached->Shapes)
+			{
+				ShapeNames.Add(Shape.ShapeName);
+			}
+		}
+	}
+}
+
 void UControlRigShapeLibraryLink::PostLoad()
 {
 	Super::PostLoad();
diff --git a/Engine/Plugins/Animation/ControlRig/Source/ControlRig/Public/ControlRigAssetUserData.h b/Engine/Plugins/Animation/ControlRig/Source/ControlRig/Public/ControlRigAssetUserData.h
--- a/Engine/Plugins/Animation/ControlRig/Source/ControlRig/Public/ControlRigAssetUserData.h
+++ b/Engine/Plugins/Animation/ControlRig/Source/ControlRig/Public/ControlRigAssetUserData.h
@@ -28,6 +28,7 @@ public:
 
 	virtual const FUserData* GetUserData(const FString& InPath, FString* OutErrorMessage = nullptr) const override;
 	virtual const TArray<const FUserData*>& GetUserDataArray(const FString& InParentPath = FString(), FString* OutErrorMessage = nullptr) const override;
+	void UpdateShapeLibraryCache() const;
 
 	virtual void PostLoad() override;
 #if WITH_EDITOR
@@ -37,10 +38,10 @@ public:
 protected:
 
 	UPROPERTY(transient)
-	TArray<FName> ShapeNames;
+	mutable TArray<FName> ShapeNames;
 
 	UPROPERTY(transient)
-	TObjectPtr<UControlRigShapeLibrary> ShapeLibraryCached;
+	mutable TObjectPtr<UControlRigShapeLibrary> ShapeLibraryCached;
 	
 
 	static inline constexpr TCHAR DefaultShapePath[] = TEXT("DefaultShape");
diff --git a/Engine/Plugins/Runtime/RigVM/Source/RigVM/Private/RigVMCore/RigVMAssetUserData.cpp b/Engine/Plugins/Runtime/RigVM/Source/RigVM/Private/RigVMCore/RigVMAssetUserData.cpp
--- a/Engine/Plugins/Runtime/RigVM/Source/RigVM/Private/RigVMCore/RigVMAssetUserData.cpp
+++ b/Engine/Plugins/Runtime/RigVM/Source/RigVM/Private/RigVMCore/RigVMAssetUserData.cpp
@@ -2,6 +2,7 @@
 
 #include "RigVMCore/RigVMAssetUserData.h"
 #include "Engine/UserDefinedStruct.h"
+#include "Misc/PackageName.h"
 
 #include UE_INLINE_GENERATED_CPP_BY_NAME(RigVMAssetUserData)
 
@@ -392,6 +393,11 @@ const UNameSpacedUserData::FUserData* UDataAssetLink::GetUserData(const FString&
 		return ResultFromSuper;
 	}
 
+	if (!DataAssetCached)
+	{
+		UpdateDataAssetCache();
+	}
+
 	if(DataAssetCached)
 	{
 		// this method caches as well - so the next time around Super::GetUserData should return the cache 
@@ -414,6 +420,11 @@ const TArray<const UNameSpacedUserData::FUserData*>& UDataAssetLink::GetUserData
 		return ResultFromSuper;
 	}
 
+	if (!DataAssetCached)
+	{
+		UpdateDataAssetCache();
+	}
+
 	if(DataAssetCached)
 	{
 		// we should only get here if we haven't cached this user data array before.
@@ -456,6 +467,22 @@ const TArray<const UNameSpacedUserData::FUserData*>& UDataAssetLink::GetUserData
 	return EmptyUserDatas;
 }
 
+void UDataAssetLink::UpdateDataAssetCache() const
+{
+	const FString PackagePath = DataAsset.ToSoftObjectPath().GetLongPackageName();
+	const FName PluginMountPoint = FPackageName::GetPackageMountPoint(PackagePath, false);
+	if (FPackageName::MountPointExists(PluginMountPoint.ToString()))
+	{
+		DataAsset.LoadSynchronous();
+		DataAssetCached = DataAsset.Get();
+
+		if(NameSpace.IsEmpty() && DataAssetCached)
+		{
+			NameSpace = DataAsset->GetName();
+		}
+	}
+}
+
 void UDataAssetLink::PostLoad()
 {
 	Super::PostLoad();
diff --git a/Engine/Plugins/Runtime/RigVM/Source/RigVM/Public/RigVMCore/RigVMAssetUserData.h b/Engine/Plugins/Runtime/RigVM/Source/RigVM/Public/RigVMCore/RigVMAssetUserData.h
--- a/Engine/Plugins/Runtime/RigVM/Source/RigVM/Public/RigVMCore/RigVMAssetUserData.h
+++ b/Engine/Plugins/Runtime/RigVM/Source/RigVM/Public/RigVMCore/RigVMAssetUserData.h
@@ -133,7 +133,7 @@ public:
 
 	/** The namespace to use when looking up values inside of the user data. */
 	UPROPERTY(EditAnywhere, Category = General)
-	FString NameSpace;
+	mutable FString NameSpace;
 
 	/** Returns true if userdata exists for a given path */
 	bool ContainsUserData(const FString& InPath) const { return GetUserData(InPath) != nullptr; }
@@ -202,6 +202,7 @@ public:
 
 	virtual const FUserData* GetUserData(const FString& InPath, FString* OutErrorMessage = nullptr) const override;
 	virtual const TArray<const FUserData*>& GetUserDataArray(const FString& InParentPath = FString(), FString* OutErrorMessage = nullptr) const override;
+	void UpdateDataAssetCache() const;
 
 	virtual void PostLoad() override;
 #if WITH_EDITOR
@@ -210,7 +211,7 @@ public:
 
 protected:
 	UPROPERTY(transient)
-	TObjectPtr<UDataAsset> DataAssetCached;
+	mutable TObjectPtr<UDataAsset> DataAssetCached;
 	
 	static inline constexpr TCHAR DataAssetNullFormat[] = TEXT("User data path '%s' could not be found (DataAsset not provided)");
 };
