Control RIg: Fix modular rig ignoring changes in module variables
#rnx

diff --git a/Engine/Plugins/Animation/ControlRig/Source/ControlRigDeveloper/Private/ControlRigBlueprint.cpp b/Engine/Plugins/Animation/ControlRig/Source/ControlRigDeveloper/Private/ControlRigBlueprint.cpp
--- a/Engine/Plugins/Animation/ControlRig/Source/ControlRigDeveloper/Private/ControlRigBlueprint.cpp
+++ b/Engine/Plugins/Animation/ControlRig/Source/ControlRigDeveloper/Private/ControlRigBlueprint.cpp
@@ -2335,7 +2335,9 @@ void UControlRigBlueprint::OnModularDependencyVMCompiled(UObject* InBlueprint, U
 
 void UControlRigBlueprint::OnModularDependencyChanged(URigVMBlueprint* InBlueprint)
 {
-	RequestConstructionOnAllModules();
+	RefreshModuleVariables();
+	RefreshModuleConnectors();
+	RecompileModularRig();
 }
 
 void UControlRigBlueprint::RequestConstructionOnAllModules()
@@ -2362,6 +2364,77 @@ void UControlRigBlueprint::RequestConstructionOnAllModules()
 
 }
 
+void UControlRigBlueprint::RefreshModuleVariables()
+{
+	if(!IsModularRig())
+	{
+		return;
+	}
+
+	ModularRigModel.ForEachModule([this](const FRigModuleReference* Element) -> bool
+	{
+		RefreshModuleVariables(Element);
+		return true;
+	});
+}
+
+void UControlRigBlueprint::RefreshModuleVariables(const FRigModuleReference* InModule)
+{
+	if(!IsModularRig())
+	{
+		return;
+	}
+
+	// avoid dead class pointers
+	const UClass* ModuleClass = InModule->Class.Get();
+	if(ModuleClass == nullptr)
+	{
+		return;
+	}
+
+	// Make sure the provided module belongs to our ModularRigModel
+	const FString& ModulePath = InModule->GetPath();
+	FRigModuleReference* Module = ModularRigModel.FindModule(ModulePath);
+	if (Module != InModule)
+	{
+		return;
+	}
+
+	Modify();
+
+	for (TFieldIterator<FProperty> PropertyIt(ModuleClass); PropertyIt; ++PropertyIt)
+	{
+		const FProperty* Property = *PropertyIt;
+		
+		// remove advanced, private or not editable properties
+		const bool bIsAdvanced = Property->HasAnyPropertyFlags(CPF_AdvancedDisplay);
+		const bool bIsPublic = Property->HasAnyPropertyFlags(CPF_Edit | CPF_EditConst);
+		const bool bIsInstanceEditable = !Property->HasAnyPropertyFlags(CPF_DisableEditOnInstance);
+		if (bIsAdvanced || !bIsPublic || !bIsInstanceEditable)
+		{
+			Module->ConfigValues.Remove(Property->GetFName());
+			Module->Bindings.Remove(Property->GetFName());
+		}
+	}
+
+	// Make sure all the types are valid
+	if (UModularRigController* Controller = GetModularRigController())
+	{
+		const TMap<FName, FString> ConfigValues = Module->ConfigValues;
+		const TMap<FName, FString> Bindings = Module->Bindings;
+		Module->ConfigValues.Reset();
+		Module->Bindings.Reset();
+		for (const TPair<FName, FString>& Pair : ConfigValues)
+		{
+			Controller->SetConfigValueInModule(ModulePath, Pair.Key, Pair.Value, false);
+		}
+		for (const TPair<FName, FString>& Pair : Bindings)
+		{
+			Controller->BindModuleVariable(ModulePath, Pair.Key, Pair.Value, false);
+		}
+	}
+}
+
 void UControlRigBlueprint::RefreshModuleConnectors()
 {
 	if(!IsModularRig())
diff --git a/Engine/Plugins/Animation/ControlRig/Source/ControlRigDeveloper/Public/ControlRigBlueprint.h b/Engine/Plugins/Animation/ControlRig/Source/ControlRigDeveloper/Public/ControlRigBlueprint.h
--- a/Engine/Plugins/Animation/ControlRig/Source/ControlRigDeveloper/Public/ControlRigBlueprint.h
+++ b/Engine/Plugins/Animation/ControlRig/Source/ControlRigDeveloper/Public/ControlRigBlueprint.h
@@ -307,6 +307,8 @@ public:
 	void OnModularDependencyVMCompiled(UObject* InBlueprint, URigVM* InVM, FRigVMExtendedExecuteContext& InExecuteContext);
 	void OnModularDependencyChanged(URigVMBlueprint* InBlueprint);
 	void RequestConstructionOnAllModules();
+	void RefreshModuleVariables();
+	void RefreshModuleVariables(const FRigModuleReference* InModule);
 	void RefreshModuleConnectors();
 	void RefreshModuleConnectors(const FRigModuleReference* InModule);
 
