Sequencer: Constraints: Only recreate channels if there is a valid world on the Control Rig.

diff --git a/Engine/Plugins/Animation/ControlRig/Source/ControlRig/Private/Sequencer/MovieSceneControlRigParameterTemplate.cpp b/Engine/Plugins/Animation/ControlRig/Source/ControlRig/Private/Sequencer/MovieSceneControlRigParameterTemplate.cpp
--- a/Engine/Plugins/Animation/ControlRig/Source/ControlRig/Private/Sequencer/MovieSceneControlRigParameterTemplate.cpp
+++ b/Engine/Plugins/Animation/ControlRig/Source/ControlRig/Private/Sequencer/MovieSceneControlRigParameterTemplate.cpp
@@ -1199,8 +1199,12 @@ static UTickableConstraint* CreateConstraintIfNeeded(const FConstraintsManagerCo
 		{
 			Controller.AddConstraint(ConstraintValue.Constraint.Get());
 			//need to reconstuct channels here.. note this is now lazy and so will recreate it next time view requests it
-			Section->ReconstructChannelProxy();
-			Section->MarkAsChanged();
+			//but only do it if the control rig has a valid world it may not for example in PIE
+			if (Section->GetControlRig() && Section->GetControlRig()->GetWorld())
+			{
+				Section->ReconstructChannelProxy();
+				Section->MarkAsChanged();
+			}
 		}
 	}
 
