[Backout] - CL32340335
Original CL Desc
-----------------------------------------------------------------
Control Rig: Fix UserDataAssets not loading when needed
#rnx

diff --git a/Engine/Plugins/Animation/ControlRig/Source/ControlRig/Private/ControlRigAssetUserData.cpp b/Engine/Plugins/Animation/ControlRig/Source/ControlRig/Private/ControlRigAssetUserData.cpp
--- a/Engine/Plugins/Animation/ControlRig/Source/ControlRig/Private/ControlRigAssetUserData.cpp
+++ b/Engine/Plugins/Animation/ControlRig/Source/ControlRig/Private/ControlRigAssetUserData.cpp
@@ -1,7 +1,6 @@
 // Copyright Epic Games, Inc. All Rights Reserved.
 
 #include "ControlRigAssetUserData.h"
-#include "Misc/PackageName.h"
 
 void UControlRigShapeLibraryLink::SetShapeLibrary(TSoftObjectPtr<UControlRigShapeLibrary> InShapeLibrary)
 {
@@ -26,12 +25,7 @@ const UNameSpacedUserData::FUserData* UControlRigShapeLibraryLink::GetUserData(c
 		return ResultFromSuper;
 	}
 
-	if(!ShapeLibraryCached)
-	{
-		UpdateShapeLibraryCache();
-	}
-
-	if(ShapeLibraryCached)
+	if(ShapeLibrary)
 	{
 		if(InPath.Equals(GET_MEMBER_NAME_STRING_CHECKED(UControlRigShapeLibraryLink, ShapeLibrary), ESearchCase::CaseSensitive))
 		{
@@ -70,12 +64,7 @@ const TArray<const UNameSpacedUserData::FUserData*>& UControlRigShapeLibraryLink
 		return ResultFromSuper;
 	}
 
-	if (!ShapeLibraryCached)
-	{
-		UpdateShapeLibraryCache();
-	}
-
-	if(ShapeLibraryCached)
+	if(ShapeLibrary)
 	{
 		// UControlRigShapeLibraryLink doesn't offer any arrays other than the top level 
 		if(InParentPath.IsEmpty())
@@ -98,26 +87,6 @@ const TArray<const UNameSpacedUserData::FUserData*>& UControlRigShapeLibraryLink
 	return EmptyUserDatas;
 }
 
-void UControlRigShapeLibraryLink::UpdateShapeLibraryCache() const
-{
-	const FString PackagePath = ShapeLibrary.ToSoftObjectPath().GetLongPackageName();
-	const FName PluginMountPoint = FPackageName::GetPackageMountPoint(PackagePath, false);
-	if (FPackageName::MountPointExists(PluginMountPoint.ToString()))
-	{
-		ShapeLibrary.LoadSynchronous();
-		ShapeLibraryCached = ShapeLibrary.Get();
-
-		if (ShapeLibraryCached)
-		{
-			ShapeNames.Reset();
-			for(const FControlRigShapeDefinition& Shape : ShapeLibraryCached->Shapes)
-			{
-				ShapeNames.Add(Shape.ShapeName);
-			}
-		}
-	}
-}
-
 void UControlRigShapeLibraryLink::PostLoad()
 {
 	Super::PostLoad();
diff --git a/Engine/Plugins/Animation/ControlRig/Source/ControlRig/Public/ControlRigAssetUserData.h b/Engine/Plugins/Animation/ControlRig/Source/ControlRig/Public/ControlRigAssetUserData.h
--- a/Engine/Plugins/Animation/ControlRig/Source/ControlRig/Public/ControlRigAssetUserData.h
+++ b/Engine/Plugins/Animation/ControlRig/Source/ControlRig/Public/ControlRigAssetUserData.h
@@ -28,7 +28,6 @@ public:
 
 	virtual const FUserData* GetUserData(const FString& InPath, FString* OutErrorMessage = nullptr) const override;
 	virtual const TArray<const FUserData*>& GetUserDataArray(const FString& InParentPath = FString(), FString* OutErrorMessage = nullptr) const override;
-	void UpdateShapeLibraryCache() const;
 
 	virtual void PostLoad() override;
 #if WITH_EDITOR
@@ -38,10 +37,10 @@ public:
 protected:
 
 	UPROPERTY(transient)
-	mutable TArray<FName> ShapeNames;
+	TArray<FName> ShapeNames;
 
 	UPROPERTY(transient)
-	mutable TObjectPtr<UControlRigShapeLibrary> ShapeLibraryCached;
+	TObjectPtr<UControlRigShapeLibrary> ShapeLibraryCached;
 	
 
 	static inline constexpr TCHAR DefaultShapePath[] = TEXT("DefaultShape");
diff --git a/Engine/Plugins/Runtime/RigVM/Source/RigVM/Private/RigVMCore/RigVMAssetUserData.cpp b/Engine/Plugins/Runtime/RigVM/Source/RigVM/Private/RigVMCore/RigVMAssetUserData.cpp
--- a/Engine/Plugins/Runtime/RigVM/Source/RigVM/Private/RigVMCore/RigVMAssetUserData.cpp
+++ b/Engine/Plugins/Runtime/RigVM/Source/RigVM/Private/RigVMCore/RigVMAssetUserData.cpp
@@ -2,7 +2,6 @@
 
 #include "RigVMCore/RigVMAssetUserData.h"
 #include "Engine/UserDefinedStruct.h"
-#include "Misc/PackageName.h"
 
 #include UE_INLINE_GENERATED_CPP_BY_NAME(RigVMAssetUserData)
 
@@ -393,11 +392,6 @@ const UNameSpacedUserData::FUserData* UDataAssetLink::GetUserData(const FString&
 		return ResultFromSuper;
 	}
 
-	if (!DataAssetCached)
-	{
-		UpdateDataAssetCache();
-	}
-
 	if(DataAssetCached)
 	{
 		// this method caches as well - so the next time around Super::GetUserData should return the cache 
@@ -420,11 +414,6 @@ const TArray<const UNameSpacedUserData::FUserData*>& UDataAssetLink::GetUserData
 		return ResultFromSuper;
 	}
 
-	if (!DataAssetCached)
-	{
-		UpdateDataAssetCache();
-	}
-
 	if(DataAssetCached)
 	{
 		// we should only get here if we haven't cached this user data array before.
@@ -467,22 +456,6 @@ const TArray<const UNameSpacedUserData::FUserData*>& UDataAssetLink::GetUserData
 	return EmptyUserDatas;
 }
 
-void UDataAssetLink::UpdateDataAssetCache() const
-{
-	const FString PackagePath = DataAsset.ToSoftObjectPath().GetLongPackageName();
-	const FName PluginMountPoint = FPackageName::GetPackageMountPoint(PackagePath, false);
-	if (FPackageName::MountPointExists(PluginMountPoint.ToString()))
-	{
-		DataAsset.LoadSynchronous();
-		DataAssetCached = DataAsset.Get();
-
-		if(NameSpace.IsEmpty() && DataAssetCached)
-		{
-			NameSpace = DataAsset->GetName();
-		}
-	}
-}
-
 void UDataAssetLink::PostLoad()
 {
 	Super::PostLoad();
diff --git a/Engine/Plugins/Runtime/RigVM/Source/RigVM/Public/RigVMCore/RigVMAssetUserData.h b/Engine/Plugins/Runtime/RigVM/Source/RigVM/Public/RigVMCore/RigVMAssetUserData.h
--- a/Engine/Plugins/Runtime/RigVM/Source/RigVM/Public/RigVMCore/RigVMAssetUserData.h
+++ b/Engine/Plugins/Runtime/RigVM/Source/RigVM/Public/RigVMCore/RigVMAssetUserData.h
@@ -133,7 +133,7 @@ public:
 
 	/** The namespace to use when looking up values inside of the user data. */
 	UPROPERTY(EditAnywhere, Category = General)
-	mutable FString NameSpace;
+	FString NameSpace;
 
 	/** Returns true if userdata exists for a given path */
 	bool ContainsUserData(const FString& InPath) const { return GetUserData(InPath) != nullptr; }
@@ -202,7 +202,6 @@ public:
 
 	virtual const FUserData* GetUserData(const FString& InPath, FString* OutErrorMessage = nullptr) const override;
 	virtual const TArray<const FUserData*>& GetUserDataArray(const FString& InParentPath = FString(), FString* OutErrorMessage = nullptr) const override;
-	void UpdateDataAssetCache() const;
 
 	virtual void PostLoad() override;
 #if WITH_EDITOR
@@ -211,7 +210,7 @@ public:
 
 protected:
 	UPROPERTY(transient)
-	mutable TObjectPtr<UDataAsset> DataAssetCached;
+	TObjectPtr<UDataAsset> DataAssetCached;
 	
 	static inline constexpr TCHAR DataAssetNullFormat[] = TEXT("User data path '%s' could not be found (DataAsset not provided)");
 };
