RigVM / ControlRig: Lazily load data asset dependencies

#rnx

diff --git a/Engine/Plugins/Animation/ControlRig/Source/ControlRig/Private/ControlRigAssetUserData.cpp b/Engine/Plugins/Animation/ControlRig/Source/ControlRig/Private/ControlRigAssetUserData.cpp
--- a/Engine/Plugins/Animation/ControlRig/Source/ControlRig/Private/ControlRigAssetUserData.cpp
+++ b/Engine/Plugins/Animation/ControlRig/Source/ControlRig/Private/ControlRigAssetUserData.cpp
@@ -1,6 +1,7 @@
 // Copyright Epic Games, Inc. All Rights Reserved.
 
 #include "ControlRigAssetUserData.h"
+#include "Misc/PackageName.h"
 
 void UControlRigShapeLibraryLink::SetShapeLibrary(TSoftObjectPtr<UControlRigShapeLibrary> InShapeLibrary)
 {
@@ -92,6 +93,37 @@ void UControlRigShapeLibraryLink::PostLoad()
 	Super::PostLoad();
 
 	ShapeLibraryCached = ShapeLibrary.Get();
+	
+	if(ShapeLibraryCached == nullptr)
+	{
+		// We need to check if the mount point exists - since the shape library link may
+		// refer to an editor-only asset in a runtime game.
+		const FString PackagePath = ShapeLibrary.GetLongPackageName();
+		const FName PluginMountPoint = FPackageName::GetPackageMountPoint(PackagePath, false);
+		if (FPackageName::MountPointExists(PluginMountPoint.ToString()))
+		{
+			const FString ObjectPath = ShapeLibrary.ToString();
+
+			// load without throwing additional warnings / errors
+			ShapeLibraryCached = LoadObject<UControlRigShapeLibrary>(nullptr, *ObjectPath, nullptr, LOAD_Quiet | LOAD_NoWarn);
+			
+			if (ShapeLibraryCached)
+			{
+				if(NameSpace.IsEmpty())
+				{
+					NameSpace = ShapeLibrary->GetName();
+				}
+				
+				ShapeNames.Reset();
+				for(const FControlRigShapeDefinition& Shape : ShapeLibraryCached->Shapes)
+				{
+					ShapeNames.Add(Shape.ShapeName);
+				}
+
+				verify(ShapeLibraryCached == ShapeLibrary.Get());
+			}
+		}
+	}
 }
 
 #if WITH_EDITOR
diff --git a/Engine/Plugins/Animation/ControlRig/Source/ControlRig/Public/ControlRigAssetUserData.h b/Engine/Plugins/Animation/ControlRig/Source/ControlRig/Public/ControlRigAssetUserData.h
--- a/Engine/Plugins/Animation/ControlRig/Source/ControlRig/Public/ControlRigAssetUserData.h
+++ b/Engine/Plugins/Animation/ControlRig/Source/ControlRig/Public/ControlRigAssetUserData.h
@@ -29,6 +29,7 @@ public:
 	virtual const FUserData* GetUserData(const FString& InPath, FString* OutErrorMessage = nullptr) const override;
 	virtual const TArray<const FUserData*>& GetUserDataArray(const FString& InParentPath = FString(), FString* OutErrorMessage = nullptr) const override;
 
+	virtual bool IsPostLoadThreadSafe() const override { return false; }
 	virtual void PostLoad() override;
 #if WITH_EDITOR
 	virtual void PostEditChangeProperty(FPropertyChangedEvent& PropertyChangedEvent) override;
diff --git a/Engine/Plugins/Runtime/RigVM/Source/RigVM/Private/RigVMCore/RigVMAssetUserData.cpp b/Engine/Plugins/Runtime/RigVM/Source/RigVM/Private/RigVMCore/RigVMAssetUserData.cpp
--- a/Engine/Plugins/Runtime/RigVM/Source/RigVM/Private/RigVMCore/RigVMAssetUserData.cpp
+++ b/Engine/Plugins/Runtime/RigVM/Source/RigVM/Private/RigVMCore/RigVMAssetUserData.cpp
@@ -2,6 +2,7 @@
 
 #include "RigVMCore/RigVMAssetUserData.h"
 #include "Engine/UserDefinedStruct.h"
+#include "Misc/PackageName.h"
 
 #include UE_INLINE_GENERATED_CPP_BY_NAME(RigVMAssetUserData)
 
@@ -461,6 +462,31 @@ void UDataAssetLink::PostLoad()
 	Super::PostLoad();
 
 	DataAssetCached = DataAsset.Get();
+
+	if(DataAssetCached == nullptr)
+	{
+		// We need to check if the mount point exists - since the data asset library link may
+		// refer to an editor-only asset in a runtime game.
+		const FString PackagePath = DataAsset.GetLongPackageName();
+		const FName PluginMountPoint = FPackageName::GetPackageMountPoint(PackagePath, false);
+		if (FPackageName::MountPointExists(PluginMountPoint.ToString()))
+		{
+			const FString ObjectPath = DataAsset.ToString();
+
+			// load without throwing additional warnings / errors
+			DataAssetCached = LoadObject<UDataAsset>(nullptr, *ObjectPath, nullptr, LOAD_Quiet | LOAD_NoWarn);
+
+			if(DataAssetCached)
+			{
+				if(NameSpace.IsEmpty())
+				{
+					NameSpace = DataAsset->GetName();
+				}
+
+				verify(DataAssetCached == DataAsset.Get());
+			}
+		}
+	}
 }
 
 #if WITH_EDITOR
diff --git a/Engine/Plugins/Runtime/RigVM/Source/RigVM/Public/RigVMCore/RigVMAssetUserData.h b/Engine/Plugins/Runtime/RigVM/Source/RigVM/Public/RigVMCore/RigVMAssetUserData.h
--- a/Engine/Plugins/Runtime/RigVM/Source/RigVM/Public/RigVMCore/RigVMAssetUserData.h
+++ b/Engine/Plugins/Runtime/RigVM/Source/RigVM/Public/RigVMCore/RigVMAssetUserData.h
@@ -203,6 +203,7 @@ public:
 	virtual const FUserData* GetUserData(const FString& InPath, FString* OutErrorMessage = nullptr) const override;
 	virtual const TArray<const FUserData*>& GetUserDataArray(const FString& InParentPath = FString(), FString* OutErrorMessage = nullptr) const override;
 
+	virtual bool IsPostLoadThreadSafe() const override { return false; }
 	virtual void PostLoad() override;
 #if WITH_EDITOR
 	virtual void PostEditChangeProperty(FPropertyChangedEvent& PropertyChangedEvent) override;
