Sequencer: Constraints: follow-up to 32371825

#rnx

diff --git a/Engine/Plugins/Animation/ControlRig/Source/ControlRig/Private/Sequencer/MovieSceneControlRigParameterTemplate.cpp b/Engine/Plugins/Animation/ControlRig/Source/ControlRig/Private/Sequencer/MovieSceneControlRigParameterTemplate.cpp
--- a/Engine/Plugins/Animation/ControlRig/Source/ControlRig/Private/Sequencer/MovieSceneControlRigParameterTemplate.cpp
+++ b/Engine/Plugins/Animation/ControlRig/Source/ControlRig/Private/Sequencer/MovieSceneControlRigParameterTemplate.cpp
@@ -1184,26 +1184,27 @@ static UControlRig* GetControlRig(const UMovieSceneControlRigParameterSection* S
 	return ControlRig;
 }
 
-static UTickableConstraint* CreateConstraintIfNeeded(const FConstraintsManagerController& Controller, FConstraintAndActiveValue& ConstraintValue, UMovieSceneControlRigParameterSection* Section)
+static UTickableConstraint* CreateConstraintIfNeeded(UWorld* InWorld, const FConstraintAndActiveValue& InConstraintValue, UMovieSceneControlRigParameterSection* InSection)
 {
-	UTickableConstraint* Constraint = ConstraintValue.Constraint.Get();
-
+	UTickableConstraint* Constraint = InConstraintValue.Constraint.Get();
 	if (!Constraint) 
 	{
 		return nullptr;
 	}
-	else // it's possible that we have it but it's not in the manager, due to manager not being saved with it (due to spawning or undo/redo).
+	
+	// it's possible that we have it but it's not in the manager, due to manager not being saved with it (due to spawning or undo/redo).
+	if (InWorld)
 	{
-		const TArray< TWeakObjectPtr<UTickableConstraint>>& ConstraintsArray = Controller.GetConstraintsArray();
+		const FConstraintsManagerController& Controller = FConstraintsManagerController::Get(InWorld);
 		if (Controller.GetConstraint(Constraint->ConstraintID) == nullptr)
 		{
-			Controller.AddConstraint(ConstraintValue.Constraint.Get());
+			Controller.AddConstraint(InConstraintValue.Constraint.Get());
 			//need to reconstuct channels here.. note this is now lazy and so will recreate it next time view requests it
 			//but only do it if the control rig has a valid world it may not for example in PIE
-			if (Section->GetControlRig() && Section->GetControlRig()->GetWorld())
+			if (InSection->GetControlRig() && InSection->GetControlRig()->GetWorld())
 			{
-				Section->ReconstructChannelProxy();
-				Section->MarkAsChanged();
+				InSection->ReconstructChannelProxy();
+				InSection->MarkAsChanged();
 			}
 		}
 	}
@@ -1401,11 +1402,10 @@ struct FControlRigParameterExecutionToken : IMovieSceneExecutionToken
 				}
 				if (BoundObject)
 				{
-					const FConstraintsManagerController& Controller = FConstraintsManagerController::Get(BoundObject->GetWorld());
 					for (FConstraintAndActiveValue& ConstraintValue : ConstraintsValues)
 					{
 						UMovieSceneControlRigParameterSection* NonConstSection = const_cast<UMovieSceneControlRigParameterSection*>(Section);
-						CreateConstraintIfNeeded(Controller, ConstraintValue, NonConstSection);
+						CreateConstraintIfNeeded(BoundObject->GetWorld(), ConstraintValue, NonConstSection);
 
 						if (ConstraintValue.Constraint.IsValid())
 						{
@@ -1420,8 +1420,9 @@ struct FControlRigParameterExecutionToken : IMovieSceneExecutionToken
 					}
 					//unfortunately for Constraints with ControlRig we need to resolve all Parents also. Don't need to do children since they wil be handled by
 					//the channel resolve above
-					TArray< TWeakObjectPtr<UTickableConstraint>> Constraints = Controller.GetAllConstraints();
-					for (TWeakObjectPtr<UTickableConstraint>& TickConstraint : Constraints)
+					const FConstraintsManagerController& Controller = FConstraintsManagerController::Get(BoundObject->GetWorld());
+					const TArray< TWeakObjectPtr<UTickableConstraint>> Constraints = Controller.GetAllConstraints();
+					for (const TWeakObjectPtr<UTickableConstraint>& TickConstraint : Constraints)
 					{
 						if (UTickableTransformConstraint* TransformConstraint = Cast< UTickableTransformConstraint>(TickConstraint.Get()))
 						{
