Control Rig: Fix layered rigs not working in PIE
#rnx

diff --git a/Engine/Plugins/Animation/ControlRig/Source/ControlRig/Private/ControlRig.cpp b/Engine/Plugins/Animation/ControlRig/Source/ControlRig/Private/ControlRig.cpp
--- a/Engine/Plugins/Animation/ControlRig/Source/ControlRig/Private/ControlRig.cpp
+++ b/Engine/Plugins/Animation/ControlRig/Source/ControlRig/Private/ControlRig.cpp
@@ -88,6 +88,7 @@ UControlRig::UControlRig(const FObjectInitializer& ObjectInitializer)
 	, ControlUndoBracketIndex(0)
 	, InteractionType((uint8)EControlRigInteractionType::None)
 	, bInteractionJustBegan(false)
+	, bIsAdditive(false)
 	, DebugBoneRadiusMultiplier(1.f)
 #if WITH_EDITOR
 	, bRecordSelectionPoseForConstructionMode(true)
diff --git a/Engine/Plugins/Animation/ControlRig/Source/ControlRig/Private/Sequencer/MovieSceneControlRigParameterTrack.cpp b/Engine/Plugins/Animation/ControlRig/Source/ControlRig/Private/Sequencer/MovieSceneControlRigParameterTrack.cpp
--- a/Engine/Plugins/Animation/ControlRig/Source/ControlRig/Private/Sequencer/MovieSceneControlRigParameterTrack.cpp
+++ b/Engine/Plugins/Animation/ControlRig/Source/ControlRig/Private/Sequencer/MovieSceneControlRigParameterTrack.cpp
@@ -879,6 +879,7 @@ UControlRig* UMovieSceneControlRigParameterTrack::GetGameWorldControlRig(UWorld*
 	if (GameWorldControlRigs.Find(InWorld) == nullptr && ControlRig)
 	{
 		UControlRig* GameWorldControlRig = NewObject<UControlRig>(this, ControlRig->GetClass(), NAME_None, RF_Transient);
+		GameWorldControlRig->SetIsAdditive(ControlRig->IsAdditive());
 		GameWorldControlRigs.Add(InWorld, GameWorldControlRig);
 	}
 	TObjectPtr<UControlRig> * GameWorldControlRig = GameWorldControlRigs.Find(InWorld);
diff --git a/Engine/Plugins/Animation/ControlRig/Source/ControlRig/Public/ControlRig.h b/Engine/Plugins/Animation/ControlRig/Source/ControlRig/Public/ControlRig.h
--- a/Engine/Plugins/Animation/ControlRig/Source/ControlRig/Public/ControlRig.h
+++ b/Engine/Plugins/Animation/ControlRig/Source/ControlRig/Public/ControlRig.h
@@ -667,7 +667,7 @@ protected:
 	 * and running the forward solve
 	 */
 	UPROPERTY()
-	bool bIsAdditive = false;
+	bool bIsAdditive;
 
 	struct FRigSetControlValueInfo
 	{
