Control Rig: Fix crash on SControlRigOutliner when closing the editor
#rnx

diff --git a/Engine/Plugins/Animation/ControlRig/Source/ControlRigEditor/Private/EditMode/SControlRigOutliner.cpp b/Engine/Plugins/Animation/ControlRig/Source/ControlRigEditor/Private/EditMode/SControlRigOutliner.cpp
--- a/Engine/Plugins/Animation/ControlRig/Source/ControlRigEditor/Private/EditMode/SControlRigOutliner.cpp
+++ b/Engine/Plugins/Animation/ControlRig/Source/ControlRigEditor/Private/EditMode/SControlRigOutliner.cpp
@@ -1240,18 +1240,21 @@ void SControlRigOutliner::HandleOnControlRigBound(UControlRig* InControlRig)
 
 void SControlRigOutliner::HandleOnObjectBoundToControlRig(UObject* InObject)
 {
-	//just refresh the views, but do so on nex tick sine with FK control rig's the controls aren't set up
+	//just refresh the views, but do so on next tick since with FK control rig's the controls aren't set up
 	//until AFTER we are bound.
-	GEditor->GetTimerManager()->SetTimerForNextTick([this]()
+	TWeakPtr<SControlRigOutliner> WeakPtr = StaticCastSharedRef<SControlRigOutliner>(AsShared()).ToWeakPtr();
+	GEditor->GetTimerManager()->SetTimerForNextTick([WeakPtr]()
 	{
-		if (FControlRigEditMode* EditMode = static_cast<FControlRigEditMode*>(ModeTools->GetActiveMode(FControlRigEditMode::ModeName)))
+		if (WeakPtr.IsValid())
 		{
-			TArrayView<TWeakObjectPtr<UControlRig>> ControlRigs = EditMode->GetControlRigs();
-			HierarchyTreeView->GetTreeView()->SetControlRigs(ControlRigs); //will refresh tree
+			TSharedPtr<SControlRigOutliner> StrongThis = WeakPtr.Pin();
+			if (FControlRigEditMode* EditMode = static_cast<FControlRigEditMode*>(StrongThis->ModeTools->GetActiveMode(FControlRigEditMode::ModeName)))
+			{
+				TArrayView<TWeakObjectPtr<UControlRig>> ControlRigs = EditMode->GetControlRigs();
+				StrongThis->HierarchyTreeView->GetTreeView()->SetControlRigs(ControlRigs); //will refresh tree
+			}
 		}
 	});
-
-
 }
 
 #undef LOCTEXT_NAMESPACE
