Sequencer: Control Rig: Use GetDisplayNameForUI when checking inside Control Rig Filter

diff --git a/Engine/Plugins/Animation/ControlRig/Source/ControlRigEditor/Private/Sequencer/ControlRigSequencerFilter.cpp b/Engine/Plugins/Animation/ControlRig/Source/ControlRigEditor/Private/Sequencer/ControlRigSequencerFilter.cpp
--- a/Engine/Plugins/Animation/ControlRig/Source/ControlRigEditor/Private/Sequencer/ControlRigSequencerFilter.cpp
+++ b/Engine/Plugins/Animation/ControlRig/Source/ControlRigEditor/Private/Sequencer/ControlRigSequencerFilter.cpp
@@ -173,11 +173,10 @@ public:
 			UControlRig *ControlRig = Track->GetControlRig();
 			if (ControlRig && ControlRig->GetHierarchy())
 			{
-				FName Name(*InText.ToString());
 				TArray<const FRigBaseElement*> SelectedControls = ControlRig->GetHierarchy()->GetSelectedElements(ERigElementType::Control);
 				for (const FRigBaseElement* SelectedControl : SelectedControls)
 				{
-					if (Name == SelectedControl->GetFName())
+					if (InText.EqualTo(ControlRig->GetHierarchy()->GetDisplayNameForUI(SelectedControl)))
 					{
 						return true;
 					}
@@ -188,9 +187,12 @@ public:
 							const TArray<FRigElementKey>& DrivenControls = ControlElement->Settings.DrivenControls;
 							for (const FRigElementKey& DrivenKey : DrivenControls)
 							{
-								if (Name == DrivenKey.Name)
+								if (FRigBaseElement* Element = ControlRig->GetHierarchy()->Find(DrivenKey))
 								{
-									return true;
+									if (InText.EqualTo(ControlRig->GetHierarchy()->GetDisplayNameForUI(Element)))
+									{
+										return true;
+									}
 								}
 							}
 						}
