Control Rig: Fix crash when setting additive control rig with invalid hierarchy
#rnx

diff --git a/Engine/Plugins/Animation/ControlRig/Source/ControlRig/Private/Sequencer/MovieSceneControlRigParameterTrack.cpp b/Engine/Plugins/Animation/ControlRig/Source/ControlRig/Private/Sequencer/MovieSceneControlRigParameterTrack.cpp
--- a/Engine/Plugins/Animation/ControlRig/Source/ControlRig/Private/Sequencer/MovieSceneControlRigParameterTrack.cpp
+++ b/Engine/Plugins/Animation/ControlRig/Source/ControlRig/Private/Sequencer/MovieSceneControlRigParameterTrack.cpp
@@ -878,9 +878,10 @@ UControlRig* UMovieSceneControlRigParameterTrack::GetGameWorldControlRig(UWorld*
 {
 	if (GameWorldControlRigs.Find(InWorld) == nullptr && ControlRig)
 	{
-		UControlRig* GameWorldControlRig = NewObject<UControlRig>(this, ControlRig->GetClass(), NAME_None, RF_Transient);
-		GameWorldControlRig->SetIsAdditive(ControlRig->IsAdditive());
-		GameWorldControlRigs.Add(InWorld, GameWorldControlRig);
+		UControlRig* NewGameWorldControlRig = NewObject<UControlRig>(this, ControlRig->GetClass(), NAME_None, RF_Transient);
+		NewGameWorldControlRig->Initialize();
+		NewGameWorldControlRig->SetIsAdditive(ControlRig->IsAdditive());
+		GameWorldControlRigs.Add(InWorld, NewGameWorldControlRig);
 	}
 	TObjectPtr<UControlRig> * GameWorldControlRig = GameWorldControlRigs.Find(InWorld);
 	if (GameWorldControlRig != nullptr)
diff --git a/Engine/Plugins/Animation/ControlRig/Source/ControlRig/Public/ControlRig.h b/Engine/Plugins/Animation/ControlRig/Source/ControlRig/Public/ControlRig.h
--- a/Engine/Plugins/Animation/ControlRig/Source/ControlRig/Public/ControlRig.h
+++ b/Engine/Plugins/Animation/ControlRig/Source/ControlRig/Public/ControlRig.h
@@ -707,7 +707,10 @@ public:
 	void SetIsAdditive(const bool bInIsAdditive)
 	{
 		bIsAdditive = bInIsAdditive;
-		GetHierarchy()->bUsePreferredEulerAngles = !bIsAdditive;
+		if (URigHierarchy* Hierarchy = GetHierarchy())
+		{
+			Hierarchy->bUsePreferredEulerAngles = !bIsAdditive;
+		}
 	}
 
 private:
