Fixed shutdown crash at ControlRigEdMode by releasing the ModeTools at ControlRigEditMode::Exit, as the destructor of the ModeTools sets a handler that refreses the TreeView and that requires Slate to be alive.

#rnx

diff --git a/Engine/Plugins/Animation/ControlRig/Source/ControlRigEditor/Private/EditMode/ControlRigEditMode.cpp b/Engine/Plugins/Animation/ControlRig/Source/ControlRigEditor/Private/EditMode/ControlRigEditMode.cpp
--- a/Engine/Plugins/Animation/ControlRig/Source/ControlRigEditor/Private/EditMode/ControlRigEditMode.cpp
+++ b/Engine/Plugins/Animation/ControlRig/Source/ControlRigEditor/Private/EditMode/ControlRigEditMode.cpp
@@ -507,6 +507,7 @@ void FControlRigEditMode::Exit()
 	if (Toolkit.IsValid())
 	{
 		FToolkitManager::Get().CloseToolkit(Toolkit.ToSharedRef());
+		Toolkit.Reset();
 	}
 
 	DestroyShapesActors(nullptr);
