ControlRig: Replace uses of static FName with FLazyName

#rnx
#prefligt https://horde.devtools.epicgames.com/job/66054a0f0f0b188ba659788b

diff --git a/Engine/Plugins/Animation/ControlRig/Source/ControlRig/Private/ModularRig.cpp b/Engine/Plugins/Animation/ControlRig/Source/ControlRig/Private/ModularRig.cpp
--- a/Engine/Plugins/Animation/ControlRig/Source/ControlRig/Private/ModularRig.cpp
+++ b/Engine/Plugins/Animation/ControlRig/Source/ControlRig/Private/ModularRig.cpp
@@ -11,8 +11,6 @@
 
 #define LOCTEXT_NAMESPACE "ModularRig"
 
-const FString UModularRig::NamespaceSeparator = TEXT(":");
-
 ////////////////////////////////////////////////////////////////////////////////
 // FModuleInstanceHandle
 ////////////////////////////////////////////////////////////////////////////////
diff --git a/Engine/Plugins/Animation/ControlRig/Source/ControlRig/Private/Rigs/FKControlRig.cpp b/Engine/Plugins/Animation/ControlRig/Source/ControlRig/Private/Rigs/FKControlRig.cpp
--- a/Engine/Plugins/Animation/ControlRig/Source/ControlRig/Private/Rigs/FKControlRig.cpp
+++ b/Engine/Plugins/Animation/ControlRig/Source/ControlRig/Private/Rigs/FKControlRig.cpp
@@ -48,8 +48,7 @@ FName UFKControlRig::GetControlName(const FName& InName, const ERigElementType&
 					ScratchString.Append(CurvePostFix);
 				}
 
-				static FString ControlPostFix = TEXT("_CONTROL");
-		
+				static const TCHAR* ControlPostFix = TEXT("_CONTROL");
 				ScratchString.Append(ControlPostFix);
 				return NameToControlMapping.Add(InName, FName(*ScratchString));
 			}
diff --git a/Engine/Plugins/Animation/ControlRig/Source/ControlRig/Private/Units/Animation/RigUnit_AnimAttribute.cpp b/Engine/Plugins/Animation/ControlRig/Source/ControlRig/Private/Units/Animation/RigUnit_AnimAttribute.cpp
--- a/Engine/Plugins/Animation/ControlRig/Source/ControlRig/Private/Units/Animation/RigUnit_AnimAttribute.cpp
+++ b/Engine/Plugins/Animation/ControlRig/Source/ControlRig/Private/Units/Animation/RigUnit_AnimAttribute.cpp
@@ -10,16 +10,6 @@
 
 #include UE_INLINE_GENERATED_CPP_BY_NAME(RigUnit_AnimAttribute)
 
-FName FRigDispatch_AnimAttributeBase::NameArgName = TEXT("Name");
-FName FRigDispatch_AnimAttributeBase::BoneNameArgName = TEXT("BoneName");
-FName FRigDispatch_AnimAttributeBase::CachedBoneNameArgName = TEXT("CachedBoneName");
-FName FRigDispatch_AnimAttributeBase::CachedBoneIndexArgName = TEXT("CachedBoneIndex");
-FName FRigDispatch_AnimAttributeBase::DefaultArgName = TEXT("Default");
-FName FRigDispatch_AnimAttributeBase::ValueArgName = TEXT("Value");
-FName FRigDispatch_AnimAttributeBase::FoundArgName = TEXT("Found");
-FName FRigDispatch_AnimAttributeBase::SuccessArgName = TEXT("Success");
-
-
 bool FRigDispatch_AnimAttributeBase::IsTypeSupported(const TRigVMTypeIndex& InTypeIndex)
 {
 	const FRigVMRegistry_NoLock& Registry = FRigVMRegistry_NoLock::GetForRead();
diff --git a/Engine/Plugins/Animation/ControlRig/Source/ControlRig/Private/Units/Execution/RigUnit_BeginExecution.cpp b/Engine/Plugins/Animation/ControlRig/Source/ControlRig/Private/Units/Execution/RigUnit_BeginExecution.cpp
--- a/Engine/Plugins/Animation/ControlRig/Source/ControlRig/Private/Units/Execution/RigUnit_BeginExecution.cpp
+++ b/Engine/Plugins/Animation/ControlRig/Source/ControlRig/Private/Units/Execution/RigUnit_BeginExecution.cpp
@@ -5,10 +5,6 @@
 
 #include UE_INLINE_GENERATED_CPP_BY_NAME(RigUnit_BeginExecution)
 
-FName FRigUnit_BeginExecution::EventName = TEXT("Forwards Solve");
-FName FRigUnit_PreBeginExecution::EventName = TEXT("Pre Forwards Solve");
-FName FRigUnit_PostBeginExecution::EventName = TEXT("Post Forwards Solve");
-
 FRigUnit_BeginExecution_Execute()
 {
     DECLARE_SCOPE_HIERARCHICAL_COUNTER_RIGUNIT()
diff --git a/Engine/Plugins/Animation/ControlRig/Source/ControlRig/Private/Units/Execution/RigUnit_InteractionExecution.cpp b/Engine/Plugins/Animation/ControlRig/Source/ControlRig/Private/Units/Execution/RigUnit_InteractionExecution.cpp
--- a/Engine/Plugins/Animation/ControlRig/Source/ControlRig/Private/Units/Execution/RigUnit_InteractionExecution.cpp
+++ b/Engine/Plugins/Animation/ControlRig/Source/ControlRig/Private/Units/Execution/RigUnit_InteractionExecution.cpp
@@ -5,8 +5,6 @@
 
 #include UE_INLINE_GENERATED_CPP_BY_NAME(RigUnit_InteractionExecution)
 
-FName FRigUnit_InteractionExecution::EventName = TEXT("Interaction");
-
 FRigUnit_InteractionExecution_Execute()
 {
     DECLARE_SCOPE_HIERARCHICAL_COUNTER_RIGUNIT()
diff --git a/Engine/Plugins/Animation/ControlRig/Source/ControlRig/Private/Units/Execution/RigUnit_InverseExecution.cpp b/Engine/Plugins/Animation/ControlRig/Source/ControlRig/Private/Units/Execution/RigUnit_InverseExecution.cpp
--- a/Engine/Plugins/Animation/ControlRig/Source/ControlRig/Private/Units/Execution/RigUnit_InverseExecution.cpp
+++ b/Engine/Plugins/Animation/ControlRig/Source/ControlRig/Private/Units/Execution/RigUnit_InverseExecution.cpp
@@ -5,8 +5,6 @@
 
 #include UE_INLINE_GENERATED_CPP_BY_NAME(RigUnit_InverseExecution)
 
-FName FRigUnit_InverseExecution::EventName = TEXT("Backwards Solve");
-
 FRigUnit_InverseExecution_Execute()
 {
     DECLARE_SCOPE_HIERARCHICAL_COUNTER_RIGUNIT()
diff --git a/Engine/Plugins/Animation/ControlRig/Source/ControlRig/Private/Units/Execution/RigUnit_PrepareForExecution.cpp b/Engine/Plugins/Animation/ControlRig/Source/ControlRig/Private/Units/Execution/RigUnit_PrepareForExecution.cpp
--- a/Engine/Plugins/Animation/ControlRig/Source/ControlRig/Private/Units/Execution/RigUnit_PrepareForExecution.cpp
+++ b/Engine/Plugins/Animation/ControlRig/Source/ControlRig/Private/Units/Execution/RigUnit_PrepareForExecution.cpp
@@ -5,8 +5,6 @@
 
 #include UE_INLINE_GENERATED_CPP_BY_NAME(RigUnit_PrepareForExecution)
 
-FName FRigUnit_PrepareForExecution::EventName = TEXT("Construction");
-
 FRigUnit_PrepareForExecution_Execute()
 {
     DECLARE_SCOPE_HIERARCHICAL_COUNTER_RIGUNIT()
diff --git a/Engine/Plugins/Animation/ControlRig/Source/ControlRig/Private/Units/Hierarchy/RigUnit_Metadata.cpp b/Engine/Plugins/Animation/ControlRig/Source/ControlRig/Private/Units/Hierarchy/RigUnit_Metadata.cpp
--- a/Engine/Plugins/Animation/ControlRig/Source/ControlRig/Private/Units/Hierarchy/RigUnit_Metadata.cpp
+++ b/Engine/Plugins/Animation/ControlRig/Source/ControlRig/Private/Units/Hierarchy/RigUnit_Metadata.cpp
@@ -13,15 +13,6 @@
 #include "RigVMModel/RigVMPin.h"
 #endif
 
-FName FRigDispatch_MetadataBase::ItemArgName = TEXT("Item");
-FName FRigDispatch_MetadataBase::NameArgName = TEXT("Name");
-FName FRigDispatch_MetadataBase::NameSpaceArgName = TEXT("NameSpace");
-FName FRigDispatch_MetadataBase::CacheArgName = TEXT("Cache");
-FName FRigDispatch_MetadataBase::DefaultArgName = TEXT("Default");
-FName FRigDispatch_MetadataBase::ValueArgName = TEXT("Value");
-FName FRigDispatch_MetadataBase::FoundArgName = TEXT("Found");
-FName FRigDispatch_MetadataBase::SuccessArgName = TEXT("Success");
-
 #if WITH_EDITOR
 
 FString FRigDispatch_MetadataBase::GetNodeTitle(const FRigVMTemplateTypeMap& InTypes) const
diff --git a/Engine/Plugins/Animation/ControlRig/Source/ControlRig/Private/Units/Modules/RigUnit_ConnectorExecution.cpp b/Engine/Plugins/Animation/ControlRig/Source/ControlRig/Private/Units/Modules/RigUnit_ConnectorExecution.cpp
--- a/Engine/Plugins/Animation/ControlRig/Source/ControlRig/Private/Units/Modules/RigUnit_ConnectorExecution.cpp
+++ b/Engine/Plugins/Animation/ControlRig/Source/ControlRig/Private/Units/Modules/RigUnit_ConnectorExecution.cpp
@@ -5,8 +5,6 @@
 
 #include UE_INLINE_GENERATED_CPP_BY_NAME(RigUnit_ConnectorExecution)
 
-FName FRigUnit_ConnectorExecution::EventName = TEXT("Connector");
-
 FRigUnit_ConnectorExecution_Execute()
 {
     DECLARE_SCOPE_HIERARCHICAL_COUNTER_RIGUNIT()
diff --git a/Engine/Plugins/Animation/ControlRig/Source/ControlRig/Public/ModularRig.h b/Engine/Plugins/Animation/ControlRig/Source/ControlRig/Public/ModularRig.h
--- a/Engine/Plugins/Animation/ControlRig/Source/ControlRig/Public/ModularRig.h
+++ b/Engine/Plugins/Animation/ControlRig/Source/ControlRig/Public/ModularRig.h
@@ -207,7 +207,7 @@ public:
 		return FModuleInstanceHandle();
 	}
 
-	static const FString NamespaceSeparator;
+	static inline const TCHAR* NamespaceSeparator = TEXT(":");
 
 protected:
 
diff --git a/Engine/Plugins/Animation/ControlRig/Source/ControlRig/Public/Rigs/RigHierarchyDefines.h b/Engine/Plugins/Animation/ControlRig/Source/ControlRig/Public/Rigs/RigHierarchyDefines.h
--- a/Engine/Plugins/Animation/ControlRig/Source/ControlRig/Public/Rigs/RigHierarchyDefines.h
+++ b/Engine/Plugins/Animation/ControlRig/Source/ControlRig/Public/Rigs/RigHierarchyDefines.h
@@ -1212,8 +1212,8 @@ template<>
 inline FString FRigControlValue::ToString<bool>() const
 {
 	const bool Value = Get<bool>();
-	static const FString True = TEXT("True");
-	static const FString False = TEXT("False");
+	static const TCHAR* True = TEXT("True");
+	static const TCHAR* False = TEXT("False");
 	return Value ? True : False;
 }
 
diff --git a/Engine/Plugins/Animation/ControlRig/Source/ControlRig/Public/Units/Animation/RigUnit_AnimAttribute.h b/Engine/Plugins/Animation/ControlRig/Source/ControlRig/Public/Units/Animation/RigUnit_AnimAttribute.h
--- a/Engine/Plugins/Animation/ControlRig/Source/ControlRig/Public/Units/Animation/RigUnit_AnimAttribute.h
+++ b/Engine/Plugins/Animation/ControlRig/Source/ControlRig/Public/Units/Animation/RigUnit_AnimAttribute.h
@@ -172,16 +172,15 @@ protected:
 	// hidden
 	mutable int32 CachedBoneNameArgIndex = INDEX_NONE;
 	mutable int32 CachedBoneIndexArgIndex = INDEX_NONE;
-	
-	static FName NameArgName;
-	static FName BoneNameArgName;
-	static FName CachedBoneNameArgName;
-	static FName CachedBoneIndexArgName;
-	static FName DefaultArgName;
-	static FName ValueArgName;
-	static FName FoundArgName;
-	static FName SuccessArgName;
 
+	static inline const FLazyName NameArgName = FLazyName(TEXT("Name"));
+	static inline const FLazyName BoneNameArgName = FLazyName(TEXT("BoneName"));
+	static inline const FLazyName CachedBoneNameArgName = FLazyName(TEXT("CachedBoneName"));
+	static inline const FLazyName CachedBoneIndexArgName = FLazyName(TEXT("CachedBoneIndex"));
+	static inline const FLazyName DefaultArgName = FLazyName(TEXT("Default"));
+	static inline const FLazyName ValueArgName = FLazyName(TEXT("Value"));
+	static inline const FLazyName FoundArgName = FLazyName(TEXT("Found"));
+	static inline const FLazyName SuccessArgName = FLazyName(TEXT("Success"));
 };
 
 
diff --git a/Engine/Plugins/Animation/ControlRig/Source/ControlRig/Public/Units/Execution/RigUnit_BeginExecution.h b/Engine/Plugins/Animation/ControlRig/Source/ControlRig/Public/Units/Execution/RigUnit_BeginExecution.h
--- a/Engine/Plugins/Animation/ControlRig/Source/ControlRig/Public/Units/Execution/RigUnit_BeginExecution.h
+++ b/Engine/Plugins/Animation/ControlRig/Source/ControlRig/Public/Units/Execution/RigUnit_BeginExecution.h
@@ -24,7 +24,7 @@ struct CONTROLRIG_API FRigUnit_BeginExecution : public FRigUnit
 	UPROPERTY(EditAnywhere, Transient, DisplayName = "Execute", Category = "BeginExecution", meta = (Output))
 	FControlRigExecuteContext ExecuteContext;
 
-	static FName EventName;
+	static inline const FLazyName EventName = FLazyName(TEXT("Forwards Solve"));
 };
 
 /**
@@ -45,7 +45,7 @@ struct CONTROLRIG_API FRigUnit_PreBeginExecution : public FRigUnit
 	UPROPERTY(EditAnywhere, Transient, DisplayName = "Execute", Category = "BeginExecution", meta = (Output))
 	FControlRigExecuteContext ExecuteContext;
 
-	static FName EventName;
+	static inline const FLazyName EventName = FLazyName(TEXT("Pre Forwards Solve"));
 };
 
 /**
@@ -66,5 +66,5 @@ struct CONTROLRIG_API FRigUnit_PostBeginExecution : public FRigUnit
 	UPROPERTY(EditAnywhere, Transient, DisplayName = "Execute", Category = "BeginExecution", meta = (Output))
 	FControlRigExecuteContext ExecuteContext;
 
-	static FName EventName;
+	static inline const FLazyName EventName = FLazyName(TEXT("Post Forwards Solve"));
 };
diff --git a/Engine/Plugins/Animation/ControlRig/Source/ControlRig/Public/Units/Execution/RigUnit_InteractionExecution.h b/Engine/Plugins/Animation/ControlRig/Source/ControlRig/Public/Units/Execution/RigUnit_InteractionExecution.h
--- a/Engine/Plugins/Animation/ControlRig/Source/ControlRig/Public/Units/Execution/RigUnit_InteractionExecution.h
+++ b/Engine/Plugins/Animation/ControlRig/Source/ControlRig/Public/Units/Execution/RigUnit_InteractionExecution.h
@@ -23,5 +23,5 @@ struct CONTROLRIG_API FRigUnit_InteractionExecution : public FRigUnit
 	UPROPERTY(EditAnywhere, Transient, DisplayName = "Execute", Category = "BeginExecution", meta = (Output))
 	FControlRigExecuteContext ExecuteContext;
 
-	static FName EventName;
+	static inline const FLazyName EventName = FLazyName(TEXT("Interaction"));
 };
diff --git a/Engine/Plugins/Animation/ControlRig/Source/ControlRig/Public/Units/Execution/RigUnit_InverseExecution.h b/Engine/Plugins/Animation/ControlRig/Source/ControlRig/Public/Units/Execution/RigUnit_InverseExecution.h
--- a/Engine/Plugins/Animation/ControlRig/Source/ControlRig/Public/Units/Execution/RigUnit_InverseExecution.h
+++ b/Engine/Plugins/Animation/ControlRig/Source/ControlRig/Public/Units/Execution/RigUnit_InverseExecution.h
@@ -24,5 +24,5 @@ struct CONTROLRIG_API FRigUnit_InverseExecution : public FRigUnit
 	UPROPERTY(EditAnywhere, Transient, DisplayName = "Execute", Category = "InverseExecution", meta = (Output))
 	FControlRigExecuteContext ExecuteContext;
 
-	static FName EventName;
+	static inline const FLazyName EventName = FLazyName(TEXT("Backwards Solve"));
 };
diff --git a/Engine/Plugins/Animation/ControlRig/Source/ControlRig/Public/Units/Execution/RigUnit_PrepareForExecution.h b/Engine/Plugins/Animation/ControlRig/Source/ControlRig/Public/Units/Execution/RigUnit_PrepareForExecution.h
--- a/Engine/Plugins/Animation/ControlRig/Source/ControlRig/Public/Units/Execution/RigUnit_PrepareForExecution.h
+++ b/Engine/Plugins/Animation/ControlRig/Source/ControlRig/Public/Units/Execution/RigUnit_PrepareForExecution.h
@@ -24,5 +24,5 @@ struct CONTROLRIG_API FRigUnit_PrepareForExecution : public FRigUnit
 	UPROPERTY(EditAnywhere, Transient, DisplayName = "Execute", Category = "PrepareForExecution", meta = (Output))
 	FControlRigExecuteContext ExecuteContext;
 
-	static FName EventName;
+	static inline const FLazyName EventName = FLazyName(TEXT("Construction"));
 };
diff --git a/Engine/Plugins/Animation/ControlRig/Source/ControlRig/Public/Units/Hierarchy/RigUnit_Metadata.h b/Engine/Plugins/Animation/ControlRig/Source/ControlRig/Public/Units/Hierarchy/RigUnit_Metadata.h
--- a/Engine/Plugins/Animation/ControlRig/Source/ControlRig/Public/Units/Hierarchy/RigUnit_Metadata.h
+++ b/Engine/Plugins/Animation/ControlRig/Source/ControlRig/Public/Units/Hierarchy/RigUnit_Metadata.h
@@ -66,14 +66,15 @@ protected:
 	mutable int32 ValueArgIndex = INDEX_NONE;
 	mutable int32 FoundArgIndex = INDEX_NONE;
 	mutable int32 SuccessArgIndex = INDEX_NONE;
-	static FName ItemArgName;
-	static FName NameArgName;
-	static FName NameSpaceArgName;
-	static FName CacheArgName;
-	static FName DefaultArgName;
-	static FName ValueArgName;
-	static FName FoundArgName;
-	static FName SuccessArgName;
+
+	static inline const FLazyName ItemArgName = FLazyName(TEXT("Item"));
+	static inline const FLazyName NameArgName = FLazyName(TEXT("Name"));
+	static inline const FLazyName NameSpaceArgName = FLazyName(TEXT("NameSpace"));
+	static inline const FLazyName CacheArgName = FLazyName(TEXT("Cache"));
+	static inline const FLazyName DefaultArgName = FLazyName(TEXT("Default"));
+	static inline const FLazyName ValueArgName = FLazyName(TEXT("Value"));
+	static inline const FLazyName FoundArgName = FLazyName(TEXT("Found"));
+	static inline const FLazyName SuccessArgName = FLazyName(TEXT("Success"));
 };
 
 template<>
diff --git a/Engine/Plugins/Animation/ControlRig/Source/ControlRig/Public/Units/Modules/RigUnit_ConnectorExecution.h b/Engine/Plugins/Animation/ControlRig/Source/ControlRig/Public/Units/Modules/RigUnit_ConnectorExecution.h
--- a/Engine/Plugins/Animation/ControlRig/Source/ControlRig/Public/Units/Modules/RigUnit_ConnectorExecution.h
+++ b/Engine/Plugins/Animation/ControlRig/Source/ControlRig/Public/Units/Modules/RigUnit_ConnectorExecution.h
@@ -23,5 +23,5 @@ struct CONTROLRIG_API FRigUnit_ConnectorExecution : public FRigUnit
 	UPROPERTY(EditAnywhere, Transient, DisplayName = "Execute", Category = "BeginExecution", meta = (Output))
 	FControlRigExecuteContext ExecuteContext;
 
-	static FName EventName;
+	static inline const FLazyName EventName = FLazyName(TEXT("Connector"));
 };
diff --git a/Engine/Plugins/Animation/ControlRig/Source/ControlRig/Public/Units/RigUnit.h b/Engine/Plugins/Animation/ControlRig/Source/ControlRig/Public/Units/RigUnit.h
--- a/Engine/Plugins/Animation/ControlRig/Source/ControlRig/Public/Units/RigUnit.h
+++ b/Engine/Plugins/Animation/ControlRig/Source/ControlRig/Public/Units/RigUnit.h
@@ -75,11 +75,7 @@ public:
 	/** The name of the method used within each rig unit */
 	static FName GetMethodName()
 	{
-		static FName MethodName = FName(NAME_None);
-		if(MethodName.IsNone())
-		{
-			MethodName = FRigVMStruct::ExecuteName;
-		}
+		static const FLazyName MethodName = FRigVMStruct::ExecuteName;
 		return MethodName;
 	}
 
diff --git a/Engine/Plugins/Animation/ControlRig/Source/ControlRigEditor/Private/ControlRigModuleDetails.cpp b/Engine/Plugins/Animation/ControlRig/Source/ControlRigEditor/Private/ControlRigModuleDetails.cpp
--- a/Engine/Plugins/Animation/ControlRig/Source/ControlRigEditor/Private/ControlRigModuleDetails.cpp
+++ b/Engine/Plugins/Animation/ControlRig/Source/ControlRigEditor/Private/ControlRigModuleDetails.cpp
@@ -1109,8 +1109,8 @@ FText FRigModuleInstanceDetails::GetBindingText(const FProperty* InProperty) con
 
 const FSlateBrush* FRigModuleInstanceDetails::GetBindingImage(const FProperty* InProperty) const
 {
-	static FName TypeIcon(TEXT("Kismet.VariableList.TypeIcon"));
-	static FName ArrayTypeIcon(TEXT("Kismet.VariableList.ArrayTypeIcon"));
+	static const FLazyName TypeIcon(TEXT("Kismet.VariableList.TypeIcon"));
+	static const FLazyName ArrayTypeIcon(TEXT("Kismet.VariableList.ArrayTypeIcon"));
 
 	if(CastField<FArrayProperty>(InProperty))
 	{
diff --git a/Engine/Plugins/Animation/ControlRig/Source/ControlRigEditor/Private/EditMode/ControlRigEditMode.cpp b/Engine/Plugins/Animation/ControlRig/Source/ControlRigEditor/Private/EditMode/ControlRigEditMode.cpp
--- a/Engine/Plugins/Animation/ControlRig/Source/ControlRigEditor/Private/EditMode/ControlRigEditMode.cpp
+++ b/Engine/Plugins/Animation/ControlRig/Source/ControlRigEditor/Private/EditMode/ControlRigEditMode.cpp
@@ -128,9 +128,6 @@ void UControlRigEditModeDelegateHelper::RemoveDelegates()
 	}
 }
 
-
-FName FControlRigEditMode::ModeName("EditMode.ControlRig");
-
 #define LOCTEXT_NAMESPACE "ControlRigEditMode"
 
 /** The different parts of a transform that manipulators can support */
diff --git a/Engine/Plugins/Animation/ControlRig/Source/ControlRigEditor/Private/Editor/ControlRigEditorEditMode.cpp b/Engine/Plugins/Animation/ControlRig/Source/ControlRigEditor/Private/Editor/ControlRigEditorEditMode.cpp
--- a/Engine/Plugins/Animation/ControlRig/Source/ControlRigEditor/Private/Editor/ControlRigEditorEditMode.cpp
+++ b/Engine/Plugins/Animation/ControlRig/Source/ControlRigEditor/Private/Editor/ControlRigEditorEditMode.cpp
@@ -10,8 +10,6 @@
 #include "Preferences/PersonaOptions.h"
 #include "SkeletalDebugRendering.h"
 
-FName FControlRigEditorEditMode::ModeName("EditMode.ControlRigEditor");
-
 void FControlRigEditorEditMode::Render(const FSceneView* View, FViewport* Viewport, FPrimitiveDrawInterface* PDI)
 {
 	FControlRigEditMode::Render(View, Viewport, PDI);
@@ -266,5 +264,3 @@ void FControlRigEditorEditMode::GetOnScreenDebugInfo(TArray<FText>& OutDebugInfo
 {
 
 }
-
-FName FModularRigEditorEditMode::ModeName("EditMode.ModularRigEditor");
diff --git a/Engine/Plugins/Animation/ControlRig/Source/ControlRigEditor/Private/Editor/ControlRigEditorEditMode.h b/Engine/Plugins/Animation/ControlRig/Source/ControlRigEditor/Private/Editor/ControlRigEditorEditMode.h
--- a/Engine/Plugins/Animation/ControlRig/Source/ControlRigEditor/Private/Editor/ControlRigEditorEditMode.h
+++ b/Engine/Plugins/Animation/ControlRig/Source/ControlRigEditor/Private/Editor/ControlRigEditorEditMode.h
@@ -9,7 +9,7 @@ class UPersonaOptions;
 class FControlRigEditorEditMode : public FControlRigEditMode
 {
 public:
-	static FName ModeName;
+	static inline const FLazyName ModeName = FLazyName(TEXT("EditMode.ControlRigEditor"));
 	
 	virtual bool IsInLevelEditor() const override { return false; }
 	virtual bool AreEditingControlRigDirectly() const override { return true; }
@@ -32,5 +32,6 @@ public:
 class FModularRigEditorEditMode : public FControlRigEditorEditMode
 {
 public:
-	static FName ModeName;
+
+	static inline const FLazyName ModeName = FLazyName(TEXT("EditMode.ModularRigEditor"));
 };
diff --git a/Engine/Plugins/Animation/ControlRig/Source/ControlRigEditor/Private/Editor/SRigHierarchyTreeView.cpp b/Engine/Plugins/Animation/ControlRig/Source/ControlRigEditor/Private/Editor/SRigHierarchyTreeView.cpp
--- a/Engine/Plugins/Animation/ControlRig/Source/ControlRigEditor/Private/Editor/SRigHierarchyTreeView.cpp
+++ b/Engine/Plugins/Animation/ControlRig/Source/ControlRigEditor/Private/Editor/SRigHierarchyTreeView.cpp
@@ -1146,7 +1146,7 @@ TPair<const FSlateBrush*, FSlateColor> SRigHierarchyItem::GetBrushForElementType
 				}
 				else
 				{
-					static FName TypeIcon(TEXT("Kismet.VariableList.TypeIcon"));
+					static const FLazyName TypeIcon(TEXT("Kismet.VariableList.TypeIcon"));
 					Brush = FAppStyle::GetBrush(TypeIcon);
 					ShapeColor = GetColorForControlType(Control->Settings.ControlType, Control->Settings.ControlEnum);
 				}
diff --git a/Engine/Plugins/Animation/ControlRig/Source/ControlRigEditor/Public/EditMode/ControlRigEditMode.h b/Engine/Plugins/Animation/ControlRig/Source/ControlRigEditor/Public/EditMode/ControlRigEditMode.h
--- a/Engine/Plugins/Animation/ControlRig/Source/ControlRigEditor/Public/EditMode/ControlRigEditMode.h
+++ b/Engine/Plugins/Animation/ControlRig/Source/ControlRigEditor/Public/EditMode/ControlRigEditMode.h
@@ -125,7 +125,7 @@ private:
 class CONTROLRIGEDITOR_API FControlRigEditMode : public IPersonaEditMode
 {
 public:
-	static FName ModeName;
+	static inline const FLazyName ModeName = FLazyName(TEXT("EditMode.ControlRig"));
 
 	FControlRigEditMode();
 	~FControlRigEditMode();
