RigVM: Fix deadlock in RigUnit_GetUserData

#rnx

diff --git a/Engine/Plugins/Animation/ControlRig/Source/ControlRig/Private/Units/Core/RigUnit_UserData.cpp b/Engine/Plugins/Animation/ControlRig/Source/ControlRig/Private/Units/Core/RigUnit_UserData.cpp
--- a/Engine/Plugins/Animation/ControlRig/Source/ControlRig/Private/Units/Core/RigUnit_UserData.cpp
+++ b/Engine/Plugins/Animation/ControlRig/Source/ControlRig/Private/Units/Core/RigUnit_UserData.cpp
@@ -65,7 +65,7 @@ FRigVMFunctionPtr FRigDispatch_GetUserData::GetDispatchFunctionImpl(const FRigVM
 	check(InTypes.FindChecked(ArgNameSpaceName) == RigVMTypeUtils::TypeIndex::FString);
 	check(InTypes.FindChecked(ArgPathName) == RigVMTypeUtils::TypeIndex::FString);
 	const TRigVMTypeIndex TypeIndex = InTypes.FindChecked(ArgDefaultName);
-	check(FRigVMRegistry::Get().CanMatchTypes(TypeIndex, InTypes.FindChecked(ArgResultName), true));
+	check(FRigVMRegistry_NoLock::GetForRead().CanMatchTypes_NoLock(TypeIndex, InTypes.FindChecked(ArgResultName), true));
 	check(InTypes.FindChecked(ArgFoundName) == RigVMTypeUtils::TypeIndex::Bool);
 	return &FRigDispatch_GetUserData::Execute;
 }
