Sequencer : Fix broken FK control rigs
#rnx

diff --git a/Engine/Plugins/Animation/ControlRig/Source/ControlRig/Private/Sequencer/MovieSceneControlRigParameterTemplate.cpp b/Engine/Plugins/Animation/ControlRig/Source/ControlRig/Private/Sequencer/MovieSceneControlRigParameterTemplate.cpp
--- a/Engine/Plugins/Animation/ControlRig/Source/ControlRig/Private/Sequencer/MovieSceneControlRigParameterTemplate.cpp
+++ b/Engine/Plugins/Animation/ControlRig/Source/ControlRig/Private/Sequencer/MovieSceneControlRigParameterTemplate.cpp
@@ -1260,7 +1260,7 @@ struct FControlRigParameterExecutionToken : IMovieSceneExecutionToken
 					ControlRig->SetBoneInitialTransformsFromSkeletalMeshComponent(SkeletalMeshComponent, true);
 					ControlRig->Evaluate_AnyThread();
 				};
-				if (ControlRig->IsA<UFKControlRig>())
+				if (GameWorld == nullptr && ControlRig->IsA<UFKControlRig>())// mz only in editor replace Fk Control rig, will look post 29.20 to see if this really needed but want to unblock folks
 				{
 					UMovieSceneControlRigParameterTrack* Track = Section->GetTypedOuter<UMovieSceneControlRigParameterTrack>();
 					if (Track)
@@ -1828,14 +1828,14 @@ void FMovieSceneControlRigParameterTemplate::Evaluate(const FMovieSceneEvaluatio
 	const UMovieSceneControlRigParameterSection* Section = Cast<UMovieSceneControlRigParameterSection>(GetSourceSection());
 	if (Section && Section->GetControlRig())
 	{
-		UControlRig* ControlRig = Section->GetControlRig();
+		UControlRig* ControlRig = Section->GetControlRig(); //this will be default(editor) control rig
 		if (Operand.ObjectBindingID.IsValid())
 		{
 			TArrayView<TWeakObjectPtr<>> BoundObjects = PersistentData.GetMovieScenePlayer().FindBoundObjects(Operand);
-			if(BoundObjects.Num() > 0 && BoundObjects[0].IsValid() && BoundObjects[0].Get()->GetWorld()) //just support one bound object per control rig
+			if(BoundObjects.Num() > 0 && BoundObjects[0].IsValid() && BoundObjects[0].Get()->GetWorld() && BoundObjects[0].Get()->GetWorld()->IsGameWorld()) //just support one bound object per control rig
 			{
-				UWorld* GameWorld = BoundObjects[0].Get()->GetWorld()->IsGameWorld() ? BoundObjects[0].Get()->GetWorld() : nullptr;
-				ControlRig = Section->GetControlRig(BoundObjects[0].Get()->GetWorld());
+				UWorld* GameWorld =  BoundObjects[0].Get()->GetWorld();
+				ControlRig = Section->GetControlRig(GameWorld);
 			}
 		}
 
diff --git a/Engine/Plugins/Animation/ControlRig/Source/ControlRig/Private/Sequencer/MovieSceneControlRigParameterTrack.cpp b/Engine/Plugins/Animation/ControlRig/Source/ControlRig/Private/Sequencer/MovieSceneControlRigParameterTrack.cpp
--- a/Engine/Plugins/Animation/ControlRig/Source/ControlRig/Private/Sequencer/MovieSceneControlRigParameterTrack.cpp
+++ b/Engine/Plugins/Animation/ControlRig/Source/ControlRig/Private/Sequencer/MovieSceneControlRigParameterTrack.cpp
@@ -8,6 +8,7 @@
 #include "MovieScene.h"
 #include "MovieSceneTimeHelpers.h"
 #include "Channels/MovieSceneChannelProxy.h"
+#include "Rigs/FKControlRig.h"
 #include "Rigs/RigHierarchyController.h"
 #include "UObject/Package.h"
 #include "Async/Async.h"
@@ -880,7 +881,17 @@ UControlRig* UMovieSceneControlRigParameterTrack::GetGameWorldControlRig(UWorld*
 	{
 		UControlRig* NewGameWorldControlRig = NewObject<UControlRig>(this, ControlRig->GetClass(), NAME_None, RF_Transient);
 		NewGameWorldControlRig->Initialize();
-		NewGameWorldControlRig->SetIsAdditive(ControlRig->IsAdditive());
+		if (UFKControlRig* FKControlRig = Cast<UFKControlRig>(Cast<UControlRig>(ControlRig)))
+		{
+			if (UFKControlRig* NewFKControlRig = Cast<UFKControlRig>(Cast<UControlRig>(NewGameWorldControlRig)))
+			{
+				NewFKControlRig->SetApplyMode(FKControlRig->GetApplyMode());
+			}
+		}
+		else
+		{
+			NewGameWorldControlRig->SetIsAdditive(ControlRig->IsAdditive());
+		}
 		GameWorldControlRigs.Add(InWorld, NewGameWorldControlRig);
 	}
 	TObjectPtr<UControlRig> * GameWorldControlRig = GameWorldControlRigs.Find(InWorld);
