Sequencer: Anim details: Shift does range select, and use new method to get display names to fix modular rig bugs with details panel

diff --git a/Engine/Plugins/Animation/ControlRig/Source/ControlRigEditor/Private/EditMode/AnimDetailsProxy.cpp b/Engine/Plugins/Animation/ControlRig/Source/ControlRigEditor/Private/EditMode/AnimDetailsProxy.cpp
--- a/Engine/Plugins/Animation/ControlRig/Source/ControlRigEditor/Private/EditMode/AnimDetailsProxy.cpp
+++ b/Engine/Plugins/Animation/ControlRig/Source/ControlRigEditor/Private/EditMode/AnimDetailsProxy.cpp
@@ -36,9 +36,13 @@
 #include "MVVM/ViewModels/SectionModel.h"
 #include "MVVM/ViewModels/SequencerEditorViewModel.h"
 #include "MVVM/SectionModelStorageExtension.h"
+#include "MVVM/ViewModels/OutlinerViewModel.h"
 #include "MVVM/ViewModels/SequenceModel.h"
 #include "MVVM/ViewModels/TrackAreaViewModel.h"
+#include "MVVM/Extensions/IOutlinerExtension.h"
 #include "MVVM/ViewModels/TrackModel.h"
+#include "MVVM/CurveEditorExtension.h"
+#include "Tree/SCurveEditorTree.h"
 #include "IKeyArea.h"
 #include "SequencerAddKeyOperation.h"
 #include "TransformConstraint.h"
@@ -115,6 +119,8 @@ static EAnimDetailSelectionState CachePropertySelection(TWeakPtr<FCurveEditor>&
 						{
 							++TotalNum;
 							EControlRigContextChannelToKey ChannelToKey = Proxy->GetChannelToKeyFromPropertyName(PropertyName);
+							FText ControlNameText = ControlRig->GetHierarchy()->GetDisplayNameForUI(ControlElement);
+							FString ControlNameString = ControlNameText.ToString();
 							for (const TPair<FCurveModelID, TUniquePtr<FCurveModel>>& Pair : Curves) //horribly slow
 							{
 								FCurveEditorTreeItemID TreeItemID = CurveEditor->GetTreeIDFromCurveID(Pair.Key);
@@ -144,8 +150,8 @@ static EAnimDetailSelectionState CachePropertySelection(TWeakPtr<FCurveEditor>&
 									Name.ParseIntoArray(StringArray, TEXT("."));
 									if (StringArray.Num() > 1)
 									{
-										if (StringArray[0] == ControlElement->GetDisplayName() ||
-											StringArray[1] == ControlElement->GetDisplayName()) //nested controls will be 2nd
+										if (StringArray[0] == ControlNameString ||
+											StringArray[1] == ControlNameString) //nested controls will be 2nd
 										{
 											FString ChannelName;
 											if (StringArray.Num() == 3)
@@ -1347,6 +1353,41 @@ bool UAnimDetailControlsProxyTransform::IsMultiple(const FName& PropertyName) co
 	return false;
 }
 
+TMap<FName, int32> UAnimDetailControlsProxyTransform::GetPropertyNames() const
+{
+	TMap<FName, int32> NameToIndex;
+	FName PropertyName = GET_MEMBER_NAME_CHECKED(FAnimDetailProxyLocation, LX);
+	int32 Index = 0;
+	NameToIndex.Add(PropertyName, Index);
+	PropertyName = GET_MEMBER_NAME_CHECKED(FAnimDetailProxyLocation, LY);
+	++Index;
+	NameToIndex.Add(PropertyName, Index);
+	PropertyName = GET_MEMBER_NAME_CHECKED(FAnimDetailProxyLocation, LZ);
+	++Index;
+	NameToIndex.Add(PropertyName, Index);
+
+	PropertyName = GET_MEMBER_NAME_CHECKED(FAnimDetailProxyRotation, RX);
+	++Index;
+	NameToIndex.Add(PropertyName, Index);
+	PropertyName = GET_MEMBER_NAME_CHECKED(FAnimDetailProxyRotation, RY);
+	++Index;
+	NameToIndex.Add(PropertyName, Index);
+	PropertyName = GET_MEMBER_NAME_CHECKED(FAnimDetailProxyRotation, RZ);
+	++Index;
+	NameToIndex.Add(PropertyName, Index);
+
+	PropertyName = GET_MEMBER_NAME_CHECKED(FAnimDetailProxyScale, SX);
+	++Index;
+	NameToIndex.Add(PropertyName, Index);
+	PropertyName = GET_MEMBER_NAME_CHECKED(FAnimDetailProxyScale, SY);
+	++Index;
+	NameToIndex.Add(PropertyName, Index);
+	PropertyName = GET_MEMBER_NAME_CHECKED(FAnimDetailProxyScale, SZ);
+	++Index;
+	NameToIndex.Add(PropertyName, Index);
+	return NameToIndex;
+}
+
 EControlRigContextChannelToKey UAnimDetailControlsProxyTransform::GetChannelToKeyFromPropertyName(const FName& PropertyName) const
 {
 	if (PropertyName == GET_MEMBER_NAME_CHECKED(UAnimDetailControlsProxyTransform, Location))
@@ -1584,6 +1625,21 @@ bool UAnimDetailControlsProxyLocation::IsMultiple(const FName& PropertyName) con
 	return false;
 }
 
+TMap<FName, int32> UAnimDetailControlsProxyLocation::GetPropertyNames() const
+{
+	TMap<FName, int32> NameToIndex;
+	FName PropertyName = GET_MEMBER_NAME_CHECKED(FAnimDetailProxyLocation, LX);
+	int32 Index = 0;
+	NameToIndex.Add(PropertyName, Index);
+	PropertyName = GET_MEMBER_NAME_CHECKED(FAnimDetailProxyLocation, LY);
+	++Index;
+	NameToIndex.Add(PropertyName, Index);
+	PropertyName = GET_MEMBER_NAME_CHECKED(FAnimDetailProxyLocation, LZ);
+	++Index;
+	NameToIndex.Add(PropertyName, Index);
+	return NameToIndex;
+}
+
 EControlRigContextChannelToKey UAnimDetailControlsProxyLocation::GetChannelToKeyFromPropertyName(const FName& PropertyName) const
 {
 	if (PropertyName == GET_MEMBER_NAME_CHECKED(FAnimDetailProxyLocation, LX))
@@ -1750,6 +1806,21 @@ bool UAnimDetailControlsProxyRotation::IsMultiple(const FName& PropertyName) con
 	return false;
 }
 
+TMap<FName, int32> UAnimDetailControlsProxyRotation::GetPropertyNames() const
+{
+	TMap<FName, int32> NameToIndex;
+	FName PropertyName = GET_MEMBER_NAME_CHECKED(FAnimDetailProxyRotation, RX);
+	int32 Index = 0;
+	NameToIndex.Add(PropertyName, Index);
+	PropertyName = GET_MEMBER_NAME_CHECKED(FAnimDetailProxyRotation, RY);
+	++Index;
+	NameToIndex.Add(PropertyName, Index);
+	PropertyName = GET_MEMBER_NAME_CHECKED(FAnimDetailProxyRotation, RZ);
+	++Index;
+	NameToIndex.Add(PropertyName, Index);
+	return NameToIndex;
+}
+
 EControlRigContextChannelToKey UAnimDetailControlsProxyRotation::GetChannelToKeyFromPropertyName(const FName& PropertyName) const
 {
 	if (PropertyName == GET_MEMBER_NAME_CHECKED(FAnimDetailProxyRotation, RX))
@@ -1917,6 +1988,21 @@ bool UAnimDetailControlsProxyScale::IsMultiple(const FName& PropertyName) const
 	return false;
 }
 
+TMap<FName, int32> UAnimDetailControlsProxyScale::GetPropertyNames() const
+{
+	TMap<FName, int32> NameToIndex;
+	FName PropertyName = GET_MEMBER_NAME_CHECKED(FAnimDetailProxyScale, SX);
+	int32 Index = 0;
+	NameToIndex.Add(PropertyName, Index);
+	PropertyName = GET_MEMBER_NAME_CHECKED(FAnimDetailProxyScale, SY);
+	++Index;
+	NameToIndex.Add(PropertyName, Index);
+	PropertyName = GET_MEMBER_NAME_CHECKED(FAnimDetailProxyScale, SZ);
+	++Index;
+	NameToIndex.Add(PropertyName, Index);
+	return NameToIndex;
+}
+
 EControlRigContextChannelToKey UAnimDetailControlsProxyScale::GetChannelToKeyFromPropertyName(const FName& PropertyName) const
 {
 	if (PropertyName == GET_MEMBER_NAME_CHECKED(FAnimDetailProxyScale, SX))
@@ -2071,6 +2157,18 @@ bool UAnimDetailControlsProxyVector2D::IsMultiple(const FName& PropertyName) con
 	return false;
 }
 
+TMap<FName, int32> UAnimDetailControlsProxyVector2D::GetPropertyNames() const
+{
+	TMap<FName, int32> NameToIndex;
+	FName PropertyName = GET_MEMBER_NAME_CHECKED(FAnimDetailProxyVector2D, X);
+	int32 Index = 0;
+	NameToIndex.Add(PropertyName, Index);
+	PropertyName = GET_MEMBER_NAME_CHECKED(FAnimDetailProxyVector2D, Y);
+	++Index;
+	NameToIndex.Add(PropertyName, Index);
+	return NameToIndex;
+}
+
 EControlRigContextChannelToKey UAnimDetailControlsProxyVector2D::GetChannelToKeyFromPropertyName(const FName& PropertyName) const
 {
 	if (PropertyName == GET_MEMBER_NAME_CHECKED(FAnimDetailProxyVector2D, X))
@@ -2265,6 +2363,15 @@ bool UAnimDetailControlsProxyFloat::IsMultiple(const FName& PropertyName) const
 	return false;
 }
 
+TMap<FName, int32> UAnimDetailControlsProxyFloat::GetPropertyNames() const
+{
+	TMap<FName, int32> NameToIndex;
+	FName PropertyName = GET_MEMBER_NAME_CHECKED(FAnimDetailProxyFloat, Float);
+	int32 Index = 0;
+	NameToIndex.Add(PropertyName, Index);
+	return NameToIndex;
+}
+
 EControlRigContextChannelToKey UAnimDetailControlsProxyFloat::GetChannelToKeyFromPropertyName(const FName& PropertyName) const
 {
 	if (PropertyName == GET_MEMBER_NAME_CHECKED(FAnimDetailProxyFloat, Float))
@@ -2458,6 +2565,15 @@ bool UAnimDetailControlsProxyBool::IsMultiple(const FName& PropertyName) const
 	return false;
 }
 
+TMap<FName, int32> UAnimDetailControlsProxyBool::GetPropertyNames() const
+{
+	TMap<FName, int32> NameToIndex;
+	FName PropertyName = GET_MEMBER_NAME_CHECKED(FAnimDetailProxyBool, Bool);
+	int32 Index = 0;
+	NameToIndex.Add(PropertyName, Index);
+	return NameToIndex;
+}
+
 EControlRigContextChannelToKey UAnimDetailControlsProxyBool::GetChannelToKeyFromPropertyName(const FName& PropertyName) const
 {
 	if (PropertyName == GET_MEMBER_NAME_CHECKED(FAnimDetailProxyBool, Bool))
@@ -2638,9 +2754,18 @@ bool UAnimDetailControlsProxyInteger::IsMultiple(const FName& PropertyName) cons
 	return false;
 }
 
+TMap<FName, int32> UAnimDetailControlsProxyInteger::GetPropertyNames() const
+{
+	TMap<FName, int32> NameToIndex;
+	FName PropertyName = GET_MEMBER_NAME_CHECKED(UAnimDetailControlsProxyInteger, Integer);
+	int32 Index = 0;
+	NameToIndex.Add(PropertyName, Index);
+	return NameToIndex;
+}
+
 EControlRigContextChannelToKey UAnimDetailControlsProxyInteger::GetChannelToKeyFromPropertyName(const FName& PropertyName) const
 {
-	if (PropertyName == GET_MEMBER_NAME_CHECKED(FAnimDetailProxyInteger, Integer))
+	if (PropertyName == GET_MEMBER_NAME_CHECKED(UAnimDetailControlsProxyInteger, Integer))
 	{
 		return EControlRigContextChannelToKey::TranslationX;
 	}
@@ -2802,6 +2927,15 @@ bool UAnimDetailControlsProxyEnum::IsMultiple(const FName& PropertyName) const
 	return false;
 }
 
+TMap<FName, int32> UAnimDetailControlsProxyEnum::GetPropertyNames() const
+{
+	TMap<FName, int32> NameToIndex;
+	FName PropertyName = GET_MEMBER_NAME_CHECKED(FControlRigEnumControlProxyValue, EnumIndex);
+	int32 Index = 0;
+	NameToIndex.Add(PropertyName, Index);
+	return NameToIndex;
+}
+
 EControlRigContextChannelToKey UAnimDetailControlsProxyEnum::GetChannelToKeyFromPropertyName(const FName& PropertyName) const
 {
 	if (PropertyName == GET_MEMBER_NAME_CHECKED(FControlRigEnumControlProxyValue, EnumIndex))
@@ -3128,6 +3262,7 @@ void UControlRigDetailPanelControlProxies::RemoveSequencerProxies(UObject* InObj
 				if (ExistingProxy)
 				{
 					SelectedControlRigProxies.Remove(ExistingProxy);
+					ClearSelectedProperty(ExistingProxy);
 					ExistingProxy->Rename(nullptr, GetTransientPackage(), REN_ForceNoResetLoaders);
 					ExistingProxy->MarkAsGarbage();
 				}
@@ -3135,7 +3270,6 @@ void UControlRigDetailPanelControlProxies::RemoveSequencerProxies(UObject* InObj
 			SequencerOnlyProxies.Remove(InObject);
 		}
 	}
-
 }
 
 void UControlRigDetailPanelControlProxies::RemoveControlRigProxies(UControlRig* ControlRig)
@@ -3157,6 +3291,7 @@ void UControlRigDetailPanelControlProxies::RemoveControlRigProxies(UControlRig*
 		}
 		ControlRigOnlyProxies.Empty();
 		SelectedControlRigProxies.SetNum(0);
+		ClearSelectedProperty();
 	}
 	else
 	{
@@ -3169,6 +3304,7 @@ void UControlRigDetailPanelControlProxies::RemoveControlRigProxies(UControlRig*
 				if (ExistingProxy)
 				{
 					SelectedControlRigProxies.Remove(ExistingProxy);
+					ClearSelectedProperty(ExistingProxy);
 					ExistingProxy->Rename(nullptr, GetTransientPackage(), REN_ForceNoResetLoaders);
 					ExistingProxy->MarkAsGarbage();
 				}
@@ -3224,6 +3360,7 @@ void UControlRigDetailPanelControlProxies::SelectProxy(UControlRig* ControlRig,
 		}
 		else
 		{
+			ClearSelectedProperty(Proxy);
 			SelectedControlRigProxies.Remove(Proxy);
 		}
 		Proxy->SelectionChanged(bSelected);
@@ -3239,6 +3376,405 @@ bool UControlRigDetailPanelControlProxies::IsSelected(UControlRig* InControlRig,
 	return false;
 }
 
+void UControlRigDetailPanelControlProxies::ClearSelectedProperty(UControlRigControlsProxy* Proxy)
+{
+	if (Proxy == nullptr || LastSelection.Key == Proxy)
+	{
+		LastSelection.Key.Reset();
+		LastSelection.Value = NAME_None;
+	}
+}
+
+static TArray<FName> GetPropertyNames(TMap<FName, int32>& NameToIndex, const TOptional<FName>& OptionalPropertyNameFirst, const TOptional<FName>& OptionalPropertyNameSecond,
+	EAnimDetailRangeDirection Direction = EAnimDetailRangeDirection::Down)
+{
+	const int32 MaxNumProperties = NameToIndex.Num();
+	TArray<FName> PropertyNames;
+	//if we have a first name then we either go to the second name or go in direction
+	if (OptionalPropertyNameFirst.IsSet())
+	{
+		int32* FirstLocation = NameToIndex.Find(OptionalPropertyNameFirst.GetValue());
+		if (FirstLocation != nullptr && (*FirstLocation >= 0 && *FirstLocation < MaxNumProperties))
+		{
+			int32 FirstIndex = INDEX_NONE, SecondIndex = INDEX_NONE;
+			if (OptionalPropertyNameSecond.IsSet())
+			{
+				int32* SecondLocation = NameToIndex.Find(OptionalPropertyNameSecond.GetValue());
+				if (SecondLocation != nullptr && (*SecondLocation >= 0 && *SecondLocation < MaxNumProperties))
+				{
+					if (*FirstLocation < *SecondLocation)
+					{
+						FirstIndex = *FirstLocation;
+						SecondIndex = *SecondLocation;
+					}
+					else
+					{
+						FirstIndex = *SecondLocation;
+						SecondIndex = *FirstLocation;
+					}
+				}
+			}
+			else if (Direction == EAnimDetailRangeDirection::Down)
+			{
+				FirstIndex = *FirstLocation;
+				SecondIndex = MaxNumProperties - 1;
+			}
+			else if (Direction == EAnimDetailRangeDirection::Up)
+			{
+				FirstIndex = 0;
+				SecondIndex = *FirstLocation;
+			}
+			if (FirstIndex != INDEX_NONE && SecondIndex != INDEX_NONE)
+			{
+				TArray<FName> AllPropertyNames;
+				NameToIndex.GenerateKeyArray(AllPropertyNames);
+				for (int32 Index = FirstIndex; Index <= SecondIndex; ++Index)
+				{
+					PropertyNames.Add(AllPropertyNames[Index]);
+				}
+			}
+		}
+	}
+	else //no first name we just get all
+	{
+		NameToIndex.GenerateKeyArray(PropertyNames);
+	}
+	return PropertyNames;
+}
+
+TArray<TPair< UControlRigControlsProxy*, FName>> UControlRigDetailPanelControlProxies::GetPropertiesFromLastSelection(UControlRigControlsProxy* Proxy, const FName& PropertyName) const
+{
+	TArray<TPair<UControlRigControlsProxy*, FName>> PropertiesToSelect;
+
+	auto GetPropertiesToSelect = ([&PropertiesToSelect](UControlRigControlsProxy* ProxyToGet, TOptional <FName>& FirstName, TOptional <FName>& SecondName, EAnimDetailRangeDirection Direction)
+	{
+		TMap<FName, int32> NameToIndex = ProxyToGet->GetPropertyNames();
+		TArray<FName> PropertyNames = GetPropertyNames(NameToIndex, FirstName, SecondName, Direction);
+		for (FName& PropName : PropertyNames)
+		{
+			TPair<UControlRigControlsProxy*, FName> Selected(ProxyToGet, PropName);
+			PropertiesToSelect.Add(Selected);
+		}
+	});
+
+	auto SelectChildProxiesFromIndex = ([&GetPropertiesToSelect,&PropertiesToSelect](UControlRigControlsProxy *ParentProxy, const FName& InPropretyName,int32 StartIndex, int32 ChildIndex)
+	{
+			//now we need to select all of the child proxies before this and including this
+			for (int32 Index = StartIndex; Index <= ChildIndex; ++Index)
+			{
+				UControlRigControlsProxy* ChildProxy = ParentProxy->ChildProxies[Index];
+
+				TMap<FName, int32> NameToIndex = ChildProxy->GetPropertyNames();
+				TArray<FName> PropertyNames;
+				NameToIndex.GenerateKeyArray(PropertyNames);
+				TOptional <FName> SecondName; //empty
+				EAnimDetailRangeDirection Direction = EAnimDetailRangeDirection::Down;
+				if (PropertyNames.Num() > 0)
+				{
+					TOptional <FName> FirstName = PropertyNames[0];
+					if (Index != ChildIndex) //not the current one need to select everything
+					{
+						GetPropertiesToSelect(ChildProxy, FirstName, SecondName, Direction);
+					}
+					else  //select from the name up
+					{
+						SecondName = InPropretyName;
+						GetPropertiesToSelect(ChildProxy, FirstName, SecondName, Direction);
+					}
+				}
+
+			}
+	});
+	if (LastSelection.Key.IsValid())
+	{
+		//same proxy so just get names to match.
+		if (Proxy == LastSelection.Key.Get())
+		{
+			TMap<FName, int32> NameToIndex = Proxy->GetPropertyNames();
+			TOptional <FName> FirstName = LastSelection.Value;
+			TOptional <FName> SecondName = PropertyName;
+			EAnimDetailRangeDirection Direction = EAnimDetailRangeDirection::Down;
+			GetPropertiesToSelect(Proxy, FirstName, SecondName, Direction);
+		}
+		else if (Proxy->bIsIndividual) //check to see if last selection is another channel with same parent or parent proxy
+		{
+			if (LastSelection.Key.Get()->bIsIndividual)
+			{
+				bool bFoundSibling = false;
+				UControlRigDetailPanelControlProxies* ProxyOwner = Proxy->GetTypedOuter<UControlRigDetailPanelControlProxies>();
+				if (ProxyOwner)
+				{
+					TArray<UControlRigControlsProxy*> AllProxies = ProxyOwner->GetAllSelectedProxies();
+					for (UControlRigControlsProxy* ParentProxy : AllProxies)
+					{
+						if (ParentProxy->ChildProxies.Num() > 1)
+						{
+							int32 OneIndex = ParentProxy->ChildProxies.Find(Proxy);
+							if (OneIndex != INDEX_NONE)
+							{
+								int32 TwoIndex = ParentProxy->ChildProxies.Find(LastSelection.Key.Get());
+								if(TwoIndex != INDEX_NONE)
+								{
+									bFoundSibling = true;
+									if (OneIndex > TwoIndex)
+									{
+										int32 Swap = TwoIndex;
+										TwoIndex = OneIndex;
+										OneIndex = Swap;
+									}
+									SelectChildProxiesFromIndex(ParentProxy, PropertyName, OneIndex, TwoIndex);
+								}
+							}
+						}
+					}
+				}
+				if (bFoundSibling == false)
+				{
+					TPair<UControlRigControlsProxy*, FName> Selected(Proxy, PropertyName);
+					PropertiesToSelect.Add(Selected);
+				}
+			}
+			else //may be it's parent
+			{
+				UControlRigControlsProxy* ParentProxy = LastSelection.Key.Get();
+				int32 ChildIndex = ParentProxy->ChildProxies.Find(Proxy);
+				if (ChildIndex != INDEX_NONE)
+				{
+					//found child so we need to select all of our properties down to the last one
+					TOptional <FName> FirstName = LastSelection.Value;
+					TOptional <FName> SecondName; //empty
+					EAnimDetailRangeDirection Direction = EAnimDetailRangeDirection::Down;
+					GetPropertiesToSelect(LastSelection.Key.Get(), FirstName, SecondName, Direction);
+					//now we need to select all of the child proxies before this and including this
+					SelectChildProxiesFromIndex(ParentProxy, PropertyName, 0, ChildIndex);
+				}
+				else
+				{
+					TPair<UControlRigControlsProxy*, FName> Selected(Proxy, PropertyName);
+					PropertiesToSelect.Add(Selected);
+				}
+			}	
+		}
+		else if (Proxy->bIsIndividual == false)
+		{
+			if (LastSelection.Key->bIsIndividual) //last may be a child
+			{
+				//select this proxy down
+				TOptional <FName> FirstName = PropertyName;
+				TOptional <FName> SecondName; //empty
+				EAnimDetailRangeDirection Direction = EAnimDetailRangeDirection::Down;
+				GetPropertiesToSelect(Proxy, FirstName, SecondName, Direction);
+				int32 ChildIndex = Proxy->ChildProxies.Find(LastSelection.Key.Get());
+				if (ChildIndex != INDEX_NONE)
+				{
+					SelectChildProxiesFromIndex(Proxy, LastSelection.Value, 0, ChildIndex);
+				}
+			}
+			else
+			{
+				TPair<UControlRigControlsProxy*, FName> Selected(Proxy, PropertyName);
+				PropertiesToSelect.Add(Selected);
+			}
+		}
+	}
+	else //no range set so 
+	{
+		TPair<UControlRigControlsProxy*, FName> Selected(Proxy, PropertyName);
+		PropertiesToSelect.Add(Selected);
+	}
+
+	return PropertiesToSelect;
+}
+
+bool UControlRigDetailPanelControlProxies::IsPropertyEditingEnabled() const
+{
+	return true;
+}
+
+void UControlRigDetailPanelControlProxies::SelectProperty(UControlRigControlsProxy* Proxy, const FName& PropertyName, EAnimDetailPropertySelectionType SelectionType)
+{
+	if (SelectionType == EAnimDetailPropertySelectionType::SelectRange && LastSelection.Key.IsValid())
+	{
+		TArray<TPair<UControlRigControlsProxy*, FName>> PropertiesToSelect = GetPropertiesFromLastSelection(Proxy, PropertyName);
+		for (TPair<UControlRigControlsProxy*, FName>& Selected : PropertiesToSelect)
+		{
+			SelectPropertyInternal(Selected.Key, Selected.Value, SelectionType);
+		}
+		LastSelection.Key = Proxy;
+		LastSelection.Value = PropertyName;
+	}
+	else
+	{
+		if (SelectPropertyInternal(Proxy, PropertyName, SelectionType))
+		{
+			LastSelection.Key = Proxy;
+			LastSelection.Value = PropertyName;
+		}
+	}
+}
+
+//may want this to be under each proxy? need to think about this.
+static bool GetChannelNameForCurve(const TArray<FString>& CurveString, const FRigControlElement* ControlElement, FString& OutChannelName)
+{
+	//if single channel expect one item and the name will match
+	if (ControlElement->Settings.ControlType == ERigControlType::ScaleFloat ||
+		ControlElement->Settings.ControlType == ERigControlType::Float ||
+		ControlElement->Settings.ControlType == ERigControlType::Bool ||
+		ControlElement->Settings.ControlType == ERigControlType::Integer)
+	{
+		if (CurveString[0] == ControlElement->GetKey().Name)
+		{
+			if (ControlElement->Settings.ControlType == ERigControlType::ScaleFloat ||
+				ControlElement->Settings.ControlType == ERigControlType::Float)
+			{
+				OutChannelName = FString("Float");
+				return true;
+			}
+			if (ControlElement->Settings.ControlType == ERigControlType::Bool)
+			{
+				OutChannelName = FString("Bool");
+				return true;
+			}
+			if (ControlElement->Settings.ControlType == ERigControlType::Integer)
+			{
+				OutChannelName = FString("Integer");
+				return true;
+			}
+		}
+	}
+	else if (CurveString.Num() > 1)
+	{
+		if (CurveString[0] == ControlElement->GetKey().Name)
+		{
+			if (CurveString.Num() == 3)
+			{
+				OutChannelName = CurveString[1] + "." + CurveString[2];
+				return true;
+			}
+			else if (CurveString.Num() == 2)
+			{
+				OutChannelName = CurveString[1];
+				return true;
+			}
+		}
+	}
+	return false;
+}
+
+bool UControlRigDetailPanelControlProxies::SelectPropertyInternal(UControlRigControlsProxy* Proxy, const FName& PropertyName, EAnimDetailPropertySelectionType SelectionType)
+{
+	using namespace UE::Sequencer;
+
+	const TSharedPtr<FSequencerEditorViewModel> SequencerViewModel = GetSequencer()->GetViewModel();
+	const FCurveEditorExtension* CurveEditorExtension = SequencerViewModel->CastDynamic<FCurveEditorExtension>();
+	check(CurveEditorExtension);
+	TSharedPtr<FCurveEditor> CurveEditor = CurveEditorExtension->GetCurveEditor();
+	TSharedPtr<SCurveEditorTree>  CurveEditorTreeView = CurveEditorExtension->GetCurveEditorTreeView();
+	TSharedPtr<FOutlinerViewModel> OutlinerViewModel = SequencerViewModel->GetOutliner();
+	bool bIsSelected = false;
+	if (Proxy)
+	{
+		for (const TPair<TWeakObjectPtr<UControlRig>, FControlRigProxyItem>& Items : Proxy->ControlRigItems)
+		{
+			if (UControlRig* ControlRig = Items.Value.ControlRig.Get())
+			{
+				for (const FName& CName : Items.Value.ControlElements)
+				{
+					if (FRigControlElement* ControlElement = Items.Value.GetControlElement(CName))
+					{
+						EControlRigContextChannelToKey ChannelToKey = Proxy->GetChannelToKeyFromPropertyName(PropertyName);
+						TParentFirstChildIterator<IOutlinerExtension> OutlinerExtenstionIt = OutlinerViewModel->GetRootItem()->GetDescendantsOfType<IOutlinerExtension>();
+						for (; OutlinerExtenstionIt; ++OutlinerExtenstionIt)
+						{
+							if (TSharedPtr<FTrackModel> TrackModel = OutlinerExtenstionIt.GetCurrentItem()->FindAncestorOfType<FTrackModel>())
+							{
+								if (UMovieSceneControlRigParameterTrack* Track = Cast<UMovieSceneControlRigParameterTrack>(TrackModel->GetTrack()))
+								{
+									if (Track->GetControlRig() != ControlRig)
+									{
+										continue;
+									}
+									FName ID = OutlinerExtenstionIt->GetIdentifier();
+									FString Name = ID.ToString();
+									TArray<FString> StringArray;
+									Name.ParseIntoArray(StringArray, TEXT("."));
+
+									FString ChannelName;
+									if (GetChannelNameForCurve(StringArray, ControlElement, ChannelName))
+									{
+										EControlRigContextChannelToKey ChannelToKeyFromCurve = Proxy->GetChannelToKeyFromChannelName(ChannelName);
+										if (ChannelToKey == ChannelToKeyFromCurve)
+										{
+											if (TViewModelPtr<ICurveEditorTreeItemExtension> CurveEditorItem = OutlinerExtenstionIt.GetCurrentItem().ImplicitCast())
+											{
+												FCurveEditorTreeItemID CurveEditorTreeItem = CurveEditorItem->GetCurveEditorItemID();
+												if (CurveEditorTreeItem != FCurveEditorTreeItemID::Invalid())
+												{
+													bIsSelected = (SelectionType == EAnimDetailPropertySelectionType::Toggle) ? !CurveEditorTreeView->IsItemSelected(CurveEditorTreeItem) : true;
+													CurveEditorTreeView->SetItemSelection(CurveEditorTreeItem, bIsSelected);
+												}
+											}
+										}
+									}
+								}
+							}
+						}
+					}
+				}
+			}
+		}
+		for (const TPair<TWeakObjectPtr<UObject>, FSequencerProxyItem>& SItems : Proxy->SequencerItems)
+		{
+			if (UObject* Object = SItems.Key.Get())
+			{
+				for (const FBindingAndTrack& Element : SItems.Value.Bindings)
+				{
+					EControlRigContextChannelToKey ChannelToKey = Proxy->GetChannelToKeyFromPropertyName(PropertyName);
+					TParentFirstChildIterator<IOutlinerExtension> OutlinerExtenstionIt = OutlinerViewModel->GetRootItem()->GetDescendantsOfType<IOutlinerExtension>();
+					for (; OutlinerExtenstionIt; ++OutlinerExtenstionIt)
+					{
+						if (TSharedPtr<FTrackModel> TrackModel = OutlinerExtenstionIt.GetCurrentItem()->FindAncestorOfType<FTrackModel>())
+						{
+							if (TrackModel->GetTrack() == Element.WeakTrack.Get())
+							{
+								FName ID = OutlinerExtenstionIt->GetIdentifier();
+								FString Name = ID.ToString();
+								TArray<FString> StringArray;
+								Name.ParseIntoArray(StringArray, TEXT("."));
+
+								FString ChannelName;
+								if (StringArray.Num() == 2)
+								{
+									ChannelName = StringArray[0] + "." + StringArray[1];
+								}
+								else if (StringArray.Num() == 0)
+								{
+									ChannelName = StringArray[0];
+								}
+
+								EControlRigContextChannelToKey ChannelToKeyFromCurve = Proxy->GetChannelToKeyFromChannelName(ChannelName);
+								if (ChannelToKey == ChannelToKeyFromCurve)
+								{
+									if (TViewModelPtr<ICurveEditorTreeItemExtension> CurveEditorItem = OutlinerExtenstionIt.GetCurrentItem().ImplicitCast())
+									{
+										FCurveEditorTreeItemID CurveEditorTreeItem = CurveEditorItem->GetCurveEditorItemID();
+										if (CurveEditorTreeItem != FCurveEditorTreeItemID::Invalid())
+										{
+											bIsSelected = (SelectionType == EAnimDetailPropertySelectionType::Toggle) ? !CurveEditorTreeView->IsItemSelected(CurveEditorTreeItem) : true;
+											CurveEditorTreeView->SetItemSelection(CurveEditorTreeItem, bIsSelected);
+										}
+									}
+								}
+							}
+						}
+					}
+				}
+			}
+		}
+	}
+	return bIsSelected;
+}
+
 void FControlRigEnumControlProxyValueDetails::CustomizeHeader(TSharedRef<IPropertyHandle> InStructPropertyHandle, FDetailWidgetRow& HeaderRow, IPropertyTypeCustomizationUtils& StructCustomizationUtils)
 {
 	TArray<UObject*> Objects;
@@ -3289,4 +3825,4 @@ void FControlRigEnumControlProxyValueDetails::OnEnumValueChanged(int32 InValue,
 		ProxyBeingCustomized->Enum.EnumIndex = InValue;
 		InStructHandle->NotifyPostChange(EPropertyChangeType::ValueSet);
 	}
-}
+}
\ No newline at end of file
diff --git a/Engine/Plugins/Animation/ControlRig/Source/ControlRigEditor/Private/EditMode/AnimDetailsProxy.h b/Engine/Plugins/Animation/ControlRig/Source/ControlRigEditor/Private/EditMode/AnimDetailsProxy.h
--- a/Engine/Plugins/Animation/ControlRig/Source/ControlRigEditor/Private/EditMode/AnimDetailsProxy.h
+++ b/Engine/Plugins/Animation/ControlRig/Source/ControlRigEditor/Private/EditMode/AnimDetailsProxy.h
@@ -252,6 +252,7 @@ public:
 	virtual void ValueChanged() override;
 	virtual EControlRigContextChannelToKey GetChannelToKeyFromPropertyName(const FName& InPropertyName) const override;
 	virtual EControlRigContextChannelToKey GetChannelToKeyFromChannelName(const FString& InChannelName) const override;
+	virtual TMap<FName, int32> GetPropertyNames() const override;
 	virtual bool IsMultiple(const FName& InPropertyName) const override;
 	virtual void SetControlRigElementValueFromCurrent(UControlRig* ControlRig, FRigControlElement* ControlElement, const FRigControlModifiedContext& Context) override;
 	virtual void GetChannelSelectionState(TWeakPtr<FCurveEditor>& CurveEditor, FAnimDetailVectorSelection& OutLocationSelection, FAnimDetailVectorSelection& OutRotationSelection,
@@ -280,6 +281,7 @@ public:
 	virtual void ValueChanged() override;
 	virtual EControlRigContextChannelToKey GetChannelToKeyFromPropertyName(const FName& InPropertyName) const override;
 	virtual EControlRigContextChannelToKey GetChannelToKeyFromChannelName(const FString& InChannelName) const override;
+	virtual TMap<FName, int32> GetPropertyNames() const override;
 	virtual bool IsMultiple(const FName& InPropertyName) const override;
 	virtual void SetControlRigElementValueFromCurrent(UControlRig* ControlRig, FRigControlElement* ControlElement, const FRigControlModifiedContext& Context) override;
 	virtual void GetChannelSelectionState(TWeakPtr<FCurveEditor>& CurveEditor, FAnimDetailVectorSelection& OutLocationSelection, FAnimDetailVectorSelection& OutRotationSelection,
@@ -309,6 +311,7 @@ public:
 	virtual void ValueChanged() override;
 	virtual EControlRigContextChannelToKey GetChannelToKeyFromPropertyName(const FName& InPropertyName) const override;
 	virtual EControlRigContextChannelToKey GetChannelToKeyFromChannelName(const FString& InChannelName) const override;
+	virtual TMap<FName, int32> GetPropertyNames() const override;
 	virtual bool IsMultiple(const FName& InPropertyName) const override;
 	virtual void SetControlRigElementValueFromCurrent(UControlRig* ControlRig, FRigControlElement* ControlElement, const FRigControlModifiedContext& Context) override;
 	virtual void GetChannelSelectionState(TWeakPtr<FCurveEditor>& CurveEditor, FAnimDetailVectorSelection& OutLocationSelection, FAnimDetailVectorSelection& OutRotationSelection,
@@ -338,6 +341,7 @@ public:
 	virtual void ValueChanged() override;
 	virtual EControlRigContextChannelToKey GetChannelToKeyFromPropertyName(const FName& InPropertyName) const override;
 	virtual EControlRigContextChannelToKey GetChannelToKeyFromChannelName(const FString& InChannelName) const override;
+	virtual TMap<FName, int32> GetPropertyNames() const override;
 	virtual bool IsMultiple(const FName& InPropertyName) const override;
 	virtual void SetControlRigElementValueFromCurrent(UControlRig* ControlRig, FRigControlElement* ControlElement, const FRigControlModifiedContext& Context) override;
 	virtual void GetChannelSelectionState(TWeakPtr<FCurveEditor>& CurveEditor, FAnimDetailVectorSelection& OutLocationSelection, FAnimDetailVectorSelection& OutRotationSelection,
@@ -367,6 +371,7 @@ public:
 	virtual void ValueChanged() override;
 	virtual EControlRigContextChannelToKey GetChannelToKeyFromPropertyName(const FName& InPropertyName) const override;
 	virtual EControlRigContextChannelToKey GetChannelToKeyFromChannelName(const FString& InChannelName) const override;
+	virtual TMap<FName, int32> GetPropertyNames() const override;
 	virtual bool IsMultiple(const FName& InPropertyName) const override;
 	virtual void SetControlRigElementValueFromCurrent(UControlRig* ControlRig, FRigControlElement* ControlElement, const FRigControlModifiedContext& Context) override;
 	virtual void SetBindingValueFromCurrent(UObject* InObject, TSharedPtr<FTrackInstancePropertyBindings>& Binding, FRigControlModifiedContext& Context, bool bInteractive = false) override;
@@ -399,6 +404,7 @@ public:
 	virtual void ValueChanged() override;
 	virtual EControlRigContextChannelToKey GetChannelToKeyFromPropertyName(const FName& InPropertyName) const override;
 	virtual EControlRigContextChannelToKey GetChannelToKeyFromChannelName(const FString& InChannelName) const override;
+	virtual TMap<FName, int32> GetPropertyNames() const override;
 	virtual void SetControlRigElementValueFromCurrent(UControlRig* ControlRig, FRigControlElement* ControlElement, const FRigControlModifiedContext& Context) override;
 	virtual bool IsMultiple(const FName& InPropertyName) const override;
 	virtual void GetChannelSelectionState(TWeakPtr<FCurveEditor>& CurveEditor, FAnimDetailVectorSelection& OutLocationSelection, FAnimDetailVectorSelection& OutRotationSelection,
@@ -423,6 +429,7 @@ public:
 	virtual void ValueChanged() override;
 	virtual EControlRigContextChannelToKey GetChannelToKeyFromPropertyName(const FName& InPropertyName) const override;
 	virtual EControlRigContextChannelToKey GetChannelToKeyFromChannelName(const FString& InChannelName) const override;
+	virtual TMap<FName, int32> GetPropertyNames() const override;
 	virtual void SetControlRigElementValueFromCurrent(UControlRig* ControlRig, FRigControlElement* ControlElement, const FRigControlModifiedContext& Context) override;
 	virtual bool IsMultiple(const FName& InPropertyName) const override;
 	virtual void GetChannelSelectionState(TWeakPtr<FCurveEditor>& CurveEditor, FAnimDetailVectorSelection& OutLocationSelection, FAnimDetailVectorSelection& OutRotationSelection,
@@ -447,6 +454,7 @@ public:
 	virtual void ValueChanged() override;
 	virtual EControlRigContextChannelToKey GetChannelToKeyFromPropertyName(const FName& InPropertyName) const override;
 	virtual EControlRigContextChannelToKey GetChannelToKeyFromChannelName(const FString& InChannelName) const override;
+	virtual TMap<FName, int32> GetPropertyNames() const override;
 	virtual void SetControlRigElementValueFromCurrent(UControlRig* ControlRig, FRigControlElement* ControlElement, const FRigControlModifiedContext& Context) override;
 	virtual bool IsMultiple(const FName& InPropertyName) const override;
 	virtual void GetChannelSelectionState(TWeakPtr<FCurveEditor>& CurveEditor, FAnimDetailVectorSelection& OutLocationSelection, FAnimDetailVectorSelection& OutRotationSelection,
@@ -471,6 +479,7 @@ public:
 	virtual void ValueChanged() override;
 	virtual EControlRigContextChannelToKey GetChannelToKeyFromPropertyName(const FName& InPropertyName) const override;
 	virtual EControlRigContextChannelToKey GetChannelToKeyFromChannelName(const FString& InChannelName) const override;
+	virtual TMap<FName, int32> GetPropertyNames() const override;
 	virtual void SetControlRigElementValueFromCurrent(UControlRig* ControlRig, FRigControlElement* ControlElement, const FRigControlModifiedContext& Context) override;
 	virtual bool IsMultiple(const FName& InPropertyName) const override;
 	virtual void GetChannelSelectionState(TWeakPtr<FCurveEditor>& CurveEditor, FAnimDetailVectorSelection& OutLocationSelection, FAnimDetailVectorSelection& OutRotationSelection,
diff --git a/Engine/Plugins/Animation/ControlRig/Source/ControlRigEditor/Private/EditMode/ControlRigControlsProxy.cpp b/Engine/Plugins/Animation/ControlRig/Source/ControlRigEditor/Private/EditMode/ControlRigControlsProxy.cpp
--- a/Engine/Plugins/Animation/ControlRig/Source/ControlRigEditor/Private/EditMode/ControlRigControlsProxy.cpp
+++ b/Engine/Plugins/Animation/ControlRig/Source/ControlRigEditor/Private/EditMode/ControlRigControlsProxy.cpp
@@ -73,7 +73,9 @@ void UControlRigControlsProxy::AddControlRigControl(UControlRig* InControlRig, c
 	}
 	else
 	{
-		Name = ControlElement->GetDisplayName();
+		FText ControlNameText = InControlRig->GetHierarchy()->GetDisplayNameForUI(ControlElement);
+		FString ControlNameString = ControlNameText.ToString();
+		Name = FName(*ControlNameString);
 	}
 }
 
diff --git a/Engine/Plugins/Animation/ControlRig/Source/ControlRigEditor/Private/EditMode/ControlRigControlsProxy.h b/Engine/Plugins/Animation/ControlRig/Source/ControlRigEditor/Private/EditMode/ControlRigControlsProxy.h
--- a/Engine/Plugins/Animation/ControlRig/Source/ControlRigEditor/Private/EditMode/ControlRigControlsProxy.h
+++ b/Engine/Plugins/Animation/ControlRig/Source/ControlRigEditor/Private/EditMode/ControlRigControlsProxy.h
@@ -32,6 +32,11 @@ enum class EAnimDetailSelectionState : uint8
 	None = 0x0, Partial = 0x1, All = 0x2
 };
 
+//direction to find range of property names
+enum class EAnimDetailRangeDirection : uint8
+{
+	Up = 0x0, Down = 0x1
+};
 struct FAnimDetailVectorSelection
 {
 	EAnimDetailSelectionState  XSelected = EAnimDetailSelectionState::None;
@@ -80,6 +85,7 @@ public:
 	virtual EPropertyKeyedStatus GetPropertyKeyedStatus(TSharedPtr<ISequencer>& Sequencer, const IPropertyHandle& PropertyHandle) const { return EPropertyKeyedStatus::NotKeyed; }
 	virtual EControlRigContextChannelToKey GetChannelToKeyFromPropertyName(const FName& PropertyName) const { return EControlRigContextChannelToKey::AllTransform; }
 	virtual EControlRigContextChannelToKey GetChannelToKeyFromChannelName(const FString& InChannelName) const { return EControlRigContextChannelToKey::AllTransform; }
+	virtual TMap<FName, int32> GetPropertyNames() const { TMap<FName, int32> Empty; return Empty; }
 	virtual bool IsMultiple(const FName& InPropertyName) const { return false; }
 	virtual void SetControlRigElementValueFromCurrent(UControlRig* ControlRig, FRigControlElement* ControlElement, const FRigControlModifiedContext& Context) {};
 	virtual void SetBindingValueFromCurrent(UObject* InObject, TSharedPtr<FTrackInstancePropertyBindings>& Binding, FRigControlModifiedContext& Context, bool bInteractive = false) {};
@@ -107,8 +113,6 @@ public:
 	void AddControlRigControl(UControlRig* InControlRig, const FName& InName);
 
 private:
-	void HandlePostEditChangedForControlRigs();
-	void HandlePostEditChangedForSequencer();
 
 	//reset the  rig controls this also owns
 	void ResetControlRigItems();
@@ -157,6 +161,11 @@ struct FSequencerProxyPerType
 	TMap<UObject*, TArray<FBindingAndTrack>> Bindings;
 };
 
+enum class EAnimDetailPropertySelectionType : uint8
+{
+	Toggle = 0x0, Select = 0x1, SelectRange = 0x2
+};
+
 /** Proxy in Details Panel */
 UCLASS()
 class UControlRigDetailPanelControlProxies :public UObject
@@ -180,6 +189,8 @@ protected:
 	UPROPERTY()
 	TArray< TObjectPtr<UControlRigControlsProxy>> SelectedSequencerProxies;
 	
+	TPair<TWeakObjectPtr<UControlRigControlsProxy>, FName> LastSelection;
+
 	TWeakPtr<ISequencer> Sequencer;
 
 public:
@@ -199,12 +210,25 @@ public:
 	void ValuesChanged();
 	void ResetSequencerProxies(TMap<ERigControlType, FSequencerProxyPerType>& ProxyPerType);
 
+	//property selection to handle shift selection ranges
+	void SelectProperty(UControlRigControlsProxy* Proxy, const FName& PropertyName, EAnimDetailPropertySelectionType SelectionType);
+
+	bool IsPropertyEditingEnabled() const;
 private:
+
+
 	UControlRigControlsProxy* FindProxy(UControlRig* InControlRig, FRigControlElement* RigElement) const;
 	UControlRigControlsProxy* FindProxy(UObject* InObject, FName PropertyName) const;
 
 	UControlRigControlsProxy* NewProxyFromType(ERigControlType Type, TObjectPtr<UEnum>& EnumPtr);
 
+	//if SelectedProxy == nullptr then always clear it
+	void ClearSelectedProperty(UControlRigControlsProxy* SelectedProxy = nullptr);
+	//select/deslect property based on selection type  return true if it is selected
+	bool SelectPropertyInternal(UControlRigControlsProxy* Proxy, const FName& PropertyName, EAnimDetailPropertySelectionType SelectionType);
+	//using the Last Section, get the range of properties from that (used for shift selections).
+	TArray<TPair< UControlRigControlsProxy*, FName>> GetPropertiesFromLastSelection(UControlRigControlsProxy* Proxy, const FName& PropertyName) const;
+
 	//delegate for changed objects
 	void OnPostPropertyChanged(UObject* InObject, FPropertyChangedEvent& InPropertyChangedEvent);
 };
diff --git a/Engine/Plugins/Animation/ControlRig/Source/ControlRigEditor/Private/EditMode/ControlsProxyDetailCustomization.cpp b/Engine/Plugins/Animation/ControlRig/Source/ControlRigEditor/Private/EditMode/ControlsProxyDetailCustomization.cpp
--- a/Engine/Plugins/Animation/ControlRig/Source/ControlRigEditor/Private/EditMode/ControlsProxyDetailCustomization.cpp
+++ b/Engine/Plugins/Animation/ControlRig/Source/ControlRigEditor/Private/EditMode/ControlsProxyDetailCustomization.cpp
@@ -27,16 +27,13 @@
 #include "Widgets/Text/STextBlock.h"
 #include "Widgets/Input/SButton.h"
 
-#include "MVVM/CurveEditorExtension.h"
 #include "MVVM/ViewModelPtr.h"
 #include "MVVM/Selection/Selection.h"
 #include "MVVM/ViewModels/SequencerEditorViewModel.h"
-#include "MVVM/ViewModels/OutlinerViewModel.h"
 #include "MVVM/ViewModels/ViewModelIterators.h"
-#include "MVVM/Extensions/IOutlinerExtension.h"
-#include "MVVM/ViewModels/TrackModel.h"
 #include "MVVM/ViewModels/ChannelModel.h"
 #include "Rigs/RigControlHierarchy.h"
+#include "MVVM/CurveEditorExtension.h"
 #include "Tree/SCurveEditorTree.h"
 
 #include "CurveEditor.h"
@@ -300,6 +297,18 @@ void FAnimDetailValueCustomization::CustomizeChildren(TSharedRef<class IProperty
 			.OnMouseButtonDown_Lambda([Proxy, PropertyName, this](const FGeometry&, const FPointerEvent& PointerEvent)
 			{
 				TogglePropertySelection(Proxy, PropertyName);
+				if (Proxy)
+				{
+					UControlRigDetailPanelControlProxies* ProxyOwner = Proxy->GetTypedOuter<UControlRigDetailPanelControlProxies>();
+				}
+				return FReply::Handled();
+			})
+			.OnMouseButtonUp_Lambda([Proxy](const FGeometry&, const FPointerEvent& PointerEvent)
+			{
+				if(Proxy)
+				{
+					UControlRigDetailPanelControlProxies* ProxyOwner = Proxy->GetTypedOuter<UControlRigDetailPanelControlProxies>();
+				}
 				return FReply::Handled();
 			})
 			.Content()
@@ -380,6 +389,12 @@ void FAnimDetailValueCustomization::MakeHeaderRow(TSharedRef<class IPropertyHand
 		[
 			StructPropertyHandle->CreatePropertyNameWidget()
 		]
+	.PasteAction(FUIAction(
+		FExecuteAction::CreateLambda([]()
+			{
+			}),
+		FCanExecuteAction::CreateLambda([]() { return false; }))
+	)
 	.ValueContent()
 		// Make enough space for each child handle
 		.MinDesiredWidth(125.f * SortedChildHandles.Num())
@@ -390,10 +405,11 @@ void FAnimDetailValueCustomization::MakeHeaderRow(TSharedRef<class IPropertyHand
 				.IsEnabled(this, &FMathStructCustomization::IsValueEnabled, StructWeakHandlePtr)
 		];
 
+
 	for (int32 ChildIndex = 0; ChildIndex < SortedChildHandles.Num(); ++ChildIndex)
 	{
 		TSharedRef<IPropertyHandle> ChildHandle = SortedChildHandles[ChildIndex];
-
+		ChildHandle->GetProperty()->SetPropertyFlags(CPF_TextExportTransient); //hack to turn off shift copy/paste
 		// Propagate metadata to child properties so that it's reflected in the nested, individual spin boxes
 		ChildHandle->SetInstanceMetaData(TEXT("UIMin"), StructPropertyHandle->GetMetaData(TEXT("UIMin")));
 		ChildHandle->SetInstanceMetaData(TEXT("UIMax"), StructPropertyHandle->GetMetaData(TEXT("UIMax")));
@@ -481,54 +497,6 @@ EAnimDetailSelectionState FAnimDetailValueCustomization::IsPropertySelected(UCon
 	
 }
 
-//may want this to be under each proxy? need to think about this.
-static bool GetChannelNameForCurve(const TArray<FString>& CurveString, const FRigControlElement* ControlElement, FString& OutChannelName)
-{
-	//if single channel expect one item and the name will match
-	if (ControlElement->Settings.ControlType == ERigControlType::ScaleFloat ||
-		ControlElement->Settings.ControlType == ERigControlType::Float ||
-		ControlElement->Settings.ControlType == ERigControlType::Bool ||
-		ControlElement->Settings.ControlType == ERigControlType::Integer)
-	{
-		if (CurveString[0] == ControlElement->GetKey().Name)
-		{
-			if (ControlElement->Settings.ControlType == ERigControlType::ScaleFloat ||
-				ControlElement->Settings.ControlType == ERigControlType::Float)
-			{
-				OutChannelName = FString("Float");
-				return true;
-			}
-			if (ControlElement->Settings.ControlType == ERigControlType::Bool)
-			{
-				OutChannelName = FString("Bool");
-				return true;
-			}
-			if (ControlElement->Settings.ControlType == ERigControlType::Integer)
-			{
-				OutChannelName = FString("Integer");
-				return true;
-			}
-		}
-	}
-	else if (CurveString.Num() > 1)
-	{
-		if (CurveString[0] == ControlElement->GetKey().Name)
-		{
-			if (CurveString.Num() == 3)
-			{
-				OutChannelName = CurveString[1] + "." + CurveString[2];
-				return true;
-			}
-			else if (CurveString.Num() == 2)
-			{
-				OutChannelName = CurveString[1];
-				return true;
-			}
-		}
-	}
-	return false;
-}
-
 void FAnimDetailValueCustomization::TogglePropertySelection(UControlRigControlsProxy* Proxy, const FName& PropertyName) const
 {
 	using namespace UE::Sequencer;
@@ -543,7 +511,6 @@ void FAnimDetailValueCustomization::TogglePropertySelection(UControlRigControlsP
 	check(CurveEditorExtension);
 	TSharedPtr<FCurveEditor> CurveEditor = CurveEditorExtension->GetCurveEditor();
 	TSharedPtr<SCurveEditorTree>  CurveEditorTreeView = CurveEditorExtension->GetCurveEditorTreeView();
-	TSharedPtr<FOutlinerViewModel> OutlinerViewModel = SequencerViewModel->GetOutliner();
 
 	const bool bIsShiftDown = FSlateApplication::Get().GetModifierKeys().IsShiftDown();
 	const bool bIsCtrlDown = FSlateApplication::Get().GetModifierKeys().IsControlDown();
@@ -551,125 +518,8 @@ void FAnimDetailValueCustomization::TogglePropertySelection(UControlRigControlsP
 	{
 		CurveEditorTreeView->ClearSelection();
 	}
-
-	if (Proxy)
-	{
-		for (const TPair<TWeakObjectPtr<UControlRig>, FControlRigProxyItem>& Items : Proxy->ControlRigItems)
-		{
-			if (UControlRig* ControlRig = Items.Value.ControlRig.Get())
-			{
-				for (const FName& CName : Items.Value.ControlElements)
-				{
-					if (FRigControlElement* ControlElement = Items.Value.GetControlElement(CName))
-					{
-						EControlRigContextChannelToKey ChannelToKey = Proxy->GetChannelToKeyFromPropertyName(PropertyName);
-						TParentFirstChildIterator<IOutlinerExtension> OutlinerExtenstionIt = OutlinerViewModel->GetRootItem()->GetDescendantsOfType<IOutlinerExtension>();
-						for (; OutlinerExtenstionIt; ++OutlinerExtenstionIt)
-						{
-							if (TSharedPtr<FTrackModel> TrackModel = OutlinerExtenstionIt.GetCurrentItem()->FindAncestorOfType<FTrackModel>())
-							{
-								if(UMovieSceneControlRigParameterTrack* Track = Cast<UMovieSceneControlRigParameterTrack>(TrackModel->GetTrack()))
-								{
-									if (Track->GetControlRig() != ControlRig)
-									{
-										continue;
-									}
-									if (TViewModelPtr<FChannelGroupOutlinerModel> ChannelModel = CastViewModel<FChannelGroupOutlinerModel>(OutlinerExtenstionIt.GetCurrentItem()))
-									{
-										if (ChannelModel->GetChannel(Track->GetSectionToKey()) == nullptr) //if not section to key we also don't select it.
-										{
-											continue;
-										}
-									}
-									else
-									{
-										continue;
-									}
-									
-									FName ID = OutlinerExtenstionIt->GetIdentifier();
-									FString Name = ID.ToString();
-									TArray<FString> StringArray;
-									Name.ParseIntoArray(StringArray, TEXT("."));
-
-									FString ChannelName;
-									if (GetChannelNameForCurve(StringArray, ControlElement, ChannelName))
-									{
-										EControlRigContextChannelToKey ChannelToKeyFromCurve = Proxy->GetChannelToKeyFromChannelName(ChannelName);
-										if (ChannelToKey == ChannelToKeyFromCurve)
-										{
-											if (TViewModelPtr<ICurveEditorTreeItemExtension> CurveEditorItem = OutlinerExtenstionIt.GetCurrentItem().ImplicitCast())
-											{
-												FCurveEditorTreeItemID CurveEditorTreeItem = CurveEditorItem->GetCurveEditorItemID();
-												if (CurveEditorTreeItem != FCurveEditorTreeItemID::Invalid())
-												{
-													const bool bSelected = bIsShiftDown ? true : !CurveEditorTreeView->IsItemSelected(CurveEditorTreeItem);
-													if (bIsCtrlDown || bSelected)
-													{
-														CurveEditorTreeView->SetItemSelection(CurveEditorTreeItem, bSelected);
-													}
-												}
-											}
-										}
-									}
-								}
-							}
-						}
-					}
-				}
-			}
-		}
-		for (const TPair<TWeakObjectPtr<UObject>, FSequencerProxyItem>& SItems : Proxy->SequencerItems)
-		{
-			if (UObject* Object = SItems.Key.Get())
-			{
-				for (const FBindingAndTrack& Element : SItems.Value.Bindings)
-				{
-					EControlRigContextChannelToKey ChannelToKey = Proxy->GetChannelToKeyFromPropertyName(PropertyName);
-					TParentFirstChildIterator<IOutlinerExtension> OutlinerExtenstionIt = OutlinerViewModel->GetRootItem()->GetDescendantsOfType<IOutlinerExtension>();
-					for (; OutlinerExtenstionIt; ++OutlinerExtenstionIt)
-					{
-						if (TSharedPtr<FTrackModel> TrackModel = OutlinerExtenstionIt.GetCurrentItem()->FindAncestorOfType<FTrackModel>())
-						{
-							if (TrackModel->GetTrack() == Element.WeakTrack.Get())
-							{
-								FName ID = OutlinerExtenstionIt->GetIdentifier();
-								FString Name = ID.ToString();
-								TArray<FString> StringArray;
-								Name.ParseIntoArray(StringArray, TEXT("."));
-
-								FString ChannelName;
-								if (StringArray.Num() == 2)
-								{
-									ChannelName = StringArray[0] + "." + StringArray[1];
-								}
-								else if (StringArray.Num() == 0)
-								{
-									ChannelName = StringArray[0];
-								}
-
-								EControlRigContextChannelToKey ChannelToKeyFromCurve = Proxy->GetChannelToKeyFromChannelName(ChannelName);
-								if (ChannelToKey == ChannelToKeyFromCurve)
-								{
-									if (TViewModelPtr<ICurveEditorTreeItemExtension> CurveEditorItem = OutlinerExtenstionIt.GetCurrentItem().ImplicitCast())
-									{
-										FCurveEditorTreeItemID CurveEditorTreeItem = CurveEditorItem->GetCurveEditorItemID();
-										if (CurveEditorTreeItem != FCurveEditorTreeItemID::Invalid())
-										{
-											const bool bSelected = bIsShiftDown ? true : !CurveEditorTreeView->IsItemSelected(CurveEditorTreeItem);
-											if (bIsCtrlDown || bSelected)
-											{
-												CurveEditorTreeView->SetItemSelection(CurveEditorTreeItem, bSelected);
-											}
-										}
-									}
-								}
-							}
-						}
-					}
-				}
-			}
-		}
-	}
+	EAnimDetailPropertySelectionType SelectionType = bIsShiftDown ? EAnimDetailPropertySelectionType::SelectRange : (bIsCtrlDown ? EAnimDetailPropertySelectionType::Toggle : EAnimDetailPropertySelectionType::Select);
+	ProxyOwner->SelectProperty(Proxy, PropertyName, SelectionType);
 }
 
 bool FAnimDetailValueCustomization::IsMultiple(UControlRigControlsProxy* Proxy,const FName& InPropertyName) const
diff --git a/Engine/Plugins/Animation/ControlRig/Source/ControlRigEditor/Private/Sequencer/ControlRigParameterTrackEditor.cpp b/Engine/Plugins/Animation/ControlRig/Source/ControlRigEditor/Private/Sequencer/ControlRigParameterTrackEditor.cpp
--- a/Engine/Plugins/Animation/ControlRig/Source/ControlRigEditor/Private/Sequencer/ControlRigParameterTrackEditor.cpp
+++ b/Engine/Plugins/Animation/ControlRig/Source/ControlRigEditor/Private/Sequencer/ControlRigParameterTrackEditor.cpp
@@ -3577,7 +3577,7 @@ void FControlRigParameterTrackEditor::GetControlRigKeys(
 				bool bKeyX = bSetKey && EnumHasAnyFlags(ChannelsToKey, EControlRigContextChannelToKey::TranslationX);
 				bool bKeyY = bSetKey && EnumHasAnyFlags(ChannelsToKey, EControlRigContextChannelToKey::TranslationY);
 				bool bKeyZ = bSetKey && EnumHasAnyFlags(ChannelsToKey, EControlRigContextChannelToKey::TranslationZ);
-				if (GetSequencer()->GetKeyGroupMode() == EKeyGroupMode::KeyGroup && (bKeyX || bKeyY || bKeyZ))
+				if (KeyMode != ESequencerKeyMode::ManualKeyForced && (GetSequencer()->GetKeyGroupMode() == EKeyGroupMode::KeyGroup && (bKeyX || bKeyY || bKeyZ)))
 				{
 					bKeyX = bKeyY = bKeyZ = true;
 				}
@@ -3614,7 +3614,7 @@ void FControlRigParameterTrackEditor::GetControlRigKeys(
 				bKeyX = bSetKey && EnumHasAnyFlags(ChannelsToKey, EControlRigContextChannelToKey::RotationX);
 				bKeyY = bSetKey && EnumHasAnyFlags(ChannelsToKey, EControlRigContextChannelToKey::RotationY);
 				bKeyZ = bSetKey && EnumHasAnyFlags(ChannelsToKey, EControlRigContextChannelToKey::RotationZ);
-				if (GetSequencer()->GetKeyGroupMode() == EKeyGroupMode::KeyGroup && (bKeyX || bKeyY || bKeyZ))
+				if (KeyMode != ESequencerKeyMode::ManualKeyForced && (GetSequencer()->GetKeyGroupMode() == EKeyGroupMode::KeyGroup && (bKeyX || bKeyY || bKeyZ)))
 				{
 					bKeyX = bKeyY = bKeyZ = true;
 				}
@@ -3641,7 +3641,7 @@ void FControlRigParameterTrackEditor::GetControlRigKeys(
 					bKeyX = bSetKey && EnumHasAnyFlags(ChannelsToKey, EControlRigContextChannelToKey::ScaleX);
 					bKeyY = bSetKey && EnumHasAnyFlags(ChannelsToKey, EControlRigContextChannelToKey::ScaleY);
 					bKeyZ = bSetKey && EnumHasAnyFlags(ChannelsToKey, EControlRigContextChannelToKey::ScaleZ);
-					if (GetSequencer()->GetKeyGroupMode() == EKeyGroupMode::KeyGroup && (bKeyX || bKeyY || bKeyZ))
+					if (KeyMode != ESequencerKeyMode::ManualKeyForced && (GetSequencer()->GetKeyGroupMode() == EKeyGroupMode::KeyGroup && (bKeyX || bKeyY || bKeyZ)))
 					{
 						bKeyX = bKeyY = bKeyZ = true;
 					}
