RigVM UserData: Use the cached ptr when loading user data assets in client builds. Otherwise use the soft object ptr.

#rnx

diff --git a/Engine/Plugins/Animation/ControlRig/Source/ControlRig/Private/ControlRigAssetUserData.cpp b/Engine/Plugins/Animation/ControlRig/Source/ControlRig/Private/ControlRigAssetUserData.cpp
--- a/Engine/Plugins/Animation/ControlRig/Source/ControlRig/Private/ControlRigAssetUserData.cpp
+++ b/Engine/Plugins/Animation/ControlRig/Source/ControlRig/Private/ControlRigAssetUserData.cpp
@@ -88,13 +88,33 @@ const TArray<const UNameSpacedUserData::FUserData*>& UControlRigShapeLibraryLink
 	return EmptyUserDatas;
 }
 
+void UControlRigShapeLibraryLink::Serialize(FArchive& Ar)
+{
+	// Treat the cached ptr as transient unless we're cooking this out.
+	const bool bIsSavingAssetToStorage = Ar.IsSaving() && Ar.IsPersistent() && !Ar.IsCooking(); 
+	UControlRigShapeLibrary* SavedShapeLibrary = ShapeLibraryCached; 
+	if (bIsSavingAssetToStorage)
+	{
+		ShapeLibraryCached = nullptr;
+	}
+	
+	Super::Serialize(Ar);
+	
+	if (bIsSavingAssetToStorage)
+	{
+		ShapeLibraryCached = SavedShapeLibrary;
+	}
+}
+
+#if WITH_EDITOR
+
 void UControlRigShapeLibraryLink::PostLoad()
 {
 	Super::PostLoad();
 
 	ShapeLibraryCached = ShapeLibrary.Get();
 	
-	if(ShapeLibraryCached == nullptr)
+	if(ShapeLibraryCached == nullptr && !ShapeLibrary.IsNull())
 	{
 		// We need to check if the mount point exists - since the shape library link may
 		// refer to an editor-only asset in a runtime game.
@@ -106,28 +126,11 @@ void UControlRigShapeLibraryLink::PostLoad()
 
 			// load without throwing additional warnings / errors
 			ShapeLibraryCached = LoadObject<UControlRigShapeLibrary>(nullptr, *ObjectPath, nullptr, LOAD_Quiet | LOAD_NoWarn);
-			
-			if (ShapeLibraryCached)
-			{
-				if(NameSpace.IsEmpty())
-				{
-					NameSpace = ShapeLibrary->GetName();
-				}
-				
-				ShapeNames.Reset();
-				for(const FControlRigShapeDefinition& Shape : ShapeLibraryCached->Shapes)
-				{
-					ShapeNames.Add(Shape.ShapeName);
-				}
-
-				verify(ShapeLibraryCached == ShapeLibrary.Get());
-			}
+			SetShapeLibrary(ShapeLibrary);
 		}
 	}
 }
 
-#if WITH_EDITOR
-
 void UControlRigShapeLibraryLink::PostEditChangeProperty(FPropertyChangedEvent& PropertyChangedEvent)
 {
 	Super::PostEditChangeProperty(PropertyChangedEvent);
diff --git a/Engine/Plugins/Animation/ControlRig/Source/ControlRig/Public/ControlRigAssetUserData.h b/Engine/Plugins/Animation/ControlRig/Source/ControlRig/Public/ControlRigAssetUserData.h
--- a/Engine/Plugins/Animation/ControlRig/Source/ControlRig/Public/ControlRigAssetUserData.h
+++ b/Engine/Plugins/Animation/ControlRig/Source/ControlRig/Public/ControlRigAssetUserData.h
@@ -29,9 +29,10 @@ public:
 	virtual const FUserData* GetUserData(const FString& InPath, FString* OutErrorMessage = nullptr) const override;
 	virtual const TArray<const FUserData*>& GetUserDataArray(const FString& InParentPath = FString(), FString* OutErrorMessage = nullptr) const override;
 
+	virtual void Serialize(FArchive& Ar) override;
+#if WITH_EDITOR
 	virtual bool IsPostLoadThreadSafe() const override { return false; }
 	virtual void PostLoad() override;
-#if WITH_EDITOR
 	virtual void PostEditChangeProperty(FPropertyChangedEvent& PropertyChangedEvent) override;
 #endif
 
@@ -40,7 +41,7 @@ protected:
 	UPROPERTY(transient)
 	TArray<FName> ShapeNames;
 
-	UPROPERTY(transient)
+	UPROPERTY()
 	TObjectPtr<UControlRigShapeLibrary> ShapeLibraryCached;
 	
 
diff --git a/Engine/Plugins/Runtime/RigVM/Source/RigVM/Private/RigVMCore/RigVMAssetUserData.cpp b/Engine/Plugins/Runtime/RigVM/Source/RigVM/Private/RigVMCore/RigVMAssetUserData.cpp
--- a/Engine/Plugins/Runtime/RigVM/Source/RigVM/Private/RigVMCore/RigVMAssetUserData.cpp
+++ b/Engine/Plugins/Runtime/RigVM/Source/RigVM/Private/RigVMCore/RigVMAssetUserData.cpp
@@ -457,13 +457,32 @@ const TArray<const UNameSpacedUserData::FUserData*>& UDataAssetLink::GetUserData
 	return EmptyUserDatas;
 }
 
+void UDataAssetLink::Serialize(FArchive& Ar)
+{
+	// Treat the cached ptr as transient unless we're cooking this out.
+	const bool bIsSavingAssetToStorage = Ar.IsSaving() && Ar.IsPersistent() && !Ar.IsCooking(); 
+	UDataAsset* SavedDataAsset = DataAssetCached; 
+	if (bIsSavingAssetToStorage)
+	{
+		DataAssetCached = nullptr;
+	}
+
+	Super::Serialize(Ar);
+	
+	if (bIsSavingAssetToStorage)
+	{
+		DataAssetCached = SavedDataAsset;
+	}
+}
+
+#if WITH_EDITOR
+
 void UDataAssetLink::PostLoad()
 {
 	Super::PostLoad();
 
 	DataAssetCached = DataAsset.Get();
-
-	if(DataAssetCached == nullptr)
+	if(DataAssetCached == nullptr && !DataAsset.IsNull())
 	{
 		// We need to check if the mount point exists - since the data asset library link may
 		// refer to an editor-only asset in a runtime game.
@@ -475,21 +494,11 @@ void UDataAssetLink::PostLoad()
 
 			// load without throwing additional warnings / errors
 			DataAssetCached = LoadObject<UDataAsset>(nullptr, *ObjectPath, nullptr, LOAD_Quiet | LOAD_NoWarn);
-
-			if(DataAssetCached)
-			{
-				if(NameSpace.IsEmpty())
-				{
-					NameSpace = DataAsset->GetName();
-				}
-
-				verify(DataAssetCached == DataAsset.Get());
-			}
+			SetDataAsset(DataAsset);
 		}
 	}
 }
 
-#if WITH_EDITOR
 
 void UDataAssetLink::PostEditChangeProperty(FPropertyChangedEvent& PropertyChangedEvent)
 {
diff --git a/Engine/Plugins/Runtime/RigVM/Source/RigVM/Public/RigVMCore/RigVMAssetUserData.h b/Engine/Plugins/Runtime/RigVM/Source/RigVM/Public/RigVMCore/RigVMAssetUserData.h
--- a/Engine/Plugins/Runtime/RigVM/Source/RigVM/Public/RigVMCore/RigVMAssetUserData.h
+++ b/Engine/Plugins/Runtime/RigVM/Source/RigVM/Public/RigVMCore/RigVMAssetUserData.h
@@ -203,14 +203,15 @@ public:
 	virtual const FUserData* GetUserData(const FString& InPath, FString* OutErrorMessage = nullptr) const override;
 	virtual const TArray<const FUserData*>& GetUserDataArray(const FString& InParentPath = FString(), FString* OutErrorMessage = nullptr) const override;
 
+	virtual void Serialize(FArchive& Ar) override;
+#if WITH_EDITOR
 	virtual bool IsPostLoadThreadSafe() const override { return false; }
 	virtual void PostLoad() override;
-#if WITH_EDITOR
 	virtual void PostEditChangeProperty(FPropertyChangedEvent& PropertyChangedEvent) override;
 #endif
 
 protected:
-	UPROPERTY(transient)
+	UPROPERTY()
 	TObjectPtr<UDataAsset> DataAssetCached;
 	
 	static inline constexpr TCHAR DataAssetNullFormat[] = TEXT("User data path '%s' could not be found (DataAsset not provided)");
