Sequencer: Anim Details: Add support for interactions and driven controls

diff --git a/Engine/Plugins/Animation/ControlRig/Source/ControlRigEditor/Private/EditMode/AnimDetailsProxy.cpp b/Engine/Plugins/Animation/ControlRig/Source/ControlRigEditor/Private/EditMode/AnimDetailsProxy.cpp
--- a/Engine/Plugins/Animation/ControlRig/Source/ControlRigEditor/Private/EditMode/AnimDetailsProxy.cpp
+++ b/Engine/Plugins/Animation/ControlRig/Source/ControlRigEditor/Private/EditMode/AnimDetailsProxy.cpp
@@ -240,6 +240,8 @@ static EAnimDetailSelectionState CachePropertySelection(TWeakPtr<FCurveEditor>&
 
 void UAnimDetailControlsKeyedProxy::SetKey(TSharedPtr<ISequencer>& Sequencer, const IPropertyHandle& KeyedPropertyHandle)
 {
+	FRigControlModifiedContext NotifyDrivenContext; //always key ever
+	NotifyDrivenContext.SetKey = EControlRigSetKey::Always;
 	for (const TPair<TWeakObjectPtr<UControlRig>, FControlRigProxyItem>& Items : ControlRigItems)
 	{
 		if (UControlRig* ControlRig = Items.Value.ControlRig.Get())
@@ -255,6 +257,7 @@ void UAnimDetailControlsKeyedProxy::SetKey(TSharedPtr<ISequencer>& Sequencer, co
 						FName PropertyName = KeyedPropertyHandle.GetProperty()->GetFName();
 						Context.KeyMask = (uint32)GetChannelToKeyFromPropertyName(PropertyName);
 						SetControlRigElementValueFromCurrent(ControlRig, ControlElement, Context);
+						FControlRigEditMode::NotifyDrivenControls(ControlRig, ControlElement->GetKey(), NotifyDrivenContext);
 					}
 				}
 			}
@@ -3694,6 +3697,17 @@ bool UControlRigDetailPanelControlProxies::SelectPropertyInternal(UControlRigCon
 									{
 										continue;
 									}
+									if (TViewModelPtr<FChannelGroupOutlinerModel> ChannelModel = CastViewModel<FChannelGroupOutlinerModel>(OutlinerExtenstionIt.GetCurrentItem()))
+									{
+										if (ChannelModel->GetChannel(Track->GetSectionToKey()) == nullptr) //if not section to key we also don't select it.
+										{
+											continue;
+										}
+									}
+									else
+									{
+										continue;
+									}
 									FName ID = OutlinerExtenstionIt->GetIdentifier();
 									FString Name = ID.ToString();
 									TArray<FString> StringArray;
diff --git a/Engine/Plugins/Animation/ControlRig/Source/ControlRigEditor/Private/EditMode/ControlRigControlsProxy.cpp b/Engine/Plugins/Animation/ControlRig/Source/ControlRigEditor/Private/EditMode/ControlRigControlsProxy.cpp
--- a/Engine/Plugins/Animation/ControlRig/Source/ControlRigEditor/Private/EditMode/ControlRigControlsProxy.cpp
+++ b/Engine/Plugins/Animation/ControlRig/Source/ControlRigEditor/Private/EditMode/ControlRigControlsProxy.cpp
@@ -253,6 +253,49 @@ void UControlRigControlsProxy::SelectionChanged(bool bInSelected)
 	}
 }
 
+void UControlRigControlsProxy::AddInteractions(EControlRigContextChannelToKey ChannelsToKey, EPropertyChangeType::Type ChangeType)
+{
+	if (ChangeType == EPropertyChangeType::Interactive || ChangeType == EPropertyChangeType::ValueSet)
+	{
+		EControlRigInteractionType InteractionType = EControlRigInteractionType::None;
+		if (EnumHasAnyFlags(ChannelsToKey, EControlRigContextChannelToKey::TranslationX)
+			|| EnumHasAnyFlags(ChannelsToKey, EControlRigContextChannelToKey::TranslationY)
+			|| EnumHasAnyFlags(ChannelsToKey, EControlRigContextChannelToKey::TranslationZ))
+		{
+			EnumAddFlags(InteractionType, EControlRigInteractionType::Translate);
+		}
+		if (EnumHasAnyFlags(ChannelsToKey, EControlRigContextChannelToKey::RotationX)
+			|| EnumHasAnyFlags(ChannelsToKey, EControlRigContextChannelToKey::RotationY)
+			|| EnumHasAnyFlags(ChannelsToKey, EControlRigContextChannelToKey::RotationZ))
+		{
+			EnumAddFlags(InteractionType, EControlRigInteractionType::Rotate);
+		}
+		if (EnumHasAnyFlags(ChannelsToKey, EControlRigContextChannelToKey::ScaleX)
+			|| EnumHasAnyFlags(ChannelsToKey, EControlRigContextChannelToKey::ScaleY)
+			|| EnumHasAnyFlags(ChannelsToKey, EControlRigContextChannelToKey::ScaleZ))
+		{
+			EnumAddFlags(InteractionType, EControlRigInteractionType::Scale);
+		}
+		for (const TPair<TWeakObjectPtr<UControlRig>, FControlRigProxyItem>& Items : ControlRigItems)
+		{
+			if (UControlRig* ControlRig = Items.Value.ControlRig.Get())
+			{
+				for (const FName& CName : Items.Value.ControlElements)
+				{
+					if (FRigControlElement* ControlElement = Items.Value.GetControlElement(CName))
+					{
+						if (InteractionScopes.Contains(ControlElement) == false)
+						{
+							FControlRigInteractionScope* InteractionScope = new FControlRigInteractionScope(ControlRig, ControlElement->GetKey(), InteractionType);
+							InteractionScopes.Add(ControlElement, InteractionScope);
+						}
+					}
+				}
+			}
+		}
+	}
+}
+
 void UControlRigControlsProxy::PostEditChangeChainProperty(struct FPropertyChangedChainEvent& PropertyChangedEvent)
 {
 	Super::PostEditChangeChainProperty(PropertyChangedEvent);
@@ -268,62 +311,12 @@ void UControlRigControlsProxy::PostEditChangeChainProperty(struct FPropertyChang
 			OwnerControlRig.Get()->SelectControl(OwnerControlElement.GetKey().Name, bSelected);
 			OwnerControlRig.Get()->Evaluate_AnyThread();
 		}
+		return;
 	}
 #if WITH_EDITOR
 	if (PropertyChangedEvent.Property)
 	{
-		if (PropertyChangedEvent.ChangeType == EPropertyChangeType::Interactive)
-		{
-			EControlRigInteractionType InteractionType = EControlRigInteractionType::None;
-			FProperty* Owner = PropertyChangedEvent.Property->GetOwnerProperty();
-			if (FEditPropertyChain::TDoubleLinkedListNode* MemberNode = PropertyChangedEvent.PropertyChain.GetActiveMemberNode())
-			{
-				if (FProperty* MemberProperty = MemberNode->GetValue())
-				{
-					if (MemberProperty->GetFName() == GET_MEMBER_NAME_CHECKED(FEulerTransform, Location))
-					{
-						InteractionType = EControlRigInteractionType::Translate;
-					}
-					else if (MemberProperty->GetFName() == GET_MEMBER_NAME_CHECKED(FEulerTransform, Rotation))
-					{
-						InteractionType = EControlRigInteractionType::Rotate;
-					}
-					else if (MemberProperty->GetFName() == GET_MEMBER_NAME_CHECKED(FEulerTransform, Scale))
-					{
-						InteractionType = EControlRigInteractionType::Scale;
-					}
-				}
-			}
-			for (const TPair<TWeakObjectPtr<UControlRig>, FControlRigProxyItem>& Items : ControlRigItems)
-			{
-				if (UControlRig* ControlRig = Items.Value.ControlRig.Get())
-				{
-					for (const FName& CName : Items.Value.ControlElements)
-					{
-						if (FRigControlElement * ControlElement = Items.Value.GetControlElement(CName))
-						{
-							if (InteractionScopes.Contains(ControlElement) == false)
-							{
-								FControlRigInteractionScope* InteractionScope = new FControlRigInteractionScope(ControlRig, ControlElement->GetKey(), InteractionType);
-								InteractionScopes.Add(ControlElement, InteractionScope);
-							}
-						}
-					}
-				}
-			}
-		}
-		else
-		{
-			for (TPair<FRigControlElement*, FControlRigInteractionScope*>& Scope : InteractionScopes)
-			{
-				if (Scope.Value)
-				{
-					delete Scope.Value;
-				}
-			}
-			InteractionScopes.Reset();
-		}
-
+		
 		//set values
 		FProperty* Property = PropertyChangedEvent.Property;
 		FProperty* MemberProperty = nullptr;
@@ -333,9 +326,14 @@ void UControlRigControlsProxy::PostEditChangeChainProperty(struct FPropertyChang
 		}
 		if (PropertyIsOnProxy(Property, MemberProperty))
 		{
+			EControlRigContextChannelToKey ChannelToKeyContext = GetChannelToKeyFromPropertyName(Property->GetFName());
+			AddInteractions(ChannelToKeyContext, PropertyChangedEvent.ChangeType);
 			FRigControlModifiedContext Context;
 			Context.SetKey = EControlRigSetKey::DoNotCare;
-			Context.KeyMask = (uint32)GetChannelToKeyFromPropertyName(Property->GetFName());
+			Context.KeyMask = (uint32)ChannelToKeyContext;
+			FRigControlModifiedContext NotifyDrivenContext; //we key all for them
+			Context.SetKey = EControlRigSetKey::DoNotCare;
+
 			UWorld* World = GCurrentLevelEditingViewportClient ? GCurrentLevelEditingViewportClient->GetWorld() : nullptr;
 			const FConstraintsManagerController& Controller = FConstraintsManagerController::Get(World);
 			Controller.EvaluateAllConstraints();
@@ -344,14 +342,21 @@ void UControlRigControlsProxy::PostEditChangeChainProperty(struct FPropertyChang
 			{
 				if (UControlRig* ControlRig = Items.Value.ControlRig.Get())
 				{
+					bool bDoReavaluate = false;
 					//we do this backwards so ValueChanged later is set up correctly since that iterates in the other direction
 					for (int32 Index = Items.Value.ControlElements.Num() - 1; Index >= 0; --Index)
 					{
 						if (FRigControlElement* ControlElement = Items.Value.GetControlElement(Items.Value.ControlElements[Index]))
 						{
 							SetControlRigElementValueFromCurrent(ControlRig, ControlElement, Context);
+							FControlRigEditMode::NotifyDrivenControls(ControlRig, ControlElement->GetKey(), NotifyDrivenContext);
+							bDoReavaluate = true;
 						}
 					}
+					if (bDoReavaluate)
+					{
+						ControlRig->Evaluate_AnyThread();
+					}
 				}
 			}
 			for (TPair <TWeakObjectPtr<UObject>, FSequencerProxyItem>& SItems : SequencerItems)
@@ -367,8 +372,19 @@ void UControlRigControlsProxy::PostEditChangeChainProperty(struct FPropertyChang
 					SetBindingValueFromCurrent(SItems.Key.Get(), Binding.Binding, Context, PropertyChangedEvent.ChangeType == EPropertyChangeType::Interactive);
 				}
 			}
+			if (PropertyChangedEvent.ChangeType != EPropertyChangeType::Interactive)
+			{
+				for (TPair<FRigControlElement*, FControlRigInteractionScope*>& Scope : InteractionScopes)
+				{
+					if (Scope.Value)
+					{
+						delete Scope.Value;
+					}
+				}
+				InteractionScopes.Reset();
+			}
+			ValueChanged();
 		}
-		ValueChanged();
 	}
 #endif
 }
diff --git a/Engine/Plugins/Animation/ControlRig/Source/ControlRigEditor/Private/EditMode/ControlRigControlsProxy.h b/Engine/Plugins/Animation/ControlRig/Source/ControlRigEditor/Private/EditMode/ControlRigControlsProxy.h
--- a/Engine/Plugins/Animation/ControlRig/Source/ControlRigEditor/Private/EditMode/ControlRigControlsProxy.h
+++ b/Engine/Plugins/Animation/ControlRig/Source/ControlRigEditor/Private/EditMode/ControlRigControlsProxy.h
@@ -120,6 +120,8 @@ private:
 
 	FCachedRigElement& GetOwnerControlElement();
 
+	void AddInteractions(EControlRigContextChannelToKey ChannelsToKey, EPropertyChangeType::Type ChangeType);
+
 public:
 
 	//if individual it will show up independently, this will happen for certain nested controls
diff --git a/Engine/Plugins/Animation/ControlRig/Source/ControlRigEditor/Private/EditMode/ControlRigEditMode.cpp b/Engine/Plugins/Animation/ControlRig/Source/ControlRigEditor/Private/EditMode/ControlRigEditMode.cpp
--- a/Engine/Plugins/Animation/ControlRig/Source/ControlRigEditor/Private/EditMode/ControlRigEditMode.cpp
+++ b/Engine/Plugins/Animation/ControlRig/Source/ControlRigEditor/Private/EditMode/ControlRigEditMode.cpp
@@ -3303,6 +3303,8 @@ void FControlRigEditMode::ZeroTransforms(bool bSelectionOnly)
 	}
 
 	FScopedTransaction Transaction(LOCTEXT("HierarchyZeroTransforms", "Zero Transforms"));
+	FRigControlModifiedContext Context;
+	Context.SetKey = EControlRigSetKey::DoNotCare;
 
 	for (UControlRig* ControlRig : ControlRigs)
 	{
@@ -3395,10 +3397,11 @@ void FControlRigEditMode::ZeroTransforms(bool bSelectionOnly)
 			}
 			if (ControlElement)
 			{
-				ControlRig->SetControlLocalTransform(ElementToReset.Name, InitialLocalTransform, true, FRigControlModifiedContext(), true, true);
+				ControlRig->SetControlLocalTransform(ElementToReset.Name, InitialLocalTransform, true, Context, true, true);
 				const FVector InitialAngles = ControlRig->GetHierarchy()->GetControlPreferredEulerAngles(ControlElement, ControlElement->Settings.PreferredRotationOrder, true);
 				ControlRig->GetHierarchy()->SetControlPreferredEulerAngles(ControlElement, InitialAngles, ControlElement->Settings.PreferredRotationOrder);
-				NotifyDrivenControls(ControlRig, ElementToReset);
+				NotifyDrivenControls(ControlRig, ElementToReset, Context);
+
 				if (bHasNonDefaultParent == false)
 				{
 					ControlRig->ControlModified().Broadcast(ControlRig, ControlElement, EControlRigSetKey::DoNotCare);
@@ -3432,7 +3435,7 @@ void FControlRigEditMode::ZeroTransforms(bool bSelectionOnly)
 						FTransform GlobalTransform = ControlRig->GetHierarchy()->GetGlobalTransform(ElementToReset);
 						GlobalTransforms.Add(ElementToReset, GlobalTransform);
 					}
-					NotifyDrivenControls(ControlRig, ElementToReset);
+					NotifyDrivenControls(ControlRig, ElementToReset,Context);
 				}
 				else
 				{
@@ -3474,7 +3477,7 @@ void FControlRigEditMode::ZeroTransforms(bool bSelectionOnly)
 							{
 								ControlRig->SetControlGlobalTransform(ElementToReset.Name, *GlobalTransform, true);
 								ControlRig->Evaluate_AnyThread();
-								NotifyDrivenControls(ControlRig, ElementToReset);
+								NotifyDrivenControls(ControlRig, ElementToReset, Context);
 							}
 						}
 						else
@@ -3504,7 +3507,7 @@ void FControlRigEditMode::ZeroTransforms(bool bSelectionOnly)
 			ControlRig->Evaluate_AnyThread();
 			for (const FRigElementKey& ControlToReset : ControlsToReset)
 			{
-				NotifyDrivenControls(ControlRig, ControlToReset);
+				NotifyDrivenControls(ControlRig, ControlToReset, Context);
 			}
 		}
 	}
@@ -4812,7 +4815,7 @@ void FControlRigEditMode::MoveControlShape(AControlRigShapeActor* ShapeActor, co
 				ControlRig->Evaluate_AnyThread();
 				SetControlShapeTransform(ShapeActor, NewTransform, ToWorldTransform, Context, bPrintPythonCommands, /*fix flips*/ true);
 				//UpdatePreferredEulerAngles(ControlRig);
-				NotifyDrivenControls(ControlRig, ShapeActor->GetElementKey());
+				NotifyDrivenControls(ControlRig, ShapeActor->GetElementKey(),Context);
 				if(const FRigControlElement* ControlElement = ControlRig->FindControl(ShapeActor->ControlName))
 				{
 					if(!bIsTransientControl)
@@ -5455,14 +5458,14 @@ void FControlRigEditMode::PostPoseUpdate() const
 
 }
 
-void FControlRigEditMode::NotifyDrivenControls(UControlRig* InControlRig, const FRigElementKey& InKey)
+void FControlRigEditMode::NotifyDrivenControls(UControlRig* InControlRig, const FRigElementKey& InKey, const FRigControlModifiedContext& InContext)
 {
 	// if we are changing a proxy control - we also need to notify the change for the driven controls
 	if (FRigControlElement* ControlElement = InControlRig->GetHierarchy()->Find<FRigControlElement>(InKey))
 	{
 		if(ControlElement->CanDriveControls())
 		{
-			FRigControlModifiedContext Context;
+			FRigControlModifiedContext Context(InContext);
 			Context.EventName = FRigUnit_BeginExecution::EventName;
 
 			for(const FRigElementKey& DrivenKey : ControlElement->Settings.DrivenControls)
diff --git a/Engine/Plugins/Animation/ControlRig/Source/ControlRigEditor/Private/Sequencer/ControlRigParameterTrackEditor.cpp b/Engine/Plugins/Animation/ControlRig/Source/ControlRigEditor/Private/Sequencer/ControlRigParameterTrackEditor.cpp
--- a/Engine/Plugins/Animation/ControlRig/Source/ControlRigEditor/Private/Sequencer/ControlRigParameterTrackEditor.cpp
+++ b/Engine/Plugins/Animation/ControlRig/Source/ControlRigEditor/Private/Sequencer/ControlRigParameterTrackEditor.cpp
@@ -1695,6 +1695,8 @@ void FControlRigParameterTrackEditor::OnAddTransformKeysForSelectedObjects(EMovi
 	FScopedTransaction KeyTransaction(LOCTEXT("SetKeysOnControls", "Set Keys On Controls"), !GIsTransacting);
 
 	static constexpr bool bInConstraintSpace = true;
+	FRigControlModifiedContext NotifyDrivenContext;
+	NotifyDrivenContext.SetKey = EControlRigSetKey::Always;
 	for (const TPair<UControlRig*, TArray<FRigElementKey>>& Selection : SelectedControls)
 	{
 		UControlRig* ControlRig = Selection.Key;
@@ -1707,8 +1709,12 @@ void FControlRigParameterTrackEditor::OnAddTransformKeysForSelectedObjects(EMovi
 				const TArray<FName> ControlNames = ControlRig->CurrentControlSelection();
 				for (const FName& ControlName : ControlNames)
 				{
-					AddControlKeys(Object, ControlRig, Name, ControlName, ChannelsToKey,
-						ESequencerKeyMode::ManualKeyForced, FLT_MAX, bInConstraintSpace);
+					if (FRigControlElement* ControlElement = ControlRig->FindControl(ControlName))
+					{
+						AddControlKeys(Object, ControlRig, Name, ControlName, ChannelsToKey,
+							ESequencerKeyMode::ManualKeyForced, FLT_MAX, bInConstraintSpace);
+						FControlRigEditMode::NotifyDrivenControls(ControlRig, ControlElement->GetKey(), NotifyDrivenContext);
+					}
 				}
 			}
 		}
diff --git a/Engine/Plugins/Animation/ControlRig/Source/ControlRigEditor/Public/EditMode/ControlRigEditMode.h b/Engine/Plugins/Animation/ControlRig/Source/ControlRigEditor/Public/EditMode/ControlRigEditMode.h
--- a/Engine/Plugins/Animation/ControlRig/Source/ControlRigEditor/Public/EditMode/ControlRigEditMode.h
+++ b/Engine/Plugins/Animation/ControlRig/Source/ControlRigEditor/Public/EditMode/ControlRigEditMode.h
@@ -506,6 +506,9 @@ private:
 	void TickControlShape(AControlRigShapeActor* ShapeActor, const FTransform& ComponentTransform) const;
 	bool ModeSupportedByShapeActor(const AControlRigShapeActor* ShapeActor, UE::Widget::EWidgetMode InMode) const;
 
+public:
+	//notify driven controls, should this be inside CR instead?
+	static void NotifyDrivenControls(UControlRig* InControlRig, const FRigElementKey& InKey, const FRigControlModifiedContext& InContext);
 
 protected:
 	
@@ -525,7 +528,6 @@ private:
 	// Post pose update handler
 	void OnPoseInitialized();
 	void PostPoseUpdate() const;
-	void NotifyDrivenControls(UControlRig* InControlRig, const FRigElementKey& InKey);
 	void UpdateSelectabilityOnSkeletalMeshes(UControlRig* InControlRig, bool bEnabled);
 
 	bool IsMovingCamera(const FViewport* InViewport) const;
