Control Rig: Fix ensure when undoing control rig sequencer track removal
#rnx

diff --git a/Engine/Plugins/Animation/ControlRig/Source/ControlRig/Private/Sequencer/MovieSceneControlRigParameterTrack.cpp b/Engine/Plugins/Animation/ControlRig/Source/ControlRig/Private/Sequencer/MovieSceneControlRigParameterTrack.cpp
--- a/Engine/Plugins/Animation/ControlRig/Source/ControlRig/Private/Sequencer/MovieSceneControlRigParameterTrack.cpp
+++ b/Engine/Plugins/Animation/ControlRig/Source/ControlRig/Private/Sequencer/MovieSceneControlRigParameterTrack.cpp
@@ -545,7 +545,7 @@ void UMovieSceneControlRigParameterTrack::HandleOnPostConstructed_GameThread()
 		TArray<FRigControlElement*> SortedControls;
 		ControlRig->GetControlsInOrder(SortedControls);
 #if WITH_EDITOR
-		const FScopedTransaction PostConstructTransation(NSLOCTEXT("ControlRig", "PostConstructTransation", "Post Construct"));
+		const FScopedTransaction PostConstructTransation(NSLOCTEXT("ControlRig", "PostConstructTransation", "Post Construct"), !GIsTransacting);
 #endif		
 		bool bSectionWasDifferent = false;
 		for (UMovieSceneSection* BaseSection : GetAllSections())
