Control Rig: Fix ensure when deleting modules
#rnx

diff --git a/Engine/Plugins/Animation/ControlRig/Source/ControlRigEditor/Private/Editor/SModularRigTreeView.cpp b/Engine/Plugins/Animation/ControlRig/Source/ControlRigEditor/Private/Editor/SModularRigTreeView.cpp
--- a/Engine/Plugins/Animation/ControlRig/Source/ControlRigEditor/Private/Editor/SModularRigTreeView.cpp
+++ b/Engine/Plugins/Animation/ControlRig/Source/ControlRigEditor/Private/Editor/SModularRigTreeView.cpp
@@ -967,7 +967,10 @@ TArray<FString> SModularRigTreeView::GetSelectedKeys() const
 	TArray<TSharedPtr<FModularRigTreeElement>> SelectedElements = GetSelectedItems();
 	for(const TSharedPtr<FModularRigTreeElement>& SelectedElement : SelectedElements)
 	{
-		Keys.AddUnique(SelectedElement->ModulePath);
+		if (SelectedElement.IsValid())
+		{
+			Keys.AddUnique(SelectedElement->ModulePath);
+		}
 	}
 	return Keys;
 }
