Constraints: fix UTransformableControlHandle not updating the ControlRig instance when resolving bound objects

#rnx

diff --git a/Engine/Plugins/Animation/ControlRig/Source/ControlRig/Private/Constraints/ControlRigTransformableHandle.cpp b/Engine/Plugins/Animation/ControlRig/Source/ControlRig/Private/Constraints/ControlRigTransformableHandle.cpp
--- a/Engine/Plugins/Animation/ControlRig/Source/ControlRig/Private/Constraints/ControlRigTransformableHandle.cpp
+++ b/Engine/Plugins/Animation/ControlRig/Source/ControlRig/Private/Constraints/ControlRigTransformableHandle.cpp
@@ -7,6 +7,7 @@
 #include "ControlRig.h"
 #include "ControlRigComponent.h"
 #include "IControlRigObjectBinding.h"
+#include "ControlRigObjectBinding.h"
 #include "TransformableHandleUtils.h"
 #include "Rigs/RigHierarchyElements.h"
 #include "Sequencer/MovieSceneControlRigParameterSection.h"
@@ -568,23 +569,14 @@ bool UTransformableControlHandle::AddTransformKeys(const TArray<FFrameNumber>& I
 //for control rig need to check to see if the control rig is different then we may need to update it based upon what we are now bound to
 void UTransformableControlHandle::ResolveBoundObjects(FMovieSceneSequenceID LocalSequenceID, IMovieScenePlayer& Player, UObject* SubObject)
 {
-	if (UControlRig* InControlRig = Cast<UControlRig>(SubObject))
+	if (const UControlRig* InControlRig = Cast<UControlRig>(SubObject))
 	{
 		if (ControlRig != InControlRig)
 		{
-			for (TWeakObjectPtr<> ParentObject : ConstraintBindingID.ResolveBoundObjects(LocalSequenceID, Player))
+			for (const TWeakObjectPtr<> ParentObject : ConstraintBindingID.ResolveBoundObjects(LocalSequenceID, Player))
 			{
-				USceneComponent* Component = nullptr;
-				if (AActor* Actor = Cast<AActor>(ParentObject.Get()))
-				{
-					Component = Actor->GetRootComponent();
-				}
-				else if (USceneComponent* Comp = Cast<USceneComponent>(ParentObject.Get()))
-				{
-					Component = Comp;
-				}
-
-				if (InControlRig->GetObjectBinding() && InControlRig->GetObjectBinding()->GetBoundObject() == Component)
+				const UObject* Bindable = FControlRigObjectBinding::GetBindableObject(ParentObject.Get());
+				if (InControlRig->GetObjectBinding() && InControlRig->GetObjectBinding()->GetBoundObject() == Bindable)
 				{
 					ControlRig = InControlRig;
 				}
