Control Rig: Add SwapModuleClass function
#rnx

diff --git a/Engine/Plugins/Animation/ControlRig/Source/ControlRig/Private/ModularRigController.cpp b/Engine/Plugins/Animation/ControlRig/Source/ControlRig/Private/ModularRigController.cpp
--- a/Engine/Plugins/Animation/ControlRig/Source/ControlRig/Private/ModularRigController.cpp
+++ b/Engine/Plugins/Animation/ControlRig/Source/ControlRig/Private/ModularRigController.cpp
@@ -44,7 +44,7 @@ FString UModularRigController::AddModule(const FName& InModuleName, TSubclassOf<
 	TSharedPtr<FScopedTransaction> TransactionPtr;
 	if (bSetupUndo)
 	{
-		TransactionPtr = MakeShared<FScopedTransaction>(NSLOCTEXT("ModularRigController", "AddModuleTransaction", "Add Module"));
+		TransactionPtr = MakeShared<FScopedTransaction>(NSLOCTEXT("ModularRigController", "AddModuleTransaction", "Add Module"), !GIsTransacting);
 		if(UBlueprint* Blueprint = Cast<UBlueprint>(GetOuter()))
 		{
 			Blueprint->Modify();
@@ -254,7 +254,7 @@ bool UModularRigController::ConnectConnectorToElement(const FRigElementKey& InCo
 	TSharedPtr<FScopedTransaction> TransactionPtr;
 	if (bSetupUndo)
 	{
-		TransactionPtr = MakeShared<FScopedTransaction>(NSLOCTEXT("ModularRigController", "ConnectModuleToElementTransaction", "Connect to Element"));
+		TransactionPtr = MakeShared<FScopedTransaction>(NSLOCTEXT("ModularRigController", "ConnectModuleToElementTransaction", "Connect to Element"), !GIsTransacting);
 		Blueprint->Modify();
 	}
 #endif 
@@ -364,7 +364,7 @@ bool UModularRigController::DisconnectConnector(const FRigElementKey& InConnecto
 	TSharedPtr<FScopedTransaction> TransactionPtr;
 	if (bSetupUndo)
 	{
-		TransactionPtr = MakeShared<FScopedTransaction>(NSLOCTEXT("ModularRigController", "ConnectModuleToElementTransaction", "Connect to Element"));
+		TransactionPtr = MakeShared<FScopedTransaction>(NSLOCTEXT("ModularRigController", "ConnectModuleToElementTransaction", "Connect to Element"), !GIsTransacting);
 		Blueprint->Modify();
 	}
 #endif 
@@ -530,7 +530,7 @@ bool UModularRigController::AutoConnectSecondaryConnectors(const TArray<FRigElem
 	TSharedPtr<FScopedTransaction> TransactionPtr;
 	if (bSetupUndo)
 	{
-		TransactionPtr = MakeShared<FScopedTransaction>(NSLOCTEXT("ModularRigController", "AutoResolveSecondaryConnectors", "Auto-Resolve Connectors"));
+		TransactionPtr = MakeShared<FScopedTransaction>(NSLOCTEXT("ModularRigController", "AutoResolveSecondaryConnectors", "Auto-Resolve Connectors"), !GIsTransacting);
 	}
 
 	Blueprint->Modify();
@@ -730,7 +730,7 @@ bool UModularRigController::SetConfigValueInModule(const FString& InModulePath,
 	TSharedPtr<FScopedTransaction> TransactionPtr;
 	if (bSetupUndo)
 	{
-		TransactionPtr = MakeShared<FScopedTransaction>(NSLOCTEXT("ModularRigController", "ConfigureModuleValueTransaction", "Configure Module Value"));
+		TransactionPtr = MakeShared<FScopedTransaction>(NSLOCTEXT("ModularRigController", "ConfigureModuleValueTransaction", "Configure Module Value"), !GIsTransacting);
 		if(UBlueprint* Blueprint = Cast<UBlueprint>(GetOuter()))
 		{
 			Blueprint->Modify();
@@ -937,7 +937,7 @@ bool UModularRigController::BindModuleVariable(const FString& InModulePath, cons
 	TSharedPtr<FScopedTransaction> TransactionPtr;
 	if (bSetupUndo)
 	{
-		TransactionPtr = MakeShared<FScopedTransaction>(NSLOCTEXT("ModularRigController", "BindModuleVariableTransaction", "Bind Module Variable"));
+		TransactionPtr = MakeShared<FScopedTransaction>(NSLOCTEXT("ModularRigController", "BindModuleVariableTransaction", "Bind Module Variable"), !GIsTransacting);
 		if(UBlueprint* Blueprint = Cast<UBlueprint>(GetOuter()))
 		{
 			Blueprint->Modify();
@@ -976,7 +976,7 @@ bool UModularRigController::UnBindModuleVariable(const FString& InModulePath, co
 	TSharedPtr<FScopedTransaction> TransactionPtr;
 	if (bSetupUndo)
 	{
-		TransactionPtr = MakeShared<FScopedTransaction>(NSLOCTEXT("ModularRigController", "BindModuleVariableTransaction", "Bind Module Variable"));
+		TransactionPtr = MakeShared<FScopedTransaction>(NSLOCTEXT("ModularRigController", "BindModuleVariableTransaction", "Bind Module Variable"), !GIsTransacting);
 		if(UBlueprint* Blueprint = Cast<UBlueprint>(GetOuter()))
 		{
 			Blueprint->Modify();
@@ -1008,7 +1008,7 @@ bool UModularRigController::DeleteModule(const FString& InModulePath, bool bSetu
 	TSharedPtr<FScopedTransaction> TransactionPtr;
 	if (bSetupUndo)
 	{
-		TransactionPtr = MakeShared<FScopedTransaction>(NSLOCTEXT("ModularRigController", "DeleteModuleTransaction", "Delete Module"));
+		TransactionPtr = MakeShared<FScopedTransaction>(NSLOCTEXT("ModularRigController", "DeleteModuleTransaction", "Delete Module"), !GIsTransacting);
 		if(UBlueprint* Blueprint = Cast<UBlueprint>(GetOuter()))
 		{
 			Blueprint->Modify();
@@ -1107,7 +1107,7 @@ FString UModularRigController::RenameModule(const FString& InModulePath, const F
 	TSharedPtr<FScopedTransaction> TransactionPtr;
 	if (bSetupUndo)
 	{
-		TransactionPtr = MakeShared<FScopedTransaction>(NSLOCTEXT("ModularRigController", "RenameModuleTransaction", "Rename Module"));
+		TransactionPtr = MakeShared<FScopedTransaction>(NSLOCTEXT("ModularRigController", "RenameModuleTransaction", "Rename Module"), !GIsTransacting);
 		if(UBlueprint* Blueprint = Cast<UBlueprint>(GetOuter()))
 		{
 			Blueprint->Modify();
@@ -1220,7 +1220,7 @@ FString UModularRigController::ReparentModule(const FString& InModulePath, const
 	TSharedPtr<FScopedTransaction> TransactionPtr;
 	if (bSetupUndo)
 	{
-		TransactionPtr = MakeShared<FScopedTransaction>(NSLOCTEXT("ModularRigController", "ReparentModuleTransaction", "Reparent Module"));
+		TransactionPtr = MakeShared<FScopedTransaction>(NSLOCTEXT("ModularRigController", "ReparentModuleTransaction", "Reparent Module"), !GIsTransacting);
 		if(UBlueprint* Blueprint = Cast<UBlueprint>(GetOuter()))
 		{
 			Blueprint->Modify();
@@ -1464,7 +1464,7 @@ bool UModularRigController::SetModuleShortName(const FString& InModulePath, cons
 	TSharedPtr<FScopedTransaction> TransactionPtr;
 	if (bSetupUndo)
 	{
-		TransactionPtr = MakeShared<FScopedTransaction>(NSLOCTEXT("ModularRigController", "SetModuleShortNameTransaction", "Set Module Display Name"));
+		TransactionPtr = MakeShared<FScopedTransaction>(NSLOCTEXT("ModularRigController", "SetModuleShortNameTransaction", "Set Module Display Name"), !GIsTransacting);
 		if(UBlueprint* Blueprint = Cast<UBlueprint>(GetOuter()))
 		{
 			Blueprint->Modify();
@@ -1498,6 +1498,209 @@ bool UModularRigController::CanSetModuleShortName(const FString& InModulePath, c
 	return true;
 }
 
+bool UModularRigController::SwapModuleClass(const FString& InModulePath, TSubclassOf<UControlRig> InNewClass, bool bSetupUndo)
+{
+	FRigModuleReference* Module = FindModule(InModulePath);
+	if (!Module)
+	{
+		UE_LOG(LogControlRig, Error, TEXT("Could not find module %s"), *InModulePath);
+		return false;
+	}
+
+	if (!InNewClass)
+	{
+		UE_LOG(LogControlRig, Error, TEXT("Invalid InClass"));
+		return false;
+	}
+
+	UControlRig* ClassDefaultObject = InNewClass->GetDefaultObject<UControlRig>();
+	if (!ClassDefaultObject->IsRigModule())
+	{
+		UE_LOG(LogControlRig, Error, TEXT("Class %s is not a rig module"), *InNewClass->GetClassPathName().ToString());
+		return false;
+	}
+
+	if (Module->Class.Get() == InNewClass)
+	{
+		// Nothing to do here
+		return true;
+	}
+
+#if WITH_EDITOR
+	TSharedPtr<FScopedTransaction> TransactionPtr;
+	if (bSetupUndo)
+	{
+		TransactionPtr = MakeShared<FScopedTransaction>(NSLOCTEXT("ModularRigController", "SwapModuleClassTransaction", "Swap Module Class"), !GIsTransacting);
+		if(UBlueprint* Blueprint = Cast<UBlueprint>(GetOuter()))
+		{
+			Blueprint->Modify();
+		}
+	}
+#endif
+
+	Module->Class = InNewClass;
+
+	// Remove invalid connectors/connections
+	{
+		const TArray<FModularRigSingleConnection>& Connections = Model->Connections.GetConnectionList();
+		const UControlRig* CDO = InNewClass->GetDefaultObject<UControlRig>();
+		const TArray<FRigModuleConnector>& ExposedConnectors = CDO->GetRigModuleSettings().ExposedConnectors;
+
+		TArray<FRigElementKey> ConnectionsToRemove;
+		for (const FModularRigSingleConnection& Connection : Connections)
+		{
+			FString Namespace, ConnectorName;
+			URigHierarchy::SplitNameSpace(Connection.Connector.Name.ToString(), &Namespace, &ConnectorName);
+			if (Namespace.Equals(InModulePath))
+			{
+				if (!ExposedConnectors.ContainsByPredicate([ConnectorName](const FRigModuleConnector& Exposed)
+				{
+				   return Exposed.Name == ConnectorName;
+				}))
+				{
+					ConnectionsToRemove.Add(Connection.Connector);
+					continue;
+				}
+
+				FText ErrorMessage;
+				if (!CanConnectConnectorToElement(Connection.Connector, Connection.Target, ErrorMessage))
+				{
+					ConnectionsToRemove.Add(Connection.Connector);
+				}
+			}
+		}
+
+		for (const FRigElementKey& ToRemove : ConnectionsToRemove)
+		{
+			DisconnectConnector(ToRemove, false, bSetupUndo);
+		}
+	}
+
+	// Remove config values and bindings that are not supported anymore
+	RefreshModuleVariables();
+
+	Notify(EModularRigNotification::ModuleClassChanged, Module);
+
+	return true;
+}
+
+void UModularRigController::RefreshModuleVariables(bool bSetupUndo)
+{
+	Model->ForEachModule([this, bSetupUndo](const FRigModuleReference* Element) -> bool
+	{
+		RefreshModuleVariables(Element, bSetupUndo);
+		return true;
+	});
+}
+
+void UModularRigController::RefreshModuleVariables(const FRigModuleReference* InModule, bool bSetupUndo)
+{
+	if (!InModule)
+	{
+		return;
+	}
+	
+	// avoid dead class pointers
+	const UClass* ModuleClass = InModule->Class.Get();
+	if(ModuleClass == nullptr)
+	{
+		return;
+	}
+
+	// Make sure the provided module belongs to our ModularRigModel
+	const FString& ModulePath = InModule->GetPath();
+	FRigModuleReference* Module = FindModule(ModulePath);
+	if (Module != InModule)
+	{
+		return;
+	}
+
+#if WITH_EDITOR
+	TSharedPtr<FScopedTransaction> TransactionPtr;
+	if (bSetupUndo)
+	{
+		TransactionPtr = MakeShared<FScopedTransaction>(NSLOCTEXT("ModularRigController", "RefreshModuleVariablesTransaction", "Refresh Module Variables"), !GIsTransacting);
+		if(UBlueprint* Blueprint = Cast<UBlueprint>(GetOuter()))
+		{
+			Blueprint->Modify();
+		}
+	}
+#endif
+
+	for (TFieldIterator<FProperty> PropertyIt(ModuleClass); PropertyIt; ++PropertyIt)
+	{
+		const FProperty* Property = *PropertyIt;
+		
+		// remove advanced, private or not editable properties
+		const bool bIsAdvanced = Property->HasAnyPropertyFlags(CPF_AdvancedDisplay);
+		const bool bIsPublic = Property->HasAnyPropertyFlags(CPF_Edit | CPF_EditConst);
+		const bool bIsInstanceEditable = !Property->HasAnyPropertyFlags(CPF_DisableEditOnInstance);
+		if (bIsAdvanced || !bIsPublic || !bIsInstanceEditable)
+		{
+			Module->ConfigValues.Remove(Property->GetFName());
+			Module->Bindings.Remove(Property->GetFName());
+		}
+	}
+
+	// Make sure all the types are valid
+	const TMap<FName, FString> ConfigValues = Module->ConfigValues;
+	const TMap<FName, FString> Bindings = Module->Bindings;
+	Module->ConfigValues.Reset();
+	Module->Bindings.Reset();
+	for (const TPair<FName, FString>& Pair : ConfigValues)
+	{
+		SetConfigValueInModule(ModulePath, Pair.Key, Pair.Value, false);
+	}
+	for (const TPair<FName, FString>& Pair : Bindings)
+	{
+		BindModuleVariable(ModulePath, Pair.Key, Pair.Value, false);
+	}
+
+	// If the module is the source of another module's binding, make sure it is still a valid binding
+	Model->ForEachModule([this, InModule, ModulePath, ModuleClass](const FRigModuleReference* OtherModule) -> bool
+	{
+		if (InModule == OtherModule)
+		{
+			return true;
+		}
+		TArray<FName> BindingsToRemove;
+		for (const TPair<FName, FString>& Binding : OtherModule->Bindings)
+		{
+			FString BindingModulePath, VariableName = Binding.Value;
+			(void)URigHierarchy::SplitNameSpace(Binding.Value, &BindingModulePath, &VariableName);
+			if (BindingModulePath == ModulePath)
+			{
+				if (const FProperty* Property = ModuleClass->FindPropertyByName(*VariableName))
+				{
+					// remove advanced, private or not editable properties
+					const bool bIsAdvanced = Property->HasAnyPropertyFlags(CPF_AdvancedDisplay);
+					const bool bIsPublic = Property->HasAnyPropertyFlags(CPF_Edit | CPF_EditConst);
+					const bool bIsInstanceEditable = !Property->HasAnyPropertyFlags(CPF_DisableEditOnInstance);
+					if (bIsAdvanced || !bIsPublic || !bIsInstanceEditable)
+					{
+						BindingsToRemove.Add(Binding.Key);
+					}
+					else 
+					{
+						FText ErrorMessage;
+						if (!CanBindModuleVariable(OtherModule->GetPath(), Binding.Key, Binding.Value, ErrorMessage))
+						{
+							BindingsToRemove.Add(Binding.Key);
+						}
+					}
+				}
+			}
+		}
+
+		for (const FName& ToRemove : BindingsToRemove)
+		{
+			UnBindModuleVariable(OtherModule->GetPath(), ToRemove);
+		}
+		return true;
+	});
+	
+}
+
 void UModularRigController::SanitizeName(FRigName& InOutName, bool bAllowNameSpaces)
 {
 	// Sanitize the name
diff --git a/Engine/Plugins/Animation/ControlRig/Source/ControlRig/Public/ModularRigController.h b/Engine/Plugins/Animation/ControlRig/Source/ControlRig/Public/ModularRigController.h
--- a/Engine/Plugins/Animation/ControlRig/Source/ControlRig/Public/ModularRigController.h
+++ b/Engine/Plugins/Animation/ControlRig/Source/ControlRig/Public/ModularRigController.h
@@ -73,6 +73,12 @@ class CONTROLRIG_API UModularRigController : public UObject
 	bool SetModuleShortName(const FString& InModulePath, const FString& InNewShortName, bool bSetupUndo = true);
 	bool CanSetModuleShortName(const FString& InModulePath, const FString& InNewShortName, FText& OutErrorMessage) const;
 
+	UFUNCTION(BlueprintCallable, Category = "ControlRig | Modules")
+	bool SwapModuleClass(const FString& InModulePath, TSubclassOf<UControlRig> InNewClass, bool bSetupUndo = true);
+
+	void RefreshModuleVariables(bool bSetupUndo = true);
+	void RefreshModuleVariables(const FRigModuleReference* InModule, bool bSetupUndo = true);
+
 	static int32 GetMaxNameLength() { return 100; }
 	static void SanitizeName(FRigName& InOutName, bool bAllowNameSpaces);
 	static FRigName GetSanitizedName(const FRigName& InName, bool bAllowNameSpaces);
diff --git a/Engine/Plugins/Animation/ControlRig/Source/ControlRig/Public/ModularRigModel.h b/Engine/Plugins/Animation/ControlRig/Source/ControlRig/Public/ModularRigModel.h
--- a/Engine/Plugins/Animation/ControlRig/Source/ControlRig/Public/ModularRigModel.h
+++ b/Engine/Plugins/Animation/ControlRig/Source/ControlRig/Public/ModularRigModel.h
@@ -30,6 +30,8 @@ enum class EModularRigNotification : uint8
 	
 	InteractionBracketCanceled, // A bracket has been canceled
 
+	ModuleClassChanged, 
+
 	/** MAX - invalid */
 	Max UMETA(Hidden),
 };
diff --git a/Engine/Plugins/Animation/ControlRig/Source/ControlRigDeveloper/Private/ControlRigBlueprint.cpp b/Engine/Plugins/Animation/ControlRig/Source/ControlRigDeveloper/Private/ControlRigBlueprint.cpp
--- a/Engine/Plugins/Animation/ControlRig/Source/ControlRigDeveloper/Private/ControlRigBlueprint.cpp
+++ b/Engine/Plugins/Animation/ControlRig/Source/ControlRigDeveloper/Private/ControlRigBlueprint.cpp
@@ -2368,67 +2368,9 @@ void UControlRigBlueprint::RefreshModuleVariables()
 		return;
 	}
 
-	ModularRigModel.ForEachModule([this](const FRigModuleReference* Element) -> bool
-	{
-		RefreshModuleVariables(Element);
-		return true;
-	});
-}
-
-void UControlRigBlueprint::RefreshModuleVariables(const FRigModuleReference* InModule)
-{
-	if(!IsModularRig())
-	{
-		return;
-	}
-
-	// avoid dead class pointers
-	const UClass* ModuleClass = InModule->Class.Get();
-	if(ModuleClass == nullptr)
-	{
-		return;
-	}
-
-	// Make sure the provided module belongs to our ModularRigModel
-	const FString& ModulePath = InModule->GetPath();
-	FRigModuleReference* Module = ModularRigModel.FindModule(ModulePath);
-	if (Module != InModule)
-	{
-		return;
-	}
-
-	Modify();
-
-	for (TFieldIterator<FProperty> PropertyIt(ModuleClass); PropertyIt; ++PropertyIt)
-	{
-		const FProperty* Property = *PropertyIt;
-		
-		// remove advanced, private or not editable properties
-		const bool bIsAdvanced = Property->HasAnyPropertyFlags(CPF_AdvancedDisplay);
-		const bool bIsPublic = Property->HasAnyPropertyFlags(CPF_Edit | CPF_EditConst);
-		const bool bIsInstanceEditable = !Property->HasAnyPropertyFlags(CPF_DisableEditOnInstance);
-		if (bIsAdvanced || !bIsPublic || !bIsInstanceEditable)
-		{
-			Module->ConfigValues.Remove(Property->GetFName());
-			Module->Bindings.Remove(Property->GetFName());
-		}
-	}
-
-	// Make sure all the types are valid
 	if (UModularRigController* Controller = GetModularRigController())
 	{
-		const TMap<FName, FString> ConfigValues = Module->ConfigValues;
-		const TMap<FName, FString> Bindings = Module->Bindings;
-		Module->ConfigValues.Reset();
-		Module->Bindings.Reset();
-		for (const TPair<FName, FString>& Pair : ConfigValues)
-		{
-			Controller->SetConfigValueInModule(ModulePath, Pair.Key, Pair.Value, false);
-		}
-		for (const TPair<FName, FString>& Pair : Bindings)
-		{
-			Controller->BindModuleVariable(ModulePath, Pair.Key, Pair.Value, false);
-		}
+		Controller->RefreshModuleVariables();
 	}
 }
 
@@ -2742,6 +2684,15 @@ void UControlRigBlueprint::HandleRigModulesModified(EModularRigNotification InNo
 			HierarchyModifiedEvent.Broadcast(ERigHierarchyNotification::HierarchyReset, Hierarchy, nullptr);
 			break;
 		}
+		case EModularRigNotification::ModuleClassChanged:
+		{
+			if (InModule)
+			{
+				RefreshModuleConnectors(InModule);
+				UpdateConnectionMapFromModel();
+			}
+			break;
+		}
 		case EModularRigNotification::ModuleShortNameChanged:
 		{
 			bRecompile = false;
