Control Rig: Change Module browser to tile display
#rnx

diff --git a/Engine/Plugins/Animation/ControlRig/Source/ControlRigEditor/Private/Editor/SRigModuleAssetBrowser.cpp b/Engine/Plugins/Animation/ControlRig/Source/ControlRigEditor/Private/Editor/SRigModuleAssetBrowser.cpp
--- a/Engine/Plugins/Animation/ControlRig/Source/ControlRigEditor/Private/Editor/SRigModuleAssetBrowser.cpp
+++ b/Engine/Plugins/Animation/ControlRig/Source/ControlRigEditor/Private/Editor/SRigModuleAssetBrowser.cpp
@@ -37,7 +37,7 @@ void SRigModuleAssetBrowser::RefreshView()
 	
 	// setup filtering
 	AssetPickerConfig.Filter.ClassPaths.Add(UControlRigBlueprint::StaticClass()->GetClassPathName());
-	AssetPickerConfig.InitialAssetViewType = EAssetViewType::List;
+	AssetPickerConfig.InitialAssetViewType = EAssetViewType::Tile;
 	AssetPickerConfig.bAddFilterUI = true;
 	AssetPickerConfig.bShowPathInColumnView = true;
 	AssetPickerConfig.bShowTypeInColumnView = true;
