Control Rig: Avoid call to GetPIEWorldContext on worker thread
#rnx

diff --git a/Engine/Plugins/Animation/ControlRig/Source/ControlRig/Private/ControlRig.cpp b/Engine/Plugins/Animation/ControlRig/Source/ControlRig/Private/ControlRig.cpp
--- a/Engine/Plugins/Animation/ControlRig/Source/ControlRig/Private/ControlRig.cpp
+++ b/Engine/Plugins/Animation/ControlRig/Source/ControlRig/Private/ControlRig.cpp
@@ -71,6 +71,7 @@ static TAutoConsoleVariable<float> CVarControlRigEnableDrawInterfaceInGame(
 UControlRig::UControlRig(const FObjectInitializer& ObjectInitializer)
 	: Super(ObjectInitializer)
 #if WITH_EDITOR
+	, bIsRunningInPIE(false)
 	, bEnableAnimAttributeTrace(false)
 #endif 
 	, DataSourceRegistry(nullptr)
@@ -192,6 +193,10 @@ void UControlRig::Initialize(bool bRequestInit)
 	{
 		GetHierarchy()->GetController(true);
 	}
+
+#if WITH_EDITOR
+	bIsRunningInPIE = GEditor && GEditor->IsPlaySessionInProgress();
+#endif
 	
 	// should refresh mapping 
 	RequestConstruction();
@@ -787,8 +792,7 @@ bool UControlRig::Execute(const FName& InEventName)
 	bool bEnableDrawInterface = false;
 #if WITH_EDITOR
 	const bool bEnabledDuringGame = CVarControlRigEnableDrawInterfaceInGame->GetInt() != 0;
-	const bool bInGame = !(GEditor && !GEditor->GetPIEWorldContext());
-	if (bEnabledDuringGame || !bInGame)
+	if (bEnabledDuringGame || !bIsRunningInPIE)
 	{
 		bEnableDrawInterface = true;
 	}	
diff --git a/Engine/Plugins/Animation/ControlRig/Source/ControlRig/Public/ControlRig.h b/Engine/Plugins/Animation/ControlRig/Source/ControlRig/Public/ControlRig.h
--- a/Engine/Plugins/Animation/ControlRig/Source/ControlRig/Public/ControlRig.h
+++ b/Engine/Plugins/Animation/ControlRig/Source/ControlRig/Public/ControlRig.h
@@ -140,6 +140,13 @@ public:
 	virtual void InitializeVMs(bool bInitRigUnits = true) { Super::Initialize(bInitRigUnits); }
 	virtual bool InitializeVMs(const FName& InEventName) { return Super::InitializeVM(InEventName); }
 
+#if WITH_EDITOR
+protected:
+	bool bIsRunningInPIE;
+#endif
+	
+public:
+
 	/** Evaluates the ControlRig */
 	virtual void Evaluate_AnyThread() override;
 
