Control Rig: Swap all modules of a class inside a modular rig
#rnx

diff --git a/Engine/Plugins/Animation/ControlRig/Source/ControlRig/Private/ModularRigController.cpp b/Engine/Plugins/Animation/ControlRig/Source/ControlRig/Private/ModularRigController.cpp
--- a/Engine/Plugins/Animation/ControlRig/Source/ControlRig/Private/ModularRigController.cpp
+++ b/Engine/Plugins/Animation/ControlRig/Source/ControlRig/Private/ModularRigController.cpp
@@ -1584,6 +1584,20 @@ bool UModularRigController::SwapModuleClass(const FString& InModulePath, TSubcla
 	return true;
 }
 
+bool UModularRigController::SwapModulesOfClass(TSubclassOf<UControlRig> InOldClass, TSubclassOf<UControlRig> InNewClass, bool bSetupUndo)
+{
+	Model->ForEachModule([this, InOldClass, InNewClass, bSetupUndo](const FRigModuleReference* Module) -> bool
+	{
+		if (Module->Class.Get() == InOldClass)
+		{
+			SwapModuleClass(Module->GetPath(), InNewClass, bSetupUndo);
+		}
+		return true;
+	});
+	
+	return true;
+}
+
 void UModularRigController::RefreshModuleVariables(bool bSetupUndo)
 {
 	Model->ForEachModule([this, bSetupUndo](const FRigModuleReference* Element) -> bool
diff --git a/Engine/Plugins/Animation/ControlRig/Source/ControlRig/Public/ModularRigController.h b/Engine/Plugins/Animation/ControlRig/Source/ControlRig/Public/ModularRigController.h
--- a/Engine/Plugins/Animation/ControlRig/Source/ControlRig/Public/ModularRigController.h
+++ b/Engine/Plugins/Animation/ControlRig/Source/ControlRig/Public/ModularRigController.h
@@ -76,6 +76,9 @@ class CONTROLRIG_API UModularRigController : public UObject
 	UFUNCTION(BlueprintCallable, Category = "ControlRig | Modules")
 	bool SwapModuleClass(const FString& InModulePath, TSubclassOf<UControlRig> InNewClass, bool bSetupUndo = true);
 
+	UFUNCTION(BlueprintCallable, Category = "ControlRig | Modules")
+	bool SwapModulesOfClass(TSubclassOf<UControlRig> InOldClass, TSubclassOf<UControlRig> InNewClass, bool bSetupUndo = true);
+
 	void RefreshModuleVariables(bool bSetupUndo = true);
 	void RefreshModuleVariables(const FRigModuleReference* InModule, bool bSetupUndo = true);
 
