Control Rig: Fix slowness when compiling module while modular rig is opened
#rnx

diff --git a/Engine/Plugins/Animation/ControlRig/Source/ControlRig/Private/ModularRigController.cpp b/Engine/Plugins/Animation/ControlRig/Source/ControlRig/Private/ModularRigController.cpp
--- a/Engine/Plugins/Animation/ControlRig/Source/ControlRig/Private/ModularRigController.cpp
+++ b/Engine/Plugins/Animation/ControlRig/Source/ControlRig/Private/ModularRigController.cpp
@@ -1602,6 +1602,7 @@ void UModularRigController::RefreshModuleVariables(bool bSetupUndo)
 {
 	Model->ForEachModule([this, bSetupUndo](const FRigModuleReference* Element) -> bool
 	{
+		TGuardValue<bool> NotificationsGuard(bSuspendNotifications, true);
 		RefreshModuleVariables(Element, bSetupUndo);
 		return true;
 	});
diff --git a/Engine/Plugins/Animation/ControlRig/Source/ControlRigDeveloper/Private/ControlRigBlueprint.cpp b/Engine/Plugins/Animation/ControlRig/Source/ControlRigDeveloper/Private/ControlRigBlueprint.cpp
--- a/Engine/Plugins/Animation/ControlRig/Source/ControlRigDeveloper/Private/ControlRigBlueprint.cpp
+++ b/Engine/Plugins/Animation/ControlRig/Source/ControlRigDeveloper/Private/ControlRigBlueprint.cpp
@@ -2381,14 +2381,20 @@ void UControlRigBlueprint::RefreshModuleConnectors()
 		return;
 	}
 
-	ModularRigModel.ForEachModule([this](const FRigModuleReference* Element) -> bool
+	if (UModularRigController* Controller = GetModularRigController())
 	{
-		RefreshModuleConnectors(Element);
-		return true;
-	});
+		TGuardValue<bool> NotificationsGuard(Controller->bSuspendNotifications, true);
+		ModularRigModel.ForEachModule([this](const FRigModuleReference* Element) -> bool
+		{
+			RefreshModuleConnectors(Element, false);
+			return true;
+		});
+	}
+
+	PropagateHierarchyFromBPToInstances();
 }
 
-void UControlRigBlueprint::RefreshModuleConnectors(const FRigModuleReference* InModule)
+void UControlRigBlueprint::RefreshModuleConnectors(const FRigModuleReference* InModule, bool bPropagateHierarchy)
 {
 	if(!IsModularRig())
 	{
@@ -2469,8 +2475,11 @@ void UControlRigBlueprint::RefreshModuleConnectors(const FRigModuleReference* In
 					}
 				}
 			}
-			
-			PropagateHierarchyFromBPToInstances();
+
+			if (bPropagateHierarchy)
+			{
+				PropagateHierarchyFromBPToInstances();
+			}
 		}
 	}
 }
diff --git a/Engine/Plugins/Animation/ControlRig/Source/ControlRigDeveloper/Public/ControlRigBlueprint.h b/Engine/Plugins/Animation/ControlRig/Source/ControlRigDeveloper/Public/ControlRigBlueprint.h
--- a/Engine/Plugins/Animation/ControlRig/Source/ControlRigDeveloper/Public/ControlRigBlueprint.h
+++ b/Engine/Plugins/Animation/ControlRig/Source/ControlRigDeveloper/Public/ControlRigBlueprint.h
@@ -310,7 +310,7 @@ public:
 	void RefreshModuleVariables();
 	void RefreshModuleVariables(const FRigModuleReference* InModule);
 	void RefreshModuleConnectors();
-	void RefreshModuleConnectors(const FRigModuleReference* InModule);
+	void RefreshModuleConnectors(const FRigModuleReference* InModule, bool bPropagateHierarchy = true);
 
 	/**
 	* Returns the modified event, which can be used to 
