Control Rig: Swap all modules of a class inside a modular rig (fix undo/redo)
#rnx

diff --git a/Engine/Plugins/Animation/ControlRig/Source/ControlRig/Private/ModularRigController.cpp b/Engine/Plugins/Animation/ControlRig/Source/ControlRig/Private/ModularRigController.cpp
--- a/Engine/Plugins/Animation/ControlRig/Source/ControlRig/Private/ModularRigController.cpp
+++ b/Engine/Plugins/Animation/ControlRig/Source/ControlRig/Private/ModularRigController.cpp
@@ -1300,6 +1300,10 @@ FString UModularRigController::ReparentModule(const FString& InModulePath, const
 
 	Notify(EModularRigNotification::ModuleReparented, Module);
 	
+#if WITH_EDITOR
+ 	TransactionPtr.Reset();
+#endif
+	
 	return NewPath;
 }
 
@@ -1580,12 +1584,28 @@ bool UModularRigController::SwapModuleClass(const FString& InModulePath, TSubcla
 	RefreshModuleVariables();
 
 	Notify(EModularRigNotification::ModuleClassChanged, Module);
+	
+#if WITH_EDITOR
+	TransactionPtr.Reset();
+#endif
 
 	return true;
 }
 
 bool UModularRigController::SwapModulesOfClass(TSubclassOf<UControlRig> InOldClass, TSubclassOf<UControlRig> InNewClass, bool bSetupUndo)
 {
+#if WITH_EDITOR
+	TSharedPtr<FScopedTransaction> TransactionPtr;
+	if (bSetupUndo)
+	{
+		TransactionPtr = MakeShared<FScopedTransaction>(NSLOCTEXT("ModularRigController", "SwapModulesOfClassTransaction", "Swap Modules of Class"), !GIsTransacting);
+		if(UBlueprint* Blueprint = Cast<UBlueprint>(GetOuter()))
+		{
+			Blueprint->Modify();
+		}
+	}
+#endif
+	
 	Model->ForEachModule([this, InOldClass, InNewClass, bSetupUndo](const FRigModuleReference* Module) -> bool
 	{
 		if (Module->Class.Get() == InOldClass)
@@ -1595,6 +1615,10 @@ bool UModularRigController::SwapModulesOfClass(TSubclassOf<UControlRig> InOldCla
 		return true;
 	});
 	
+#if WITH_EDITOR
+	TransactionPtr.Reset();
+#endif
+	
 	return true;
 }
 
@@ -1714,6 +1738,9 @@ void UModularRigController::RefreshModuleVariables(const FRigModuleReference* In
 		return true;
 	});
 	
+#if WITH_EDITOR
+	TransactionPtr.Reset();
+#endif
 }
 
 void UModularRigController::SanitizeName(FRigName& InOutName, bool bAllowNameSpaces)
