Control Rig: Fix modular rigs being marked as dirty during load
#rnx

diff --git a/Engine/Plugins/Animation/ControlRig/Source/ControlRigDeveloper/Private/ControlRigBlueprint.cpp b/Engine/Plugins/Animation/ControlRig/Source/ControlRigDeveloper/Private/ControlRigBlueprint.cpp
--- a/Engine/Plugins/Animation/ControlRig/Source/ControlRigDeveloper/Private/ControlRigBlueprint.cpp
+++ b/Engine/Plugins/Animation/ControlRig/Source/ControlRigDeveloper/Private/ControlRigBlueprint.cpp
@@ -2413,8 +2413,6 @@ void UControlRigBlueprint::RefreshModuleConnectors(const FRigModuleReference* In
 	{
 		if (UControlRig* CDO = GetControlRigClass()->GetDefaultObject<UControlRig>())
 		{
-			Hierarchy->Modify();
-
 			const FString Namespace = InModule->GetNamespace();
 			const TArray<FRigElementKey> AllConnectors = Hierarchy->GetKeysOfType<FRigConnectorElement>();
 			const TArray<FRigElementKey> ExistingConnectors = AllConnectors.FilterByPredicate([Namespace](const FRigElementKey& ConnectorKey) -> bool
@@ -2448,6 +2446,7 @@ void UControlRigBlueprint::RefreshModuleConnectors(const FRigModuleReference* In
 				
 				if(bRemoveAllConnectors || !bConnectorExpected)
 				{
+					Hierarchy->Modify();
 					(void)Controller->RemoveElement(Connector);
 					ConnectionMap.Remove(Connector);
 				}
@@ -2465,6 +2464,7 @@ void UControlRigBlueprint::RefreshModuleConnectors(const FRigModuleReference* In
 					{
 						FRigHierarchyExecuteContextBracket HierarchyContextGuard(Hierarchy, &Context);
 						FControlRigExecuteContextRigModuleGuard RigModuleGuard(PublicContext, InModule->GetNamespace());
+						Hierarchy->Modify();
 						(void)Controller->AddConnector(ConnectorName, Connector.Settings);
 					}
 					else
