Control RIg: Add FindModulesOfClass in the modular model
#rnx

diff --git a/Engine/Plugins/Animation/ControlRig/Source/ControlRig/Private/ModularRigModel.cpp b/Engine/Plugins/Animation/ControlRig/Source/ControlRig/Private/ModularRigModel.cpp
--- a/Engine/Plugins/Animation/ControlRig/Source/ControlRig/Private/ModularRigModel.cpp
+++ b/Engine/Plugins/Animation/ControlRig/Source/ControlRig/Private/ModularRigModel.cpp
@@ -259,6 +259,20 @@ bool FModularRigModel::IsModuleParentedTo(const FRigModuleReference* InChildModu
 	return false;
 }
 
+TArray<const FRigModuleReference*> FModularRigModel::FindModuleInstancesOfClass(TSoftClassPtr<UControlRig> InClass) const
+{
+	TArray<const FRigModuleReference*> Result;
+	ForEachModule([&Result, InClass](const FRigModuleReference* Module) -> bool
+	{
+		if (Module->Class == InClass)
+		{
+			Result.Add(Module);
+		}
+		return true;
+	});
+	return Result;
+}
+
 FString FModularRigModel::GetParentPath(const FString& InPath) const
 {
 	if (const FRigModuleReference* Element = FindModule(InPath))
diff --git a/Engine/Plugins/Animation/ControlRig/Source/ControlRig/Public/ModularRigModel.h b/Engine/Plugins/Animation/ControlRig/Source/ControlRig/Public/ModularRigModel.h
--- a/Engine/Plugins/Animation/ControlRig/Source/ControlRig/Public/ModularRigModel.h
+++ b/Engine/Plugins/Animation/ControlRig/Source/ControlRig/Public/ModularRigModel.h
@@ -318,6 +318,8 @@ public:
 	
 	bool IsModuleParentedTo(const FRigModuleReference* InChildModule, const FRigModuleReference* InParentModule) const;
 
+	TArray<const FRigModuleReference*> FindModuleInstancesOfClass(TSoftClassPtr<UControlRig> InClass) const;
+
 private:
 	TWeakObjectPtr<UObject> OuterClientHost;
 
diff --git a/Engine/Plugins/Animation/ControlRig/Source/ControlRigDeveloper/Private/ControlRigBlueprint.cpp b/Engine/Plugins/Animation/ControlRig/Source/ControlRigDeveloper/Private/ControlRigBlueprint.cpp
--- a/Engine/Plugins/Animation/ControlRig/Source/ControlRigDeveloper/Private/ControlRigBlueprint.cpp
+++ b/Engine/Plugins/Animation/ControlRig/Source/ControlRigDeveloper/Private/ControlRigBlueprint.cpp
@@ -955,6 +955,89 @@ void UControlRigBlueprint::PreSave(FObjectPreSaveContext ObjectSaveContext)
 		ControlRigType = EControlRigType::IndependentRig;
 		ItemTypeDisplayName = TEXT("Control Rig");
 	}
+
+	if (IsModularRig())
+	{
+		ModuleReferenceData = GetModuleReferenceData();
+		IAssetRegistry::GetChecked().AssetTagsFinalized(*this);
+	}
+}
+
+TArray<FModuleReferenceData> UControlRigBlueprint::FindReferencesToModule() const
+{
+	TArray<FModuleReferenceData> Result;
+	if (!IsControlRigModule())
+	{
+		return Result;
+	}
+
+	const UClass* RigModuleClass = GetControlRigClass();
+	if (!RigModuleClass)
+	{
+		return Result;
+	}
+
+	// Load the asset registry module
+	const FAssetRegistryModule& AssetRegistryModule = FModuleManager::LoadModuleChecked<FAssetRegistryModule>(TEXT("AssetRegistry"));
+
+	// Collect a full list of assets with the control rig class
+	TArray<FAssetData> AssetDataList;
+	AssetRegistryModule.Get().GetAssetsByClass(UControlRigBlueprint::StaticClass()->GetClassPathName(), AssetDataList, true);
+
+	static const FLazyName ControlRigTypeName(GET_MEMBER_NAME_CHECKED(UControlRigBlueprint, ControlRigType));
+	FProperty* ControlRigTypeProperty = CastField<FProperty>(UControlRigBlueprint::StaticClass()->FindPropertyByName(ControlRigTypeName));
+	static const FLazyName ModuleReferenceDataName(GET_MEMBER_NAME_CHECKED(UControlRigBlueprint, ModuleReferenceData));
+	FArrayProperty* ModuleReferenceDataProperty = CastField<FArrayProperty>(UControlRigBlueprint::StaticClass()->FindPropertyByName(ModuleReferenceDataName));
+
+	for(const FAssetData& AssetData : AssetDataList)
+	{
+		// Check only modular rigs
+		{
+			const FString ControlRigTypeString = AssetData.GetTagValueRef<FString>(ControlRigTypeName);
+			if (ControlRigTypeString.IsEmpty())
+			{
+				continue;
+			}
+
+			EControlRigType RigType;
+			ControlRigTypeProperty->ImportText_Direct(*ControlRigTypeString, &RigType, nullptr, EPropertyPortFlags::PPF_None);
+			if (RigType != EControlRigType::ModularRig)
+			{
+				continue;
+			}
+		}
+		
+		const FString ModularRigDataString = AssetData.GetTagValueRef<FString>(ModuleReferenceDataName);
+		if (ModularRigDataString.IsEmpty())
+		{
+			continue;
+		}
+
+		TArray<FModuleReferenceData> Modules;
+		ModuleReferenceDataProperty->ImportText_Direct(*ModularRigDataString, &Modules, nullptr, EPropertyPortFlags::PPF_None);
+
+		for (FModuleReferenceData& Module : Modules)
+		{
+			if (Module.ReferencedModule == RigModuleClass)
+			{
+				Result.Add(Module);
+			}
+		}
+	}
+
+	return Result;
+}
+
+TArray<FModuleReferenceData> UControlRigBlueprint::GetModuleReferenceData() const
+{
+	TArray<FModuleReferenceData> Result;
+	Result.Reserve(ModularRigModel.Modules.Num());
+	ModularRigModel.ForEachModule([&Result](const FRigModuleReference* Module) -> bool
+	{
+		Result.Add(Module);
+		return true;
+	});
+	return Result;
 }
 
 void UControlRigBlueprint::UpdateExposedModuleConnectors() const
diff --git a/Engine/Plugins/Animation/ControlRig/Source/ControlRigDeveloper/Public/ControlRigBlueprint.h b/Engine/Plugins/Animation/ControlRig/Source/ControlRigDeveloper/Public/ControlRigBlueprint.h
--- a/Engine/Plugins/Animation/ControlRig/Source/ControlRigDeveloper/Public/ControlRigBlueprint.h
+++ b/Engine/Plugins/Animation/ControlRig/Source/ControlRigDeveloper/Public/ControlRigBlueprint.h
@@ -39,6 +39,33 @@ enum class EControlRigType : uint8
 	ModularRig =2,
 };
 
+USTRUCT(BlueprintType)
+struct CONTROLRIGDEVELOPER_API FModuleReferenceData
+{
+	GENERATED_BODY()
+
+public:
+	
+	FModuleReferenceData(){}
+
+	FModuleReferenceData(const FRigModuleReference* InModule)
+	{
+		if (InModule)
+		{
+			ModulePath = InModule->GetPath();
+			if (InModule->Class.IsValid())
+			{
+				ReferencedModule = InModule->Class.Get();
+			}
+		}
+	}
+
+	UPROPERTY()
+	FString ModulePath;
+
+	UPROPERTY()
+	FSoftClassPath ReferencedModule;
+};
 
 UCLASS(BlueprintType, meta=(IgnoreClassThumbnail))
 class CONTROLRIGDEVELOPER_API UControlRigBlueprint : public URigVMBlueprint, public IInterface_PreviewMeshProvider, public IRigHierarchyProvider
@@ -162,11 +189,20 @@ public:
 	UPROPERTY(EditAnywhere, Category = "Hierarchy", AssetRegistrySearchable)
 	FRigModuleSettings RigModuleSettings;
 
+	/** Asset searchable information module references in this rig */
+	UPROPERTY(AssetRegistrySearchable)
+	TArray<FModuleReferenceData> ModuleReferenceData;
+
 	UPROPERTY()
 	TMap<FRigElementKey, FRigElementKey> ConnectionMap;
 
+	UFUNCTION(BlueprintPure, Category = "Control Rig Blueprint")
+	TArray<FModuleReferenceData> FindReferencesToModule() const;
+
 protected:
 
+	TArray<FModuleReferenceData> GetModuleReferenceData() const;
+
 	FOnRigTypeChanged OnRigTypeChangedDelegate;
 	
 	void UpdateExposedModuleConnectors() const;
