Control Rig: Fix rename module affecting sibling modules with similar name
#rnx

diff --git a/Engine/Plugins/Animation/ControlRig/Source/ControlRig/Private/ModularRigController.cpp b/Engine/Plugins/Animation/ControlRig/Source/ControlRig/Private/ModularRigController.cpp
--- a/Engine/Plugins/Animation/ControlRig/Source/ControlRig/Private/ModularRigController.cpp
+++ b/Engine/Plugins/Animation/ControlRig/Source/ControlRig/Private/ModularRigController.cpp
@@ -1131,15 +1131,17 @@ FString UModularRigController::RenameModule(const FString& InModulePath, const F
 
 	// Fix connections
 	{
+		const FString OldNamespace = OldPath + TEXT(":");
+		const FString NewNamespace = NewPath + TEXT(":");
 		for (FModularRigSingleConnection& Connection : Model->Connections)
 		{
-			if (Connection.Connector.Name.ToString().StartsWith(OldPath, ESearchCase::CaseSensitive))
+			if (Connection.Connector.Name.ToString().StartsWith(OldNamespace, ESearchCase::CaseSensitive))
 			{
-				Connection.Connector.Name = *FString::Printf(TEXT("%s%s"), *NewPath, *Connection.Connector.Name.ToString().RightChop(OldPath.Len()));
+				Connection.Connector.Name = *FString::Printf(TEXT("%s%s"), *NewNamespace, *Connection.Connector.Name.ToString().RightChop(OldNamespace.Len()));
 			}
-			if (Connection.Target.Name.ToString().StartsWith(OldPath, ESearchCase::CaseSensitive))
+			if (Connection.Target.Name.ToString().StartsWith(OldNamespace, ESearchCase::CaseSensitive))
 			{
-				Connection.Target.Name = *FString::Printf(TEXT("%s%s"), *NewPath, *Connection.Target.Name.ToString().RightChop(OldPath.Len()));
+				Connection.Target.Name = *FString::Printf(TEXT("%s%s"), *NewNamespace, *Connection.Target.Name.ToString().RightChop(OldNamespace.Len()));
 			}
 		}
 		Model->Connections.UpdateFromConnectionList();
