Control Rig: Fix ensure !GIsTransacting after undoing move control action
#rnx

diff --git a/Engine/Plugins/Animation/ControlRig/Source/ControlRigEditor/Private/EditMode/SControlRigOutliner.cpp b/Engine/Plugins/Animation/ControlRig/Source/ControlRigEditor/Private/EditMode/SControlRigOutliner.cpp
--- a/Engine/Plugins/Animation/ControlRig/Source/ControlRigEditor/Private/EditMode/SControlRigOutliner.cpp
+++ b/Engine/Plugins/Animation/ControlRig/Source/ControlRigEditor/Private/EditMode/SControlRigOutliner.cpp
@@ -1131,7 +1131,9 @@ void SControlRigOutliner::HandleSelectionChanged(TSharedPtr<FMultiRigTreeElement
 			}
 			if (GEditor)
 			{
-				GEditor->Exec(GEditor->GetWorld(), TEXT("SELECT NONE"));
+				// Replicating the UEditorEngine::HandleSelectCommand, without the transaction to avoid ensure(!GIsTransacting)
+				GEditor->SelectNone( true, true );
+				GEditor->RedrawLevelEditingViewports();
 			}
 		}
 	}
