Control Rig: Fix color and tag of pre 5.4 layered control rig tracks

diff --git a/Engine/Plugins/Animation/ControlRig/Source/ControlRigEditor/Private/ControlRigSequencerEditorLibrary.cpp b/Engine/Plugins/Animation/ControlRig/Source/ControlRigEditor/Private/ControlRigSequencerEditorLibrary.cpp
--- a/Engine/Plugins/Animation/ControlRig/Source/ControlRigEditor/Private/ControlRigSequencerEditorLibrary.cpp
+++ b/Engine/Plugins/Animation/ControlRig/Source/ControlRigEditor/Private/ControlRigSequencerEditorLibrary.cpp
@@ -270,20 +270,8 @@ static UMovieSceneControlRigParameterTrack* AddControlRig(ULevelSequence* LevelS
 		UMovieSceneSection* NewSection = Track->CreateControlRigSection(0, ControlRig, bSequencerOwnsControlRig);
 		NewSection->Modify();
 
-		if (bIsAdditiveControlRig)
-		{
-			const FString AdditiveObjectName = ObjectName + TEXT(" (Layered)");
-			Track->SetTrackName(FName(*ObjectName));
-			Track->SetDisplayName(FText::FromString(AdditiveObjectName));
-			Track->SetColorTint(UMovieSceneControlRigParameterTrack::LayeredRigTrackColor);
-		}
-		else
-		{
-			//mz todo need to have multiple rigs with same class
-			Track->SetTrackName(FName(*ObjectName));
-			Track->SetDisplayName(FText::FromString(ObjectName));
-			Track->SetColorTint(UMovieSceneControlRigParameterTrack::AbsoluteRigTrackColor);
-		}
+		Track->SetTrackName(FName(*ObjectName));
+		UControlRigSequencerEditorLibrary::MarkLayeredModeOnTrackDisplay(Track);
 
 		if (SharedSequencer.IsValid())
 		{
@@ -3097,6 +3085,38 @@ bool UControlRigSequencerEditorLibrary::IsLayeredControlRig(UControlRig* InContr
 	return (InControlRig && InControlRig->IsAdditive());
 }
 
+bool UControlRigSequencerEditorLibrary::MarkLayeredModeOnTrackDisplay(UMovieSceneControlRigParameterTrack* InTrack)
+{
+	if (!InTrack)
+	{
+		FFrame::KismetExecutionMessage(TEXT("Invalid track"), ELogVerbosity::Error);
+		return false;
+	}
+
+	UControlRig* ControlRig = InTrack->GetControlRig();
+	if (!ControlRig)
+	{
+		FFrame::KismetExecutionMessage(TEXT("Track does not have a control rig"), ELogVerbosity::Error);
+		return false;
+	}
+
+	static constexpr TCHAR LayeredTag[] =  TEXT(" (Layered)");
+	const FString TrackName = InTrack->GetTrackName().ToString();
+	if (ControlRig->IsAdditive())
+	{
+		const FString AdditiveObjectName = TrackName + LayeredTag;
+		InTrack->SetDisplayName(FText::FromString(AdditiveObjectName));
+		InTrack->SetColorTint(UMovieSceneControlRigParameterTrack::LayeredRigTrackColor);
+	}
+	else
+	{
+		InTrack->SetDisplayName(FText::FromString(TrackName));
+		InTrack->SetColorTint(UMovieSceneControlRigParameterTrack::AbsoluteRigTrackColor);
+	}
+
+	return true;
+}
+
 bool UControlRigSequencerEditorLibrary::SetControlRigLayeredMode(UMovieSceneControlRigParameterTrack* InTrack, bool bSetIsLayered)
 {
 	if (!InTrack)
@@ -3144,20 +3164,9 @@ bool UControlRigSequencerEditorLibrary::SetControlRigLayeredMode(UMovieSceneCont
 
 	FString ObjectName = ControlRig->GetClass()->GetName(); //GetDisplayNameText().ToString();
 	ObjectName.RemoveFromEnd(TEXT("_C"));
+	InTrack->SetTrackName(FName(*ObjectName));
 	
-	if (bSetIsLayered)
-	{
-		const FString AdditiveObjectName = ObjectName + TEXT(" (Layered)");
-		InTrack->SetTrackName(FName(*ObjectName));
-		InTrack->SetDisplayName(FText::FromString(AdditiveObjectName));
-		InTrack->SetColorTint(UMovieSceneControlRigParameterTrack::LayeredRigTrackColor);
-	}
-	else
-	{
-		InTrack->SetTrackName(FName(*ObjectName));
-		InTrack->SetDisplayName(FText::FromString(ObjectName));
-		InTrack->SetColorTint(UMovieSceneControlRigParameterTrack::AbsoluteRigTrackColor);
-	}
+	MarkLayeredModeOnTrackDisplay(InTrack);
 
 	FControlRigEditMode* ControlRigEditMode = static_cast<FControlRigEditMode*>(GLevelEditorModeTools().GetActiveMode(FControlRigEditMode::ModeName));
 	if (ControlRigEditMode)
diff --git a/Engine/Plugins/Animation/ControlRig/Source/ControlRigEditor/Private/Sequencer/ControlRigParameterTrackEditor.cpp b/Engine/Plugins/Animation/ControlRig/Source/ControlRigEditor/Private/Sequencer/ControlRigParameterTrackEditor.cpp
--- a/Engine/Plugins/Animation/ControlRig/Source/ControlRigEditor/Private/Sequencer/ControlRigParameterTrackEditor.cpp
+++ b/Engine/Plugins/Animation/ControlRig/Source/ControlRigEditor/Private/Sequencer/ControlRigParameterTrackEditor.cpp
@@ -420,6 +420,9 @@ FControlRigParameterTrackEditor::FControlRigParameterTrackEditor(TSharedRef<ISeq
 			BindControlRig(ControlRig);
 		}
 
+		// Mark layered mode on track color and display name 
+		UControlRigSequencerEditorLibrary::MarkLayeredModeOnTrackDisplay(Track);
+
 		return false;
 	});
 }
diff --git a/Engine/Plugins/Animation/ControlRig/Source/ControlRigEditor/Public/ControlRigSequencerEditorLibrary.h b/Engine/Plugins/Animation/ControlRig/Source/ControlRigEditor/Public/ControlRigSequencerEditorLibrary.h
--- a/Engine/Plugins/Animation/ControlRig/Source/ControlRigEditor/Public/ControlRigSequencerEditorLibrary.h
+++ b/Engine/Plugins/Animation/ControlRig/Source/ControlRigEditor/Public/ControlRigSequencerEditorLibrary.h
@@ -1080,6 +1080,12 @@ public:
 	UFUNCTION(BlueprintCallable, Category = "Editor Scripting | Sequencer Tools | Control Rig")
 	static bool IsLayeredControlRig(UControlRig* InControlRig);
 
+	/**
+	 * Mark the layered mode of a control rig on the color and display name of a track
+	 * @param InTrack The track to modify the color and display name if necessary
+	 */
+	static bool MarkLayeredModeOnTrackDisplay(UMovieSceneControlRigParameterTrack* InTrack);
+
 	/*
 	 * Convert the control rig track into absolute or layered rig
 	 *
