Sequencer: Control Rig outliner, speculative fix for dev-cine, reset delegate handles directly and check for null hierarchies
#prefglight 6614589a7f454980bd90132b

diff --git a/Engine/Plugins/Animation/ControlRig/Source/ControlRigEditor/Private/EditMode/SControlRigEditModeTools.cpp b/Engine/Plugins/Animation/ControlRig/Source/ControlRigEditor/Private/EditMode/SControlRigEditModeTools.cpp
--- a/Engine/Plugins/Animation/ControlRig/Source/ControlRigEditor/Private/EditMode/SControlRigEditModeTools.cpp
+++ b/Engine/Plugins/Animation/ControlRig/Source/ControlRigEditor/Private/EditMode/SControlRigEditModeTools.cpp
@@ -62,19 +62,20 @@ void SControlRigEditModeTools::Cleanup()
 
 void SControlRigEditModeTools::SetControlRigs(const TArrayView<TWeakObjectPtr<UControlRig>>& InControlRigs)
 {
-	for (TWeakObjectPtr<UControlRig>& ControlRig : ControlRigs)
+	for (FDelegateHandle& Delegate : HandlesToClear)
 	{
-		if (ControlRig.IsValid())
+		if (Delegate.IsValid())
 		{
-			ControlRig->ControlSelected().RemoveAll(this);
+			Delegate.Reset();
 		}
 	}
+	HandlesToClear.Reset();
 	ControlRigs = InControlRigs;
 	for (TWeakObjectPtr<UControlRig>& InControlRig : InControlRigs)
 	{
 		if (InControlRig.IsValid())
 		{
-			InControlRig->ControlSelected().AddRaw(this, &SControlRigEditModeTools::OnRigElementSelected);
+			HandlesToClear.Add(InControlRig->ControlSelected().AddRaw(this, &SControlRigEditModeTools::OnRigElementSelected));
 		}
 	}
 
@@ -850,11 +851,14 @@ void SControlRigEditModeTools::OnRigElementSelected(UControlRig* Subject, FRigCo
 	}
 #endif
 
-	if (Subject)
+	if (Subject && Subject->GetHierarchy())
 	{
 		// get the selected controls
 		TArray<FRigElementKey> SelectedControls = Subject->GetHierarchy()->GetSelectedKeys(ERigElementType::Control);
-		SpacePickerWidget->SetControls(Subject->GetHierarchy(), SelectedControls);
+		if (SpacePickerWidget)
+		{
+			SpacePickerWidget->SetControls(Subject->GetHierarchy(), SelectedControls);
+		}
 		if (ConstraintsEditionWidget)
 		{
 			ConstraintsEditionWidget->InvalidateConstraintList();
@@ -862,7 +866,6 @@ void SControlRigEditModeTools::OnRigElementSelected(UControlRig* Subject, FRigCo
 	}
 }
 
-
 const FRigControlElementCustomization* SControlRigEditModeTools::HandleGetControlElementCustomization(URigHierarchy* InHierarchy, const FRigElementKey& InControlKey)
 {
 	UControlRig* Rig = ControlRigs.Num() > 0 ? ControlRigs[0].Get() : nullptr;
diff --git a/Engine/Plugins/Animation/ControlRig/Source/ControlRigEditor/Private/EditMode/SControlRigEditModeTools.h b/Engine/Plugins/Animation/ControlRig/Source/ControlRigEditor/Private/EditMode/SControlRigEditModeTools.h
--- a/Engine/Plugins/Animation/ControlRig/Source/ControlRigEditor/Private/EditMode/SControlRigEditModeTools.h
+++ b/Engine/Plugins/Animation/ControlRig/Source/ControlRigEditor/Private/EditMode/SControlRigEditModeTools.h
@@ -148,6 +148,9 @@ private:
 	// The toolkit that created this UI
 	TWeakPtr<FControlRigEditModeToolkit> OwningToolkit;
 
+	//array of handles to clear when getting new control rigs
+	TArray<FDelegateHandle> HandlesToClear;
+
 public:
 	/** Modes Panel Header Information **/
 	void CustomizeToolBarPalette(FToolBarBuilder& ToolBarBuilder);
