Sequencer: Constraints: fix for constraint creation on BPs with multiple SKM components (follow-up to 32821448)
    - ControlRig track auto-generation disabled to avoid spawning useless skms / rigs

#rnx

diff --git a/Engine/Plugins/Animation/ControlRig/Source/ControlRigEditor/Private/Editor/Constraints/SConstraintsWidget.cpp b/Engine/Plugins/Animation/ControlRig/Source/ControlRigEditor/Private/Editor/Constraints/SConstraintsWidget.cpp
--- a/Engine/Plugins/Animation/ControlRig/Source/ControlRigEditor/Private/Editor/Constraints/SConstraintsWidget.cpp
+++ b/Engine/Plugins/Animation/ControlRig/Source/ControlRigEditor/Private/Editor/Constraints/SConstraintsWidget.cpp
@@ -29,6 +29,7 @@
 #include "Styling/SlateIconFinder.h"
 #include "Widgets/Input/NumericTypeInterface.h"
 #include "FrameNumberDetailsCustomization.h"
+#include "Sequencer/ControlRigParameterTrackEditor.h"
 
 #define LOCTEXT_NAMESPACE "SConstraintsWidget"
 
@@ -492,7 +493,7 @@ void SConstraintMenuEntry::CreateConstraint(
 						{
 							Constraint->Evaluate();
 						}
-						FMovieSceneConstraintChannelHelper::AddConstraintToSequencer(WeakSequencer.Pin(), Constraint);
+						FControlRigParameterTrackEditor::AddConstraintToSequencer(WeakSequencer.Pin(), Constraint);
 					}
 					else
 					{
diff --git a/Engine/Plugins/Animation/ControlRig/Source/ControlRigEditor/Private/Sequencer/ControlRigParameterTrackEditor.cpp b/Engine/Plugins/Animation/ControlRig/Source/ControlRigEditor/Private/Sequencer/ControlRigParameterTrackEditor.cpp
--- a/Engine/Plugins/Animation/ControlRig/Source/ControlRigEditor/Private/Sequencer/ControlRigParameterTrackEditor.cpp
+++ b/Engine/Plugins/Animation/ControlRig/Source/ControlRigEditor/Private/Sequencer/ControlRigParameterTrackEditor.cpp
@@ -1967,6 +1967,12 @@ void FControlRigParameterTrackEditor::OnChannelChanged(const FMovieSceneChannelM
 	}
 }
 
+void FControlRigParameterTrackEditor::AddConstraintToSequencer(const TSharedPtr<ISequencer>& InSequencer, UTickableTransformConstraint* InConstraint)
+{
+	TGuardValue<bool> DisableTrackCreation(bAutoGenerateControlRigTrack, false);
+	FMovieSceneConstraintChannelHelper::AddConstraintToSequencer(InSequencer, InConstraint);
+}
+
 void FControlRigParameterTrackEditor::AddTrackForComponent(USceneComponent* InComponent, FGuid InBinding) 
 {
 	if (USkeletalMeshComponent* SkelMeshComp = Cast<USkeletalMeshComponent>(InComponent))
@@ -2035,7 +2041,7 @@ void FControlRigParameterTrackEditor::HandleActorAdded(AActor* Actor, FGuid Targ
 				const AActor* TargetActor = Child->IsA<AActor>() ? Cast<AActor>(Child) : Child->GetTypedOuter<AActor>();
 				if (TargetActor == Actor)
 				{
-					FMovieSceneConstraintChannelHelper::AddConstraintToSequencer(GetSequencer(), Constraint);
+					AddConstraintToSequencer(GetSequencer(), Constraint);
 				}		
 			}
 		}
diff --git a/Engine/Plugins/Animation/ControlRig/Source/ControlRigEditor/Private/Sequencer/ControlRigParameterTrackEditor.h b/Engine/Plugins/Animation/ControlRig/Source/ControlRigEditor/Private/Sequencer/ControlRigParameterTrackEditor.h
--- a/Engine/Plugins/Animation/ControlRig/Source/ControlRigEditor/Private/Sequencer/ControlRigParameterTrackEditor.h
+++ b/Engine/Plugins/Animation/ControlRig/Source/ControlRigEditor/Private/Sequencer/ControlRigParameterTrackEditor.h
@@ -249,6 +249,10 @@ public:
 
 	//** Function to smart reduce all keys on all constraints based on the parameters
 	static void SmartReduce(TSharedPtr<ISequencer>& SequencerPtr, const FSmartReduceParams& InParams, UMovieSceneControlRigParameterSection* InSection);
+
+	/** Adds the constraint to sequencer and create the bindings if needed. Note that ControlRig Track auto-generation is disabled in this function. */
+	static void AddConstraintToSequencer(const TSharedPtr<ISequencer>& InSequencer, UTickableTransformConstraint* InConstraint);
+
 private:
 	FDelegateHandle SelectionChangedHandle;
 	FDelegateHandle SequencerChangedHandle;
