Control Rig: Allow to create connectors on rig modules without connectors

#rnx

diff --git a/Engine/Plugins/Animation/ControlRig/Source/ControlRigEditor/Private/Editor/SRigHierarchy.cpp b/Engine/Plugins/Animation/ControlRig/Source/ControlRigEditor/Private/Editor/SRigHierarchy.cpp
--- a/Engine/Plugins/Animation/ControlRig/Source/ControlRigEditor/Private/Editor/SRigHierarchy.cpp
+++ b/Engine/Plugins/Animation/ControlRig/Source/ControlRigEditor/Private/Editor/SRigHierarchy.cpp
@@ -2076,21 +2076,24 @@ void SRigHierarchy::HandleNewItem(ERigElementType InElementType, bool bIsAnimati
 						{
 							if(ControlRigBlueprint->Hierarchy->Num(ERigElementType::Connector) == 0)
 							{
-								static constexpr TCHAR Format[] = TEXT("Connector cannot be created: %s");
-								UE_LOG(LogControlRig, Warning, Format, *FailureReason);
-								FNotificationInfo Info(FText::FromString(FString::Printf(Format, *FailureReason)));
-								Info.bUseSuccessFailIcons = true;
-								Info.Image = FAppStyle::GetBrush(TEXT("MessageLog.Warning"));
-								Info.bFireAndForget = true;
-								Info.bUseThrobber = true;
-								Info.FadeOutDuration = 2.f;
-								Info.ExpireDuration = 8.f;;
-								TSharedPtr<SNotificationItem> NotificationPtr = FSlateNotificationManager::Get().AddNotification(Info);
-								if (NotificationPtr)
+								if(!ControlRigBlueprint->IsControlRigModule())
 								{
-									NotificationPtr->SetCompletionState(SNotificationItem::CS_Fail);
+									static constexpr TCHAR Format[] = TEXT("Connector cannot be created: %s");
+									UE_LOG(LogControlRig, Warning, Format, *FailureReason);
+									FNotificationInfo Info(FText::FromString(FString::Printf(Format, *FailureReason)));
+									Info.bUseSuccessFailIcons = true;
+									Info.Image = FAppStyle::GetBrush(TEXT("MessageLog.Warning"));
+									Info.bFireAndForget = true;
+									Info.bUseThrobber = true;
+									Info.FadeOutDuration = 2.f;
+									Info.ExpireDuration = 8.f;;
+									TSharedPtr<SNotificationItem> NotificationPtr = FSlateNotificationManager::Get().AddNotification(Info);
+									if (NotificationPtr)
+									{
+										NotificationPtr->SetCompletionState(SNotificationItem::CS_Fail);
+									}
+									return;
 								}
-								return;
 							}
 						}
 
