Control Rig: Hide non-related tabs in the modular rig editor
#rnx

diff --git a/Engine/Plugins/Animation/ControlRig/Source/ControlRigEditor/Private/Editor/ControlRigEditorMode.cpp b/Engine/Plugins/Animation/ControlRig/Source/ControlRigEditor/Private/Editor/ControlRigEditorMode.cpp
--- a/Engine/Plugins/Animation/ControlRig/Source/ControlRigEditor/Private/Editor/ControlRigEditorMode.cpp
+++ b/Engine/Plugins/Animation/ControlRig/Source/ControlRigEditor/Private/Editor/ControlRigEditorMode.cpp
@@ -204,4 +204,37 @@ FModularRigEditorMode::FModularRigEditorMode(const TSharedRef<FControlRigEditor>
 				)
 			)
 		);
-}
\ No newline at end of file
+}
+
+void FModularRigEditorMode::RegisterTabFactories(TSharedPtr<FTabManager> InTabManager)
+{
+	TSharedPtr<FBlueprintEditor> BP = MyBlueprintEditor.Pin();
+	
+	BP->RegisterToolbarTab(InTabManager.ToSharedRef());
+
+	static const TArray<FName> DisallowedTabs = {
+		FBlueprintEditorTabs::PaletteID,
+		FBlueprintEditorTabs::ReplaceNodeReferencesID,
+		FBlueprintEditorTabs::CompilerResultsID,
+		FBlueprintEditorTabs::FindResultsID,
+		FBlueprintEditorTabs::BookmarksID,
+		FRigVMExecutionStackTabSummoner::TabID
+	};
+
+	auto PushTabFactories = [&](FWorkflowAllowedTabSet& Tabs)
+	{
+		for (auto FactoryIt = Tabs.CreateIterator(); FactoryIt; ++FactoryIt)
+		{
+			if (DisallowedTabs.Contains(FactoryIt->Key))
+			{
+				continue;
+			}
+			FactoryIt.Value()->RegisterTabSpawner(InTabManager.ToSharedRef(), BP->GetCurrentModePtr().Get());
+		}
+	};
+
+	// Mode-specific setup
+	PushTabFactories(CoreTabFactories);
+	PushTabFactories(BlueprintEditorTabFactories);
+	PushTabFactories(TabFactories);
+}
diff --git a/Engine/Plugins/Animation/ControlRig/Source/ControlRigEditor/Private/Editor/ControlRigEditorMode.h b/Engine/Plugins/Animation/ControlRig/Source/ControlRigEditor/Private/Editor/ControlRigEditorMode.h
--- a/Engine/Plugins/Animation/ControlRig/Source/ControlRigEditor/Private/Editor/ControlRigEditorMode.h
+++ b/Engine/Plugins/Animation/ControlRig/Source/ControlRigEditor/Private/Editor/ControlRigEditorMode.h
@@ -31,6 +31,9 @@ class FModularRigEditorMode : public FControlRigEditorMode
 public:
 	FModularRigEditorMode(const TSharedRef<FControlRigEditor>& InControlRigEditor);
 
+	// FApplicationMode interface
+	virtual void RegisterTabFactories(TSharedPtr<FTabManager> InTabManager) override;
+
 	// for now just don't open up the previous edited documents
 	virtual void PostActivateMode() override {}
 };
