Control Rig: Cleanup hierarchy for cache validation
#rnx

diff --git a/Engine/Plugins/Animation/ControlRig/Source/ControlRig/Private/Rigs/RigHierarchy.cpp b/Engine/Plugins/Animation/ControlRig/Source/ControlRig/Private/Rigs/RigHierarchy.cpp
--- a/Engine/Plugins/Animation/ControlRig/Source/ControlRig/Private/Rigs/RigHierarchy.cpp
+++ b/Engine/Plugins/Animation/ControlRig/Source/ControlRig/Private/Rigs/RigHierarchy.cpp
@@ -5357,10 +5357,17 @@ void URigHierarchy::PropagateDirtyFlags(FRigTransformElement* InTransformElement
 void URigHierarchy::CleanupInvalidCaches()
 {
 	// create a copy of this hierarchy and pre compute all transforms
-	if(HierarchyForCacheValidation == nullptr)
+	bool bHierarchyWasCreatedDuringCleanup = false;
+	if(!IsValid(HierarchyForCacheValidation))
 	{
-		HierarchyForCacheValidation = NewObject<URigHierarchy>(this, NAME_None, RF_Transient);
+		// Give the validation hierarchy a unique name for debug purposes and to avoid possible memory reuse
+		static TAtomic<uint32> HierarchyNameIndex = 0;
+		static constexpr TCHAR Format[] = TEXT("CacheValidationHierarchy_%zu");
+		const FString ValidationHierarchyName = FString::Printf(Format, uint32(++HierarchyNameIndex));
+		
+		HierarchyForCacheValidation = NewObject<URigHierarchy>(this, *ValidationHierarchyName, RF_Transient);
 		HierarchyForCacheValidation->bEnableCacheValidityCheck = false;
+		bHierarchyWasCreatedDuringCleanup = true;
 	}
 	HierarchyForCacheValidation->CopyHierarchy(this);
 
@@ -5497,6 +5504,14 @@ void URigHierarchy::CleanupInvalidCaches()
 
 	ResetPoseToInitial(ERigElementType::All);
 	EnsureCacheValidity();
+
+	if(bHierarchyWasCreatedDuringCleanup && HierarchyForCacheValidation)
+	{
+		HierarchyForCacheValidation->Rename(nullptr, GetTransientPackage(), REN_ForceNoResetLoaders | REN_DoNotDirty | REN_DontCreateRedirectors | REN_NonTransactional);
+		HierarchyForCacheValidation->RemoveFromRoot();
+		HierarchyForCacheValidation->MarkAsGarbage();
+		HierarchyForCacheValidation = nullptr;
+	}
 }
 
 void URigHierarchy::FMetadataStorage::Reset()
