Sequencer: Control Rig: When getting object replaced  delegate don't force sequencer to rerun if the track is muted, eventually will change this to a 'track is evaluating' test but that doesn't exist(need introspection Api).

diff --git a/Engine/Plugins/Animation/ControlRig/Source/ControlRigEditor/Private/Sequencer/ControlRigParameterTrackEditor.cpp b/Engine/Plugins/Animation/ControlRig/Source/ControlRigEditor/Private/Sequencer/ControlRigParameterTrackEditor.cpp
--- a/Engine/Plugins/Animation/ControlRig/Source/ControlRigEditor/Private/Sequencer/ControlRigParameterTrackEditor.cpp
+++ b/Engine/Plugins/Animation/ControlRig/Source/ControlRigEditor/Private/Sequencer/ControlRigParameterTrackEditor.cpp
@@ -260,7 +260,7 @@ FControlRigParameterTrackEditor::FControlRigParameterTrackEditor(TSharedRef<ISeq
 		//the other is if bound object on the edit mode is null we request a re-evaluate which will reset it up.
 		const FDelegateHandle OnObjectsReplacedHandle = FCoreUObjectDelegates::OnObjectsReplaced.AddLambda([this](const TMap<UObject*, UObject*>& ReplacementMap)
 		{
-			if (GetSequencer().IsValid())
+			if (GetSequencer().IsValid() && GetSequencer()->GetFocusedMovieSceneSequence() && GetSequencer()->GetFocusedMovieSceneSequence()->GetMovieScene())
 			{
 				static bool bHasEnteredSilent = false;
 				
@@ -269,6 +269,9 @@ FControlRigParameterTrackEditor::FControlRigParameterTrackEditor(TSharedRef<ISeq
 				FControlRigEditMode* ControlRigEditMode = GetEditMode();
 				if (ControlRigEditMode)
 				{
+					UMovieScene* MovieScene = GetSequencer()->GetFocusedMovieSceneSequence()->GetMovieScene();
+					TArray<FString>& MuteNodes = MovieScene->GetMuteNodes();
+
 					TArrayView<TWeakObjectPtr<UControlRig>> ControlRigs = ControlRigEditMode->GetControlRigs();
 					for (TWeakObjectPtr<UControlRig>& ControlRigPtr : ControlRigs)
 					{
@@ -276,8 +279,21 @@ FControlRigParameterTrackEditor::FControlRigParameterTrackEditor(TSharedRef<ISeq
 						{
 							if (ControlRig->GetObjectBinding() && ControlRig->GetObjectBinding()->GetBoundObject() == nullptr)
 							{
-								bRequestEvaluate = true;
-								break;
+								IterateTracks([this, ControlRig, &MuteNodes, &bRequestEvaluate](UMovieSceneControlRigParameterTrack* Track)
+									{
+										if (ControlRig == Track->GetControlRig())
+										{
+											if (MuteNodes.Find(Track->GetName()) != INDEX_NONE) //only re-evalute if  not muted, TODO, need function to test to see if a track is evaluating at a certain time
+											{
+												bRequestEvaluate = true;
+											}
+										}
+										return false;
+									});
+								if (bRequestEvaluate == true)
+								{
+									break;
+								}
 							}
 						}
 					}
