Sequencer: 2nd speculative fix for dev-cine, remove and then reset delegate handles
#prefglight 66170e43157150b53df06266

diff --git a/Engine/Plugins/Animation/ControlRig/Source/ControlRigEditor/Private/EditMode/SControlRigEditModeTools.cpp b/Engine/Plugins/Animation/ControlRig/Source/ControlRigEditor/Private/EditMode/SControlRigEditModeTools.cpp
--- a/Engine/Plugins/Animation/ControlRig/Source/ControlRigEditor/Private/EditMode/SControlRigEditModeTools.cpp
+++ b/Engine/Plugins/Animation/ControlRig/Source/ControlRigEditor/Private/EditMode/SControlRigEditModeTools.cpp
@@ -58,15 +58,31 @@ void SControlRigEditModeTools::Cleanup()
 	{
 		RigOptionsDetailsView->SetKeyframeHandler(nullptr);
 	}
+	for (TPair<FDelegateHandle, TWeakObjectPtr<UControlRig>>& Handles : HandlesToClear)
+	{
+		if (Handles.Value.IsValid())
+		{
+			Handles.Value->ControlSelected().RemoveAll(this);
+		}
+		if (Handles.Key.IsValid())
+		{
+			Handles.Key.Reset();
+		}
+	}
+	HandlesToClear.Reset();
 }
 
 void SControlRigEditModeTools::SetControlRigs(const TArrayView<TWeakObjectPtr<UControlRig>>& InControlRigs)
 {
-	for (FDelegateHandle& Delegate : HandlesToClear)
+	for (TPair<FDelegateHandle, TWeakObjectPtr<UControlRig>>& Handles : HandlesToClear)
 	{
-		if (Delegate.IsValid())
+		if (Handles.Value.IsValid())
+		{
+			Handles.Value->ControlSelected().RemoveAll(this);
+		}
+		if (Handles.Key.IsValid())
 		{
-			Delegate.Reset();
+			Handles.Key.Reset();
 		}
 	}
 	HandlesToClear.Reset();
@@ -75,7 +91,10 @@ void SControlRigEditModeTools::SetControlRigs(const TArrayView<TWeakObjectPtr<UC
 	{
 		if (InControlRig.IsValid())
 		{
-			HandlesToClear.Add(InControlRig->ControlSelected().AddRaw(this, &SControlRigEditModeTools::OnRigElementSelected));
+			TPair<FDelegateHandle, TWeakObjectPtr<UControlRig>> Handles;
+			Handles.Key = InControlRig->ControlSelected().AddRaw(this, &SControlRigEditModeTools::OnRigElementSelected);
+			Handles.Value = InControlRig;
+			HandlesToClear.Add(Handles);
 		}
 	}
 
diff --git a/Engine/Plugins/Animation/ControlRig/Source/ControlRigEditor/Private/EditMode/SControlRigEditModeTools.h b/Engine/Plugins/Animation/ControlRig/Source/ControlRigEditor/Private/EditMode/SControlRigEditModeTools.h
--- a/Engine/Plugins/Animation/ControlRig/Source/ControlRigEditor/Private/EditMode/SControlRigEditModeTools.h
+++ b/Engine/Plugins/Animation/ControlRig/Source/ControlRigEditor/Private/EditMode/SControlRigEditModeTools.h
@@ -149,7 +149,7 @@ private:
 	TWeakPtr<FControlRigEditModeToolkit> OwningToolkit;
 
 	//array of handles to clear when getting new control rigs
-	TArray<FDelegateHandle> HandlesToClear;
+	TArray<TPair<FDelegateHandle, TWeakObjectPtr<UControlRig>>> HandlesToClear;
 
 public:
 	/** Modes Panel Header Information **/
