Control Rig: Delete child modules when a parent module is deleted
#rnx

diff --git a/Engine/Plugins/Animation/ControlRig/Source/ControlRig/Private/ModularRigController.cpp b/Engine/Plugins/Animation/ControlRig/Source/ControlRig/Private/ModularRigController.cpp
--- a/Engine/Plugins/Animation/ControlRig/Source/ControlRig/Private/ModularRigController.cpp
+++ b/Engine/Plugins/Animation/ControlRig/Source/ControlRig/Private/ModularRigController.cpp
@@ -1016,11 +1016,12 @@ bool UModularRigController::DeleteModule(const FString& InModulePath, bool bSetu
 	}
 #endif
 
-	// Unparent children (add them to root)
-	TArray<FRigModuleReference*> PreviousChildren = Module->CachedChildren;
-	for (const FRigModuleReference* Child : PreviousChildren)
+	// Delete children
+	TArray<FString> ChildrenPaths;
+	Algo::Transform(Module->CachedChildren, ChildrenPaths, [](const FRigModuleReference* Child){ return Child->GetPath(); });
+	for (const FString& ChildPath : ChildrenPaths)
 	{
-		ReparentModule(Child->GetPath(), FString(), bSetupUndo);
+		DeleteModule(ChildPath, bSetupUndo);
 	}
 
 	Model->DeletedModules.Add(*Module);
