Verify that executing the Set Transform Rig Unit translates the chosen Bone in the Control Rig Preview viewport
#rnx

diff --git a/Engine/Plugins/Animation/ControlRig/Source/ControlRigEditor/Private/Tests/ControlRigEditorlTransform.spec.cpp b/Engine/Plugins/Animation/ControlRig/Source/ControlRigEditor/Private/Tests/ControlRigEditorlTransform.spec.cpp
new file mode 100644
--- /dev/null
+++ b/Engine/Plugins/Animation/ControlRig/Source/ControlRigEditor/Private/Tests/ControlRigEditorlTransform.spec.cpp
@@ -0,0 +1,242 @@
+// Copyright Epic Games, Inc. All Rights Reserved.
+
+#if WITH_DEV_AUTOMATION_TESTS
+
+#if WITH_EDITOR
+
+#include "Misc/AutomationTest.h"
+#include "Editor.h"
+#include "Subsystems/EditorAssetSubsystem.h"
+#include "Subsystems/AssetEditorSubsystem.h"
+#include "ControlRigBlueprintFactory.h"
+#include "ControlRigBlueprint.h"
+#include "Rigs/RigHierarchyController.h"
+#include "Rigs/RigHierarchy.h"
+#include "Graph/ControlRigGraph.h"
+#include "RigVMModel/RigVMController.h"
+#include "ObjectTools.h"
+#include "Units/Hierarchy/RigUnit_GetTransform.h"
+#include "Units/Hierarchy/RigUnit_SetTransform.h"
+#include "Editor/ControlRigEditor.h"
+#include "SAdvancedTransformInputBox.h"
+
+
+BEGIN_DEFINE_SPEC(
+	FControlRigEditorTransformSpec, "Animation.Editor",
+	EAutomationTestFlags::EditorContext | EAutomationTestFlags::EngineFilter)
+
+	UControlRigBlueprint* ControlRigAsset = nullptr;
+	URigHierarchyController* HierarchyController = nullptr;
+	UControlRigGraph* RigGraph = nullptr;
+	URigVMController* VMController = nullptr;
+	
+	TStrongObjectPtr<URigVMUnitNode> ControlTransformNode;
+	TStrongObjectPtr<URigVMUnitNode> BoneTransformNode;
+	
+	FControlRigEditor* ControlRigEditor = nullptr;
+	UControlRig* ControlRig = nullptr;
+	URigHierarchy* Hierarchy = nullptr;
+	
+	FRigElementKey ControlKey;
+
+	void AddControlElement();
+	const FRigElementKey GetFirstBoneKey() const;
+	void CreateAndLinkTransformNodes();
+	void SelectControlElement();
+	bool ApplyTransformToControl(const ESlateTransformComponent::Type ComponentType, const double Value) const;
+	bool VerifyBoneTransformUpdate() const;
+
+END_DEFINE_SPEC(FControlRigEditorTransformSpec)
+
+void FControlRigEditorTransformSpec::Define()
+{
+	Describe("Control Rig Transform", [this]()
+		{
+			BeforeEach([this]()
+				{
+					const FString SkeletalMeshName = "SK_Box_Morph_1";
+					const FString SkeletalMeshPath = FString::Printf(TEXT("/Game/Tests/Animation/Curves/%s"), *SkeletalMeshName);
+					FAssetData AssetData = GEditor->GetEditorSubsystem<UEditorAssetSubsystem>()->FindAssetData(*SkeletalMeshPath);
+					UObject* SkeletalMesh = nullptr;
+					if (AssetData.IsValid())
+					{
+						SkeletalMesh = AssetData.GetAsset();
+					}
+					TestNotNull(FString::Printf(TEXT("Asset '%s' loaded successfully"), *SkeletalMeshName), SkeletalMesh);
+
+					ControlRigAsset = UControlRigBlueprintFactory::CreateControlRigFromSkeletalMeshOrSkeleton(SkeletalMesh);
+					TestNotNull(TEXT("Control Rig Asset created successfully from Skeletal Mesh"), ControlRigAsset);
+
+					HierarchyController = ControlRigAsset->GetHierarchyController();
+
+					if(UEdGraph* EdGraph = ControlRigAsset->GetEdGraph(ControlRigAsset->GetModel()))
+					{
+						RigGraph = CastChecked<UControlRigGraph>(EdGraph);
+						VMController = ControlRigAsset->GetController(RigGraph);
+					}
+
+					const bool bShowProgressWindow = false;
+					if(GEditor->GetEditorSubsystem<UAssetEditorSubsystem>()->OpenEditorForAsset(ControlRigAsset, EToolkitMode::Standalone, TSharedPtr<IToolkitHost>(), bShowProgressWindow))
+					{
+						const bool bFocusIfOpen = true;
+						if (IAssetEditorInstance* ControlRigEditorInstance = GEditor->GetEditorSubsystem<UAssetEditorSubsystem>()->FindEditorForAsset(ControlRigAsset, bFocusIfOpen))
+						{
+							ControlRigEditor = static_cast<FControlRigEditor*>(ControlRigEditorInstance);
+							ControlRig = ControlRigEditor->GetInstanceRig();
+							if (ControlRig)
+							{
+								Hierarchy = ControlRig->GetHierarchy();
+							}
+						}
+					}
+				});
+
+			It("Verify that executing the Set Transform Rig Unit translates the chosen Bone in the Control Rig Preview viewport", [this]()
+				{
+					AddControlElement();
+					CreateAndLinkTransformNodes();
+					SelectControlElement();
+
+					if (TestTrue(TEXT("A new Location is applied to the Control Element"), ApplyTransformToControl(ESlateTransformComponent::Location, 50)))
+					{
+						TestTrue(TEXT("Bone Location updated"), VerifyBoneTransformUpdate());
+					}
+				});
+			
+			It("Verify that executing the Set Transform Rig Unit rotates the chosen Bone in the Control Rig Preview viewport", [this]()
+				{
+					AddControlElement();
+					CreateAndLinkTransformNodes();
+					SelectControlElement();
+
+					if (TestTrue(TEXT("A new Rotation is applied to the Control Element"), ApplyTransformToControl(ESlateTransformComponent::Rotation, 50)))
+					{
+						TestTrue(TEXT("Bone Rotation updated"), VerifyBoneTransformUpdate());
+					}
+				});
+
+			It("Verify that executing the Set Transform Rig Unit scales the chosen Bone in the Control Rig Preview viewport", [this]()
+				{
+					AddControlElement();
+					CreateAndLinkTransformNodes();
+					SelectControlElement();
+
+					if(TestTrue(TEXT("A new Scale is applied to the Control Element"), ApplyTransformToControl(ESlateTransformComponent::Scale, 50)))
+					{
+						TestTrue(TEXT("Bone Scale updated"), VerifyBoneTransformUpdate());
+					}
+				});
+
+			AfterEach([this]()
+				{
+					if(ControlRigAsset)
+					{
+						GEditor->GetEditorSubsystem<UAssetEditorSubsystem>()->CloseAllEditorsForAsset(ControlRigAsset);
+						ObjectTools::DeleteObjectsUnchecked({ ControlRigAsset });
+					}
+				});
+		});
+}
+
+void FControlRigEditorTransformSpec::AddControlElement()
+{
+	if(HierarchyController)
+	{
+		ControlKey = HierarchyController->AddControl(TEXT("TestControl"), FRigElementKey(), FRigControlSettings(), FRigControlValue::Make(FTransform::Identity));
+	}
+}
+
+const FRigElementKey FControlRigEditorTransformSpec::GetFirstBoneKey() const
+{
+	if(ControlRigAsset)
+	{
+		const TArray<FRigBoneElement*>& Bones = ControlRigAsset->GetHierarchy()->GetBones();
+		if(Bones.Num() > 0)
+		{
+			return Bones[0]->GetKey();
+		}
+	}
+	return FRigElementKey();
+}
+
+void FControlRigEditorTransformSpec::CreateAndLinkTransformNodes()
+{
+	if(VMController)
+	{
+		FRigUnit_GetTransform GetTransform;
+		GetTransform.Item = ControlKey;
+		ControlTransformNode.Reset(VMController->AddUnitNode(GetTransform));
+
+		FRigUnit_SetTransform SetTransform;
+		SetTransform.Item = GetFirstBoneKey();
+		BoneTransformNode.Reset(VMController->AddUnitNode(SetTransform));
+
+		if(RigGraph && ControlTransformNode.IsValid() && BoneTransformNode.IsValid())
+		{
+			const URigVMGraph* RigVMGraph = RigGraph->GetModel();
+			const TArray<URigVMNode*>& Nodes = RigVMGraph->GetNodes();
+
+			if(Nodes.Num() > 0)
+			{
+				VMController->AddLink(
+					FString::Printf(TEXT("%s.ExecuteContext"), *Nodes[0]->GetName()),
+					FString::Printf(TEXT("%s.ExecuteContext"), *BoneTransformNode->GetName())
+				);
+			}
+
+			VMController->AddLink(
+				FString::Printf(TEXT("%s.Transform"), *ControlTransformNode->GetName()),
+				FString::Printf(TEXT("%s.Transform"), *BoneTransformNode->GetName())
+			);
+		}
+	}
+}
+
+void FControlRigEditorTransformSpec::SelectControlElement()
+{
+	if(HierarchyController && ControlKey.IsValid())
+	{
+		HierarchyController->SetSelection({ ControlKey });
+	}
+}
+
+bool FControlRigEditorTransformSpec::ApplyTransformToControl(const ESlateTransformComponent::Type ComponentType, const double Value) const
+{
+	if(Hierarchy && ControlRig && ControlKey.IsValid())
+	{
+		Hierarchy->Modify();
+
+		const FTransform ControlLocalTransform = ControlRig->GetControlLocalTransform(ControlKey.Name);
+
+		FEulerTransform RelativeTransform = FEulerTransform::Identity;
+		RelativeTransform.FromFTransform(ControlLocalTransform);
+
+		SAdvancedTransformInputBox<FEulerTransform>::ApplyNumericValueChange(RelativeTransform, Value,
+			ComponentType, ESlateRotationRepresentation::EulerXYZ, ESlateTransformSubComponent::Z);
+
+		ControlRig->Evaluate_AnyThread();
+		ControlRig->SetControlLocalTransform(ControlKey.Name, RelativeTransform.ToFTransform(), true, FRigControlModifiedContext(), false, false);
+		ControlRig->Evaluate_AnyThread();
+
+		return true;
+	}
+
+	return false;
+}
+
+bool FControlRigEditorTransformSpec::VerifyBoneTransformUpdate() const
+{
+	const FRigElementKey BoneKey = GetFirstBoneKey();
+	if (Hierarchy && ControlKey.IsValid() && BoneKey.IsValid())
+	{
+		const FTransform ControlTransform = Hierarchy->GetLocalTransform(ControlKey, false);
+		const FTransform BoneTransform = Hierarchy->GetLocalTransform(BoneKey, false);
+
+		return ControlTransform.Equals(BoneTransform);
+	}
+	return false;
+}
+
+#endif // WITH_DEV_AUTOMATION_TESTS
+
+#endif // WITH_EDITOR
\ No newline at end of file
