Move ControlRig tests from Editor.Content to Automation.Editor

diff --git a/Engine/Plugins/Animation/ControlRig/Source/ControlRigEditor/Private/Tests/ControlRigEditor.spec.cpp b/Engine/Plugins/Animation/ControlRig/Source/ControlRigEditor/Private/Tests/ControlRigEditor.spec.cpp
--- a/Engine/Plugins/Animation/ControlRig/Source/ControlRigEditor/Private/Tests/ControlRigEditor.spec.cpp
+++ b/Engine/Plugins/Animation/ControlRig/Source/ControlRigEditor/Private/Tests/ControlRigEditor.spec.cpp
@@ -18,7 +18,7 @@
 
 
 BEGIN_DEFINE_SPEC(
-	FControlRigEditorSpec, "Editor.Content",
+	FControlRigEditorSpec, "Animation.Editor",
 	EAutomationTestFlags::EditorContext | EAutomationTestFlags::EngineFilter)
 	
 	UControlRigBlueprint* ControlRigBlueprint = nullptr;
