ControlRig: Fix static analysis warning
#rnx

diff --git a/Engine/Plugins/Animation/ControlRig/Source/ControlRigDeveloper/Private/ControlRigBlueprint.cpp b/Engine/Plugins/Animation/ControlRig/Source/ControlRigDeveloper/Private/ControlRigBlueprint.cpp
--- a/Engine/Plugins/Animation/ControlRig/Source/ControlRigDeveloper/Private/ControlRigBlueprint.cpp
+++ b/Engine/Plugins/Animation/ControlRig/Source/ControlRigDeveloper/Private/ControlRigBlueprint.cpp
@@ -1979,14 +1979,15 @@ void UControlRigBlueprint::PatchFunctionsOnLoad()
 			Store.AddFunction(Header, bIsPublic);
 			
 		}
-	}
 
-	// Update dependencies and external variables if needed
-	for (URigVMLibraryNode* LibraryNode : FunctionLibrary->GetFunctions())
-	{
-		GetRigVMClient()->UpdateExternalVariablesForFunction(LibraryNode);
-		GetRigVMClient()->UpdateDependenciesForFunction(LibraryNode);
+		// Update dependencies and external variables if needed
+		for (URigVMLibraryNode* LibraryNode : FunctionLibrary->GetFunctions())
+		{
+			GetRigVMClient()->UpdateExternalVariablesForFunction(LibraryNode);
+			GetRigVMClient()->UpdateDependenciesForFunction(LibraryNode);
+		}
 	}
+
 }
 
 void UControlRigBlueprint::CreateMemberVariablesOnLoad()
