Control Rig: Fix slowness when compiling module while modular rig is opened
#rnx

diff --git a/Engine/Plugins/Animation/ControlRig/Source/ControlRigDeveloper/Private/ControlRigBlueprint.cpp b/Engine/Plugins/Animation/ControlRig/Source/ControlRigDeveloper/Private/ControlRigBlueprint.cpp
--- a/Engine/Plugins/Animation/ControlRig/Source/ControlRigDeveloper/Private/ControlRigBlueprint.cpp
+++ b/Engine/Plugins/Animation/ControlRig/Source/ControlRigDeveloper/Private/ControlRigBlueprint.cpp
@@ -2408,6 +2408,7 @@ void UControlRigBlueprint::UpdateModularDependencyDelegates()
 				if(!VisitList.Contains(Blueprint))
 				{
 					Blueprint->OnVMCompiled().RemoveAll(this);
+					Blueprint->OnModularRigCompiled().RemoveAll(this);
 					Blueprint->OnVMCompiled().AddUObject(this, &UControlRigBlueprint::OnModularDependencyVMCompiled);
 					Blueprint->OnModularRigCompiled().AddUObject(this, &UControlRigBlueprint::OnModularDependencyChanged);
 					VisitList.Add(Blueprint);
