Fixed a lot of unused or potentially-uninitialized variables, and also a couple of unreachable code warnings.

diff --git a/Engine/Plugins/Animation/ControlRig/Source/ControlRig/Private/Sequencer/MovieSceneControlRigParameterSection.cpp b/Engine/Plugins/Animation/ControlRig/Source/ControlRig/Private/Sequencer/MovieSceneControlRigParameterSection.cpp
--- a/Engine/Plugins/Animation/ControlRig/Source/ControlRig/Private/Sequencer/MovieSceneControlRigParameterSection.cpp
+++ b/Engine/Plugins/Animation/ControlRig/Source/ControlRig/Private/Sequencer/MovieSceneControlRigParameterSection.cpp
@@ -742,7 +742,7 @@ struct FParameterTransformChannelEditorData
 		EvalTrack.Interrogate(Context, InterrogationData, Object);
 
 		FVector CurrentPos = FVector::ZeroVector;
-		FRotator CurrentRot;
+		FRotator CurrentRot = FRotator::ZeroRotator;
 		FVector CurrentScale = FVector::ZeroVector;
 
 		for (const FEulerTransformInterrogationData& Transform : InterrogationData.Iterate<FEulerTransformInterrogationData>(UMovieSceneControlRigParameterSection::GetTransformInterrogationKey()))
diff --git a/Engine/Plugins/Animation/GameplayInsights/Source/GameplayInsights/Private/MontagesTrack.cpp b/Engine/Plugins/Animation/GameplayInsights/Source/GameplayInsights/Private/MontagesTrack.cpp
--- a/Engine/Plugins/Animation/GameplayInsights/Source/GameplayInsights/Private/MontagesTrack.cpp
+++ b/Engine/Plugins/Animation/GameplayInsights/Source/GameplayInsights/Private/MontagesTrack.cpp
@@ -251,7 +251,7 @@ static FLinearColor MakeMontageCurveColor(uint32 InSeed, bool bInLine = false)
 
 TSharedPtr<SWidget> FMontageTrack::GetTimelineViewInternal()
 {
-	FLinearColor Color;
+	FLinearColor Color = FLinearColor::Black;
 	switch(CurveType)
 	{
 	case ECurveType::BlendWeight:
diff --git a/Engine/Plugins/Experimental/TextureGraph/Source/TextureGraph/Private/TG_GraphEvaluation.cpp b/Engine/Plugins/Experimental/TextureGraph/Source/TextureGraph/Private/TG_GraphEvaluation.cpp
--- a/Engine/Plugins/Experimental/TextureGraph/Source/TextureGraph/Private/TG_GraphEvaluation.cpp
+++ b/Engine/Plugins/Experimental/TextureGraph/Source/TextureGraph/Private/TG_GraphEvaluation.cpp
@@ -357,7 +357,7 @@ void FTG_VariantToFTG_Texture_Converter(FTG_Evaluation::VarConverterInfo& Info)
 
 	FTG_Variant::EType SourceType = Input.GetType();
 
-	FLinearColor Color;
+	FLinearColor Color = FLinearColor::Black;
 
 	switch (SourceType)
 	{
diff --git a/Engine/Plugins/Experimental/TextureGraph/Source/TextureGraph/Private/TG_Variant.cpp b/Engine/Plugins/Experimental/TextureGraph/Source/TextureGraph/Private/TG_Variant.cpp
--- a/Engine/Plugins/Experimental/TextureGraph/Source/TextureGraph/Private/TG_Variant.cpp
+++ b/Engine/Plugins/Experimental/TextureGraph/Source/TextureGraph/Private/TG_Variant.cpp
@@ -137,7 +137,7 @@ FTG_Texture FTG_Variant::GetTexture(FTG_EvaluationContext* InContext, FTG_Textur
 		return Texture;
 	}
 
-	FLinearColor Color;
+	FLinearColor Color = FLinearColor::Black;
 	FString Name;
 	BufferFormat TexelFormat = BufferFormat::Byte;
 
diff --git a/Engine/Plugins/FX/Niagara/Source/Niagara/Private/NiagaraDataChannel_Islands.cpp b/Engine/Plugins/FX/Niagara/Source/Niagara/Private/NiagaraDataChannel_Islands.cpp
--- a/Engine/Plugins/FX/Niagara/Source/Niagara/Private/NiagaraDataChannel_Islands.cpp
+++ b/Engine/Plugins/FX/Niagara/Source/Niagara/Private/NiagaraDataChannel_Islands.cpp
@@ -255,7 +255,6 @@ FNDCIsland* UNiagaraDataChannelHandler_Islands::FindOrCreateIsland(const FNiagar
 		for (int32 i : ActiveIslands)
 		{
 			FNDCIsland& Island = IslandPool[i];
-			FBoxSphereBounds GrowthBounds;
 			if (Island.TryGrow(Location, PerElementExtents, MaxExtents))
 			{
 				IslandToUse = &Island;
diff --git a/Engine/Plugins/MovieScene/SequencerScripting/Source/SequencerScripting/Private/ExtensionLibraries/MovieSceneFolderExtensions.cpp b/Engine/Plugins/MovieScene/SequencerScripting/Source/SequencerScripting/Private/ExtensionLibraries/MovieSceneFolderExtensions.cpp
--- a/Engine/Plugins/MovieScene/SequencerScripting/Source/SequencerScripting/Private/ExtensionLibraries/MovieSceneFolderExtensions.cpp
+++ b/Engine/Plugins/MovieScene/SequencerScripting/Source/SequencerScripting/Private/ExtensionLibraries/MovieSceneFolderExtensions.cpp
@@ -41,8 +41,9 @@ FColor UMovieSceneFolderExtensions::GetFolderColor(UMovieSceneFolder* Folder)
 
 #if WITH_EDITORONLY_DATA
 	return Folder->GetFolderColor();
-#endif //WITH_EDITORONLY_DATA
+#else
 	return FColor();
+#endif //WITH_EDITORONLY_DATA
 }
 
 bool UMovieSceneFolderExtensions::SetFolderColor(UMovieSceneFolder* Folder, FColor InFolderColor)
diff --git a/Engine/Plugins/MovieScene/SequencerScripting/Source/SequencerScripting/Private/ExtensionLibraries/MovieSceneTrackExtensions.cpp b/Engine/Plugins/MovieScene/SequencerScripting/Source/SequencerScripting/Private/ExtensionLibraries/MovieSceneTrackExtensions.cpp
--- a/Engine/Plugins/MovieScene/SequencerScripting/Source/SequencerScripting/Private/ExtensionLibraries/MovieSceneTrackExtensions.cpp
+++ b/Engine/Plugins/MovieScene/SequencerScripting/Source/SequencerScripting/Private/ExtensionLibraries/MovieSceneTrackExtensions.cpp
@@ -152,8 +152,9 @@ FColor UMovieSceneTrackExtensions::GetColorTint(UMovieSceneTrack* Track)
 
 #if WITH_EDITORONLY_DATA
 	return Track->GetColorTint(); 
-#endif
+#else
 	return FColor();
+#endif
 }
 
 void UMovieSceneTrackExtensions::SetColorTint(UMovieSceneTrack* Track, const FColor& ColorTint) 
diff --git a/Engine/Plugins/Runtime/Metasound/Source/MetasoundEditor/Public/MetasoundEditorGraphSchema.h b/Engine/Plugins/Runtime/Metasound/Source/MetasoundEditor/Public/MetasoundEditorGraphSchema.h
--- a/Engine/Plugins/Runtime/Metasound/Source/MetasoundEditor/Public/MetasoundEditorGraphSchema.h
+++ b/Engine/Plugins/Runtime/Metasound/Source/MetasoundEditor/Public/MetasoundEditorGraphSchema.h
@@ -82,7 +82,7 @@ struct METASOUNDEDITOR_API FMetasoundGraphSchemaAction : public FEdGraphSchemaAc
 
 	virtual const FLinearColor& GetIconColor() const
 	{
-		static const FLinearColor DefaultColor;
+		static const FLinearColor DefaultColor = FLinearColor::Black;
 		return DefaultColor;
 	}
 };
diff --git a/Engine/Source/Editor/BlueprintGraph/Private/EdGraphSchema_K2.cpp b/Engine/Source/Editor/BlueprintGraph/Private/EdGraphSchema_K2.cpp
--- a/Engine/Source/Editor/BlueprintGraph/Private/EdGraphSchema_K2.cpp
+++ b/Engine/Source/Editor/BlueprintGraph/Private/EdGraphSchema_K2.cpp
@@ -4447,7 +4447,6 @@ bool UEdGraphSchema_K2::DefaultValueSimpleValidation(const FEdGraphPinType& PinT
 			}
 			else if (StructType == RotatorStruct)
 			{
-				FRotator Rot;
 				if (!FDefaultValueHelper::IsStringValidRotator(NewDefaultValue))
 				{
 					DVSV_RETURN_MSG(TEXT("Invalid value for an FRotator"));
diff --git a/Engine/Source/Editor/DetailCustomizations/Private/ComponentTransformDetails.cpp b/Engine/Source/Editor/DetailCustomizations/Private/ComponentTransformDetails.cpp
--- a/Engine/Source/Editor/DetailCustomizations/Private/ComponentTransformDetails.cpp
+++ b/Engine/Source/Editor/DetailCustomizations/Private/ComponentTransformDetails.cpp
@@ -1071,7 +1071,7 @@ void FComponentTransformDetails::OnZScaleMirrored()
 void FComponentTransformDetails::CacheDetails()
 {
 	FVector CurLoc = FVector::ZeroVector;
-	FRotator CurRot;
+	FRotator CurRot = FRotator::ZeroRotator;
 	FVector CurScale = FVector::ZeroVector;
 	bIsEnabledCache = true;
 	for( int32 ObjectIndex = 0; ObjectIndex < SelectedObjects.Num(); ++ObjectIndex )
diff --git a/Engine/Source/Editor/GraphEditor/Private/MaterialNodes/SGraphSubstrateMaterial.cpp b/Engine/Source/Editor/GraphEditor/Private/MaterialNodes/SGraphSubstrateMaterial.cpp
--- a/Engine/Source/Editor/GraphEditor/Private/MaterialNodes/SGraphSubstrateMaterial.cpp
+++ b/Engine/Source/Editor/GraphEditor/Private/MaterialNodes/SGraphSubstrateMaterial.cpp
@@ -230,7 +230,7 @@ void FSubstrateWidget::GetPinColor(TSharedPtr<SGraphPin>& Out, const UMaterialGr
 	const FLinearColor Color0 = USlateThemeManager::Get().GetColor(GetSubstrateWidgetColor0());
 	const FLinearColor Color1 = USlateThemeManager::Get().GetColor(GetSubstrateWidgetColor1());
 
-	FLinearColor ColorModifier;
+	FLinearColor ColorModifier = FLinearColor::Black;
 	bool bHasColorModifier = false;
 	// Substrate operator override pin color to ease material topology visualization
 	const UEdGraphPin* Pin = Out->SGraphPin::GetPinObj();
diff --git a/Engine/Source/Editor/LandscapeEditor/Private/LandscapeEdModePaintTools.cpp b/Engine/Source/Editor/LandscapeEditor/Private/LandscapeEdModePaintTools.cpp
--- a/Engine/Source/Editor/LandscapeEditor/Private/LandscapeEdModePaintTools.cpp
+++ b/Engine/Source/Editor/LandscapeEditor/Private/LandscapeEdModePaintTools.cpp
@@ -518,7 +518,7 @@ public:
 			SculptStrength = FMath::Max(SculptStrength, 1.0f);
 		}
 
-		FPlane BrushPlane;
+		FPlane BrushPlane(ForceInit);
 		TArray<FVector> Normals;
 
 		if (bUseClayBrush)
diff --git a/Engine/Source/Editor/UnrealEd/Private/Fbx/FbxMainExport.cpp b/Engine/Source/Editor/UnrealEd/Private/Fbx/FbxMainExport.cpp
--- a/Engine/Source/Editor/UnrealEd/Private/Fbx/FbxMainExport.cpp
+++ b/Engine/Source/Editor/UnrealEd/Private/Fbx/FbxMainExport.cpp
@@ -1281,7 +1281,7 @@ void FFbxExporter::ExportLandscape(ALandscapeProxy* Actor, bool bSelectedOnly, I
 FbxDouble3 SetMaterialComponent(FColorMaterialInput& MatInput, bool ToLinear)
 {
 	FColor RGBColor;
-	FLinearColor LinearColor;
+	FLinearColor LinearColor = FLinearColor::Black;
 	bool LinearSet = false;
 	
 	if (MatInput.Expression)
diff --git a/Engine/Source/Runtime/Engine/Private/Actor.cpp b/Engine/Source/Runtime/Engine/Private/Actor.cpp
--- a/Engine/Source/Runtime/Engine/Private/Actor.cpp
+++ b/Engine/Source/Runtime/Engine/Private/Actor.cpp
@@ -1000,7 +1000,7 @@ void AActor::PostLoadSubobjects(FObjectInstancingGraph* OuterInstanceGraph)
 	USceneComponent* OldRoot = RootComponent;
 	USceneComponent* OldRootParent = (OldRoot ? OldRoot->GetAttachParent() : nullptr);
 	bool bHadRoot = !!OldRoot;
-	FRotator OldRotation;
+	FRotator OldRotation = FRotator::ZeroRotator;
 	FVector OldTranslation = FVector::ZeroVector;
 	FVector OldScale = FVector::ZeroVector;
 	if (bHadRoot)
diff --git a/Engine/Source/Runtime/Experimental/Chaos/Public/Chaos/PBDBendingConstraints.h b/Engine/Source/Runtime/Experimental/Chaos/Public/Chaos/PBDBendingConstraints.h
--- a/Engine/Source/Runtime/Experimental/Chaos/Public/Chaos/PBDBendingConstraints.h
+++ b/Engine/Source/Runtime/Experimental/Chaos/Public/Chaos/PBDBendingConstraints.h
@@ -136,15 +136,6 @@ public:
 		InitColor(InParticles);
 	}
 
-	UE_DEPRECATED(5.2, "Use one of the other constructors instead.")
-	CHAOS_API FPBDBendingConstraints(
-		const FSolverParticles& InParticles,
-		TArray<TVec4<int32>>&& InConstraints,
-		const FSolverReal InStiffness = (FSolverReal)1.)
-PRAGMA_DISABLE_DEPRECATION_WARNINGS
-		: Base(InParticles, MoveTemp(InConstraints), InStiffness) {}
-PRAGMA_ENABLE_DEPRECATION_WARNINGS
-
 	virtual ~FPBDBendingConstraints() override {}
 
 	using Base::SetProperties;
diff --git a/Engine/Source/Runtime/MovieSceneTracks/Private/Systems/MovieSceneMaterialParameterSystem.cpp b/Engine/Source/Runtime/MovieSceneTracks/Private/Systems/MovieSceneMaterialParameterSystem.cpp
--- a/Engine/Source/Runtime/MovieSceneTracks/Private/Systems/MovieSceneMaterialParameterSystem.cpp
+++ b/Engine/Source/Runtime/MovieSceneTracks/Private/Systems/MovieSceneMaterialParameterSystem.cpp
@@ -273,7 +273,7 @@ struct FVectorMixin
 
 		if (!bHasInitialValue)
 		{
-			FLinearColor ColorValue;
+			FLinearColor ColorValue = FLinearColor::Black;
 			if (UMaterialInstanceDynamic* MID = Cast<UMaterialInstanceDynamic>(BoundMaterial.GetObject()))
 			{
 				bHasInitialValue = MID->GetVectorParameterValue(ParameterInfo, ColorValue);
diff --git a/Engine/Source/Runtime/Renderer/Private/BasePassRendering.cpp b/Engine/Source/Runtime/Renderer/Private/BasePassRendering.cpp
--- a/Engine/Source/Runtime/Renderer/Private/BasePassRendering.cpp
+++ b/Engine/Source/Runtime/Renderer/Private/BasePassRendering.cpp
@@ -1064,7 +1064,7 @@ void FDeferredShadingSceneRenderer::RenderBasePass(
 	Substrate::AppendSubstrateMRTs(*this, BasePassTextureCount, BasePassTextures);
 	TArrayView<FTextureRenderTargetBinding> BasePassTexturesView = MakeArrayView(BasePassTextures.GetData(), BasePassTextureCount);
 	FRDGTextureRef BasePassDepthTexture = SceneTextures.Depth.Target;
-	FLinearColor SceneColorClearValue;
+	FLinearColor SceneColorClearValue = FLinearColor::Black;
 
 	if (bRequiresRHIClear)
 	{
diff --git a/Engine/Source/Runtime/Renderer/Private/RendererScene.cpp b/Engine/Source/Runtime/Renderer/Private/RendererScene.cpp
--- a/Engine/Source/Runtime/Renderer/Private/RendererScene.cpp
+++ b/Engine/Source/Runtime/Renderer/Private/RendererScene.cpp
@@ -5124,7 +5124,7 @@ void FScene::UpdateRayTracingGroupBounds_AddPrimitives(const Experimental::TRobi
 		if (GroupId != -1)
 		{
 			bool bInMap = false;
-			static const FRayTracingCullingGroup DefaultGroup;
+			static const FRayTracingCullingGroup DefaultGroup = {};
 			FRayTracingCullingGroup* const Group = PrimitiveRayTracingGroups.FindOrAdd(GroupId, DefaultGroup, bInMap);
 			if (bInMap)
 			{
diff --git a/Engine/Source/Runtime/SlateCore/Private/Rendering/ElementBatcher.cpp b/Engine/Source/Runtime/SlateCore/Private/Rendering/ElementBatcher.cpp
--- a/Engine/Source/Runtime/SlateCore/Private/Rendering/ElementBatcher.cpp
+++ b/Engine/Source/Runtime/SlateCore/Private/Rendering/ElementBatcher.cpp
@@ -792,7 +792,7 @@ void FSlateElementBatcher::AddBoxElements(const FSlateDrawElementArray<ElementTy
 		// Add Shader Parameters for extra RoundedBox parameters
 		ESlateShader ShaderType = ESlateShader::Default;
 		FShaderParams ShaderParams;
-		FColor SecondaryColor;
+		FColor SecondaryColor = FColor::Black;
 		if constexpr (std::is_same<ElementType, FSlateRoundedBoxElement>::value)
 		{
 			ShaderType = ESlateShader::RoundedBox;
diff --git a/Engine/Source/Runtime/SoundFieldRendering/Private/SoundFieldRendering.cpp b/Engine/Source/Runtime/SoundFieldRendering/Private/SoundFieldRendering.cpp
--- a/Engine/Source/Runtime/SoundFieldRendering/Private/SoundFieldRendering.cpp
+++ b/Engine/Source/Runtime/SoundFieldRendering/Private/SoundFieldRendering.cpp
@@ -491,7 +491,6 @@ void FSoundFieldDecoder::RotateFirstOrderAmbisonicsBed(const FAmbisonicsSoundfie
 	else
 	{
 		FMatrix CurrentRotationMatrix;
-		FMatrix RotationMatrixDelta;
 
 		// we already have our target matrix above.
 		// now we calculate our "current" matrix (where we left off)
