Fixed potentially uninitialized variables.

diff --git a/Engine/Plugins/2D/Paper2D/Source/Paper2D/Private/Terrain/PaperTerrainComponent.cpp b/Engine/Plugins/2D/Paper2D/Source/Paper2D/Private/Terrain/PaperTerrainComponent.cpp
--- a/Engine/Plugins/2D/Paper2D/Source/Paper2D/Private/Terrain/PaperTerrainComponent.cpp
+++ b/Engine/Plugins/2D/Paper2D/Source/Paper2D/Private/Terrain/PaperTerrainComponent.cpp
@@ -765,10 +765,10 @@ void UPaperTerrainComponent::SpawnSegments(const TArray<FTerrainSegment>& Terrai
 				BoxPoints[3] = LocalTransformAtCenter.TransformPosition(FVector(BoxExtents.X, 0, -BoxExtents.Y));
 
 				FVector2D BoxPoints2D[4];
-				BoxPoints2D[0].Set(BoxPoints[0].X, BoxPoints[0].Z);
-				BoxPoints2D[1].Set(BoxPoints[1].X, BoxPoints[1].Z);
-				BoxPoints2D[2].Set(BoxPoints[2].X, BoxPoints[2].Z);
-				BoxPoints2D[3].Set(BoxPoints[3].X, BoxPoints[3].Z);
+				BoxPoints2D[0] = FVector2D(BoxPoints[0].X, BoxPoints[0].Z);
+				BoxPoints2D[1] = FVector2D(BoxPoints[1].X, BoxPoints[1].Z);
+				BoxPoints2D[2] = FVector2D(BoxPoints[2].X, BoxPoints[2].Z);
+				BoxPoints2D[3] = FVector2D(BoxPoints[3].X, BoxPoints[3].Z);
 
 				// If there is a previous polygon, try to merge
 				if (CollisionPolygonPoints.Num() >= 4)
diff --git a/Engine/Plugins/Animation/ControlRig/Source/ControlRigEditor/Private/ControlRigSequencerEditorLibrary.cpp b/Engine/Plugins/Animation/ControlRig/Source/ControlRigEditor/Private/ControlRigSequencerEditorLibrary.cpp
--- a/Engine/Plugins/Animation/ControlRig/Source/ControlRigEditor/Private/ControlRigSequencerEditorLibrary.cpp
+++ b/Engine/Plugins/Animation/ControlRig/Source/ControlRigEditor/Private/ControlRigSequencerEditorLibrary.cpp
@@ -1786,7 +1786,7 @@ void UControlRigSequencerEditorLibrary::SetLocalControlRigInts(ULevelSequence* L
 
 FVector2D UControlRigSequencerEditorLibrary::GetLocalControlRigVector2D(ULevelSequence* LevelSequence, UControlRig* ControlRig, FName ControlName, FFrameNumber Frame, EMovieSceneTimeUnit TimeUnit)
 {
-	FVector2D Value;
+	FVector2D Value = FVector2D::ZeroVector;
 	TWeakPtr<ISequencer> WeakSequencer = GetSequencerFromAsset();
 
 	if (WeakSequencer.IsValid() && ControlRig)
@@ -1904,7 +1904,7 @@ void UControlRigSequencerEditorLibrary::SetLocalControlRigVector2Ds(ULevelSequen
 
 FVector UControlRigSequencerEditorLibrary::GetLocalControlRigPosition(ULevelSequence* LevelSequence, UControlRig* ControlRig, FName ControlName, FFrameNumber Frame, EMovieSceneTimeUnit TimeUnit)
 {
-	FVector Value;
+	FVector Value = FVector::ZeroVector;
 	TWeakPtr<ISequencer> WeakSequencer = GetSequencerFromAsset();
 
 	if (WeakSequencer.IsValid() && ControlRig)
@@ -2020,7 +2020,7 @@ void UControlRigSequencerEditorLibrary::SetLocalControlRigPositions(ULevelSequen
 
 FRotator UControlRigSequencerEditorLibrary::GetLocalControlRigRotator(ULevelSequence* LevelSequence, UControlRig* ControlRig, FName ControlName, FFrameNumber Frame, EMovieSceneTimeUnit TimeUnit)
 {
-	FRotator Value;
+	FRotator Value = FRotator::ZeroRotator;
 	TWeakPtr<ISequencer> WeakSequencer = GetSequencerFromAsset();
 
 	if (WeakSequencer.IsValid() && ControlRig)
@@ -2136,7 +2136,7 @@ void UControlRigSequencerEditorLibrary::SetLocalControlRigRotators(ULevelSequenc
 
 FVector UControlRigSequencerEditorLibrary::GetLocalControlRigScale(ULevelSequence* LevelSequence, UControlRig* ControlRig, FName ControlName, FFrameNumber Frame,  EMovieSceneTimeUnit TimeUnit)
 {
-	FVector Value;
+	FVector Value = FVector::ZeroVector;
 	TWeakPtr<ISequencer> WeakSequencer = GetSequencerFromAsset();
 
 	if (WeakSequencer.IsValid() && ControlRig)
diff --git a/Engine/Plugins/Editor/UVEditor/Source/UVEditorTools/Private/Operators/UVEditorUVTransformOp.cpp b/Engine/Plugins/Editor/UVEditor/Source/UVEditorTools/Private/Operators/UVEditorUVTransformOp.cpp
--- a/Engine/Plugins/Editor/UVEditor/Source/UVEditorTools/Private/Operators/UVEditorUVTransformOp.cpp
+++ b/Engine/Plugins/Editor/UVEditor/Source/UVEditorTools/Private/Operators/UVEditorUVTransformOp.cpp
@@ -267,7 +267,7 @@ void FUVEditorUVTransformBaseOp::CalculateResult(FProgressCancel* Progress)
 
 FVector2f FUVEditorUVTransformOp::GetPivotFromMode(int32 ElementID, EUVEditorPivotTypeBackend Mode)
 {
-	FVector2f Pivot;
+	FVector2f Pivot = FVector2f::ZeroVector;
 	const int32* Component;
 
 	switch (Mode)
diff --git a/Engine/Plugins/FX/Niagara/Source/NiagaraEditor/Private/ViewModels/NiagaraScriptInputCollectionViewModel.cpp b/Engine/Plugins/FX/Niagara/Source/NiagaraEditor/Private/ViewModels/NiagaraScriptInputCollectionViewModel.cpp
--- a/Engine/Plugins/FX/Niagara/Source/NiagaraEditor/Private/ViewModels/NiagaraScriptInputCollectionViewModel.cpp
+++ b/Engine/Plugins/FX/Niagara/Source/NiagaraEditor/Private/ViewModels/NiagaraScriptInputCollectionViewModel.cpp
@@ -85,7 +85,7 @@ FText FNiagaraScriptInputCollectionViewModel::GetDisplayName() const
 
 FVector2D GetNewNodeLocation(UNiagaraGraph* Graph, UNiagaraNode* NewInputNode, float VerticalNodeOffset, float HorizontalNodeOffset)
 {
-	FVector2D PlacementLocation;
+	FVector2D PlacementLocation = FVector2D::ZeroVector;
 	TArray<UNiagaraNodeInput*> InputNodes;
 	Graph->GetNodesOfClass(InputNodes);
 	if (InputNodes.Num() > 1)
diff --git a/Engine/Plugins/Interchange/Editor/Source/Pipelines/Private/InterchangeEditorPipelineDetails.cpp b/Engine/Plugins/Interchange/Editor/Source/Pipelines/Private/InterchangeEditorPipelineDetails.cpp
--- a/Engine/Plugins/Interchange/Editor/Source/Pipelines/Private/InterchangeEditorPipelineDetails.cpp
+++ b/Engine/Plugins/Interchange/Editor/Source/Pipelines/Private/InterchangeEditorPipelineDetails.cpp
@@ -1706,7 +1706,7 @@ void FInterchangeBaseNodeDetailsCustomization::BuildVectorValueContent(IDetailCa
 
 	auto GetValue = [](UInterchangeBaseNode* BaseNode, UE::Interchange::FAttributeKey& Key)->VectorType
 	{
-		VectorType VectorValue;
+		VectorType VectorValue = {};
 		const UE::Interchange::FAttributeStorage::TAttributeHandle<VectorType> AttributeHandle = BaseNode->GetAttributeHandle<VectorType>(Key);
 		if (AttributeHandle.IsValid())
 		{
diff --git a/Engine/Plugins/Interchange/Runtime/Source/Import/Private/MaterialX/MaterialXUtils/MaterialXBase.cpp b/Engine/Plugins/Interchange/Runtime/Source/Import/Private/MaterialX/MaterialXUtils/MaterialXBase.cpp
--- a/Engine/Plugins/Interchange/Runtime/Source/Import/Private/MaterialX/MaterialXUtils/MaterialXBase.cpp
+++ b/Engine/Plugins/Interchange/Runtime/Source/Import/Private/MaterialX/MaterialXUtils/MaterialXBase.cpp
@@ -243,7 +243,7 @@ MaterialX::InputPtr FMaterialXBase::GetInput(MaterialX::NodePtr Node, const char
 FLinearColor FMaterialXBase::GetLinearColor(MaterialX::InputPtr Input)
 {
 	//we assume that the default color space is linear
-	FLinearColor LinearColor;
+	FLinearColor LinearColor = FLinearColor::Black;
 
 	if(Input->getType() == mx::Type::Color3)
 	{
diff --git a/Engine/Plugins/Interchange/Runtime/Source/Import/Private/MaterialX/MaterialXUtils/MaterialXSurfaceShaderAbstract.cpp b/Engine/Plugins/Interchange/Runtime/Source/Import/Private/MaterialX/MaterialXUtils/MaterialXSurfaceShaderAbstract.cpp
--- a/Engine/Plugins/Interchange/Runtime/Source/Import/Private/MaterialX/MaterialXUtils/MaterialXSurfaceShaderAbstract.cpp
+++ b/Engine/Plugins/Interchange/Runtime/Source/Import/Private/MaterialX/MaterialXUtils/MaterialXSurfaceShaderAbstract.cpp
@@ -1354,7 +1354,7 @@ FString FMaterialXSurfaceShaderAbstract::GetFilePrefix(MaterialX::ElementPtr Ele
 
 FLinearColor FMaterialXSurfaceShaderAbstract::GetVector(MaterialX::InputPtr Input) const
 {
-	FLinearColor LinearColor;
+	FLinearColor LinearColor = FLinearColor::Black;
 
 	if(Input->getType() == mx::Type::Vector2)
 	{
diff --git a/Engine/Plugins/Runtime/Metasound/Source/MetasoundEngine/Private/Interfaces/MetasoundDeprecatedInterfaces.cpp b/Engine/Plugins/Runtime/Metasound/Source/MetasoundEngine/Private/Interfaces/MetasoundDeprecatedInterfaces.cpp
--- a/Engine/Plugins/Runtime/Metasound/Source/MetasoundEngine/Private/Interfaces/MetasoundDeprecatedInterfaces.cpp
+++ b/Engine/Plugins/Runtime/Metasound/Source/MetasoundEngine/Private/Interfaces/MetasoundDeprecatedInterfaces.cpp
@@ -490,7 +490,7 @@ namespace Metasound::Engine
 				FOutputHandle RightOutputToReconnect = IOutputController::GetInvalidHandle();
 
 #if WITH_EDITOR
-				FVector2D StereoFormatLocation;
+				FVector2D StereoFormatLocation = FVector2D::ZeroVector;
 #endif // WITH_EDITOR
 
 				if (StereoFormatOutput->IsValid())
diff --git a/Engine/Source/Editor/DetailCustomizations/Private/SlateBrushCustomization.cpp b/Engine/Source/Editor/DetailCustomizations/Private/SlateBrushCustomization.cpp
--- a/Engine/Source/Editor/DetailCustomizations/Private/SlateBrushCustomization.cpp
+++ b/Engine/Source/Editor/DetailCustomizations/Private/SlateBrushCustomization.cpp
@@ -1269,7 +1269,7 @@ private:
 			}
 
 			using ImageSizeType = decltype(FSlateBrush::ImageSize);
-			ImageSizeType CachedTextureSize;
+			ImageSizeType CachedTextureSize = {};
 
 			TArray<void*> RawData;
 			ImageSizeProperty->AccessRawData(RawData);
diff --git a/Engine/Source/Editor/GraphEditor/Private/ConnectionDrawingPolicy.cpp b/Engine/Source/Editor/GraphEditor/Private/ConnectionDrawingPolicy.cpp
--- a/Engine/Source/Editor/GraphEditor/Private/ConnectionDrawingPolicy.cpp
+++ b/Engine/Source/Editor/GraphEditor/Private/ConnectionDrawingPolicy.cpp
@@ -64,7 +64,7 @@ FVector2D FGeometryHelper::FindClosestPointOnGeom(const FGeometry& Geom, const F
 	FGeometryHelper::ConvertToPoints(Geom, Points);
 
 	float BestDistanceSquared = MAX_FLT;
-	FVector2D BestPoint;
+	FVector2D BestPoint = FVector2D::ZeroVector;
 	for (int32 i = 0; i < Points.Num(); ++i)
 	{
 		const FVector2D Candidate = FindClosestPointOnLine(Points[i], Points[(i + 1) % Points.Num()], TestPoint);
diff --git a/Engine/Source/Editor/UMGEditor/Private/Designer/SDesignerView.cpp b/Engine/Source/Editor/UMGEditor/Private/Designer/SDesignerView.cpp
--- a/Engine/Source/Editor/UMGEditor/Private/Designer/SDesignerView.cpp
+++ b/Engine/Source/Editor/UMGEditor/Private/Designer/SDesignerView.cpp
@@ -1407,7 +1407,7 @@ FVector2D SDesignerView::GetExtensionPosition(TSharedRef<FDesignerSurfaceElement
 			}
 			else
 			{
-				FVector2D WidgetPosition;
+				FVector2D WidgetPosition = FVector2D::ZeroVector;
 
 				// Get the initial offset based on the location around the selected object.
 				switch (ExtensionElement->GetLocation())
diff --git a/Engine/Source/Editor/UnrealEd/Private/Factories/EditorFactories.cpp b/Engine/Source/Editor/UnrealEd/Private/Factories/EditorFactories.cpp
--- a/Engine/Source/Editor/UnrealEd/Private/Factories/EditorFactories.cpp
+++ b/Engine/Source/Editor/UnrealEd/Private/Factories/EditorFactories.cpp
@@ -1648,7 +1648,7 @@ UObject* UPolysFactory::FactoryCreateText
 	FFeedbackContext*	Warn
 )
 {
-	FVector3f PointPool[4096];
+	FVector3f PointPool[4096] = {};
 	int32 NumPoints = 0;
 
 	GEditor->GetEditorSubsystem<UImportSubsystem>()->BroadcastAssetPreImport(this, Class, InParent, Name, Type);
diff --git a/Engine/Source/Editor/UnrealEd/Private/GeomFitUtils.cpp b/Engine/Source/Editor/UnrealEd/Private/GeomFitUtils.cpp
--- a/Engine/Source/Editor/UnrealEd/Private/GeomFitUtils.cpp
+++ b/Engine/Source/Editor/UnrealEd/Private/GeomFitUtils.cpp
@@ -155,8 +155,8 @@ static void CalcBoundingSphere(const FMeshDescription* MeshDescription, FSphere&
 		return;
 
 	FBox Box;
-	FVector MinIx[3];
-	FVector MaxIx[3];
+	FVector MinIx[3] = { FVector::ZeroVector, FVector::ZeroVector, FVector::ZeroVector };
+	FVector MaxIx[3] = { FVector::ZeroVector, FVector::ZeroVector, FVector::ZeroVector };
 
 	FStaticMeshConstAttributes Attributes(*MeshDescription);
 
diff --git a/Engine/Source/Editor/ViewportInteraction/Private/Gizmo/VIGizmoHandle.cpp b/Engine/Source/Editor/ViewportInteraction/Private/Gizmo/VIGizmoHandle.cpp
--- a/Engine/Source/Editor/ViewportInteraction/Private/Gizmo/VIGizmoHandle.cpp
+++ b/Engine/Source/Editor/ViewportInteraction/Private/Gizmo/VIGizmoHandle.cpp
@@ -83,7 +83,7 @@ FString UGizmoHandleGroup::MakeHandleName( const FTransformGizmoHandlePlacement
 
 FVector UGizmoHandleGroup::GetAxisVector( const int32 AxisIndex, const ETransformGizmoHandleDirection HandleDirection )
 {
-	FVector AxisVector;
+	FVector AxisVector = FVector::ZeroVector;
 
 	if (HandleDirection == ETransformGizmoHandleDirection::Center)
 	{
diff --git a/Engine/Source/Runtime/AppFramework/Private/Widgets/Colors/SColorGradingPicker.cpp b/Engine/Source/Runtime/AppFramework/Private/Widgets/Colors/SColorGradingPicker.cpp
--- a/Engine/Source/Runtime/AppFramework/Private/Widgets/Colors/SColorGradingPicker.cpp
+++ b/Engine/Source/Runtime/AppFramework/Private/Widgets/Colors/SColorGradingPicker.cpp
@@ -250,7 +250,7 @@ void SColorGradingPicker::TransformColorGradingRangeToLinearColorRange(float &Fl
 
 FLinearColor SColorGradingPicker::GetCurrentLinearColor()
 {
-	FLinearColor CurrentColor;
+	FLinearColor CurrentColor = FLinearColor::Black;
 	FVector4 CurrentValue;
 
 	if (OnQueryCurrentColor.IsBound())
diff --git a/Engine/Source/Runtime/ClientPilot/Private/ClientPilotBlackboard.cpp b/Engine/Source/Runtime/ClientPilot/Private/ClientPilotBlackboard.cpp
--- a/Engine/Source/Runtime/ClientPilot/Private/ClientPilotBlackboard.cpp
+++ b/Engine/Source/Runtime/ClientPilot/Private/ClientPilotBlackboard.cpp
@@ -21,7 +21,7 @@ FString UClientPilotBlackboard::GetStringValue(FString KeyName)
 
 FVector UClientPilotBlackboard::GetVectorValue(FString KeyName)
 {
-	FVector RetVal;
+	FVector RetVal = FVector::ZeroVector;
 	if (Blackboard.Contains(KeyName))
 	{
 		RetVal.InitFromString(*Blackboard[KeyName]);
diff --git a/Engine/Source/Runtime/Engine/Private/Actor.cpp b/Engine/Source/Runtime/Engine/Private/Actor.cpp
--- a/Engine/Source/Runtime/Engine/Private/Actor.cpp
+++ b/Engine/Source/Runtime/Engine/Private/Actor.cpp
@@ -4395,7 +4395,7 @@ float AActor::GetInputAxisKeyValue(const FKey InputAxisKey) const
 
 FVector AActor::GetInputVectorAxisValue(const FKey InputAxisKey) const
 {
-	FVector Value;
+	FVector Value = FVector::ZeroVector;
 
 	if (InputComponent)
 	{
diff --git a/Engine/Source/Runtime/Engine/Private/Components/InputComponent.cpp b/Engine/Source/Runtime/Engine/Private/Components/InputComponent.cpp
--- a/Engine/Source/Runtime/Engine/Private/Components/InputComponent.cpp
+++ b/Engine/Source/Runtime/Engine/Private/Components/InputComponent.cpp
@@ -201,7 +201,7 @@ float UInputComponent::GetAxisKeyValue( const FKey AxisKey ) const
 
 FVector UInputComponent::GetVectorAxisValue( const FKey AxisKey ) const
 {
-	FVector AxisValue;
+	FVector AxisValue = FVector::ZeroVector;
 	bool bFound = false;
 
 	for (const FInputVectorAxisBinding& AxisBinding : VectorAxisBindings)
diff --git a/Engine/Source/Runtime/Engine/Private/LocalPlayer.cpp b/Engine/Source/Runtime/Engine/Private/LocalPlayer.cpp
--- a/Engine/Source/Runtime/Engine/Private/LocalPlayer.cpp
+++ b/Engine/Source/Runtime/Engine/Private/LocalPlayer.cpp
@@ -610,7 +610,7 @@ private:
 	 */
 	static FViewPoint GetViewPointFromStrings(const FString* Strings, int32 NumStrings)
 	{
-		FViewPoint ViewPoint;
+		FViewPoint ViewPoint = {};
 		if (NumStrings == 6)
 		{
 			ViewPoint.Location.X = FCString::Atof(*Strings[0]);
diff --git a/Engine/Source/Runtime/Engine/Private/StaticMeshLight.cpp b/Engine/Source/Runtime/Engine/Private/StaticMeshLight.cpp
--- a/Engine/Source/Runtime/Engine/Private/StaticMeshLight.cpp
+++ b/Engine/Source/Runtime/Engine/Private/StaticMeshLight.cpp
@@ -182,8 +182,8 @@ FLightRayIntersection FStaticMeshStaticLightingMesh::IntersectLightRay(const FVe
 	}
 	else
 	{
-		IntersectionVertex.WorldPosition.Set(0,0,0);
-		IntersectionVertex.WorldTangentZ.Set(0,0,1);
+		IntersectionVertex.WorldPosition = FVector(0,0,0);
+		IntersectionVertex.WorldTangentZ = FVector(0,0,1);
 	}
 	return FLightRayIntersection(bIntersects,IntersectionVertex);
 }
diff --git a/Engine/Source/Runtime/GameplayDebugger/Private/GameplayDebuggerPlayerManager.cpp b/Engine/Source/Runtime/GameplayDebugger/Private/GameplayDebuggerPlayerManager.cpp
--- a/Engine/Source/Runtime/GameplayDebugger/Private/GameplayDebuggerPlayerManager.cpp
+++ b/Engine/Source/Runtime/GameplayDebugger/Private/GameplayDebuggerPlayerManager.cpp
@@ -286,7 +286,7 @@ void AGameplayDebuggerPlayerManager::GetViewPoint(const APlayerController& Owner
 {
 	UWorld* World = OwnerPC.GetWorld();
 	FVector CameraLocation = FVector::ZeroVector;
-	FRotator CameraRotation;
+	FRotator CameraRotation = FRotator::ZeroRotator;
 	if (OwnerPC.Player)
 	{
 		// normal game
diff --git a/Engine/Source/Runtime/GeometryCore/Private/Generators/StairGenerator.cpp b/Engine/Source/Runtime/GeometryCore/Private/Generators/StairGenerator.cpp
--- a/Engine/Source/Runtime/GeometryCore/Private/Generators/StairGenerator.cpp
+++ b/Engine/Source/Runtime/GeometryCore/Private/Generators/StairGenerator.cpp
@@ -1076,7 +1076,7 @@ FVector3f FCurvedStairGenerator::GenerateNormal(ESide Side, int VertexId)
 	const int Col = VertexIdsToColumnRow[VertexId].A;
 	const int Row = VertexIdsToColumnRow[VertexId].B;
 
-	FVector3f N;
+	FVector3f N = FVector3f::ZeroVector;
 	switch (Side)
 	{
 	case ESide::Right:
diff --git a/Engine/Source/Runtime/GeometryCore/Public/Sampling/SphericalFibonacci.h b/Engine/Source/Runtime/GeometryCore/Public/Sampling/SphericalFibonacci.h
--- a/Engine/Source/Runtime/GeometryCore/Public/Sampling/SphericalFibonacci.h
+++ b/Engine/Source/Runtime/GeometryCore/Public/Sampling/SphericalFibonacci.h
@@ -247,7 +247,7 @@ public:
 	{
 		checkSlow(Index >= 0 && Index < N);
 
-		TVector<RealType> Point;
+		TVector<RealType> Point = {};
 		switch (Distribution)
 		{
 		case EDistribution::Uniform:
diff --git a/Engine/Source/Runtime/Landscape/Private/LandscapeLight.cpp b/Engine/Source/Runtime/Landscape/Private/LandscapeLight.cpp
--- a/Engine/Source/Runtime/Landscape/Private/LandscapeLight.cpp
+++ b/Engine/Source/Runtime/Landscape/Private/LandscapeLight.cpp
@@ -662,7 +662,7 @@ FLightRayIntersection FLandscapeStaticLightingMesh::IntersectLightRay(const FVec
 	const bool bIntersects = LandscapeComponent->LineTraceComponent( Result, Start, End, NewTraceParams );
 
 	// Setup a vertex to represent the intersection.
-	FStaticLightingVertex IntersectionVertex;
+	FStaticLightingVertex IntersectionVertex = {};
 	if(bIntersects)
 	{
 		IntersectionVertex.WorldPosition = Result.Location;
