Sequencer: Anim Details: Rotators not working with anim details and selection not working with actors.
#prefglight 6627382a37485b66bbab6916

diff --git a/Engine/Plugins/Animation/ControlRig/Source/ControlRigEditor/Private/EditMode/AnimDetailsProxy.cpp b/Engine/Plugins/Animation/ControlRig/Source/ControlRigEditor/Private/EditMode/AnimDetailsProxy.cpp
--- a/Engine/Plugins/Animation/ControlRig/Source/ControlRigEditor/Private/EditMode/AnimDetailsProxy.cpp
+++ b/Engine/Plugins/Animation/ControlRig/Source/ControlRigEditor/Private/EditMode/AnimDetailsProxy.cpp
@@ -54,7 +54,6 @@
 
 #include UE_INLINE_GENERATED_CPP_BY_NAME(AnimDetailsProxy)
 
-
 static void KeyTrack(TSharedPtr<ISequencer>& Sequencer, UControlRigControlsProxy* Proxy, UMovieScenePropertyTrack* Track, EControlRigContextChannelToKey ChannelToKey)
 {
 	using namespace UE::Sequencer;
@@ -1714,7 +1713,10 @@ void UAnimDetailControlsProxyRotation::SetControlRigElementValueFromCurrent(UCon
 	{
 		FVector3f Val = Rotation.ToVector3f();
 		SetRotationValuesFromContext(ControlRig, ControlElement, Context, Val);
+		FVector EulerAngle(Rotation.ToRotator().Roll, Rotation.ToRotator().Pitch, Rotation.ToRotator().Yaw);
+		ControlRig->GetHierarchy()->SetControlSpecifiedEulerAngle(ControlElement, EulerAngle);
 		ControlRig->SetControlValue<FVector3f>(ControlElement->GetKey().Name, Val, true, Context, false);
+		ControlRig->GetHierarchy()->SetControlSpecifiedEulerAngle(ControlElement, EulerAngle);
 		ControlRig->Evaluate_AnyThread();
 	}
 }
@@ -3693,13 +3695,14 @@ bool UControlRigDetailPanelControlProxies::SelectPropertyInternal(UControlRigCon
 							{
 								if (UMovieSceneControlRigParameterTrack* Track = Cast<UMovieSceneControlRigParameterTrack>(TrackModel->GetTrack()))
 								{
-									if (Track->GetControlRig() != ControlRig)
+									if (Track->GetControlRig() != ControlRig || Track->GetAllSections().Num() < 1)
 									{
 										continue;
 									}
+									UMovieSceneSection* SectionToKey = Track->GetSectionToKey() ? Track->GetSectionToKey() : Track->GetAllSections()[0];
 									if (TViewModelPtr<FChannelGroupOutlinerModel> ChannelModel = CastViewModel<FChannelGroupOutlinerModel>(OutlinerExtenstionIt.GetCurrentItem()))
 									{
-										if (ChannelModel->GetChannel(Track->GetSectionToKey()) == nullptr) //if not section to key we also don't select it.
+										if (ChannelModel->GetChannel(SectionToKey) == nullptr) //if not section to key we also don't select it.
 										{
 											continue;
 										}
@@ -3751,13 +3754,22 @@ bool UControlRigDetailPanelControlProxies::SelectPropertyInternal(UControlRigCon
 						{
 							if (TrackModel->GetTrack() == Element.WeakTrack.Get())
 							{
+								if (TrackModel->GetTrack()->GetAllSections().Num() < 1)
+								{
+									continue;
+								}
+								UMovieSceneSection* SectionToKey = TrackModel->GetTrack()->GetSectionToKey() ? TrackModel->GetTrack()->GetSectionToKey() : TrackModel->GetTrack()->GetAllSections()[0];
 								if (TViewModelPtr<FChannelGroupOutlinerModel> ChannelModel = CastViewModel<FChannelGroupOutlinerModel>(OutlinerExtenstionIt.GetCurrentItem()))
 								{
-									if (ChannelModel->GetChannel(TrackModel->GetTrack()->GetSectionToKey()) == nullptr) //if not section to key we also don't select it.
+									if (ChannelModel->GetChannel(SectionToKey) == nullptr) //if not section to key we also don't select it.
 									{
 										continue;
 									}
 								}
+								else
+								{
+									continue;
+								}
 								FName ID = OutlinerExtenstionIt->GetIdentifier();
 								FString Name = ID.ToString();
 								TArray<FString> StringArray;
@@ -3768,7 +3780,7 @@ bool UControlRigDetailPanelControlProxies::SelectPropertyInternal(UControlRigCon
 								{
 									ChannelName = StringArray[0] + "." + StringArray[1];
 								}
-								else if (StringArray.Num() == 0)
+								else if (StringArray.Num() == 1)
 								{
 									ChannelName = StringArray[0];
 								}
@@ -3846,4 +3858,4 @@ void FControlRigEnumControlProxyValueDetails::OnEnumValueChanged(int32 InValue,
 		ProxyBeingCustomized->Enum.EnumIndex = InValue;
 		InStructHandle->NotifyPostChange(EPropertyChangeType::ValueSet);
 	}
-}
\ No newline at end of file
+}
