Sequencer: Control Rig: When collapsing, recording animation keys use preferred rotators.
#prefglight 662829342267e5ebbc3d169f

diff --git a/Engine/Plugins/Animation/ControlRig/Source/ControlRig/Private/Sequencer/MovieSceneControlRigParameterSection.cpp b/Engine/Plugins/Animation/ControlRig/Source/ControlRig/Private/Sequencer/MovieSceneControlRigParameterSection.cpp
--- a/Engine/Plugins/Animation/ControlRig/Source/ControlRig/Private/Sequencer/MovieSceneControlRigParameterSection.cpp
+++ b/Engine/Plugins/Animation/ControlRig/Source/ControlRig/Private/Sequencer/MovieSceneControlRigParameterSection.cpp
@@ -2978,7 +2978,8 @@ void UMovieSceneControlRigParameterSection::RecordControlRigKey(FFrameNumber Fra
 				case ERigControlType::Scale:
 				case ERigControlType::Rotator:
 				{
-					FVector3f Val = ControlRig->GetControlValue(ControlElement, ERigControlValueType::Current).Get<FVector3f>();
+					FVector3f Val = (ControlElement->Settings.ControlType == ERigControlType::Rotator)
+						? FVector3f(ControlRig->GetHierarchy()->GetControlSpecifiedEulerAngle(ControlElement)): ControlRig->GetControlValue(ControlElement, ERigControlValueType::Current).Get<FVector3f>();
 					if (ControlElement->Settings.ControlType == ERigControlType::Rotator &&
 						FloatChannels[ChannelIndex]->GetNumKeys() > 0)
 					{
@@ -3031,8 +3032,7 @@ void UMovieSceneControlRigParameterSection::RecordControlRigKey(FFrameNumber Fra
 					}
 
 					AddVectorKeyToFloatChannels(ChannelIndex, FrameNumber, CurrentVector);
-
-					CurrentVector = Val.GetRotation().Euler();
+					CurrentVector = ControlRig->GetHierarchy()->GetControlSpecifiedEulerAngle(ControlElement);
 					if (FloatChannels[ChannelIndex]->GetNumKeys() > 0)
 					{
 						float LastVal = FloatChannels[ChannelIndex]->GetValues()[FloatChannels[ChannelIndex]->GetNumKeys() - 1].Value;
diff --git a/Engine/Plugins/Animation/ControlRig/Source/ControlRigEditor/Private/Sequencer/ControlRigParameterTrackEditor.cpp b/Engine/Plugins/Animation/ControlRig/Source/ControlRigEditor/Private/Sequencer/ControlRigParameterTrackEditor.cpp
--- a/Engine/Plugins/Animation/ControlRig/Source/ControlRigEditor/Private/Sequencer/ControlRigParameterTrackEditor.cpp
+++ b/Engine/Plugins/Animation/ControlRig/Source/ControlRigEditor/Private/Sequencer/ControlRigParameterTrackEditor.cpp
@@ -4780,7 +4780,7 @@ bool FControlRigParameterTrackEditor::CollapseAllLayers(TSharedPtr<ISequencer>&S
 					}
 
 					//Store transforms
-					TArray<TPair<FName, TArray<FTransform>>> ControlLocalTransforms;
+					TArray<TPair<FName, TArray<FEulerTransform>>> ControlLocalTransforms;
 					TArray<FRigControlElement*> Controls;
 					ControlRig->GetControlsInOrder(Controls);
 
@@ -4790,7 +4790,7 @@ bool FControlRigParameterTrackEditor::CollapseAllLayers(TSharedPtr<ISequencer>&S
 						{
 							continue;
 						}
-						TPair<FName, TArray<FTransform>> NameTransforms;
+						TPair<FName, TArray<FEulerTransform>> NameTransforms;
 						NameTransforms.Key = ControlElement->GetFName();
 						NameTransforms.Value.SetNum(Frames.Num());
 						ControlLocalTransforms.Add(NameTransforms);
@@ -4808,9 +4808,12 @@ bool FControlRigParameterTrackEditor::CollapseAllLayers(TSharedPtr<ISequencer>&S
 
 						SequencerPtr->GetEvaluationTemplate().EvaluateSynchronousBlocking(Context);
 						ControlRig->Evaluate_AnyThread();
-						for (TPair<FName, TArray<FTransform>>& TrailControlTransform : ControlLocalTransforms)
+						for (TPair<FName, TArray<FEulerTransform>>& TrailControlTransform : ControlLocalTransforms)
 						{
-							TrailControlTransform.Value[Index] = ControlRig->GetControlLocalTransform(TrailControlTransform.Key);
+							FEulerTransform EulerTransform(ControlRig->GetControlLocalTransform(TrailControlTransform.Key));
+							FRigElementKey ControlKey(TrailControlTransform.Key, ERigElementType::Control);
+							EulerTransform.Rotation = ControlRig->GetHierarchy()->GetControlPreferredRotator(ControlKey);
+							TrailControlTransform.Value[Index] = EulerTransform;
 						}
 					}
 					//delete other sections
@@ -4842,9 +4845,14 @@ bool FControlRigParameterTrackEditor::CollapseAllLayers(TSharedPtr<ISequencer>&S
 						//need to do the twice hack since controls aren't really in order
 						for (int32 TwiceHack = 0; TwiceHack < 2; ++TwiceHack)
 						{
-							for (TPair<FName, TArray<FTransform>>& TrailControlTransform : ControlLocalTransforms)
+							for (TPair<FName, TArray<FEulerTransform>>& TrailControlTransform : ControlLocalTransforms)
 							{
-								ControlRig->SetControlLocalTransform(TrailControlTransform.Key, TrailControlTransform.Value[Index], false, Context, false /*undo*/, true/* bFixEulerFlips*/);
+								FRigElementKey ControlKey(TrailControlTransform.Key, ERigElementType::Control);
+								ControlRig->GetHierarchy()->SetControlPreferredRotator(ControlKey, TrailControlTransform.Value[Index].Rotation);
+								FTransform Transform(TrailControlTransform.Value[Index].ToFTransform());
+								ControlRig->SetControlLocalTransform(TrailControlTransform.Key, Transform, false, Context, false /*undo*/, true/* bFixEulerFlips*/);
+								ControlRig->GetHierarchy()->SetControlPreferredRotator(ControlKey, TrailControlTransform.Value[Index].Rotation);
+
 							}
 						}
 						ControlRig->Evaluate_AnyThread();
