Constraints: fix UTransformableControlHandle not updating its delegates when updating the ControlRig it acts on

diff --git a/Engine/Plugins/Animation/ControlRig/Source/ControlRig/Private/Constraints/ControlRigTransformableHandle.cpp b/Engine/Plugins/Animation/ControlRig/Source/ControlRig/Private/Constraints/ControlRigTransformableHandle.cpp
--- a/Engine/Plugins/Animation/ControlRig/Source/ControlRig/Private/Constraints/ControlRigTransformableHandle.cpp
+++ b/Engine/Plugins/Animation/ControlRig/Source/ControlRig/Private/Constraints/ControlRigTransformableHandle.cpp
@@ -578,7 +578,9 @@ void UTransformableControlHandle::ResolveBoundObjects(FMovieSceneSequenceID Loca
 				const UObject* Bindable = FControlRigObjectBinding::GetBindableObject(ParentObject.Get());
 				if (InControlRig->GetObjectBinding() && InControlRig->GetObjectBinding()->GetBoundObject() == Bindable)
 				{
+					UnregisterDelegates();
 					ControlRig = InControlRig;
+					RegisterDelegates();
 				}
 				break; //just do one
 			}
@@ -620,17 +622,9 @@ void UTransformableControlHandle::OnObjectsReplaced(const TMap<UObject*, UObject
 	{
 		if (UControlRig* NewControlRig = Cast<UControlRig>(NewObject))
 		{
-			if (URigHierarchy* Hierarchy = ControlRig->GetHierarchy())
-			{
-				Hierarchy->OnModified().RemoveAll(this);
-			}
-			
+			UnregisterDelegates();
 			ControlRig = NewControlRig;
-			
-			if (URigHierarchy* Hierarchy = ControlRig->GetHierarchy())
-			{
-				Hierarchy->OnModified().AddUObject(this, &UTransformableControlHandle::OnHierarchyModified);
-			}
+			RegisterDelegates();
 		}
 	}
 }
