Control Rig: Fix crash on close of SControlRigEditModeTools

#rnx

diff --git a/Engine/Plugins/Animation/ControlRig/Source/ControlRigEditor/Private/EditMode/SControlRigEditModeTools.cpp b/Engine/Plugins/Animation/ControlRig/Source/ControlRigEditor/Private/EditMode/SControlRigEditModeTools.cpp
--- a/Engine/Plugins/Animation/ControlRig/Source/ControlRigEditor/Private/EditMode/SControlRigEditModeTools.cpp
+++ b/Engine/Plugins/Animation/ControlRig/Source/ControlRigEditor/Private/EditMode/SControlRigEditModeTools.cpp
@@ -50,13 +50,16 @@ FRigSpacePickerBakeSettings SControlRigEditModeTools::BakeSpaceSettings;
 void SControlRigEditModeTools::Cleanup()
 {
 	// This is required as these hold a shared pointer to THIS OBJECT and make this class not to be destroyed when the parent class releases the shared pointer of this object
-	if (SettingsDetailsView)
-	{
-		SettingsDetailsView->SetKeyframeHandler(nullptr);
-	}
-	if (RigOptionsDetailsView)
+	if(FSlateApplication::IsInitialized())
 	{
-		RigOptionsDetailsView->SetKeyframeHandler(nullptr);
+		if (SettingsDetailsView)
+		{
+			SettingsDetailsView->SetKeyframeHandler(nullptr);
+		}
+		if (RigOptionsDetailsView)
+		{
+			RigOptionsDetailsView->SetKeyframeHandler(nullptr);
+		}
 	}
 	for (TPair<FDelegateHandle, TWeakObjectPtr<UControlRig>>& Handles : HandlesToClear)
 	{
