Sequencer: Fix crash eith Edit in Sequencer do to sequencer getting out of scope
#trivial

diff --git a/Engine/Plugins/Animation/ControlRig/Source/ControlRigEditor/Private/ControlRigEditorModule.cpp b/Engine/Plugins/Animation/ControlRig/Source/ControlRigEditor/Private/ControlRigEditorModule.cpp
--- a/Engine/Plugins/Animation/ControlRig/Source/ControlRigEditor/Private/ControlRigEditorModule.cpp
+++ b/Engine/Plugins/Animation/ControlRig/Source/ControlRigEditor/Private/ControlRigEditorModule.cpp
@@ -1140,8 +1140,6 @@ void FControlRigEditorModule::BakeToControlRig(UClass* ControlRigClass, UAnimSeq
 		ILevelSequenceEditorToolkit* LevelSequenceEditor = static_cast<ILevelSequenceEditorToolkit*>(AssetEditor);
 		TWeakPtr<ISequencer> WeakSequencer = LevelSequenceEditor ? LevelSequenceEditor->GetSequencer() : nullptr;
 
-
-
 		if (WeakSequencer.IsValid())
 		{
 			TSharedPtr<ISequencer> SequencerPtr = WeakSequencer.Pin();
@@ -1226,62 +1224,69 @@ void FControlRigEditorModule::BakeToControlRig(UClass* ControlRigClass, UAnimSeq
 				Track->SetTrackName(FName(*ObjectName));
 				Track->SetDisplayName(FText::FromString(ObjectName));
 				UMovieSceneControlRigParameterSection* ParamSection = Cast<UMovieSceneControlRigParameterSection>(NewSection);
-				FBakeToControlDelegate BakeCallback = FBakeToControlDelegate::CreateLambda([this, &SequencerPtr, LevelSequence,
+				FBakeToControlDelegate BakeCallback = FBakeToControlDelegate::CreateLambda([this, LevelSequence,
 					AnimSequence, MovieScene, ControlRig, ParamSection,ActorTrackGuid, SkelMeshComp]
 				(bool bKeyReduce, float KeyReduceTolerance, FFrameRate BakeFrameRate, bool bResetControls)
 				{
-					if (ParamSection)
-					{
-						FSmartReduceParams SmartReduce;
-						SmartReduce.TolerancePercentage = KeyReduceTolerance;
-						SmartReduce.SampleRate = BakeFrameRate;
-						FControlRigParameterTrackEditor::LoadAnimationIntoSection(SequencerPtr, AnimSequence, SkelMeshComp, FFrameNumber(0),
-							bKeyReduce, SmartReduce,bResetControls, ParamSection);
-					}
-					SequencerPtr->EmptySelection();
-					SequencerPtr->SelectSection(ParamSection);
-					SequencerPtr->ThrobSectionSelection();
-					SequencerPtr->ObjectImplicitlyAdded(ControlRig);
-					SequencerPtr->NotifyMovieSceneDataChanged(EMovieSceneDataChangeType::MovieSceneStructureItemAdded);
-					FControlRigEditMode* ControlRigEditMode = static_cast<FControlRigEditMode*>(GLevelEditorModeTools().GetActiveMode(FControlRigEditMode::ModeName));
-					if (!ControlRigEditMode)
-					{
-						GLevelEditorModeTools().ActivateMode(FControlRigEditMode::ModeName);
-						ControlRigEditMode = static_cast<FControlRigEditMode*>(GLevelEditorModeTools().GetActiveMode(FControlRigEditMode::ModeName));
-					}
-					if (ControlRigEditMode)
+					IAssetEditorInstance* AssetEditor = GEditor->GetEditorSubsystem<UAssetEditorSubsystem>()->FindEditorForAsset(LevelSequence, false);
+					ILevelSequenceEditorToolkit* LevelSequenceEditor = static_cast<ILevelSequenceEditorToolkit*>(AssetEditor);
+					TWeakPtr<ISequencer> WeakSequencer = LevelSequenceEditor ? LevelSequenceEditor->GetSequencer() : nullptr;
+					if (WeakSequencer.IsValid())
 					{
-						ControlRigEditMode->AddControlRigObject(ControlRig, SequencerPtr);
-					}
+						TSharedPtr<ISequencer> SequencerPtr = WeakSequencer.Pin();
+						if (ParamSection)
+						{
+							FSmartReduceParams SmartReduce;
+							SmartReduce.TolerancePercentage = KeyReduceTolerance;
+							SmartReduce.SampleRate = BakeFrameRate;
+							FControlRigParameterTrackEditor::LoadAnimationIntoSection(SequencerPtr, AnimSequence, SkelMeshComp, FFrameNumber(0),
+								bKeyReduce, SmartReduce, bResetControls, ParamSection);
+						}
+						SequencerPtr->EmptySelection();
+						SequencerPtr->SelectSection(ParamSection);
+						SequencerPtr->ThrobSectionSelection();
+						SequencerPtr->ObjectImplicitlyAdded(ControlRig);
+						SequencerPtr->NotifyMovieSceneDataChanged(EMovieSceneDataChangeType::MovieSceneStructureItemAdded);
+						FControlRigEditMode* ControlRigEditMode = static_cast<FControlRigEditMode*>(GLevelEditorModeTools().GetActiveMode(FControlRigEditMode::ModeName));
+						if (!ControlRigEditMode)
+						{
+							GLevelEditorModeTools().ActivateMode(FControlRigEditMode::ModeName);
+							ControlRigEditMode = static_cast<FControlRigEditMode*>(GLevelEditorModeTools().GetActiveMode(FControlRigEditMode::ModeName));
+						}
+						if (ControlRigEditMode)
+						{
+							ControlRigEditMode->AddControlRigObject(ControlRig, SequencerPtr);
+						}
 
-					//create soft links to each other
-					if (IInterface_AssetUserData* AssetUserDataInterface = Cast< IInterface_AssetUserData >(LevelSequence))
-					{
-						ULevelSequenceAnimSequenceLink* LevelAnimLink = NewObject<ULevelSequenceAnimSequenceLink>(LevelSequence, NAME_None, RF_Public | RF_Transactional);
-						FLevelSequenceAnimSequenceLinkItem LevelAnimLinkItem;
-						LevelAnimLinkItem.SkelTrackGuid = ActorTrackGuid;
-						LevelAnimLinkItem.PathToAnimSequence = FSoftObjectPath(AnimSequence);
-						LevelAnimLinkItem.bExportMorphTargets = true; 
-						LevelAnimLinkItem.bExportAttributeCurves = true;
-						LevelAnimLinkItem.Interpolation = EAnimInterpolationType::Linear;
-						LevelAnimLinkItem.CurveInterpolation = ERichCurveInterpMode::RCIM_Linear;
-						LevelAnimLinkItem.bExportMaterialCurves = true;
-						LevelAnimLinkItem.bExportTransforms = true;
-						LevelAnimLinkItem.bRecordInWorldSpace = false;
-						LevelAnimLinkItem.bEvaluateAllSkeletalMeshComponents = true;
-						LevelAnimLink->AnimSequenceLinks.Add(LevelAnimLinkItem);
-						AssetUserDataInterface->AddAssetUserData(LevelAnimLink);
-					}
-					if (IInterface_AssetUserData* AnimAssetUserData = Cast< IInterface_AssetUserData >(AnimSequence))
-					{
-						UAnimSequenceLevelSequenceLink* AnimLevelLink = AnimAssetUserData->GetAssetUserData< UAnimSequenceLevelSequenceLink >();
-						if (!AnimLevelLink)
+						//create soft links to each other
+						if (IInterface_AssetUserData* AssetUserDataInterface = Cast< IInterface_AssetUserData >(LevelSequence))
+						{
+							ULevelSequenceAnimSequenceLink* LevelAnimLink = NewObject<ULevelSequenceAnimSequenceLink>(LevelSequence, NAME_None, RF_Public | RF_Transactional);
+							FLevelSequenceAnimSequenceLinkItem LevelAnimLinkItem;
+							LevelAnimLinkItem.SkelTrackGuid = ActorTrackGuid;
+							LevelAnimLinkItem.PathToAnimSequence = FSoftObjectPath(AnimSequence);
+							LevelAnimLinkItem.bExportMorphTargets = true;
+							LevelAnimLinkItem.bExportAttributeCurves = true;
+							LevelAnimLinkItem.Interpolation = EAnimInterpolationType::Linear;
+							LevelAnimLinkItem.CurveInterpolation = ERichCurveInterpMode::RCIM_Linear;
+							LevelAnimLinkItem.bExportMaterialCurves = true;
+							LevelAnimLinkItem.bExportTransforms = true;
+							LevelAnimLinkItem.bRecordInWorldSpace = false;
+							LevelAnimLinkItem.bEvaluateAllSkeletalMeshComponents = true;
+							LevelAnimLink->AnimSequenceLinks.Add(LevelAnimLinkItem);
+							AssetUserDataInterface->AddAssetUserData(LevelAnimLink);
+						}
+						if (IInterface_AssetUserData* AnimAssetUserData = Cast< IInterface_AssetUserData >(AnimSequence))
 						{
-							AnimLevelLink = NewObject<UAnimSequenceLevelSequenceLink>(AnimSequence, NAME_None, RF_Public | RF_Transactional);
-							AnimAssetUserData->AddAssetUserData(AnimLevelLink);
+							UAnimSequenceLevelSequenceLink* AnimLevelLink = AnimAssetUserData->GetAssetUserData< UAnimSequenceLevelSequenceLink >();
+							if (!AnimLevelLink)
+							{
+								AnimLevelLink = NewObject<UAnimSequenceLevelSequenceLink>(AnimSequence, NAME_None, RF_Public | RF_Transactional);
+								AnimAssetUserData->AddAssetUserData(AnimLevelLink);
+							}
+							AnimLevelLink->SetLevelSequence(LevelSequence);
+							AnimLevelLink->SkelTrackGuid = ActorTrackGuid;
 						}
-						AnimLevelLink->SetLevelSequence(LevelSequence);
-						AnimLevelLink->SkelTrackGuid = ActorTrackGuid;
 					}
 				});
 
