Sequencer: Group keying properties in anim details to one undo.
#trivial

diff --git a/Engine/Plugins/Animation/ControlRig/Source/ControlRigEditor/Private/EditMode/ControlRigEditMode.cpp b/Engine/Plugins/Animation/ControlRig/Source/ControlRigEditor/Private/EditMode/ControlRigEditMode.cpp
--- a/Engine/Plugins/Animation/ControlRig/Source/ControlRigEditor/Private/EditMode/ControlRigEditMode.cpp
+++ b/Engine/Plugins/Animation/ControlRig/Source/ControlRigEditor/Private/EditMode/ControlRigEditMode.cpp
@@ -5656,6 +5656,7 @@ void FDetailKeyFrameCacheAndHandler::OnKeyPropertyClicked(const IPropertyHandle&
 	{
 		return;
 	}
+	FScopedTransaction ScopedTransaction(LOCTEXT("KeyAttribute", "Key Attribute"), !GIsTransacting);
 	TSharedPtr<ISequencer> SequencerPtr = WeakSequencer.Pin();
 
 	TArray<UObject*> Objects;
