Sequencer:  Make sure to tick the mode manager so pending deactivated modes are really deleted before enabling a new control rig edit mode.

diff --git a/Engine/Plugins/Animation/ControlRig/Source/ControlRigEditor/Private/Sequencer/ControlRigParameterTrackEditor.cpp b/Engine/Plugins/Animation/ControlRig/Source/ControlRigEditor/Private/Sequencer/ControlRigParameterTrackEditor.cpp
--- a/Engine/Plugins/Animation/ControlRig/Source/ControlRigEditor/Private/Sequencer/ControlRigParameterTrackEditor.cpp
+++ b/Engine/Plugins/Animation/ControlRig/Source/ControlRigEditor/Private/Sequencer/ControlRigParameterTrackEditor.cpp
@@ -2106,11 +2106,22 @@ void FControlRigParameterTrackEditor::OnActivateSequenceChanged(FMovieSceneSeque
 	{
 		GEditor->GetTimerManager()->SetTimerForNextTick([this]()
 		{
-			//true here will turn it on
-			FControlRigEditMode* ControlRigEditMode = GetEditMode(true);
-			if (ControlRigEditMode)
+			//we need to make sure pending deactivated edit modes, including a possible control rig edit mode
+			//get totally removed which only happens on a tick 
+			if (GEditor->GetActiveViewport() && GEditor->GetActiveViewport()->GetClient())
 			{
-				GEditor->GetTimerManager()->SetTimerForNextTick([this, ControlRigEditMode]()
+				if (FEditorModeTools* EditorModeTools = GetEditorModeTools())
+				{
+					FViewport* ActiveViewport = GEditor->GetActiveViewport();
+					FEditorViewportClient* EditorViewClient = (FEditorViewportClient*)ActiveViewport->GetClient();
+					EditorModeTools->Tick(EditorViewClient, 0.033f);
+				}
+			}
+			GEditor->GetTimerManager()->SetTimerForNextTick([this]()
+			{
+				//now we can recreate it
+				FControlRigEditMode* ControlRigEditMode = GetEditMode(true);
+				if (ControlRigEditMode)
 				{
 					bool bSequencerSet = false;
 					for (TWeakObjectPtr<UControlRig>& ControlRig : BoundControlRigs)
@@ -2137,9 +2148,10 @@ void FControlRigParameterTrackEditor::OnActivateSequenceChanged(FMovieSceneSeque
 					{
 						ControlRigEditMode->SetSequencer(GetSequencer());
 					}
-					PreviousSelectedControlRigs.Reset();
-				});
-			}
+				
+				}
+				PreviousSelectedControlRigs.Reset();
+			});
 		});
 	}
 }
