Control Rig: Add metadata limits for float and integer properties

diff --git a/Engine/Plugins/Animation/ControlRig/Source/ControlRig/Private/Sequencer/MovieSceneControlRigParameterSection.cpp b/Engine/Plugins/Animation/ControlRig/Source/ControlRig/Private/Sequencer/MovieSceneControlRigParameterSection.cpp
--- a/Engine/Plugins/Animation/ControlRig/Source/ControlRig/Private/Sequencer/MovieSceneControlRigParameterSection.cpp
+++ b/Engine/Plugins/Animation/ControlRig/Source/ControlRig/Private/Sequencer/MovieSceneControlRigParameterSection.cpp
@@ -1559,6 +1559,9 @@ void UMovieSceneControlRigParameterSection::ReconstructChannelProxy()
 
 EMovieSceneChannelProxyType UMovieSceneControlRigParameterSection::CacheChannelProxy()
 {
+	const FName UIMin("UIMin");
+	const FName UIMax("UIMax");
+
 	FMovieSceneChannelProxyData Channels;
 	ControlChannelMap.Empty();
 	// Need to create the channels in sorted orders
@@ -1711,6 +1714,8 @@ EMovieSceneChannelProxyType UMovieSceneControlRigParameterSection::CacheChannelP
 
 							FParameterFloatChannelEditorData EditorData(ControlRig, Scalar.ParameterName, bEnabled, Group, TotalIndex);
 							EditorData.MetaData.DisplayText = Hierarchy->GetDisplayNameForUI(ControlElement);
+							EditorData.MetaData.PropertyMetaData.Add(UIMin, FString::SanitizeFloat(ControlElement->Settings.MinimumValue.Get<float>()));
+							EditorData.MetaData.PropertyMetaData.Add(UIMax, FString::SanitizeFloat(ControlElement->Settings.MaximumValue.Get<float>()));
 							Channels.Add(Scalar.ParameterCurve, EditorData.MetaData, EditorData.ExternalValues);
 							FloatChannelIndex += 1;
 							TotalIndex += 1;
@@ -1820,6 +1825,8 @@ EMovieSceneChannelProxyType UMovieSceneControlRigParameterSection::CacheChannelP
 								MetaData.SortOrder = TotalIndex++;
 								// Prevent single channels from collapsing to the track node
 								MetaData.bCanCollapseToTrack = false;
+								MetaData.PropertyMetaData.Add(UIMin, FString::FromInt(ControlElement->Settings.MinimumValue.Get<int32>()));
+								MetaData.PropertyMetaData.Add(UIMax, FString::FromInt(ControlElement->Settings.MaximumValue.Get<int32>()));
 								Channels.Add(Integer.ParameterCurve, MetaData, TMovieSceneExternalValue<int32>());
 								break;
 							}
