Sequencer: Add bake settings to load animation into section for control rig so you can do proper key reduce and keep controls

diff --git a/Engine/Plugins/Animation/ControlRig/Source/ControlRigEditor/Private/Sequencer/ControlRigParameterTrackEditor.cpp b/Engine/Plugins/Animation/ControlRig/Source/ControlRigEditor/Private/Sequencer/ControlRigParameterTrackEditor.cpp
--- a/Engine/Plugins/Animation/ControlRig/Source/ControlRigEditor/Private/Sequencer/ControlRigParameterTrackEditor.cpp
+++ b/Engine/Plugins/Animation/ControlRig/Source/ControlRigEditor/Private/Sequencer/ControlRigParameterTrackEditor.cpp
@@ -5390,17 +5390,30 @@ void FControlRigParameterSection::OnAnimationAssetSelectedForFK(const FAssetData
 
 		if (AnimSequence && SkelMeshComp && AnimSequence->GetDataModel()->GetNumBoneTracks() > 0)
 		{
+			FBakeToControlDelegate BakeCallback = FBakeToControlDelegate::CreateLambda([this, Section,
+				AnimSequence, SkelMeshComp]
+				(bool bKeyReduce, float KeyReduceTolerance, FFrameRate BakeFrameRate, bool bResetControls)
+				{
+					if (TSharedPtr<ISequencer> SequencerPtr = WeakSequencer.Pin())
+					{
+						FScopedTransaction Transaction(LOCTEXT("BakeAnimation_Transaction", "Bake Animation To FK Control Rig"));
+						Section->Modify();
+						FFrameNumber StartFrame = SequencerPtr->GetLocalTime().Time.GetFrame();
+						FSmartReduceParams SmartReduce;
+						SmartReduce.TolerancePercentage = KeyReduceTolerance;
+						SmartReduce.SampleRate = BakeFrameRate;
+						if (!FControlRigParameterTrackEditor::LoadAnimationIntoSection(SequencerPtr, AnimSequence, SkelMeshComp, StartFrame,
+							bKeyReduce, SmartReduce, bResetControls, Section))
+						{
+							Transaction.Cancel();
+						}
+						SequencerPtr->NotifyMovieSceneDataChanged(EMovieSceneDataChangeType::MovieSceneStructureItemAdded);
+					}
+					
+				});
 
-			FScopedTransaction Transaction(LOCTEXT("BakeAnimation_Transaction", "Bake Animation To FK Control Rig"));
-			Section->Modify();
-			FFrameNumber StartFrame = SequencerPtr->GetLocalTime().Time.GetFrame();
-			FSmartReduceParams SmartReduce;
-			if(!FControlRigParameterTrackEditor::LoadAnimationIntoSection(SequencerPtr, AnimSequence, SkelMeshComp,StartFrame,
-				false, SmartReduce, true, Section))
-			{
-				Transaction.Cancel();
-			}
-			SequencerPtr->NotifyMovieSceneDataChanged(EMovieSceneDataChangeType::MovieSceneStructureItemAdded);
+			FOnWindowClosed BakeClosedCallback = FOnWindowClosed::CreateLambda([](const TSharedRef<SWindow>&) {});
+			BakeToControlRigDialog::GetBakeParams(BakeCallback, BakeClosedCallback);
 		}
 	}
 }
