Control Rig: Fix crash when mirroring modules
#rnx

diff --git a/Engine/Plugins/Animation/ControlRig/Source/ControlRig/Private/ModularRigController.cpp b/Engine/Plugins/Animation/ControlRig/Source/ControlRig/Private/ModularRigController.cpp
--- a/Engine/Plugins/Animation/ControlRig/Source/ControlRig/Private/ModularRigController.cpp
+++ b/Engine/Plugins/Animation/ControlRig/Source/ControlRig/Private/ModularRigController.cpp
@@ -1325,16 +1325,22 @@ FString UModularRigController::MirrorModule(const FString& InModulePath, const F
 		NewModuleName = GetSafeNewName(OriginalModule->ParentPath, FRigName(NewModuleName)).ToString();
 	}
 
+	// Before any changes, gather all the information we need from the OriginalModule, as the pointer might become invalid afterwards
+	const TMap<FRigElementKey, FRigElementKey> OriginalConnectionMap = Model->Connections.GetModuleConnectionMap(InModulePath);
+	const TMap<FName, FString> OriginalBindings = OriginalModule->Bindings;
+	const TSubclassOf<UControlRig> OriginalClass = OriginalModule->Class.Get();
+	const FString OriginalParentPath = OriginalModule->ParentPath;
+	const TMap<FName, FString> OriginalConfigValues = OriginalModule->ConfigValues;
+
 	FModularRigControllerCompileBracketScope CompileBracketScope(this);
 
-	FString NewModulePath = AddModule(*NewModuleName, OriginalModule->Class.Get(), OriginalModule->ParentPath, bSetupUndo);
+	FString NewModulePath = AddModule(*NewModuleName, OriginalClass, OriginalParentPath, bSetupUndo);
 	FRigModuleReference* NewModule = FindModule(NewModulePath);
 	if (!NewModule)
 	{
 		return FString();
 	}
 
-	const TMap<FRigElementKey, FRigElementKey> OriginalConnectionMap = Model->Connections.GetModuleConnectionMap(InModulePath);
 	for (const TPair<FRigElementKey, FRigElementKey>& Pair : OriginalConnectionMap)
 	{
 		FString OriginalTargetPath = Pair.Value.Name.ToString();
@@ -1349,7 +1355,7 @@ FString UModularRigController::MirrorModule(const FString& InModulePath, const F
 		NewModulePath = NewModule->GetPath();
 	}
 
-	for (const TPair<FName, FString>& Pair : OriginalModule->Bindings)
+	for (const TPair<FName, FString>& Pair : OriginalBindings)
 	{
 		FString NewSourcePath = Pair.Value.Replace(*InSettings.SearchString, *InSettings.ReplaceString, ESearchCase::CaseSensitive);
 		BindModuleVariable(NewModulePath, Pair.Key, NewSourcePath, bSetupUndo);
@@ -1357,7 +1363,7 @@ FString UModularRigController::MirrorModule(const FString& InModulePath, const F
 
 	TSet<FName> ConfigValueSet;
 #if WITH_EDITOR
-	for (TFieldIterator<FProperty> PropertyIt(OriginalModule->Class.Get()); PropertyIt; ++PropertyIt)
+	for (TFieldIterator<FProperty> PropertyIt(OriginalClass); PropertyIt; ++PropertyIt)
 	{
 		const FProperty* Property = *PropertyIt;
 		
@@ -1391,7 +1397,7 @@ FString UModularRigController::MirrorModule(const FString& InModulePath, const F
 		}
 
 		FString NewValueStr;
-		if (const FString* OriginalValue = OriginalModule->ConfigValues.Find(Property->GetFName()))
+		if (const FString* OriginalValue = OriginalConfigValues.Find(Property->GetFName()))
 		{
 			if (bIsVector)
 			{
@@ -1410,7 +1416,7 @@ FString UModularRigController::MirrorModule(const FString& InModulePath, const F
 		}
 		else
 		{
-			if (UControlRig* CDO = OriginalModule->Class->GetDefaultObject<UControlRig>())
+			if (UControlRig* CDO = OriginalClass->GetDefaultObject<UControlRig>())
 			{
 				if (bIsVector)
 				{
@@ -1433,7 +1439,7 @@ FString UModularRigController::MirrorModule(const FString& InModulePath, const F
 #endif
 
 	// Add any other config value that was set in the original module, but was not mirrored
-	for (const TPair<FName, FString>& Pair : OriginalModule->ConfigValues)
+	for (const TPair<FName, FString>& Pair : OriginalConfigValues)
 	{
 		if (!ConfigValueSet.Contains(Pair.Key))
 		{
