Constraints & Layered Control Rig: changes to handle UControlRig::OnControlModified being called when evaluating additive rigs (instead of when setting values)

    - ControlRigEditMode: skip compensation and evaluate the rig to force notifications: auto-key and constraints updates (among others) are based on OnControlModified broadcast
    - ConstraintsEvaluationGraph: do not mark constraints for evaluation while flushing (this can happen with OnControlModified being called while evaluating the rig)
    - TransformableControlHandle: guard from re-entrant notification (this can happen with OnControlModified being called while evaluating the rig)
    - CompensationEvaluator: evaluate the rig to force notifications after sequencer has been evaluated

diff --git a/Engine/Plugins/Animation/ControlRig/Source/ControlRig/Private/Constraints/ControlRigTransformableHandle.cpp b/Engine/Plugins/Animation/ControlRig/Source/ControlRig/Private/Constraints/ControlRigTransformableHandle.cpp
--- a/Engine/Plugins/Animation/ControlRig/Source/ControlRig/Private/Constraints/ControlRigTransformableHandle.cpp
+++ b/Engine/Plugins/Animation/ControlRig/Source/ControlRig/Private/Constraints/ControlRigTransformableHandle.cpp
@@ -17,6 +17,18 @@
 
 #include UE_INLINE_GENERATED_CPP_BY_NAME(ControlRigTransformableHandle)
 
+namespace ControlHandleLocals
+{
+	using RigGuard = TGuardValue_Bitfield_Cleanup<TFunction<void()>>;
+	
+	TSet<UControlRig*> NotifyingRigs;
+
+	bool IsRigNotifying(const UControlRig* InControlRig)
+	{
+		return InControlRig ? NotifyingRigs.Contains(InControlRig) : false;
+	}
+}
+
 /**
  * UTransformableControlHandle
  */
@@ -63,11 +75,38 @@ void UTransformableControlHandle::PreEvaluate(const bool bTick) const
 	{
 		return;
 	}
+
+	if (ControlRig->IsAdditive())
+	{
+		if (ControlHandleLocals::IsRigNotifying(ControlRig.Get()))
+		{
+			return;
+		}
+		
+		if (const USkeletalMeshComponent* SkeletalMeshComponent = GetSkeletalMesh())
+		{
+			if (!SkeletalMeshComponent->PoseTickedThisFrame())
+			{
+				return TickTarget();
+			}
+		}
+	}
+	
 	return bTick ? TickTarget() : ControlRig->Evaluate_AnyThread();
 }
 
 void UTransformableControlHandle::TickTarget() const
 {
+	if (!ControlRig.IsValid())
+	{
+		return;
+	}
+	
+	if (ControlRig->IsAdditive() && ControlHandleLocals::IsRigNotifying(ControlRig.Get()))
+	{
+		return;
+	}
+	
 	if (const USkeletalMeshComponent* SkeletalMeshComponent = GetSkeletalMesh())
 	{
 		return TransformableHandleUtils::TickDependantComponents(SkeletalMeshComponent);
@@ -98,7 +137,7 @@ void UTransformableControlHandle::SetGlobalTransform(const FTransform& InGlobal)
 	const FRigElementKey& ControlKey = ControlElement->GetKey();
 	const FTransform& ComponentTransform = BoundComponent->GetComponentTransform();
 
-	static const FRigControlModifiedContext Context;
+	static const FRigControlModifiedContext Context(EControlRigSetKey::Never);
 	static constexpr bool bNotify = false, bSetupUndo = false, bPrintPython = false, bFixEulerFlips = false;
 
 	//use this function so we don't set the preferred angles
@@ -152,6 +191,11 @@ FTransform UTransformableControlHandle::GetLocalTransform() const
 	{
 		return FTransform::Identity;
 	}
+
+	if (ControlRig->IsAdditive())
+	{
+		return ControlRig->GetControlLocalTransform(ControlName);
+	}
 	
 	const FRigElementKey& ControlKey = ControlElement->GetKey();
 	const URigHierarchy* Hierarchy = ControlRig->GetHierarchy();
@@ -438,6 +482,14 @@ void UTransformableControlHandle::OnControlModified(
 			{
 				GetEvaluationBinding().bPendingFlush = true;
 			}
+			
+			// guard from re-entrant notification
+			const ControlHandleLocals::RigGuard NotificationGuard([ControlRig = ControlRig.Get()]()
+			{
+				ControlHandleLocals::NotifyingRigs.Remove(ControlRig);
+			});
+			ControlHandleLocals::NotifyingRigs.Add(ControlRig.Get());
+			
 			Notify(Event);
 		}
 		else if (Event == EHandleEvent::GlobalTransformUpdated)
@@ -449,7 +501,15 @@ void UTransformableControlHandle::OnControlModified(
 				{
 					GetEvaluationBinding().bPendingFlush = true;
 				}
-				static constexpr  bool bPreTick = true;
+
+				// guard from re-entrant notification 
+				const ControlHandleLocals::RigGuard NotificationGuard([ControlRig = ControlRig.Get()]()
+				{
+					ControlHandleLocals::NotifyingRigs.Remove(ControlRig);
+				});
+				ControlHandleLocals::NotifyingRigs.Add(ControlRig.Get());
+				
+				const bool bPreTick = !ControlRig->IsAdditive();
 				Notify(EHandleEvent::UpperDependencyUpdated, bPreTick);
 			}
 		}
diff --git a/Engine/Plugins/Animation/ControlRig/Source/ControlRigEditor/Private/EditMode/ControlRigEditMode.cpp b/Engine/Plugins/Animation/ControlRig/Source/ControlRigEditor/Private/EditMode/ControlRigEditMode.cpp
--- a/Engine/Plugins/Animation/ControlRig/Source/ControlRigEditor/Private/EditMode/ControlRigEditMode.cpp
+++ b/Engine/Plugins/Animation/ControlRig/Source/ControlRigEditor/Private/EditMode/ControlRigEditMode.cpp
@@ -4608,6 +4608,18 @@ void FControlRigEditMode::SetControlShapeTransform(
 			InShapeActor->ControlName, InGlobalTransform, bNotify, InContext, bUndo, bPrintPython, bFixEulerFlips);
 		return;
 	}
+
+	auto EvaluateRigIfAdditive = [ControlRig]()
+	{
+		// skip compensation and evaluate the rig to force notifications: auto-key and constraints updates (among others) are based on
+		// UControlRig::OnControlModified being broadcast but this only happens on evaluation for additive rigs.
+		// constraint compensation is disabled while manipulating in that case to avoid re-entrant evaluations 
+		if (ControlRig->IsAdditive())
+		{
+			TGuardValue<bool> CompensateGuard(FMovieSceneConstraintChannelHelper::bDoNotCompensate, true);
+			ControlRig->Evaluate_AnyThread();
+		}
+	};
 	
 	// find the last constraint in the stack (this could be cached on mouse press)
 	TArray< TWeakObjectPtr<UTickableConstraint> > Constraints;
@@ -4622,6 +4634,7 @@ void FControlRigEditMode::SetControlShapeTransform(
 		TGuardValue<bool> CompensateGuard(FMovieSceneConstraintChannelHelper::bDoNotCompensate, true);
 		ControlRig->SetControlGlobalTransform(
 			InShapeActor->ControlName, InGlobalTransform, bNotify, InContext, bUndo, bPrintPython, bFixEulerFlips);
+		EvaluateRigIfAdditive();
 	}
 
 	if (!bNeedsConstraintPostProcess)
@@ -4644,7 +4657,8 @@ void FControlRigEditMode::SetControlShapeTransform(
 	Context.bConstraintUpdate = false;
 	
 	ControlRig->SetControlLocalTransform(InShapeActor->ControlName, LocalTransform, bNotify, Context, bUndo, bFixEulerFlips);
-
+	EvaluateRigIfAdditive();
+	
 	const FConstraintsManagerController& Controller = FConstraintsManagerController::Get(ControlRig->GetWorld());
 	Controller.EvaluateAllConstraints();
 }
diff --git a/Engine/Source/Editor/MovieSceneTools/Private/Constraints/MovieSceneConstraintChannelHelper.cpp b/Engine/Source/Editor/MovieSceneTools/Private/Constraints/MovieSceneConstraintChannelHelper.cpp
--- a/Engine/Source/Editor/MovieSceneTools/Private/Constraints/MovieSceneConstraintChannelHelper.cpp
+++ b/Engine/Source/Editor/MovieSceneTools/Private/Constraints/MovieSceneConstraintChannelHelper.cpp
@@ -479,7 +479,7 @@ void FCompensationEvaluator::ComputeCompensation(UWorld* InWorld, const TSharedP
 		}
 	}
 	
-	auto EvaluateAt = [InSequencer, &AllConstraints, &BakeHelpers](FFrameNumber InFrame)
+	auto EvaluateAt = [Handle = this->Handle, InSequencer, &AllConstraints, &BakeHelpers](FFrameNumber InFrame)
 	{
 		FMovieSceneSequenceTransform RootToLocalTransform = InSequencer->GetFocusedMovieSceneSequenceTransform();
 		InFrame = (FFrameTime(InFrame) * RootToLocalTransform.InverseNoLooping()).GetFrame();
@@ -516,7 +516,10 @@ void FCompensationEvaluator::ComputeCompensation(UWorld* InWorld, const TSharedP
 			}
 		}
 
-		// ControlRig->Evaluate_AnyThread();
+		if (Handle)
+		{
+			Handle->PreEvaluate();
+		}
 	};
 
 
@@ -557,8 +560,10 @@ void FCompensationEvaluator::ComputeCompensation(UWorld* InWorld, const TSharedP
 	{
 		const FTransform ChildLocal = ChildLocals[0];
 		Handle->SetGlobalTransform(ChildGlobals[0]);
+		Handle->PreEvaluate();
 		ChildLocals[0] = Handle->GetLocalTransform();
 		Handle->SetLocalTransform(ChildLocal);
+		Handle->PreEvaluate();
 	}
 }
 
diff --git a/Engine/Source/Runtime/Experimental/Animation/Constraints/Private/ConstraintsEvaluationGraph.cpp b/Engine/Source/Runtime/Experimental/Animation/Constraints/Private/ConstraintsEvaluationGraph.cpp
--- a/Engine/Source/Runtime/Experimental/Animation/Constraints/Private/ConstraintsEvaluationGraph.cpp
+++ b/Engine/Source/Runtime/Experimental/Animation/Constraints/Private/ConstraintsEvaluationGraph.cpp
@@ -256,6 +256,13 @@ void FConstraintsEvaluationGraph::MarkForEvaluation(const TWeakObjectPtr<UTickab
 	{
 		Rebuild();
 	}
+
+	if (State == Flushing)
+	{
+		// do not mark this constraint for evaluation while flushing.
+		// this can happen with UControlRig::OnControlModified being called while evaluating additive rigs
+		return;
+	}
 	
 	if (FConstraintNode* Node = FindNode(InConstraint))
 	{
