Sequencer: Constraints: Make sure to evaluate all constraints when baking

diff --git a/Engine/Plugins/Animation/ControlRig/Source/ControlRigEditor/Private/Tools/ControlRigSnapper.cpp b/Engine/Plugins/Animation/ControlRig/Source/ControlRigEditor/Private/Tools/ControlRigSnapper.cpp
--- a/Engine/Plugins/Animation/ControlRig/Source/ControlRigEditor/Private/Tools/ControlRigSnapper.cpp
+++ b/Engine/Plugins/Animation/ControlRig/Source/ControlRigEditor/Private/Tools/ControlRigSnapper.cpp
@@ -128,7 +128,8 @@ static bool LocalGetControlRigControlTransforms(IMovieScenePlayer* Player, const
 	}
 	if (UMovieScene* MovieScene = MovieSceneSequence->GetMovieScene())
 	{
-
+		UWorld* World = ControlRig->GetWorld();
+		const FConstraintsManagerController& Controller = FConstraintsManagerController::Get(World);
 		FFrameRate TickResolution = MovieScene->GetTickResolution();
 		FFrameRate DisplayRate = MovieScene->GetDisplayRate();
 		
@@ -146,6 +147,7 @@ static bool LocalGetControlRigControlTransforms(IMovieScenePlayer* Player, const
 
 				DeltaTime = 1.0/Context.GetFrameRate().AsDecimal();
 				Player->GetEvaluationTemplate().EvaluateSynchronousBlocking(Context);
+				Controller.EvaluateAllConstraints();
 			}
 			if (ControlRig->IsAdditive())
 			{
