Control Rig: FIx infinite loop when attempting to autoresolve connections
#rnx

diff --git a/Engine/Plugins/Animation/ControlRig/Source/ControlRig/Private/ModularRigController.cpp b/Engine/Plugins/Animation/ControlRig/Source/ControlRig/Private/ModularRigController.cpp
--- a/Engine/Plugins/Animation/ControlRig/Source/ControlRig/Private/ModularRigController.cpp
+++ b/Engine/Plugins/Animation/ControlRig/Source/ControlRig/Private/ModularRigController.cpp
@@ -587,7 +587,7 @@ bool UModularRigController::AutoConnectSecondaryConnectors(const TArray<FRigElem
 		const UModularRigRuleManager* RuleManager = Hierarchy->GetRuleManager();
 		const FRigModuleInstance* ModuleInstance = ModularRig->FindModule(Module->GetPath());
 		
-		if (!Model->Connections.HasConnection(ConnectorKey) || bReplaceExistingConnections)
+		if (bReplaceExistingConnections || !Model->Connections.HasConnection(ConnectorKey))
 		{
 			if (const FRigConnectorElement* OtherConnectorElement = Cast<FRigConnectorElement>(Hierarchy->Find(ConnectorKey)))
 			{
diff --git a/Engine/Plugins/Animation/ControlRig/Source/ControlRigEditor/Private/Editor/ControlRigEditor.cpp b/Engine/Plugins/Animation/ControlRig/Source/ControlRigEditor/Private/Editor/ControlRigEditor.cpp
--- a/Engine/Plugins/Animation/ControlRig/Source/ControlRigEditor/Private/Editor/ControlRigEditor.cpp
+++ b/Engine/Plugins/Animation/ControlRig/Source/ControlRigEditor/Private/Editor/ControlRigEditor.cpp
@@ -420,6 +420,7 @@ void FControlRigEditor::CreatePersonaToolKitIfRequired()
 	PersonaToolkit->GetPreviewScene()->RegisterOnPreviewMeshChanged(FOnPreviewMeshChanged::CreateSP(this, &FControlRigEditor::HandlePreviewMeshChanged));
 	
 	// Set a default preview mesh, if any
+	TGuardValue<bool> AutoResolveGuard(ControlRigBlueprint->ModularRigSettings.bAutoResolve, false);
 	PersonaToolkit->SetPreviewMesh(ControlRigBlueprint->GetPreviewMesh(), false);
 }
 
@@ -2884,11 +2885,13 @@ void FControlRigEditor::HandlePreviewMeshChanged(USkeletalMesh* InOldSkeletalMes
 						
 						// try to reestablish the connections.
 						UModularRigController* ModularRigController = ControlRigBP->GetModularRigController();
+						const bool bAutoResolve = ControlRigBP->ModularRigSettings.bAutoResolve;
 						Model->ForEachModule(
-							[Model, Hierarchy, ModularRigController, PreviousConnections]
+							[Model, Hierarchy, ModularRigController, PreviousConnections, bAutoResolve]
 							(const FRigModuleReference* Module) -> bool
 							{
 								bool bContinueResolval;
+								TArray<uint32> AttemptedTargets;
 								do
 								{
 									bContinueResolval = false;
@@ -2933,7 +2936,7 @@ void FControlRigEditor::HandlePreviewMeshChanged(USkeletalMesh* InOldSkeletalMes
 											}
 
 											// try to auto resolve it
-											if(!bContinueResolval && bIsSecondary)
+											if(!bContinueResolval && bIsSecondary && bAutoResolve)
 											{
 												if(ModularRigController->AutoConnectSecondaryConnectors({ConnectorKey}, true, true))
 												{
@@ -2945,6 +2948,26 @@ void FControlRigEditor::HandlePreviewMeshChanged(USkeletalMesh* InOldSkeletalMes
 											break;
 										}
 									}
+
+									// Avoid looping forever
+									if (bContinueResolval)
+									{
+										uint32 Attempt = 0;
+										for(const FRigElementKey& ConnectorKey : ConnectorKeys)
+										{
+											const FString ConnectionStr = FString::Printf(TEXT("%s -> %s"), *ConnectorKey.ToString(), *Model->Connections.FindTargetFromConnector(ConnectorKey).ToString());
+											const uint32 ConnectionHash = HashCombine(GetTypeHash(ConnectorKey), GetTypeHash(Model->Connections.FindTargetFromConnector(ConnectorKey)));
+											Attempt = HashCombine(Attempt, ConnectionHash);
+										}
+										if (AttemptedTargets.Contains(Attempt))
+										{
+										   bContinueResolval = false;
+										}
+										else
+										{
+										   AttemptedTargets.Add(Attempt);
+										}
+									}
 								}
 								while (bContinueResolval);
 
