Control Rig: Use correct asset editor tab icon

#rnx

diff --git a/Engine/Plugins/Animation/ControlRig/Source/ControlRigEditor/Private/ControlRigEditorStyle.h b/Engine/Plugins/Animation/ControlRig/Source/ControlRigEditor/Private/ControlRigEditorStyle.h
--- a/Engine/Plugins/Animation/ControlRig/Source/ControlRigEditor/Private/ControlRigEditorStyle.h
+++ b/Engine/Plugins/Animation/ControlRig/Source/ControlRigEditor/Private/ControlRigEditorStyle.h
@@ -64,6 +64,11 @@ public:
 			Set("ClassIcon.ControlRigPose", new IMAGE_BRUSH("Slate/ControlRigPose_16", Icon16x16));
 		}
 
+		// Editor Icons
+		{
+			Set("ControlRig.Editor.TabIcon", new IMAGE_BRUSH_SVG("Slate/ControlRigEditorTabIcon_16x", Icon16x16));
+		}
+
 		// Sequencer styles
 		{
 			Set("ControlRig.ExportAnimSequence", new IMAGE_BRUSH("Slate/ExportAnimSequence_24x", Icon24x24));
diff --git a/Engine/Plugins/Animation/ControlRig/Source/ControlRigEditor/Private/Editor/ControlRigEditor.cpp b/Engine/Plugins/Animation/ControlRig/Source/ControlRigEditor/Private/Editor/ControlRigEditor.cpp
--- a/Engine/Plugins/Animation/ControlRig/Source/ControlRigEditor/Private/Editor/ControlRigEditor.cpp
+++ b/Engine/Plugins/Animation/ControlRig/Source/ControlRigEditor/Private/Editor/ControlRigEditor.cpp
@@ -450,6 +450,12 @@ TSharedPtr<FApplicationMode> FControlRigEditor::CreateEditorMode()
 	return MakeShareable(new FControlRigEditorMode(SharedThis(this)));
 }
 
+const FSlateBrush* FControlRigEditor::GetDefaultTabIcon() const
+{
+	static const FSlateIcon TabIcon = FSlateIcon(FControlRigEditorStyle::Get().GetStyleSetName(), "ControlRig.Editor.TabIcon");
+	return TabIcon.GetIcon();
+}
+
 FText FControlRigEditor::GetTestAssetName() const
 {
 	if(TestDataStrongPtr.IsValid())
diff --git a/Engine/Plugins/Animation/ControlRig/Source/ControlRigEditor/Private/Editor/ControlRigEditor.h b/Engine/Plugins/Animation/ControlRig/Source/ControlRigEditor/Private/Editor/ControlRigEditor.h
--- a/Engine/Plugins/Animation/ControlRig/Source/ControlRigEditor/Private/Editor/ControlRigEditor.h
+++ b/Engine/Plugins/Animation/ControlRig/Source/ControlRigEditor/Private/Editor/ControlRigEditor.h
@@ -68,6 +68,7 @@ public:
 	virtual const FName GetEditorAppName() const override;
 	virtual const FName GetEditorModeName() const override;
 	virtual TSharedPtr<FApplicationMode> CreateEditorMode() override;
+	virtual const FSlateBrush* GetDefaultTabIcon() const override;
 
 public:
 	FControlRigEditor();
