Control Rig: Mark control rig as garbage when control rig component changes control rig class
#rnx

diff --git a/Engine/Plugins/Animation/ControlRig/Source/ControlRig/Private/ControlRigComponent.cpp b/Engine/Plugins/Animation/ControlRig/Source/ControlRig/Private/ControlRigComponent.cpp
--- a/Engine/Plugins/Animation/ControlRig/Source/ControlRig/Private/ControlRigComponent.cpp
+++ b/Engine/Plugins/Animation/ControlRig/Source/ControlRig/Private/ControlRigComponent.cpp
@@ -1233,13 +1233,7 @@ UControlRig* UControlRigComponent::SetupControlRigIfRequired()
 	{
 		if (ControlRig->GetClass() != ControlRigClass)
 		{
-			ControlRig->OnInitialized_AnyThread().RemoveAll(this);
-			ControlRig->OnPreConstruction_AnyThread().RemoveAll(this);
-			ControlRig->OnPostConstruction_AnyThread().RemoveAll(this);
-			ControlRig->OnPreForwardsSolve_AnyThread().RemoveAll(this);
-			ControlRig->OnPostForwardsSolve_AnyThread().RemoveAll(this);
-			ControlRig->OnExecuted_AnyThread().RemoveAll(this);
-			ControlRig = nullptr;
+			SetControlRig(nullptr);
 		}
 		else
 		{
@@ -1273,27 +1267,37 @@ void UControlRigComponent::SetControlRig(UControlRig* InControlRig)
 		ControlRig->OnPreForwardsSolve_AnyThread().RemoveAll(this);
 		ControlRig->OnPostForwardsSolve_AnyThread().RemoveAll(this);
 		ControlRig->OnExecuted_AnyThread().RemoveAll(this);
+
+		// rename the previous rig.
+		// GC will pick it up eventually - since we won't have any
+		// owning pointers to it anymore.
+		ControlRig->Rename(nullptr, GetTransientPackage(), REN_ForceNoResetLoaders | REN_DoNotDirty | REN_DontCreateRedirectors | REN_NonTransactional);
+		ControlRig->MarkAsGarbage();
 	}
 	ControlRig = InControlRig;
-	ControlRig->OnInitialized_AnyThread().AddUObject(this, &UControlRigComponent::HandleControlRigInitializedEvent);
-	ControlRig->OnPreConstruction_AnyThread().AddUObject(this, &UControlRigComponent::HandleControlRigPreConstructionEvent);
-	ControlRig->OnPostConstruction_AnyThread().AddUObject(this, &UControlRigComponent::HandleControlRigPostConstructionEvent);
-	ControlRig->OnPreForwardsSolve_AnyThread().AddUObject(this, &UControlRigComponent::HandleControlRigPreForwardsSolveEvent);
-	ControlRig->OnPostForwardsSolve_AnyThread().AddUObject(this, &UControlRigComponent::HandleControlRigPostForwardsSolveEvent);
-	ControlRig->OnExecuted_AnyThread().AddUObject(this, &UControlRigComponent::HandleControlRigExecutedEvent);
 
-	ControlRig->GetDataSourceRegistry()->RegisterDataSource(UControlRig::OwnerComponent, this);
-	if(ObjectBinding.IsValid())
+	if (ControlRig)
 	{
-		ControlRig->SetObjectBinding(ObjectBinding);
-	}
+		ControlRig->OnInitialized_AnyThread().AddUObject(this, &UControlRigComponent::HandleControlRigInitializedEvent);
+		ControlRig->OnPreConstruction_AnyThread().AddUObject(this, &UControlRigComponent::HandleControlRigPreConstructionEvent);
+		ControlRig->OnPostConstruction_AnyThread().AddUObject(this, &UControlRigComponent::HandleControlRigPostConstructionEvent);
+		ControlRig->OnPreForwardsSolve_AnyThread().AddUObject(this, &UControlRigComponent::HandleControlRigPreForwardsSolveEvent);
+		ControlRig->OnPostForwardsSolve_AnyThread().AddUObject(this, &UControlRigComponent::HandleControlRigPostForwardsSolveEvent);
+		ControlRig->OnExecuted_AnyThread().AddUObject(this, &UControlRigComponent::HandleControlRigExecutedEvent);
+
+		ControlRig->GetDataSourceRegistry()->RegisterDataSource(UControlRig::OwnerComponent, this);
+		if(ObjectBinding.IsValid())
+		{
+			ControlRig->SetObjectBinding(ObjectBinding);
+		}
 
-	ControlRig->Initialize();
+		ControlRig->Initialize();
+	}
 }
 
 void UControlRigComponent::SetControlRigClass(TSubclassOf<UControlRig> InControlRigClass)
 {
-	ControlRig = nullptr;
+	SetControlRig(nullptr);
 	ControlRigClass = InControlRigClass;
 	Initialize();
 }
