Control RIg: Only change control rig class on the control rig component if the input is of a different class
#rnx

diff --git a/Engine/Plugins/Animation/ControlRig/Source/ControlRig/Private/ControlRigComponent.cpp b/Engine/Plugins/Animation/ControlRig/Source/ControlRig/Private/ControlRigComponent.cpp
--- a/Engine/Plugins/Animation/ControlRig/Source/ControlRig/Private/ControlRigComponent.cpp
+++ b/Engine/Plugins/Animation/ControlRig/Source/ControlRig/Private/ControlRigComponent.cpp
@@ -1297,9 +1297,12 @@ void UControlRigComponent::SetControlRig(UControlRig* InControlRig)
 
 void UControlRigComponent::SetControlRigClass(TSubclassOf<UControlRig> InControlRigClass)
 {
-	SetControlRig(nullptr);
-	ControlRigClass = InControlRigClass;
-	Initialize();
+	if (ControlRigClass != InControlRigClass)
+	{
+		SetControlRig(nullptr);
+		ControlRigClass = InControlRigClass;
+		Initialize();
+	}
 }
 
 void UControlRigComponent::SetObjectBinding(UObject* InObjectToBind)
