Control Rig: Swap asset within modular rig or across project
#rnx

diff --git a/Engine/Plugins/Animation/ControlRig/Source/ControlRig/Private/ModularRigModel.cpp b/Engine/Plugins/Animation/ControlRig/Source/ControlRig/Private/ModularRigModel.cpp
--- a/Engine/Plugins/Animation/ControlRig/Source/ControlRig/Private/ModularRigModel.cpp
+++ b/Engine/Plugins/Animation/ControlRig/Source/ControlRig/Private/ModularRigModel.cpp
@@ -259,12 +259,14 @@ bool FModularRigModel::IsModuleParentedTo(const FRigModuleReference* InChildModu
 	return false;
 }
 
-TArray<const FRigModuleReference*> FModularRigModel::FindModuleInstancesOfClass(TSoftClassPtr<UControlRig> InClass) const
+TArray<const FRigModuleReference*> FModularRigModel::FindModuleInstancesOfClass(const FString& InModuleClassPath) const
 {
 	TArray<const FRigModuleReference*> Result;
-	ForEachModule([&Result, InClass](const FRigModuleReference* Module) -> bool
+	ForEachModule([&Result, InModuleClassPath](const FRigModuleReference* Module) -> bool
 	{
-		if (Module->Class == InClass)
+		FString PackageName = Module->Class.ToSoftObjectPath().GetAssetPathString();
+		PackageName.RemoveFromEnd(TEXT("_C"));
+		if (PackageName == InModuleClassPath)
 		{
 			Result.Add(Module);
 		}
@@ -273,6 +275,16 @@ TArray<const FRigModuleReference*> FModularRigModel::FindModuleInstancesOfClass(
 	return Result;
 }
 
+TArray<const FRigModuleReference*> FModularRigModel::FindModuleInstancesOfClass(const FAssetData& InModuleAsset) const
+{
+	return FindModuleInstancesOfClass(InModuleAsset.GetObjectPathString());
+}
+
+TArray<const FRigModuleReference*> FModularRigModel::FindModuleInstancesOfClass(TSoftClassPtr<UControlRig> InClass) const
+{
+	return FindModuleInstancesOfClass(*InClass.ToString());
+}
+
 FString FModularRigModel::GetParentPath(const FString& InPath) const
 {
 	if (const FRigModuleReference* Element = FindModule(InPath))
diff --git a/Engine/Plugins/Animation/ControlRig/Source/ControlRig/Public/ModularRigModel.h b/Engine/Plugins/Animation/ControlRig/Source/ControlRig/Public/ModularRigModel.h
--- a/Engine/Plugins/Animation/ControlRig/Source/ControlRig/Public/ModularRigModel.h
+++ b/Engine/Plugins/Animation/ControlRig/Source/ControlRig/Public/ModularRigModel.h
@@ -318,6 +318,8 @@ public:
 	
 	bool IsModuleParentedTo(const FRigModuleReference* InChildModule, const FRigModuleReference* InParentModule) const;
 
+	TArray<const FRigModuleReference*> FindModuleInstancesOfClass(const FString& InModuleClassPath) const;
+	TArray<const FRigModuleReference*> FindModuleInstancesOfClass(const FAssetData& InModuleAsset) const;
 	TArray<const FRigModuleReference*> FindModuleInstancesOfClass(TSoftClassPtr<UControlRig> InClass) const;
 
 private:
diff --git a/Engine/Plugins/Animation/ControlRig/Source/ControlRigDeveloper/Private/ControlRigBlueprint.cpp b/Engine/Plugins/Animation/ControlRig/Source/ControlRigDeveloper/Private/ControlRigBlueprint.cpp
--- a/Engine/Plugins/Animation/ControlRig/Source/ControlRigDeveloper/Private/ControlRigBlueprint.cpp
+++ b/Engine/Plugins/Animation/ControlRig/Source/ControlRigDeveloper/Private/ControlRigBlueprint.cpp
@@ -984,27 +984,15 @@ TArray<FModuleReferenceData> UControlRigBlueprint::FindReferencesToModule() cons
 	TArray<FAssetData> AssetDataList;
 	AssetRegistryModule.Get().GetAssetsByClass(UControlRigBlueprint::StaticClass()->GetClassPathName(), AssetDataList, true);
 
-	static const FLazyName ControlRigTypeName(GET_MEMBER_NAME_CHECKED(UControlRigBlueprint, ControlRigType));
-	FProperty* ControlRigTypeProperty = CastField<FProperty>(UControlRigBlueprint::StaticClass()->FindPropertyByName(ControlRigTypeName));
 	static const FLazyName ModuleReferenceDataName(GET_MEMBER_NAME_CHECKED(UControlRigBlueprint, ModuleReferenceData));
 	FArrayProperty* ModuleReferenceDataProperty = CastField<FArrayProperty>(UControlRigBlueprint::StaticClass()->FindPropertyByName(ModuleReferenceDataName));
 
 	for(const FAssetData& AssetData : AssetDataList)
 	{
 		// Check only modular rigs
+		if (UControlRigBlueprint::GetRigType(AssetData) != EControlRigType::ModularRig)
 		{
-			const FString ControlRigTypeString = AssetData.GetTagValueRef<FString>(ControlRigTypeName);
-			if (ControlRigTypeString.IsEmpty())
-			{
-				continue;
-			}
-
-			EControlRigType RigType;
-			ControlRigTypeProperty->ImportText_Direct(*ControlRigTypeString, &RigType, nullptr, EPropertyPortFlags::PPF_None);
-			if (RigType != EControlRigType::ModularRig)
-			{
-				continue;
-			}
+			continue;
 		}
 		
 		const FString ModularRigDataString = AssetData.GetTagValueRef<FString>(ModuleReferenceDataName);
@@ -1028,6 +1016,93 @@ TArray<FModuleReferenceData> UControlRigBlueprint::FindReferencesToModule() cons
 	return Result;
 }
 
+EControlRigType UControlRigBlueprint::GetRigType(const FAssetData& InAsset)
+{
+	EControlRigType Result = EControlRigType::MAX;
+	static const FLazyName ControlRigTypeName(GET_MEMBER_NAME_CHECKED(UControlRigBlueprint, ControlRigType));
+	FProperty* ControlRigTypeProperty = CastField<FProperty>(UControlRigBlueprint::StaticClass()->FindPropertyByName(ControlRigTypeName));
+	const FString ControlRigTypeString = InAsset.GetTagValueRef<FString>(ControlRigTypeName);
+	if (ControlRigTypeString.IsEmpty())
+	{
+		return Result;
+	}
+
+	EControlRigType RigType;
+	ControlRigTypeProperty->ImportText_Direct(*ControlRigTypeString, &RigType, nullptr, EPropertyPortFlags::PPF_None);
+	return RigType;
+}
+
+TArray<FSoftObjectPath> UControlRigBlueprint::GetReferencesToRigModule(const FAssetData& InModuleAsset)
+{
+	TArray<FSoftObjectPath> Result;
+	FAssetRegistryModule& AssetRegistryModule = FModuleManager::LoadModuleChecked<FAssetRegistryModule>(TEXT("AssetRegistry"));
+	IAssetRegistry& AssetRegistry = AssetRegistryModule.GetRegistry();
+	
+	TArray<FName> PackageDependencies;
+	AssetRegistry.GetReferencers(InModuleAsset.PackageName, PackageDependencies);
+
+	for (FName& DependencyPath : PackageDependencies)
+	{
+		TArray<FAssetData> Assets;
+		AssetRegistry.GetAssetsByPackageName(DependencyPath, Assets);
+
+		for (const FAssetData& DependencyData : Assets)
+		{
+			if (DependencyData.IsAssetLoaded())
+			{
+				if (UControlRigBlueprint* Blueprint = Cast<UControlRigBlueprint>(DependencyData.GetAsset()))
+				{
+					if (Blueprint->IsModularRig())
+					{
+						TArray<const FRigModuleReference*> Modules = Blueprint->ModularRigModel.FindModuleInstancesOfClass(InModuleAsset);
+						for (const FRigModuleReference* Module : Modules)
+						{
+							FSoftObjectPath ModulePath = DependencyData.GetSoftObjectPath();
+							ModulePath.SetSubPathString(Module->GetPath());
+							Result.Add(ModulePath);
+						}
+					}
+				}
+			}
+			else
+			{
+				// Check only modular rigs
+				if (UControlRigBlueprint::GetRigType(DependencyData) != EControlRigType::ModularRig)
+				{
+					continue;
+				}
+
+				static const FLazyName ModuleReferenceDataName(GET_MEMBER_NAME_CHECKED(UControlRigBlueprint, ModuleReferenceData));
+				FArrayProperty* ModuleReferenceDataProperty = CastField<FArrayProperty>(UControlRigBlueprint::StaticClass()->FindPropertyByName(ModuleReferenceDataName));
+				const FString ModularRigDataString = DependencyData.GetTagValueRef<FString>(ModuleReferenceDataName);
+				if (ModularRigDataString.IsEmpty())
+				{
+					continue;
+				}
+
+				TArray<FModuleReferenceData> Modules;
+				ModuleReferenceDataProperty->ImportText_Direct(*ModularRigDataString, &Modules, nullptr, EPropertyPortFlags::PPF_None);
+
+				for (const FModuleReferenceData& Module : Modules)
+				{
+					FTopLevelAssetPath ModulePath = Module.ReferencedModule.GetAssetPath();
+					FString AssetName = ModulePath.GetAssetName().ToString();
+					AssetName.RemoveFromEnd(TEXT("_C"));
+					ModulePath = FTopLevelAssetPath(ModulePath.GetPackageName(), *AssetName);
+					if (ModulePath == InModuleAsset.GetSoftObjectPath().GetAssetPath())
+					{
+						FSoftObjectPath ResultModulePath = DependencyData.GetSoftObjectPath();
+						ResultModulePath.SetSubPathString(Module.ModulePath);
+						Result.Add(ResultModulePath);
+					}
+				}
+			}
+		}
+	}
+	
+	return Result;
+}
+
 TArray<FModuleReferenceData> UControlRigBlueprint::GetModuleReferenceData() const
 {
 	TArray<FModuleReferenceData> Result;
diff --git a/Engine/Plugins/Animation/ControlRig/Source/ControlRigDeveloper/Public/ControlRigBlueprint.h b/Engine/Plugins/Animation/ControlRig/Source/ControlRigDeveloper/Public/ControlRigBlueprint.h
--- a/Engine/Plugins/Animation/ControlRig/Source/ControlRigDeveloper/Public/ControlRigBlueprint.h
+++ b/Engine/Plugins/Animation/ControlRig/Source/ControlRigDeveloper/Public/ControlRigBlueprint.h
@@ -36,7 +36,8 @@ enum class EControlRigType : uint8
 {
 	IndependentRig = 0,
 	RigModule = 1,
-	ModularRig =2,
+	ModularRig = 2,
+	MAX // Invalid
 };
 
 USTRUCT(BlueprintType)
@@ -199,6 +200,9 @@ public:
 	UFUNCTION(BlueprintPure, Category = "Control Rig Blueprint")
 	TArray<FModuleReferenceData> FindReferencesToModule() const;
 
+	static EControlRigType GetRigType(const FAssetData& InAsset);
+	static TArray<FSoftObjectPath> GetReferencesToRigModule(const FAssetData& InModuleAsset);
+
 protected:
 
 	TArray<FModuleReferenceData> GetModuleReferenceData() const;
diff --git a/Engine/Plugins/Animation/ControlRig/Source/ControlRigEditor/Private/ControlRigModularRigCommands.cpp b/Engine/Plugins/Animation/ControlRig/Source/ControlRigEditor/Private/ControlRigModularRigCommands.cpp
--- a/Engine/Plugins/Animation/ControlRig/Source/ControlRigEditor/Private/ControlRigModularRigCommands.cpp
+++ b/Engine/Plugins/Animation/ControlRig/Source/ControlRigEditor/Private/ControlRigModularRigCommands.cpp
@@ -11,7 +11,7 @@ void FControlRigModularRigCommands::RegisterCommands()
 	UI_COMMAND(DeleteModuleItem, "Delete", "Delete module from the rig.", EUserInterfaceActionType::Button, FInputChord(EKeys::Delete));
 	UI_COMMAND(MirrorModuleItem, "Mirror", "Mirror module from the rig.", EUserInterfaceActionType::Button, FInputChord());
 	UI_COMMAND(ReresolveModuleItem, "Auto Resolve", "Auto Resolve secondary connectors.", EUserInterfaceActionType::Button, FInputChord());
-	UI_COMMAND(SwapModuleClassItem, "Swap Module", "Swap module class.", EUserInterfaceActionType::Button, FInputChord());
+	UI_COMMAND(SwapModuleClassItem, "Swap Module", "Swap module of common class.", EUserInterfaceActionType::Button, FInputChord());
 }
 
 #undef LOCTEXT_NAMESPACE
diff --git a/Engine/Plugins/Animation/ControlRig/Source/ControlRigEditor/Private/Editor/ControlRigEditor.cpp b/Engine/Plugins/Animation/ControlRig/Source/ControlRigEditor/Private/Editor/ControlRigEditor.cpp
--- a/Engine/Plugins/Animation/ControlRig/Source/ControlRigEditor/Private/Editor/ControlRigEditor.cpp
+++ b/Engine/Plugins/Animation/ControlRig/Source/ControlRigEditor/Private/Editor/ControlRigEditor.cpp
@@ -108,6 +108,7 @@
 #include "RigVMModel/Nodes/RigVMAggregateNode.h"
 #include "AnimationEditorViewportClient.h"
 #include "DragAndDrop/AssetDragDropOp.h"
+#include "Editor/RigVMEditorTools.h"
 #include "SchematicGraphPanel/SSchematicGraphPanel.h"
 #include "RigVMCore/RigVMExecuteContext.h"
 #include "Editor/RigVMGraphDetailCustomization.h"
@@ -119,140 +120,6 @@
 TAutoConsoleVariable<bool> CVarControlRigShowTestingToolbar(TEXT("ControlRig.Test.EnableTestingToolbar"), false, TEXT("When true we'll show the testing toolbar in Control Rig Editor."));
 TAutoConsoleVariable<bool> CVarShowSchematicPanelOverlay(TEXT("ControlRig.Preview.ShowSchematicPanelOverlay"), true, TEXT("When true we'll add an overlay to the persona viewport to show modular rig information."));
 
-FAutoConsoleCommand FCmdModularRigSwapModuleWidget
-(
-	TEXT("ControlRig.SwapModules"),
-	TEXT("Create swap modules wizard."),
-	FConsoleCommandDelegate::CreateLambda([]()
-	{
-		SRigVMSwapAssetReferencesWidget::FArguments WidgetArgs;
-		WidgetArgs
-			.EnableUndo(true)
-			.CloseOnSuccess(true)
-			.OnGetReferences_Lambda([](const FAssetData& ReferencedAsset) -> TArray<FSoftObjectPath>
-			{
-				TArray<FSoftObjectPath> Result;
-				FAssetRegistryModule& AssetRegistryModule = FModuleManager::LoadModuleChecked<FAssetRegistryModule>(TEXT("AssetRegistry"));
-				IAssetRegistry& AssetRegistry = AssetRegistryModule.GetRegistry();
-
-				UClass* ReferencedClass = nullptr;
-				if (UControlRigBlueprint* ReferencedBlueprint = Cast<UControlRigBlueprint>(ReferencedAsset.GetAsset()))
-				{
-					ReferencedClass = ReferencedBlueprint->GetRigVMBlueprintGeneratedClass();
-				}
-				
-				TArray<FName> PackageDependencies;
-				AssetRegistry.GetReferencers(ReferencedAsset.PackageName, PackageDependencies);
-
-				for (FName& DependencyPath : PackageDependencies)
-				{
-					TArray<FAssetData> Assets;
-					AssetRegistry.GetAssetsByPackageName(DependencyPath, Assets);
-
-					for (const FAssetData& DependencyData : Assets)
-					{
-						if (DependencyData.IsAssetLoaded())
-						{
-							if (UControlRigBlueprint* Blueprint = Cast<UControlRigBlueprint>(DependencyData.GetAsset()))
-							{
-								if (Blueprint->IsModularRig())
-								{
-									TArray<const FRigModuleReference*> Modules = Blueprint->ModularRigModel.FindModuleInstancesOfClass(ReferencedClass);
-									for (const FRigModuleReference* Module : Modules)
-									{
-										FSoftObjectPath ModulePath = DependencyData.GetSoftObjectPath();
-										ModulePath.SetSubPathString(Module->GetPath());
-										Result.Add(ModulePath);
-									}
-								}
-							}
-						}
-						else
-						{
-							// Check only modular rigs
-							{
-								static const FLazyName ControlRigTypeName(GET_MEMBER_NAME_CHECKED(UControlRigBlueprint, ControlRigType));
-								FProperty* ControlRigTypeProperty = CastField<FProperty>(UControlRigBlueprint::StaticClass()->FindPropertyByName(ControlRigTypeName));
-								const FString ControlRigTypeString = DependencyData.GetTagValueRef<FString>(ControlRigTypeName);
-								if (ControlRigTypeString.IsEmpty())
-								{
-									continue;
-								}
-
-								EControlRigType RigType;
-								ControlRigTypeProperty->ImportText_Direct(*ControlRigTypeString, &RigType, nullptr, EPropertyPortFlags::PPF_None);
-								if (RigType != EControlRigType::ModularRig)
-								{
-									continue;
-								}
-							}
-
-							static const FLazyName ModuleReferenceDataName(GET_MEMBER_NAME_CHECKED(UControlRigBlueprint, ModuleReferenceData));
-							FArrayProperty* ModuleReferenceDataProperty = CastField<FArrayProperty>(UControlRigBlueprint::StaticClass()->FindPropertyByName(ModuleReferenceDataName));
-							const FString ModularRigDataString = DependencyData.GetTagValueRef<FString>(ModuleReferenceDataName);
-							if (ModularRigDataString.IsEmpty())
-							{
-								continue;
-							}
-
-							TArray<FModuleReferenceData> Modules;
-							ModuleReferenceDataProperty->ImportText_Direct(*ModularRigDataString, &Modules, nullptr, EPropertyPortFlags::PPF_None);
-
-							for (const FModuleReferenceData& Module : Modules)
-							{
-								if (Module.ReferencedModule == ReferencedClass)
-								{
-									FSoftObjectPath ModulePath = DependencyData.GetSoftObjectPath();
-									ModulePath.SetSubPathString(Module.ModulePath);
-									Result.Add(ModulePath);
-								}
-							}
-						}
-					}
-				}
-				
-				return Result;
-			})
-			.OnSwapReference_Lambda([](const FSoftObjectPath& ModulePath, const FAssetData& NewModuleAsset) -> bool
-			{
-				TSubclassOf<UControlRig> NewModuleClass = nullptr;
-				if (const UControlRigBlueprint* ModuleBlueprint = Cast<UControlRigBlueprint>(NewModuleAsset.GetAsset()))
-				{
-					NewModuleClass = ModuleBlueprint->GetRigVMBlueprintGeneratedClass();
-				}
-				if (NewModuleClass)
-				{
-					if (UControlRigBlueprint* RigBlueprint = Cast<UControlRigBlueprint>(ModulePath.GetWithoutSubPath().ResolveObject()))
-					{
-						return RigBlueprint->GetModularRigController()->SwapModuleClass(ModulePath.GetSubPathString(), NewModuleClass);
-					}
-				}
-				return false;
-			})
-			.ExtraAssetFilter_Lambda([](const FAssetData& AssetData) -> bool
-			{
-				static const FLazyName ControlRigTypeName(GET_MEMBER_NAME_CHECKED(UControlRigBlueprint, ControlRigType));
-				FProperty* ControlRigTypeProperty = CastField<FProperty>(UControlRigBlueprint::StaticClass()->FindPropertyByName(ControlRigTypeName));
-				const FString ControlRigTypeString = AssetData.GetTagValueRef<FString>(ControlRigTypeName);
-				if (ControlRigTypeString.IsEmpty())
-				{
-					return false;
-				}
-
-				EControlRigType RigType;
-				ControlRigTypeProperty->ImportText_Direct(*ControlRigTypeString, &RigType, nullptr, EPropertyPortFlags::PPF_None);
-				return RigType == EControlRigType::RigModule;
-			});
-
-		const TSharedRef<SRigVMBulkEditDialog<SRigVMSwapAssetReferencesWidget>> SwapModulesDialog =
-		 	SNew(SRigVMBulkEditDialog<SRigVMSwapAssetReferencesWidget>)
-		 	.WindowSize(FVector2D(800.0f, 640.0f))
-		 	.WidgetArgs(WidgetArgs);
-		
-		SwapModulesDialog->ShowNormal();
-	})
-);
-
 const FName FControlRigEditorModes::ControlRigEditorMode = TEXT("Rigging");
 const TArray<FName> FControlRigEditor::ForwardsSolveEventQueue = {FRigUnit_BeginExecution::EventName};
 const TArray<FName> FControlRigEditor::BackwardsSolveEventQueue = {FRigUnit_InverseExecution::EventName};
@@ -1714,6 +1581,15 @@ bool FControlRigEditor::IsModularRig() const
 	return false;
 }
 
+bool FControlRigEditor::IsRigModule() const
+{
+	if(UControlRigBlueprint* RigBlueprint = Cast<UControlRigBlueprint>(GetBlueprintObj()))
+	{
+		return RigBlueprint->IsControlRigModule();
+	}
+	return false;
+}
+
 FName FControlRigEditor::GetToolkitFName() const
 {
 	return FName("ControlRigEditor");
@@ -3802,6 +3678,35 @@ void FControlRigEditor::BindCommands()
 		FExecuteAction::CreateSP(this, &FControlRigEditor::HandleToggleSchematicViewport),
 		FCanExecuteAction(),
 		FIsActionChecked::CreateSP(this, &FControlRigEditor::IsSchematicViewportActive));
+
+	GetToolkitCommands()->MapAction(
+		FControlRigEditorCommands::Get().SwapModuleWithinAsset,
+		FExecuteAction::CreateSP(this, &FControlRigEditor::SwapModuleWithinAsset),
+		FCanExecuteAction::CreateSP(this, &FControlRigEditor::IsModularRig));
+
+	GetToolkitCommands()->MapAction(
+		FControlRigEditorCommands::Get().SwapModuleAcrossProject,
+		FExecuteAction::CreateSP(this, &FControlRigEditor::SwapModuleAcrossProject),
+		FCanExecuteAction::CreateSP(this, &FControlRigEditor::IsRigModule));
+}
+
+FMenuBuilder FControlRigEditor::GenerateBulkEditMenu()
+{
+	FMenuBuilder MenuBuilder = IControlRigEditor::GenerateBulkEditMenu();
+	MenuBuilder.BeginSection(TEXT("Asset"), LOCTEXT("Asset", "Asset"));
+	if (UControlRigBlueprint* Blueprint = GetControlRigBlueprint())
+	{
+		if (Blueprint->IsModularRig())
+		{
+			MenuBuilder.AddMenuEntry(FControlRigEditorCommands::Get().SwapModuleWithinAsset, TEXT("SwapModuleWithinAsset"), TAttribute<FText>(), TAttribute<FText>(), FSlateIcon());
+		}
+		else if (Blueprint->IsControlRigModule())
+		{
+			MenuBuilder.AddMenuEntry(FControlRigEditorCommands::Get().SwapModuleAcrossProject, TEXT("SwapModuleAcrossProject"), TAttribute<FText>(), TAttribute<FText>(), FSlateIcon());
+		}
+	}
+	MenuBuilder.EndSection();
+	return MenuBuilder;
 }
 
 void FControlRigEditor::OnHierarchyChanged()
@@ -4283,6 +4188,133 @@ void FControlRigEditor::HandlePostCompileModularRigs(URigVMBlueprint* InBlueprin
 	RefreshDetailView();
 }
 
+void FControlRigEditor::SwapModuleWithinAsset()
+{
+	const UControlRigBlueprint* Blueprint = GetControlRigBlueprint();
+	const FAssetData Asset = UE::RigVM::Editor::Tools::FindAssetFromAnyPath(GetRigVMBlueprint()->GetPathName(), true);
+	SRigVMSwapAssetReferencesWidget::FArguments WidgetArgs;
+
+	FRigVMAssetDataFilter FilterModules = FRigVMAssetDataFilter::CreateLambda([](const FAssetData& AssetData)
+	{
+		return UControlRigBlueprint::GetRigType(AssetData) == EControlRigType::RigModule;
+	});
+	FRigVMAssetDataFilter FilterSourceModules = FRigVMAssetDataFilter::CreateLambda([Blueprint](const FAssetData& AssetData)
+	{
+		if (Blueprint)
+		{
+			return !Blueprint->ModularRigModel.FindModuleInstancesOfClass(AssetData).IsEmpty();
+		}
+		return false;
+	});
+
+	TArray<FRigVMAssetDataFilter> SourceFilters = {FilterModules, FilterSourceModules};
+	TArray<FRigVMAssetDataFilter> TargetFilters = {FilterModules};
+	
+	WidgetArgs
+		.EnableUndo(true)
+		.CloseOnSuccess(true)
+		.OnGetReferences_Lambda([Blueprint, Asset](const FAssetData& ReferencedAsset) -> TArray<FSoftObjectPath>
+		{
+			TArray<FSoftObjectPath> Result;
+			FAssetRegistryModule& AssetRegistryModule = FModuleManager::LoadModuleChecked<FAssetRegistryModule>(TEXT("AssetRegistry"));
+			IAssetRegistry& AssetRegistry = AssetRegistryModule.GetRegistry();
+
+			UClass* ReferencedClass = nullptr;
+			if (UControlRigBlueprint* ReferencedBlueprint = Cast<UControlRigBlueprint>(ReferencedAsset.GetAsset()))
+			{
+				ReferencedClass = ReferencedBlueprint->GetRigVMBlueprintGeneratedClass();
+			}
+
+			if (Blueprint)
+			{
+				if (Blueprint->IsModularRig())
+				{
+					TArray<const FRigModuleReference*> Modules = Blueprint->ModularRigModel.FindModuleInstancesOfClass(ReferencedAsset);
+					for (const FRigModuleReference* Module : Modules)
+					{
+						FSoftObjectPath ModulePath = Asset.GetSoftObjectPath();
+						ModulePath.SetSubPathString(Module->GetPath());
+						Result.Add(ModulePath);
+					}
+				}
+			}
+			
+			return Result;
+		})
+		.OnSwapReference_Lambda([](const FSoftObjectPath& ModulePath, const FAssetData& NewModuleAsset) -> bool
+		{
+			TSubclassOf<UControlRig> NewModuleClass = nullptr;
+			if (const UControlRigBlueprint* ModuleBlueprint = Cast<UControlRigBlueprint>(NewModuleAsset.GetAsset()))
+			{
+				NewModuleClass = ModuleBlueprint->GetRigVMBlueprintGeneratedClass();
+			}
+			if (NewModuleClass)
+			{
+				if (UControlRigBlueprint* RigBlueprint = Cast<UControlRigBlueprint>(ModulePath.GetWithoutSubPath().ResolveObject()))
+				{
+					return RigBlueprint->GetModularRigController()->SwapModuleClass(ModulePath.GetSubPathString(), NewModuleClass);
+				}
+			}
+			return false;
+		})
+		.SourceAssetFilters(SourceFilters)
+		.TargetAssetFilters(TargetFilters);
+
+	const TSharedRef<SRigVMBulkEditDialog<SRigVMSwapAssetReferencesWidget>> SwapModulesDialog =
+		SNew(SRigVMBulkEditDialog<SRigVMSwapAssetReferencesWidget>)
+		.WindowSize(FVector2D(800.0f, 640.0f))
+		.WidgetArgs(WidgetArgs);
+	
+	SwapModulesDialog->ShowNormal();
+}
+
+void FControlRigEditor::SwapModuleAcrossProject()
+{
+	const UControlRigBlueprint* Blueprint = GetControlRigBlueprint();
+	const FAssetData Asset = UE::RigVM::Editor::Tools::FindAssetFromAnyPath(GetRigVMBlueprint()->GetPathName(), true);
+	SRigVMSwapAssetReferencesWidget::FArguments WidgetArgs;
+
+	FRigVMAssetDataFilter FilterModules = FRigVMAssetDataFilter::CreateLambda([](const FAssetData& AssetData)
+	{
+		return UControlRigBlueprint::GetRigType(AssetData) == EControlRigType::RigModule;
+	});
+
+	TArray<FRigVMAssetDataFilter> TargetFilters = {FilterModules};
+	
+	WidgetArgs
+		.EnableUndo(false)
+		.CloseOnSuccess(true)
+		.OnGetReferences_Lambda([Blueprint, Asset](const FAssetData& ReferencedAsset) -> TArray<FSoftObjectPath>
+		{
+			return UControlRigBlueprint::GetReferencesToRigModule(Asset);
+		})
+		.OnSwapReference_Lambda([](const FSoftObjectPath& ModulePath, const FAssetData& NewModuleAsset) -> bool
+		{
+			TSubclassOf<UControlRig> NewModuleClass = nullptr;
+			if (const UControlRigBlueprint* ModuleBlueprint = Cast<UControlRigBlueprint>(NewModuleAsset.GetAsset()))
+			{
+				NewModuleClass = ModuleBlueprint->GetRigVMBlueprintGeneratedClass();
+			}
+			if (NewModuleClass)
+			{
+				if (UControlRigBlueprint* RigBlueprint = Cast<UControlRigBlueprint>(ModulePath.GetWithoutSubPath().ResolveObject()))
+				{
+					return RigBlueprint->GetModularRigController()->SwapModuleClass(ModulePath.GetSubPathString(), NewModuleClass);
+				}
+			}
+			return false;
+		})
+		.Source(Asset)
+		.TargetAssetFilters(TargetFilters);
+
+	const TSharedRef<SRigVMBulkEditDialog<SRigVMSwapAssetReferencesWidget>> SwapModulesDialog =
+		SNew(SRigVMBulkEditDialog<SRigVMSwapAssetReferencesWidget>)
+		.WindowSize(FVector2D(800.0f, 640.0f))
+		.WidgetArgs(WidgetArgs);
+	
+	SwapModulesDialog->ShowNormal();
+}
+
 void FControlRigEditor::SynchronizeViewportBoneSelection()
 {
 	UControlRigBlueprint* RigBlueprint = GetControlRigBlueprint();
diff --git a/Engine/Plugins/Animation/ControlRig/Source/ControlRigEditor/Private/Editor/ControlRigEditor.h b/Engine/Plugins/Animation/ControlRig/Source/ControlRigEditor/Private/Editor/ControlRigEditor.h
--- a/Engine/Plugins/Animation/ControlRig/Source/ControlRigEditor/Private/Editor/ControlRigEditor.h
+++ b/Engine/Plugins/Animation/ControlRig/Source/ControlRigEditor/Private/Editor/ControlRigEditor.h
@@ -94,6 +94,7 @@ public:
 	int32 GetModularRigHierarchyTabCount() const { return ModularRigHierarchyTabCount; }
 
 	bool IsModularRig() const;
+	bool IsRigModule() const;
 
 public:
 	
@@ -164,6 +165,8 @@ public:
 
 	void HandleModularRigModified(EModularRigNotification InNotification, const FRigModuleReference* InModule);
 	void HandlePostCompileModularRigs(URigVMBlueprint* InBlueprint);
+	void SwapModuleWithinAsset();
+	void SwapModuleAcrossProject();
 
 	const FName RigHierarchyToGraphDragAndDropMenuName = TEXT("ControlRigEditor.RigHierarchyToGraphDragAndDropMenu");
 	void CreateRigHierarchyToGraphDragAndDropMenu() const;
@@ -200,6 +203,7 @@ public:
 protected:
 
 	virtual void BindCommands() override;
+	virtual FMenuBuilder GenerateBulkEditMenu() override;
 
 	void OnHierarchyChanged();
 
diff --git a/Engine/Plugins/Animation/ControlRig/Source/ControlRigEditor/Private/Editor/ControlRigEditorCommands.cpp b/Engine/Plugins/Animation/ControlRig/Source/ControlRigEditor/Private/Editor/ControlRigEditorCommands.cpp
--- a/Engine/Plugins/Animation/ControlRig/Source/ControlRigEditor/Private/Editor/ControlRigEditorCommands.cpp
+++ b/Engine/Plugins/Animation/ControlRig/Source/ControlRigEditor/Private/Editor/ControlRigEditorCommands.cpp
@@ -14,6 +14,8 @@ void FControlRigEditorCommands::RegisterCommands()
 	UI_COMMAND(RequestDirectManipulationRotation, "Request Direct Manipulation for Rotation", "Request per node direct manipulation on a rotation", EUserInterfaceActionType::Button, FInputChord(EKeys::E));
 	UI_COMMAND(RequestDirectManipulationScale, "Request Direct Manipulation for Scale", "Request per node direct manipulation on a scale", EUserInterfaceActionType::Button, FInputChord(EKeys::R));
 	UI_COMMAND(ToggleSchematicViewportVisibility, "Show Schematic Viewport", "Toggles the visibility of the viewport schematic", EUserInterfaceActionType::ToggleButton, FInputChord(EKeys::Y));
+	UI_COMMAND(SwapModuleWithinAsset, "Swap Module (Asset)", "Swaps a module for all occurrences within this asset.", EUserInterfaceActionType::Button, FInputChord());
+	UI_COMMAND(SwapModuleAcrossProject, "Swap Module (Project)", "Swaps a module for all occurrences in the project.", EUserInterfaceActionType::Button, FInputChord());
 }
 
 #undef LOCTEXT_NAMESPACE
diff --git a/Engine/Plugins/Animation/ControlRig/Source/ControlRigEditor/Private/Editor/ControlRigEditorCommands.h b/Engine/Plugins/Animation/ControlRig/Source/ControlRigEditor/Private/Editor/ControlRigEditorCommands.h
--- a/Engine/Plugins/Animation/ControlRig/Source/ControlRigEditor/Private/Editor/ControlRigEditorCommands.h
+++ b/Engine/Plugins/Animation/ControlRig/Source/ControlRigEditor/Private/Editor/ControlRigEditorCommands.h
@@ -42,6 +42,12 @@ public:
 	/** Toggle visibility of the schematic */
 	TSharedPtr< FUICommandInfo > ToggleSchematicViewportVisibility;
 
+	/** Swap Module (Asset) */
+	TSharedPtr< FUICommandInfo > SwapModuleWithinAsset;
+
+	/** Swap Module (Project) */
+	TSharedPtr< FUICommandInfo > SwapModuleAcrossProject;
+
 	/**
 	 * Initialize commands
 	 */
diff --git a/Engine/Plugins/Animation/ControlRig/Source/ControlRigEditor/Private/Editor/SModularRigModel.cpp b/Engine/Plugins/Animation/ControlRig/Source/ControlRigEditor/Private/Editor/SModularRigModel.cpp
--- a/Engine/Plugins/Animation/ControlRig/Source/ControlRigEditor/Private/Editor/SModularRigModel.cpp
+++ b/Engine/Plugins/Animation/ControlRig/Source/ControlRigEditor/Private/Editor/SModularRigModel.cpp
@@ -51,6 +51,7 @@
 #include "Algo/MaxElement.h"
 #include "AssetRegistry/AssetRegistryModule.h"
 #include "DragAndDrop/AssetDragDropOp.h"
+#include "Editor/RigVMEditorTools.h"
 #include "Kismet2/SClassPickerDialog.h"
 #include "RigVMFunctions/Math/RigVMMathLibrary.h"
 #include "Preferences/PersonaOptions.h"
@@ -953,18 +954,20 @@ void SModularRigModel::HandleSwapClassForModules(const TArray<FString>& InPaths)
 		return;
 	}
 
-	TArray<FAssetData> SourceAssets;
-	IAssetRegistry::Get()->GetAssetsByPackageName(*SourceClass->GetPackage()->GetPathName(),SourceAssets);
-	if (SourceAssets.IsEmpty())
-	{
-		return;
-	}
+	const FAssetData SourceAsset = UE::RigVM::Editor::Tools::FindAssetFromAnyPath(SourceClass.GetLongPackageName(), true);
 	
 	SRigVMSwapAssetReferencesWidget::FArguments WidgetArgs;
+	FRigVMAssetDataFilter FilterModules = FRigVMAssetDataFilter::CreateLambda([](const FAssetData& AssetData)
+		{
+			return UControlRigBlueprint::GetRigType(AssetData) == EControlRigType::RigModule;
+		});
+	TArray<FRigVMAssetDataFilter> SourceFilters = {FilterModules};
+	TArray<FRigVMAssetDataFilter> TargetFilters = {FilterModules};
+	
 	WidgetArgs
 		.EnableUndo(true)
 		.CloseOnSuccess(true)
-		.Source(SourceAssets[0])
+		.Source(SourceAsset)
 		.ReferencePaths(ModulePaths)
 		.SkipPickingRefs(true)
 		.OnSwapReference_Lambda([](const FSoftObjectPath& ModulePath, const FAssetData& NewModuleAsset) -> bool
@@ -983,20 +986,8 @@ void SModularRigModel::HandleSwapClassForModules(const TArray<FString>& InPaths)
 			}
 			return false;
 		})
-		.ExtraAssetFilter_Lambda([](const FAssetData& AssetData) -> bool
-		{
-			static const FLazyName ControlRigTypeName(GET_MEMBER_NAME_CHECKED(UControlRigBlueprint, ControlRigType));
-			FProperty* ControlRigTypeProperty = CastField<FProperty>(UControlRigBlueprint::StaticClass()->FindPropertyByName(ControlRigTypeName));
-			const FString ControlRigTypeString = AssetData.GetTagValueRef<FString>(ControlRigTypeName);
-			if (ControlRigTypeString.IsEmpty())
-			{
-				return false;
-			}
-
-			EControlRigType RigType;
-			ControlRigTypeProperty->ImportText_Direct(*ControlRigTypeString, &RigType, nullptr, EPropertyPortFlags::PPF_None);
-			return RigType == EControlRigType::RigModule;
-		});
+		.SourceAssetFilters(SourceFilters)
+		.TargetAssetFilters(TargetFilters);
 
 	const TSharedRef<SRigVMBulkEditDialog<SRigVMSwapAssetReferencesWidget>> SwapModulesDialog =
 		SNew(SRigVMBulkEditDialog<SRigVMSwapAssetReferencesWidget>)
diff --git a/Engine/Plugins/Runtime/RigVM/Source/RigVMDeveloper/Private/RigVMModel/RigVMBuildData.cpp b/Engine/Plugins/Runtime/RigVM/Source/RigVMDeveloper/Private/RigVMModel/RigVMBuildData.cpp
--- a/Engine/Plugins/Runtime/RigVM/Source/RigVMDeveloper/Private/RigVMModel/RigVMBuildData.cpp
+++ b/Engine/Plugins/Runtime/RigVM/Source/RigVMDeveloper/Private/RigVMModel/RigVMBuildData.cpp
@@ -568,6 +568,12 @@ TArray<FRigVMVariantRef> URigVMBuildData::GatherAllAssetVariantRefs()
 		const FString VariantStr = Asset.GetTagValueRef<FString>(AssetVariantPropertyName);
 		FRigVMVariant AssetVariant;
 		AssetVariantProperty->ImportText_Direct(*VariantStr, &AssetVariant, nullptr, EPropertyPortFlags::PPF_None);
+
+		if (!AssetVariant.Guid.IsValid())
+		{
+			AssetVariant.Guid = FRigVMVariant::GenerateGUID(Asset.PackageName.ToString());
+		}
+		
 		FRigVMVariantRef VariantRef(Asset.ToSoftObjectPath(), AssetVariant);
 		Result.Add(VariantRef);
 	}
diff --git a/Engine/Plugins/Runtime/RigVM/Source/RigVMEditor/Private/Editor/RigVMEditor.cpp b/Engine/Plugins/Runtime/RigVM/Source/RigVMEditor/Private/Editor/RigVMEditor.cpp
--- a/Engine/Plugins/Runtime/RigVM/Source/RigVMEditor/Private/Editor/RigVMEditor.cpp
+++ b/Engine/Plugins/Runtime/RigVM/Source/RigVMEditor/Private/Editor/RigVMEditor.cpp
@@ -1466,16 +1466,22 @@ TSharedRef<SWidget> FRigVMEditor::GenerateExecutionModeMenuContent()
 	return MenuBuilder.MakeWidget();
 }
 
-TSharedRef<SWidget> FRigVMEditor::GenerateBulkEditMenuContent()
+FMenuBuilder FRigVMEditor::GenerateBulkEditMenu()
 {
 	FMenuBuilder MenuBuilder(true, GetToolkitCommands());
 	MenuBuilder.BeginSection(TEXT("Functions"), LOCTEXT("Functions", "Functions"));
 	MenuBuilder.AddMenuEntry(FRigVMEditorCommands::Get().SwapFunctionWithinAsset, TEXT("SwapFunctionWithinAsset"), TAttribute<FText>(), TAttribute<FText>(), FSlateIcon());
 	MenuBuilder.AddMenuEntry(FRigVMEditorCommands::Get().SwapFunctionAcrossProject, TEXT("SwapFunctionAcrossProject"), TAttribute<FText>(), TAttribute<FText>(), FSlateIcon());
 	MenuBuilder.EndSection();
-	MenuBuilder.BeginSection(TEXT("Asset"), LOCTEXT("Asset", "Asset"));
-	MenuBuilder.AddMenuEntry(FRigVMEditorCommands::Get().SwapAssetReferences, TEXT("SwapAssetReferences"), TAttribute<FText>(), TAttribute<FText>(), FSlateIcon());
-	MenuBuilder.EndSection();
+	// MenuBuilder.BeginSection(TEXT("Asset"), LOCTEXT("Asset", "Asset"));
+	// MenuBuilder.AddMenuEntry(FRigVMEditorCommands::Get().SwapAssetReferences, TEXT("SwapAssetReferences"), TAttribute<FText>(), TAttribute<FText>(), FSlateIcon());
+	// MenuBuilder.EndSection();
+	return MenuBuilder;
+}
+
+TSharedRef<SWidget> FRigVMEditor::GenerateBulkEditMenuContent()
+{
+	FMenuBuilder MenuBuilder = GenerateBulkEditMenu();
 	return MenuBuilder.MakeWidget();
 }
 
@@ -3479,16 +3485,11 @@ void FRigVMEditor::SwapFunctionForAssets(const TArray<FAssetData>& InAssets, boo
 
 void FRigVMEditor::SwapAssetReferences()
 {
-	TArray<FAssetData> SourceAssets;
-	IAssetRegistry::Get()->GetAssetsByPackageName(*GetRigVMBlueprint()->GetPackage()->GetPathName(),SourceAssets);
-	if (SourceAssets.IsEmpty())
-	{
-		return;
-	}
+	const FAssetData Asset = UE::RigVM::Editor::Tools::FindAssetFromAnyPath(GetRigVMBlueprint()->GetPathName(), true);
 	
 	SRigVMSwapAssetReferencesWidget::FArguments WidgetArgs;
 	WidgetArgs
-		.Source(SourceAssets[0])
+		.Source(Asset)
 		.EnableUndo(false)
 		.CloseOnSuccess(true);
 
diff --git a/Engine/Plugins/Runtime/RigVM/Source/RigVMEditor/Private/Widgets/SRigVMSwapAssetReferencesWidget.cpp b/Engine/Plugins/Runtime/RigVM/Source/RigVMEditor/Private/Widgets/SRigVMSwapAssetReferencesWidget.cpp
--- a/Engine/Plugins/Runtime/RigVM/Source/RigVMEditor/Private/Widgets/SRigVMSwapAssetReferencesWidget.cpp
+++ b/Engine/Plugins/Runtime/RigVM/Source/RigVMEditor/Private/Widgets/SRigVMSwapAssetReferencesWidget.cpp
@@ -29,8 +29,11 @@ TArray<TSharedRef<FRigVMTreeNode>> FRigVMTreeAssetRefAssetNode::GetChildrenImpl(
 		{
 			for (const FSoftObjectPath& Reference : SwapContext->GetReferences())
 			{
-				TSharedRef<FRigVMTreeReferenceNode> RefNode = FRigVMTreeReferenceNode::Create(Reference);
-				AssetRefNodes.Add(RefNode);
+				if (SoftObjectPath == Reference.GetWithoutSubPath())
+				{
+					TSharedRef<FRigVMTreeReferenceNode> RefNode = FRigVMTreeReferenceNode::Create(Reference);
+					AssetRefNodes.Add(RefNode);
+				}
 			}
 		}
 	}
@@ -70,7 +73,7 @@ bool FRigVMTreeAssetVariantFilter::Filters(TSharedRef<FRigVMTreeNode>& InNode, c
 
 				if (!AssetVariant.Guid.IsValid())
 				{
-					AssetVariant.Guid = FRigVMVariant::GenerateGUID(SwapContext->GetSourceAsset().GetPackage()->GetPathName());
+					AssetVariant.Guid = FRigVMVariant::GenerateGUID(SwapContext->GetSourceAsset().PackageName.ToString());
 				}
 
 				SourceVariants.FindOrAdd(SwapContext->GetSourceAsset().GetFullName()) = URigVMBuildData::Get()->FindAssetVariantRefs(AssetVariant.Guid);
@@ -114,7 +117,8 @@ void SRigVMSwapAssetReferencesWidget::Construct(const FArguments& InArgs)
 
 	OnGetReferences = InArgs._OnGetReferences;
 	OnSwapReference = InArgs._OnSwapReference;
-	ExtraAssetFilter = InArgs._ExtraAssetFilter;
+	SourceAssetFilters = InArgs._SourceAssetFilters;
+	TargetAssetFilters = InArgs._TargetAssetFilters;
 	PickTargetContext = FRigVMSwapAssetReferencesContext::Create();
 	PickAssetRefsContext = FRigVMSwapAssetReferencesContext::Create();
 	bSkipPickingRefs = InArgs._SkipPickingRefs;
@@ -131,7 +135,7 @@ void SRigVMSwapAssetReferencesWidget::Construct(const FArguments& InArgs)
 		TSharedRef<FRigVMTreePhase> Phase = FRigVMTreePhase::Create(PHASE_PICKSOURCE, TEXT("Pick Source Asset"), FRigVMTreeContext::Create());
 		Phase->GetContext()->Filters = DefaultFilters;
 		Phase->GetContext()->Filters.Add(DefaultPathFilter);
-		Phase->SetNodes(GetAssetNodes(InArgs));
+		Phase->SetNodes(GetAssetNodes(InArgs, PHASE_PICKSOURCE));
 		Phase->PrimaryButtonText().Set(LOCTEXT("Next", "Next"));
 		Phase->IsPrimaryButtonVisible().BindLambda([this]()
 		{
@@ -163,7 +167,7 @@ void SRigVMSwapAssetReferencesWidget::Construct(const FArguments& InArgs)
 		Phase->GetContext()->Filters.Add(DefaultPathFilter);
 		Phase->GetContext()->Filters.Add(FRigVMTreeTargetAssetFilter::Create());
 		Phase->GetContext()->Filters.Add(FRigVMTreeAssetVariantFilter::Create());
-		Phase->SetNodes(GetAssetNodes(InArgs));
+		Phase->SetNodes(GetAssetNodes(InArgs, PHASE_PICKTARGET));
 		Phase->PrimaryButtonText().Set(LOCTEXT("Next", "Next"));
 		Phase->IsPrimaryButtonVisible().BindLambda([this]()
 		{
@@ -289,16 +293,20 @@ void SRigVMSwapAssetReferencesWidget::Construct(const FArguments& InArgs)
 	];
 }
 
-TArray<TSharedRef<FRigVMTreeNode>> SRigVMSwapAssetReferencesWidget::GetAssetNodes(const FArguments& InArgs)
+TArray<TSharedRef<FRigVMTreeNode>> SRigVMSwapAssetReferencesWidget::GetAssetNodes(const FArguments& InArgs, const int32& InPhase)
 {
 	TArray<TSharedRef<FRigVMTreeNode>> Nodes;
 
 	const TArray<FAssetData> Assets = URigVMEditorBlueprintLibrary::GetAssetsWithFilter(URigVMBlueprint::StaticClass(),
-		FRigVMAssetDataFilter::CreateLambda([this] (const FAssetData& AssetData)
+		FRigVMAssetDataFilter::CreateLambda([this, InPhase] (const FAssetData& AssetData)
 		{
-			if (ExtraAssetFilter.IsBound())
+			TArray<FRigVMAssetDataFilter>& Filters = (InPhase == PHASE_PICKSOURCE) ? SourceAssetFilters : TargetAssetFilters;
+			for (const FRigVMAssetDataFilter& Filter : Filters)
 			{
-				return ExtraAssetFilter.Execute(AssetData);
+				if (!Filter.Execute(AssetData))
+				{
+					return false;
+				}
 			}
 			return true;
 		}));
diff --git a/Engine/Plugins/Runtime/RigVM/Source/RigVMEditor/Public/Editor/RigVMEditor.h b/Engine/Plugins/Runtime/RigVM/Source/RigVMEditor/Public/Editor/RigVMEditor.h
--- a/Engine/Plugins/Runtime/RigVM/Source/RigVMEditor/Public/Editor/RigVMEditor.h
+++ b/Engine/Plugins/Runtime/RigVM/Source/RigVMEditor/Public/Editor/RigVMEditor.h
@@ -220,6 +220,7 @@ protected:
 	void ToggleExecutionMode();
 	TSharedRef<SWidget> GenerateEventQueueMenuContent();
 	TSharedRef<SWidget> GenerateExecutionModeMenuContent();
+	virtual FMenuBuilder GenerateBulkEditMenu();
 	TSharedRef<SWidget> GenerateBulkEditMenuContent();
 	virtual void GenerateEventQueueMenuContent(FMenuBuilder& MenuBuilder);
 
diff --git a/Engine/Plugins/Runtime/RigVM/Source/RigVMEditor/Public/Widgets/SRigVMSwapAssetReferencesWidget.h b/Engine/Plugins/Runtime/RigVM/Source/RigVMEditor/Public/Widgets/SRigVMSwapAssetReferencesWidget.h
--- a/Engine/Plugins/Runtime/RigVM/Source/RigVMEditor/Public/Widgets/SRigVMSwapAssetReferencesWidget.h
+++ b/Engine/Plugins/Runtime/RigVM/Source/RigVMEditor/Public/Widgets/SRigVMSwapAssetReferencesWidget.h
@@ -182,7 +182,8 @@ public:
 	SLATE_ARGUMENT(bool, SkipPickingRefs)
 	SLATE_ARGUMENT(bool, EnableUndo)
 	SLATE_ARGUMENT(bool, CloseOnSuccess)
-	SLATE_EVENT(FRigVMAssetDataFilter, ExtraAssetFilter)
+	SLATE_ARGUMENT(TArray<FRigVMAssetDataFilter>, SourceAssetFilters)
+	SLATE_ARGUMENT(TArray<FRigVMAssetDataFilter>, TargetAssetFilters)
 	SLATE_EVENT(FOnGetReferences, OnGetReferences)
 	SLATE_EVENT(FOnSwapReference, OnSwapReference)
 	SLATE_END_ARGS()
@@ -201,11 +202,13 @@ private:
 	TSharedPtr<FRigVMSwapAssetReferencesContext> PickAssetRefsContext;
 	bool bSkipPickingRefs = false;
 
-	FRigVMAssetDataFilter ExtraAssetFilter;
+	TArray<FRigVMAssetDataFilter> SourceAssetFilters;
+	TArray<FRigVMAssetDataFilter> TargetAssetFilters;
+	
 	FOnGetReferences OnGetReferences;
 	FOnSwapReference OnSwapReference;
 
-	TArray<TSharedRef<FRigVMTreeNode>> GetAssetNodes(const FArguments& InArgs);
+	TArray<TSharedRef<FRigVMTreeNode>> GetAssetNodes(const FArguments& InArgs, const int32& InPhase);
 	void OnPhaseActivated(TSharedRef<FRigVMTreePhase> Phase);
 	FReply OnNodeSelected(TSharedRef<FRigVMTreeNode> Node);
 	FReply OnNodeDoubleClicked(TSharedRef<FRigVMTreeNode> Node);
