Control Rig: Fixed creating new control rig and inmediately marking it as garbage
#rnx

diff --git a/Engine/Plugins/Animation/ControlRig/Source/ControlRig/Private/ControlRigComponent.cpp b/Engine/Plugins/Animation/ControlRig/Source/ControlRig/Private/ControlRigComponent.cpp
--- a/Engine/Plugins/Animation/ControlRig/Source/ControlRig/Private/ControlRigComponent.cpp
+++ b/Engine/Plugins/Animation/ControlRig/Source/ControlRig/Private/ControlRigComponent.cpp
@@ -1243,9 +1243,9 @@ UControlRig* UControlRigComponent::SetupControlRigIfRequired()
 
 	if(ControlRigClass)
 	{
-		ControlRig = NewObject<UControlRig>(this, ControlRigClass);
+		UControlRig* NewControlRig = NewObject<UControlRig>(this, ControlRigClass);
 
-		SetControlRig(ControlRig);
+		SetControlRig(NewControlRig);
 
 		if (ControlRigCreatedEvent.IsBound())
 		{
@@ -1259,22 +1259,25 @@ UControlRig* UControlRigComponent::SetupControlRigIfRequired()
 }
 void UControlRigComponent::SetControlRig(UControlRig* InControlRig)
 {
-	if (ControlRig)
+	if (ControlRig != InControlRig)
 	{
-		ControlRig->OnInitialized_AnyThread().RemoveAll(this);
-		ControlRig->OnPreConstruction_AnyThread().RemoveAll(this);
-		ControlRig->OnPostConstruction_AnyThread().RemoveAll(this);
-		ControlRig->OnPreForwardsSolve_AnyThread().RemoveAll(this);
-		ControlRig->OnPostForwardsSolve_AnyThread().RemoveAll(this);
-		ControlRig->OnExecuted_AnyThread().RemoveAll(this);
+		if (ControlRig)
+		{
+			ControlRig->OnInitialized_AnyThread().RemoveAll(this);
+			ControlRig->OnPreConstruction_AnyThread().RemoveAll(this);
+			ControlRig->OnPostConstruction_AnyThread().RemoveAll(this);
+			ControlRig->OnPreForwardsSolve_AnyThread().RemoveAll(this);
+			ControlRig->OnPostForwardsSolve_AnyThread().RemoveAll(this);
+			ControlRig->OnExecuted_AnyThread().RemoveAll(this);
 
-		// rename the previous rig.
-		// GC will pick it up eventually - since we won't have any
-		// owning pointers to it anymore.
-		ControlRig->Rename(nullptr, GetTransientPackage(), REN_ForceNoResetLoaders | REN_DoNotDirty | REN_DontCreateRedirectors | REN_NonTransactional);
-		ControlRig->MarkAsGarbage();
+			// rename the previous rig.
+			// GC will pick it up eventually - since we won't have any
+			// owning pointers to it anymore.
+			ControlRig->Rename(nullptr, GetTransientPackage(), REN_ForceNoResetLoaders | REN_DoNotDirty | REN_DontCreateRedirectors | REN_NonTransactional);
+			ControlRig->MarkAsGarbage();
+		}
+		ControlRig = InControlRig;
 	}
-	ControlRig = InControlRig;
 
 	if (ControlRig)
 	{
