Control Rig: Fix compilation issue
#rnx

diff --git a/Engine/Plugins/Animation/ControlRig/Source/ControlRig/Private/ModularRigModel.cpp b/Engine/Plugins/Animation/ControlRig/Source/ControlRig/Private/ModularRigModel.cpp
--- a/Engine/Plugins/Animation/ControlRig/Source/ControlRig/Private/ModularRigModel.cpp
+++ b/Engine/Plugins/Animation/ControlRig/Source/ControlRig/Private/ModularRigModel.cpp
@@ -3,6 +3,7 @@
 #include "ModularRigModel.h"
 #include "ModularRigController.h"
 #include "ModularRig.h"
+#include "AssetRegistry/AssetData.h"
 
 #include UE_INLINE_GENERATED_CPP_BY_NAME(ModularRigModel)
 
