Control Rig: Make sure control min value respects the max value when setting a new max value
#rnx

diff --git a/Engine/Plugins/Animation/ControlRig/Source/ControlRig/Private/Rigs/RigHierarchy.cpp b/Engine/Plugins/Animation/ControlRig/Source/ControlRig/Private/Rigs/RigHierarchy.cpp
--- a/Engine/Plugins/Animation/ControlRig/Source/ControlRig/Private/Rigs/RigHierarchy.cpp
+++ b/Engine/Plugins/Animation/ControlRig/Source/ControlRig/Private/Rigs/RigHierarchy.cpp
@@ -4299,11 +4299,29 @@ void URigHierarchy::SetControlValue(FRigControlElement* InControlElement, const
 
 				if(InValueType == ERigControlValueType::Minimum)
 				{
-					InControlElement->Settings.MinimumValue = InValue;
+					FRigControlSettings& Settings = InControlElement->Settings;
+					Settings.MinimumValue = InValue;
+
+					// Make sure the maximum value respects the new minimum value
+					TArray<FRigControlLimitEnabled> NoMaxLimits = Settings.LimitEnabled;
+					for (FRigControlLimitEnabled& NoMaxLimit : NoMaxLimits)
+					{
+						NoMaxLimit.bMaximum = false;
+					}
+					Settings.MaximumValue.ApplyLimits(NoMaxLimits, Settings.ControlType, Settings.MinimumValue, Settings.MaximumValue);
 				}
 				else
 				{
-					InControlElement->Settings.MaximumValue = InValue;
+					FRigControlSettings& Settings = InControlElement->Settings;
+					Settings.MaximumValue = InValue;
+					
+					// Make sure the minimum value respects the new maximum value
+					TArray<FRigControlLimitEnabled> NoMinLimits = Settings.LimitEnabled;
+					for (FRigControlLimitEnabled& NoMinLimit : NoMinLimits)
+					{
+						NoMinLimit.bMinimum = false;
+					}
+					Settings.MinimumValue.ApplyLimits(NoMinLimits, Settings.ControlType, Settings.MinimumValue, Settings.MaximumValue);
 				}
 				
 				Notify(ERigHierarchyNotification::ControlSettingChanged, InControlElement);
