Control Rig: Make sure the correct DOF are limited on Rotators when activated in the details panel
#rnx

diff --git a/Engine/Plugins/Animation/ControlRig/Source/ControlRigEditor/Private/ControlRigElementDetails.cpp b/Engine/Plugins/Animation/ControlRig/Source/ControlRigEditor/Private/ControlRigElementDetails.cpp
--- a/Engine/Plugins/Animation/ControlRig/Source/ControlRigEditor/Private/ControlRigElementDetails.cpp
+++ b/Engine/Plugins/Animation/ControlRig/Source/ControlRigEditor/Private/ControlRigElementDetails.cpp
@@ -1736,11 +1736,26 @@ FDetailWidgetRow& FRigTransformElementDetails::CreateEulerTransformValueWidgetRo
 								{
 									if(ControlElement->Settings.LimitEnabled.Num() == 3)
 									{
-										const int32 Index = ControlElement->Settings.ControlType == ERigControlType::Rotator ?
-											int32(SubComponent) - int32(ESlateTransformSubComponent::Pitch) :
-											int32(SubComponent) - int32(ESlateTransformSubComponent::X);
+										int32 Index = INDEX_NONE;
+										if (ControlElement->Settings.ControlType == ERigControlType::Rotator)
+										{
+											// TRotator is ordered Roll,Pitch,Yaw, while SNumericRotatorInputBox is ordered Pitch,Yaw,Roll
+											switch (SubComponent)
+											{
+												case ESlateTransformSubComponent::Pitch: Index = 1; break;
+												case ESlateTransformSubComponent::Yaw: Index = 2; break;
+												case ESlateTransformSubComponent::Roll: Index = 0; break;
+											}
+										}
+										else
+										{
+											Index = int32(SubComponent) - int32(ESlateTransformSubComponent::X);
+										}
 
-										Value = ControlElement->Settings.LimitEnabled[Index].GetForValueType(ValueType);
+										if (Index != INDEX_NONE)
+										{
+											Value = ControlElement->Settings.LimitEnabled[Index].GetForValueType(ValueType);
+										}
 									}
 									break;
 								}
@@ -1888,11 +1903,26 @@ FDetailWidgetRow& FRigTransformElementDetails::CreateEulerTransformValueWidgetRo
 							{
 								if(ControlElement->Settings.LimitEnabled.Num() == 3)
 								{
-									const int32 Index = ControlElement->Settings.ControlType == ERigControlType::Rotator ?
-										int32(SubComponent) - int32(ESlateTransformSubComponent::Pitch) :
-										int32(SubComponent) - int32(ESlateTransformSubComponent::X);
+									int32 Index = INDEX_NONE;
+									if (ControlElement->Settings.ControlType == ERigControlType::Rotator)
+									{
+										// TRotator is ordered Roll,Pitch,Yaw, while SNumericRotatorInputBox is ordered Pitch,Yaw,Roll
+										switch (SubComponent)
+										{
+											case ESlateTransformSubComponent::Pitch: Index = 1; break;
+											case ESlateTransformSubComponent::Yaw: Index = 2; break;
+											case ESlateTransformSubComponent::Roll: Index = 0; break;
+										}
+									}
+									else
+									{
+										Index = int32(SubComponent) - int32(ESlateTransformSubComponent::X);
+									}
 
-									ControlElement->Settings.LimitEnabled[Index].SetForValueType(ValueType, Value);
+									if (Index != INDEX_NONE)
+									{
+										ControlElement->Settings.LimitEnabled[Index].SetForValueType(ValueType, Value);
+									}
 								}
 								break;
 							}
