Constraints: Space Switches: Add buttons to compensate key or all for spaces like we can for constraints.
Had to do copypasta since new SControlRigSpacePicker was added and is copypasta-ing from SControlRigEditModeTools. Also note SControlRigSpacePicker needs to support multiple control rigs. Added BP functions for new functionality. Also added unused currently support for specifiy to compensate the previous tick, default behavior or the curren time. We may need that at some point in the future but not 100% sure so left it

diff --git a/Engine/Plugins/Animation/ControlRig/Source/ControlRigEditor/Private/ControlRigSequencerEditorLibrary.cpp b/Engine/Plugins/Animation/ControlRig/Source/ControlRigEditor/Private/ControlRigSequencerEditorLibrary.cpp
--- a/Engine/Plugins/Animation/ControlRig/Source/ControlRigEditor/Private/ControlRigSequencerEditorLibrary.cpp
+++ b/Engine/Plugins/Animation/ControlRig/Source/ControlRigEditor/Private/ControlRigSequencerEditorLibrary.cpp
@@ -549,7 +549,7 @@ bool UControlRigSequencerEditorLibrary::Compensate(UTickableConstraint* InConstr
 			InTime = FFrameRate::TransformTime(FFrameTime(InTime, 0), DisplayRate, TickResolution).RoundToFrame();
 		}
 		TOptional<FFrameNumber> OptTime(InTime);
-		FMovieSceneConstraintChannelHelper::Compensate(Sequencer, Constraint, OptTime);
+		FMovieSceneConstraintChannelHelper::Compensate(Sequencer, Constraint, OptTime, true /*bCompPreviousTick*/);
 		return true;
 	}
 	else
@@ -569,7 +569,7 @@ bool UControlRigSequencerEditorLibrary::CompensateAll(UTickableConstraint* InCon
 	}
 	if (UTickableTransformConstraint* Constraint = Cast<UTickableTransformConstraint>(InConstraint))
 	{
-		FMovieSceneConstraintChannelHelper::Compensate(WeakSequencer.Pin(), Constraint, TOptional<FFrameNumber>());
+		FMovieSceneConstraintChannelHelper::Compensate(WeakSequencer.Pin(), Constraint, TOptional<FFrameNumber>(), true /*bCompPreviousTick*/);
 		return true;
 	}
 	else
@@ -2760,6 +2760,60 @@ bool UControlRigSequencerEditorLibrary::SetControlRigSpace(ULevelSequence* Level
 	return bValid;
 }
 
+
+bool UControlRigSequencerEditorLibrary::SpaceCompensate(UControlRig* InControlRig, FFrameNumber InTime, EMovieSceneTimeUnit TimeUnit)
+{
+	TWeakPtr<ISequencer> WeakSequencer = GetSequencerFromAsset();
+	if (WeakSequencer.IsValid() && InControlRig)
+	{
+		TSharedPtr<ISequencer>  Sequencer = WeakSequencer.Pin();
+		TOptional<FFrameNumber> OptionalTime;
+		if (TimeUnit == EMovieSceneTimeUnit::DisplayRate)
+		{
+			const FFrameRate TickResolution = Sequencer->GetFocusedTickResolution();
+			const FFrameRate& FrameRate = Sequencer->GetFocusedDisplayRate();
+			OptionalTime = FFrameRate::TransformTime(FFrameTime(InTime, 0), FrameRate, TickResolution).RoundToFrame();
+		}
+		else
+		{
+			OptionalTime = InTime;
+		}
+		FScopedTransaction Transaction(LOCTEXT("SpacecCompensate", "Space Compensate"));
+		// compensate spaces
+		if (UMovieSceneControlRigParameterSection* CRSection = FControlRigSpaceChannelHelpers::GetControlRigSection(Sequencer.Get(), InControlRig))
+		{
+			// compensate spaces
+			FControlRigSpaceChannelHelpers::CompensateIfNeeded(
+				InControlRig, Sequencer.Get(), CRSection,
+				OptionalTime, true /*comp previous*/);
+		}
+		return true;
+	}
+	return false;
+}
+
+bool UControlRigSequencerEditorLibrary::SpaceCompensateAll(UControlRig* InControlRig)
+{
+	TWeakPtr<ISequencer> WeakSequencer = GetSequencerFromAsset();
+	bool bValid = false;
+	TSharedPtr<ISequencer>  Sequencer = WeakSequencer.Pin();
+	if (WeakSequencer.IsValid() && InControlRig)
+	{
+		TOptional<FFrameNumber> OptionalTime;
+		FScopedTransaction Transaction(LOCTEXT("SpacecCompensateAll", "Space Compensate All"));
+		// compensate spaces
+		if (UMovieSceneControlRigParameterSection* CRSection = FControlRigSpaceChannelHelpers::GetControlRigSection(Sequencer.Get(), InControlRig))
+		{
+			// compensate spaces
+			FControlRigSpaceChannelHelpers::CompensateIfNeeded(
+				InControlRig, Sequencer.Get(), CRSection,
+				OptionalTime, true /*comp previous*/);
+		}
+		return true;
+	}
+	return false;
+}
+
 bool UControlRigSequencerEditorLibrary::BakeControlRigSpace(ULevelSequence* InSequence, UControlRig* InControlRig, const TArray<FName>& InControlNames, FRigSpacePickerBakeSettings InSettings, EMovieSceneTimeUnit TimeUnit)
 {
 	TWeakPtr<ISequencer> WeakSequencer = GetSequencerFromAsset();
diff --git a/Engine/Plugins/Animation/ControlRig/Source/ControlRigEditor/Private/EditMode/SControlRigEditModeTools.cpp b/Engine/Plugins/Animation/ControlRig/Source/ControlRigEditor/Private/EditMode/SControlRigEditModeTools.cpp
--- a/Engine/Plugins/Animation/ControlRig/Source/ControlRigEditor/Private/EditMode/SControlRigEditModeTools.cpp
+++ b/Engine/Plugins/Animation/ControlRig/Source/ControlRigEditor/Private/EditMode/SControlRigEditModeTools.cpp
@@ -364,12 +364,13 @@ void SControlRigEditModeTools::Construct(const FArguments& InArgs, TSharedPtr<FC
 					.AllowDelete(true)
 					.AllowReorder(true)
 					.AllowAdd(false)
-					.ShowBakeButton(true)
+					.ShowBakeAndCompensateButton(true)
 					.GetControlCustomization(this, &SControlRigEditModeTools::HandleGetControlElementCustomization)
 					.OnActiveSpaceChanged(this, &SControlRigEditModeTools::HandleActiveSpaceChanged)
 					.OnSpaceListChanged(this, &SControlRigEditModeTools::HandleSpaceListChanged)
-					.OnBakeButtonClicked(this, &SControlRigEditModeTools::OnBakeControlsToNewSpaceButtonClicked)
-					// todo: implement GetAdditionalSpacesDelegate to pull spaces from sequencer
+					.OnCompensateKeyButtonClicked(this, &SControlRigEditModeTools::OnCompensateKeyClicked)
+					.OnCompensateAllButtonClicked(this, &SControlRigEditModeTools::OnCompensateAllClicked)
+					.OnBakeButtonClicked(this, &SControlRigEditModeTools::OnBakeControlsToNewSpaceButtonClicked)					// todo: implement GetAdditionalSpacesDelegate to pull spaces from sequencer
 				]
 			]
 
@@ -1008,19 +1009,19 @@ bool SControlRigEditModeTools::IsSpaceSwitchingRestricted() const
 	return SpacePickerWidget->IsRestricted();
 }
 
-FReply SControlRigEditModeTools::OnBakeControlsToNewSpaceButtonClicked()
+bool SControlRigEditModeTools::ReadyForBakeOrCompensation() const
 {
 	if (SpacePickerWidget->GetHierarchy() == nullptr)
 	{
-		return FReply::Unhandled();
+		return false;
 	}
 	if (SpacePickerWidget->GetControls().Num() == 0)
 	{
-		return FReply::Unhandled();
+		return false;
 	}
 
 	bool bNoValidControlRig = true;
-	for (TWeakObjectPtr<UControlRig>& ControlRig : ControlRigs)
+	for (const TWeakObjectPtr<UControlRig>& ControlRig : ControlRigs)
 	{
 		if (ControlRig.IsValid() && SpacePickerWidget->GetHierarchy() == ControlRig->GetHierarchy())
 		{
@@ -1031,27 +1032,82 @@ FReply SControlRigEditModeTools::OnBakeControlsToNewSpaceButtonClicked()
 
 	if (bNoValidControlRig)
 	{
-		return FReply::Unhandled();
+		return false;
 	}
 	ISequencer* Sequencer = WeakSequencer.Pin().Get();
 	if (Sequencer == nullptr || Sequencer->GetFocusedMovieSceneSequence() == nullptr || Sequencer->GetFocusedMovieSceneSequence()->GetMovieScene() == nullptr)
+	{
+		return false;
+	}
+	return true;
+}
+
+FReply SControlRigEditModeTools::OnCompensateKeyClicked()
+{
+	if (ReadyForBakeOrCompensation() == false)
 	{
 		return FReply::Unhandled();
 	}
+	ISequencer* Sequencer = WeakSequencer.Pin().Get();
+	const FFrameRate TickResolution = Sequencer->GetFocusedTickResolution();
+	const FFrameTime FrameTime = Sequencer->GetLocalTime().ConvertTo(TickResolution);
+	const TOptional<FFrameNumber> OptionalKeyTime = FrameTime.GetFrame();
+	const bool bSetPreviousKey = true;
+	Compensate(OptionalKeyTime, bSetPreviousKey);
+	return FReply::Handled();
+}
+
+FReply SControlRigEditModeTools::OnCompensateAllClicked()
+{
+	if (ReadyForBakeOrCompensation() == false)
+	{
+		return FReply::Unhandled();
+	}
+	const TOptional<FFrameNumber> OptionalKeyTime;
+	ISequencer* Sequencer = WeakSequencer.Pin().Get();
+	const bool bSetPreviousKey = true;
+	Compensate(OptionalKeyTime, bSetPreviousKey);
+	return FReply::Handled();
+}
+
+void SControlRigEditModeTools::Compensate(TOptional<FFrameNumber> OptionalKeyTime, bool bSetPreviousTick)
+{
+	if (ReadyForBakeOrCompensation() == false)
+	{
+		return;
+	}
+	ISequencer* Sequencer = WeakSequencer.Pin().Get();
 	for (TWeakObjectPtr<UControlRig>& ControlRig : ControlRigs)
 	{
 		if (ControlRig.IsValid() && SpacePickerWidget->GetHierarchy() == ControlRig->GetHierarchy())
 		{
+			// compensate spaces
+			if(UMovieSceneControlRigParameterSection* CRSection = FControlRigSpaceChannelHelpers::GetControlRigSection(Sequencer, ControlRig.Get()))
+			{	
+				// compensate spaces
+				FControlRigSpaceChannelHelpers::CompensateIfNeeded(
+					ControlRig.Get(), Sequencer, CRSection,
+				OptionalKeyTime, bSetPreviousTick);
+			}
+		}
+	}
+}
+
+FReply SControlRigEditModeTools::OnBakeControlsToNewSpaceButtonClicked()
+{
+	if (ReadyForBakeOrCompensation() == false)
+	{
+		return FReply::Unhandled();
+	}
+	ISequencer* Sequencer = WeakSequencer.Pin().Get();
 
-			//Find default target space, just use first control and find space at current sequencer time
-			//Then Find range
+	for (TWeakObjectPtr<UControlRig>& ControlRig : ControlRigs)
+	{
+		if (ControlRig.IsValid() && SpacePickerWidget->GetHierarchy() == ControlRig->GetHierarchy())
+		{
 
-			// FindSpaceChannelAndSectionForControl() will trigger RecreateCurveEditor(), which will deselect the controls
-			// but in theory the selection will be recovered in the next tick, so here we just cache the selected controls
-			// and use it throughout this function. If this deselection is causing other problems, this part could use a revisit.
 			TArray<FRigElementKey> ControlKeys = SpacePickerWidget->GetControls();
-
-			FSpaceChannelAndSection SpaceChannelAndSection = FControlRigSpaceChannelHelpers::FindSpaceChannelAndSectionForControl(ControlRig.Get(), ControlKeys[0].Name, Sequencer, true /*bCreateIfNeeded*/);
+			FSpaceChannelAndSection SpaceChannelAndSection = FControlRigSpaceChannelHelpers::FindSpaceChannelAndSectionForControl(ControlRig.Get(), ControlKeys[0].Name, Sequencer, false /*bCreateIfNeeded*/);
 			if (SpaceChannelAndSection.SpaceChannel != nullptr)
 			{
 				const FFrameRate TickResolution = Sequencer->GetFocusedTickResolution();
diff --git a/Engine/Plugins/Animation/ControlRig/Source/ControlRigEditor/Private/EditMode/SControlRigEditModeTools.h b/Engine/Plugins/Animation/ControlRig/Source/ControlRigEditor/Private/EditMode/SControlRigEditModeTools.h
--- a/Engine/Plugins/Animation/ControlRig/Source/ControlRigEditor/Private/EditMode/SControlRigEditModeTools.h
+++ b/Engine/Plugins/Animation/ControlRig/Source/ControlRigEditor/Private/EditMode/SControlRigEditModeTools.h
@@ -121,7 +121,10 @@ private:
 	EVisibility GetAddSpaceButtonVisibility() const;
 	bool IsSpaceSwitchingRestricted() const;
 	FReply OnBakeControlsToNewSpaceButtonClicked();
-
+	FReply OnCompensateKeyClicked();
+	FReply OnCompensateAllClicked();
+	void Compensate(TOptional<FFrameNumber> OptionalKeyTime, bool bSetPreviousTick);
+	bool ReadyForBakeOrCompensation() const;
 	FReply HandleAddConstraintClicked();
 
 	EVisibility GetRigOptionExpanderVisibility() const;
diff --git a/Engine/Plugins/Animation/ControlRig/Source/ControlRigEditor/Private/EditMode/SControlRigSpacePicker.cpp b/Engine/Plugins/Animation/ControlRig/Source/ControlRigEditor/Private/EditMode/SControlRigSpacePicker.cpp
--- a/Engine/Plugins/Animation/ControlRig/Source/ControlRigEditor/Private/EditMode/SControlRigSpacePicker.cpp
+++ b/Engine/Plugins/Animation/ControlRig/Source/ControlRigEditor/Private/EditMode/SControlRigSpacePicker.cpp
@@ -88,10 +88,12 @@ void SControlRigSpacePicker::Construct(const FArguments& InArgs, FControlRigEdit
 						.AllowDelete(true)
 						.AllowReorder(true)
 						.AllowAdd(false)
-						.ShowBakeButton(true)
+						.ShowBakeAndCompensateButton(true)
 						.GetControlCustomization(this, &SControlRigSpacePicker::HandleGetControlElementCustomization)
 						.OnActiveSpaceChanged(this, &SControlRigSpacePicker::HandleActiveSpaceChanged)
 						.OnSpaceListChanged(this, &SControlRigSpacePicker::HandleSpaceListChanged)
+						.OnCompensateKeyButtonClicked(this, &SControlRigSpacePicker::OnCompensateKeyClicked)
+						.OnCompensateAllButtonClicked(this, &SControlRigSpacePicker::OnCompensateAllClicked)
 						.OnBakeButtonClicked(this, &SControlRigSpacePicker::OnBakeControlsToNewSpaceButtonClicked)
 						// todo: implement GetAdditionalSpacesDelegate to pull spaces from sequencer
 					]
@@ -108,7 +110,7 @@ SControlRigSpacePicker::~SControlRigSpacePicker()
 	//base class handles control rig related cleanup
 }
 
-UControlRig* SControlRigSpacePicker::GetControlRig()
+UControlRig* SControlRigSpacePicker::GetControlRig() const
 {
 	TArray<UControlRig*> ControlRigs = GetControlRigs();
 	for(UControlRig* ControlRig: ControlRigs)
@@ -234,25 +236,88 @@ FReply SControlRigSpacePicker::HandleAddSpaceClicked()
 	return SpacePickerWidget->HandleAddElementClicked();
 }
 
-FReply SControlRigSpacePicker::OnBakeControlsToNewSpaceButtonClicked()
+bool SControlRigSpacePicker::ReadyForBakeOrCompensation() const
 {
 	if (SpacePickerWidget->GetHierarchy() == nullptr)
 	{
-		return FReply::Unhandled();
+		return false;
 	}
 	if (SpacePickerWidget->GetControls().Num() == 0)
 	{
-		return FReply::Unhandled();
+		return false;
 	}
 	if (!GetControlRig())
 	{
-		return FReply::Unhandled();
+		return false;
 	}
 	ISequencer* Sequencer = GetSequencer();
 	if (Sequencer == nullptr || Sequencer->GetFocusedMovieSceneSequence() == nullptr || Sequencer->GetFocusedMovieSceneSequence()->GetMovieScene() == nullptr)
+	{
+		return false;
+	}
+	return true;
+}
+
+FReply SControlRigSpacePicker::OnCompensateKeyClicked()
+{
+	if (ReadyForBakeOrCompensation() == false)
+	{
+		return FReply::Unhandled();
+	}
+	ISequencer* Sequencer = GetSequencer();
+	const FFrameRate TickResolution = Sequencer->GetFocusedTickResolution();
+	const FFrameTime FrameTime = Sequencer->GetLocalTime().ConvertTo(TickResolution);
+	const TOptional<FFrameNumber> OptionalKeyTime = FrameTime.GetFrame();
+	const bool bSetPreviousKey = true;
+	Compensate(OptionalKeyTime, bSetPreviousKey);
+	return FReply::Handled();
+}
+
+FReply SControlRigSpacePicker::OnCompensateAllClicked()
+{
+	if (ReadyForBakeOrCompensation() == false)
 	{
 		return FReply::Unhandled();
 	}
+	const TOptional<FFrameNumber> OptionalKeyTime;
+	ISequencer* Sequencer = GetSequencer();
+	const bool bSetPreviousKey = true;
+	Compensate(OptionalKeyTime, bSetPreviousKey);
+	return FReply::Handled();
+}
+
+
+
+void SControlRigSpacePicker::Compensate(TOptional<FFrameNumber> OptionalKeyTime, bool bSetPreviousTick)
+{
+	if (ReadyForBakeOrCompensation() == false)
+	{
+		return;
+	}
+	ISequencer* Sequencer = GetSequencer();
+	UControlRig* ControlRig = GetControlRig(); //!!!!! THIS SHOULD SUPPORT MULTIPLE!
+
+	if (ControlRig && SpacePickerWidget->GetHierarchy() == ControlRig->GetHierarchy())
+	{
+		// compensate spaces
+		if (UMovieSceneControlRigParameterSection* CRSection = FControlRigSpaceChannelHelpers::GetControlRigSection(Sequencer, ControlRig))
+		{
+			// compensate spaces
+			FControlRigSpaceChannelHelpers::CompensateIfNeeded(
+				ControlRig, Sequencer, CRSection,
+				OptionalKeyTime, bSetPreviousTick);
+		}
+	}
+}
+
+FReply SControlRigSpacePicker::OnBakeControlsToNewSpaceButtonClicked()
+{
+	if (ReadyForBakeOrCompensation() == false)
+	{
+		return FReply::Unhandled();
+	}
+
+	ISequencer* Sequencer = GetSequencer();
 	UControlRig* ControlRig = GetControlRig();
 
 	FRigSpacePickerBakeSettings Settings;
diff --git a/Engine/Plugins/Animation/ControlRig/Source/ControlRigEditor/Private/EditMode/SControlRigSpacePicker.h b/Engine/Plugins/Animation/ControlRig/Source/ControlRigEditor/Private/EditMode/SControlRigSpacePicker.h
--- a/Engine/Plugins/Animation/ControlRig/Source/ControlRigEditor/Private/EditMode/SControlRigSpacePicker.h
+++ b/Engine/Plugins/Animation/ControlRig/Source/ControlRigEditor/Private/EditMode/SControlRigSpacePicker.h
@@ -40,10 +40,14 @@ private:
 	void HandleSpaceListChanged(URigHierarchy* InHierarchy, const FRigElementKey& InControlKey, const TArray<FRigElementKey>& InSpaceList);
 	FReply HandleAddSpaceClicked();
 	FReply OnBakeControlsToNewSpaceButtonClicked();
+	FReply OnCompensateKeyClicked();
+	FReply OnCompensateAllClicked();
+	void Compensate(TOptional<FFrameNumber> OptionalKeyTime, bool bSetPreviousTick);
 	EVisibility GetAddSpaceButtonVisibility() const;
 
+	bool ReadyForBakeOrCompensation() const;
 	//for now picker works off of one ControlRig, this function gets the first control rig with a selection
-	UControlRig* GetControlRig();
+	UControlRig* GetControlRig() const;
 
 };
 
diff --git a/Engine/Plugins/Animation/ControlRig/Source/ControlRigEditor/Private/Editor/Constraints/SConstraintsWidget.cpp b/Engine/Plugins/Animation/ControlRig/Source/ControlRigEditor/Private/Editor/Constraints/SConstraintsWidget.cpp
--- a/Engine/Plugins/Animation/ControlRig/Source/ControlRigEditor/Private/Editor/Constraints/SConstraintsWidget.cpp
+++ b/Engine/Plugins/Animation/ControlRig/Source/ControlRigEditor/Private/Editor/Constraints/SConstraintsWidget.cpp
@@ -1476,7 +1476,7 @@ TSharedPtr<SWidget> SConstraintsEditionWidget::CreateContextMenu()
 					const FFrameRate TickResolution = Sequencer->GetFocusedTickResolution();
 					const FFrameTime FrameTime = Sequencer->GetLocalTime().ConvertTo(TickResolution);
 					const FFrameNumber Time = FrameTime.GetFrame();
-					FMovieSceneConstraintChannelHelper::Compensate(WeakSequencer.Pin(), TransformConstraint, TOptional<FFrameNumber>(Time));
+					FMovieSceneConstraintChannelHelper::Compensate(WeakSequencer.Pin(), TransformConstraint, TOptional<FFrameNumber>(Time), true /*bCompPreviousTick*/);
 				}), IsCompensationEnabled),
 				NAME_None,
 				EUserInterfaceActionType::Button);
@@ -1487,7 +1487,7 @@ TSharedPtr<SWidget> SConstraintsEditionWidget::CreateContextMenu()
 				FSlateIcon(),
 				FUIAction(FExecuteAction::CreateLambda([TransformConstraint, this]()
 				{
-					FMovieSceneConstraintChannelHelper::Compensate(WeakSequencer.Pin(), TransformConstraint, TOptional<FFrameNumber>());
+					FMovieSceneConstraintChannelHelper::Compensate(WeakSequencer.Pin(), TransformConstraint, TOptional<FFrameNumber>(), true /*bCompPreviousTick*/);
 				}), IsCompensationEnabled),
 				NAME_None,
 				EUserInterfaceActionType::Button);
diff --git a/Engine/Plugins/Animation/ControlRig/Source/ControlRigEditor/Private/Editor/SRigSpacePickerWidget.cpp b/Engine/Plugins/Animation/ControlRig/Source/ControlRigEditor/Private/Editor/SRigSpacePickerWidget.cpp
--- a/Engine/Plugins/Animation/ControlRig/Source/ControlRigEditor/Private/Editor/SRigSpacePickerWidget.cpp
+++ b/Engine/Plugins/Animation/ControlRig/Source/ControlRigEditor/Private/Editor/SRigSpacePickerWidget.cpp
@@ -40,7 +40,7 @@ void SRigSpacePickerWidget::Construct(const FArguments& InArgs)
 	bAllowReorder = InArgs._AllowReorder;
 	bAllowDelete = InArgs._AllowDelete;
 	bAllowAdd = InArgs._AllowAdd;
-	bShowBakeButton = InArgs._ShowBakeButton;
+	bShowBakeAndCompensateButton = InArgs._ShowBakeAndCompensateButton;
 	GetActiveSpaceDelegate = InArgs._GetActiveSpace;
 	GetControlCustomizationDelegate = InArgs._GetControlCustomization;
 	GetAdditionalSpacesDelegate = InArgs._GetAdditionalSpaces;
@@ -132,7 +132,7 @@ void SRigSpacePickerWidget::Construct(const FArguments& InArgs)
 		]
 	];
 
-	if(bAllowAdd || bShowBakeButton)
+	if(bAllowAdd || bShowBakeAndCompensateButton)
 	{
 		TopLevelListBox->AddSlot()
 		.AutoHeight()
@@ -175,8 +175,40 @@ void SRigSpacePickerWidget::Construct(const FArguments& InArgs)
 			SNew(SSpacer)
 		];
 
-		if(bShowBakeButton)
+		if(bShowBakeAndCompensateButton)
 		{
+			BottomButtonsListBox->AddSlot()
+			.AutoWidth()
+			.VAlign(VAlign_Center)
+			.HAlign(HAlign_Left)
+			.Padding(0.f)
+			[
+				SNew(SButton)
+				.ButtonStyle(FAppStyle::Get(), "FlatButton.Default")
+				.Text(LOCTEXT("CompensateKeyButton", "Comp Key"))
+				.OnClicked(InArgs._OnCompensateKeyButtonClicked)
+				.IsEnabled_Lambda([this]()
+				{
+					return (ControlKeys.Num() > 0);
+				})
+				.ToolTipText(LOCTEXT("CompensateKeyTooltip", "Compensate key at the current time."))
+			];
+			BottomButtonsListBox->AddSlot()
+			.AutoWidth()
+			.VAlign(VAlign_Center)
+			.HAlign(HAlign_Left)
+			.Padding(0.f)
+			[
+				SNew(SButton)
+				.ButtonStyle(FAppStyle::Get(), "FlatButton.Default")
+				.Text(LOCTEXT("CompensateAllButton", "Comp All"))
+				.OnClicked(InArgs._OnCompensateAllButtonClicked)
+				.IsEnabled_Lambda([this]()
+				{
+					return (ControlKeys.Num() > 0);
+				})
+				.ToolTipText(LOCTEXT("CompensateAllTooltip", "Compensate all space switch keys."))
+			];
 			BottomButtonsListBox->AddSlot()
 			.AutoWidth()
 			.VAlign(VAlign_Center)
@@ -1166,7 +1198,7 @@ void SRigSpacePickerBakeWidget::Construct(const FArguments& InArgs)
 				.AllowDelete(false)
 				.AllowReorder(false)
 				.AllowAdd(true)
-				.ShowBakeButton(false)
+				.ShowBakeAndCompensateButton(false)
 				.GetControlCustomization_Lambda([this] (URigHierarchy*, const FRigElementKey)
 				{
 					return &Customization;
diff --git a/Engine/Plugins/Animation/ControlRig/Source/ControlRigEditor/Private/Editor/SRigSpacePickerWidget.h b/Engine/Plugins/Animation/ControlRig/Source/ControlRigEditor/Private/Editor/SRigSpacePickerWidget.h
--- a/Engine/Plugins/Animation/ControlRig/Source/ControlRigEditor/Private/Editor/SRigSpacePickerWidget.h
+++ b/Engine/Plugins/Animation/ControlRig/Source/ControlRigEditor/Private/Editor/SRigSpacePickerWidget.h
@@ -35,7 +35,7 @@ public:
 		, _AllowReorder(false)
 		, _AllowDelete(false)
 		, _AllowAdd(false)
-		, _ShowBakeButton(false)
+		, _ShowBakeAndCompensateButton(false)
 		, _Title()
 		, _BackgroundBrush(FAppStyle::GetBrush("Menu.Background"))
 		{}
@@ -47,7 +47,7 @@ public:
 		SLATE_ARGUMENT(bool, AllowReorder)
 		SLATE_ARGUMENT(bool, AllowDelete)
 		SLATE_ARGUMENT(bool, AllowAdd)
-		SLATE_ARGUMENT(bool, ShowBakeButton)
+		SLATE_ARGUMENT(bool, ShowBakeAndCompensateButton)
 		SLATE_ARGUMENT(FText, Title)
 		SLATE_ARGUMENT(const FSlateBrush*, BackgroundBrush)
 	
@@ -56,7 +56,9 @@ public:
 		SLATE_EVENT(FRigSpacePickerActiveSpaceChanged::FDelegate, OnActiveSpaceChanged)
 		SLATE_EVENT(FRigSpacePickerSpaceListChanged::FDelegate, OnSpaceListChanged)
 		SLATE_ARGUMENT(FRigSpacePickerGetAdditionalSpaces, GetAdditionalSpaces)
-		SLATE_EVENT( FOnClicked, OnBakeButtonClicked )
+		SLATE_EVENT(FOnClicked, OnCompensateKeyButtonClicked)
+		SLATE_EVENT(FOnClicked, OnCompensateAllButtonClicked)
+		SLATE_EVENT(FOnClicked, OnBakeButtonClicked )
 	SLATE_END_ARGS()
 
 	void Construct(const FArguments& InArgs);
@@ -156,7 +158,7 @@ private:
 	bool bAllowReorder;
 	bool bAllowDelete;
 	bool bAllowAdd;
-	bool bShowBakeButton;
+	bool bShowBakeAndCompensateButton;
 	bool bLaunchingContextMenu;
 
 	FRigSpacePickerGetControlCustomization GetControlCustomizationDelegate;
diff --git a/Engine/Plugins/Animation/ControlRig/Source/ControlRigEditor/Private/Sequencer/ControlRigParameterTrackEditor.cpp b/Engine/Plugins/Animation/ControlRig/Source/ControlRigEditor/Private/Sequencer/ControlRigParameterTrackEditor.cpp
--- a/Engine/Plugins/Animation/ControlRig/Source/ControlRigEditor/Private/Sequencer/ControlRigParameterTrackEditor.cpp
+++ b/Engine/Plugins/Animation/ControlRig/Source/ControlRigEditor/Private/Sequencer/ControlRigParameterTrackEditor.cpp
@@ -3782,11 +3782,11 @@ FKeyPropertyResult FControlRigParameterTrackEditor::AddKeysToControlRigHandle(UO
 				// compensate spaces
 				FControlRigSpaceChannelHelpers::CompensateIfNeeded(
 					SectionControlRig, GetSequencer().Get(), ParamSection,
-					RigControlName, OptionalKeyTime);
+					OptionalKeyTime, true /*comp previous*/);
 
 				// compensate constraints
 				const uint32 ControlHash = UTransformableControlHandle::ComputeHash(SectionControlRig, RigControlName);
-				FMovieSceneConstraintChannelHelper::CompensateIfNeeded(GetSequencer(), ParamSection, OptionalKeyTime, ControlHash);
+				FMovieSceneConstraintChannelHelper::CompensateIfNeeded(GetSequencer(), ParamSection, OptionalKeyTime, true /*bCompPreviousTick*/, ControlHash);
 			}
 		}
 	}
diff --git a/Engine/Plugins/Animation/ControlRig/Source/ControlRigEditor/Private/Sequencer/ControlRigSpaceChannelEditors.cpp b/Engine/Plugins/Animation/ControlRig/Source/ControlRigEditor/Private/Sequencer/ControlRigSpaceChannelEditors.cpp
--- a/Engine/Plugins/Animation/ControlRig/Source/ControlRigEditor/Private/Sequencer/ControlRigSpaceChannelEditors.cpp
+++ b/Engine/Plugins/Animation/ControlRig/Source/ControlRigEditor/Private/Sequencer/ControlRigSpaceChannelEditors.cpp
@@ -1273,13 +1273,40 @@ void FControlRigSpaceChannelHelpers::HandleSpaceKeyTimeChanged(UControlRig* Cont
 	}
 }
 
+UMovieSceneControlRigParameterSection* FControlRigSpaceChannelHelpers::GetControlRigSection(ISequencer* Sequencer, const UControlRig* ControlRig)
+{
+
+	if (ControlRig == nullptr || Sequencer == nullptr)
+	{
+		return nullptr;
+	}
+	UMovieScene* MovieScene = Sequencer->GetFocusedMovieSceneSequence()->GetMovieScene();
+	if (!MovieScene)
+	{
+		return nullptr;
+	}
+	const TArray<FMovieSceneBinding>& Bindings = MovieScene->GetBindings();
+	for (const FMovieSceneBinding& Binding : Bindings)
+	{
+		UMovieSceneControlRigParameterTrack* ControlRigParameterTrack = Cast<UMovieSceneControlRigParameterTrack>(MovieScene->FindTrack(UMovieSceneControlRigParameterTrack::StaticClass(), Binding.GetObjectGuid(), NAME_None));
+		if (ControlRigParameterTrack && ControlRigParameterTrack->GetControlRig() == ControlRig)
+		{
+			UMovieSceneControlRigParameterSection* ActiveSection = Cast<UMovieSceneControlRigParameterSection>(ControlRigParameterTrack->GetSectionToKey());
+			if (ActiveSection)
+			{
+				return ActiveSection;
+			}
+		}
+	}
+	return nullptr;
+}
 
 void FControlRigSpaceChannelHelpers::CompensateIfNeeded(
 	UControlRig* ControlRig,
 	ISequencer* Sequencer,
 	UMovieSceneControlRigParameterSection* Section,
-	FName ControlName,
-	TOptional<FFrameNumber>& OptionalTime)
+	TOptional<FFrameNumber>& OptionalTime, 
+	bool bCompPreviousTick)
 {
 	if (bDoNotCompensate == true)
 	{
@@ -1315,7 +1342,7 @@ void FControlRigSpaceChannelHelpers::CompensateIfNeeded(
 	const TArray<FRigControlElement*> Controls = RigHierarchy->GetControls();
 	for (const FRigControlElement* Control: Controls)
 	{ 
-		if(Control)// ac && Control->GetName() != ControlName)
+		if(Control)
 		{ 
 			//only if we have a channel
 			if (FSpaceControlNameAndChannel* Channel = Section->GetSpaceChannel(Control->GetFName()))
@@ -1325,10 +1352,12 @@ void FControlRigSpaceChannelHelpers::CompensateIfNeeded(
 				{
 					for (const FFrameNumber& Time : FramesToCompensate)
 					{
+						const FFrameNumber TimeToCompensate = bCompPreviousTick ? (Time - 1) : Time;
+						const FFrameNumber TimeToCompare = bCompPreviousTick ? Time : (Time - 1);
 						FMovieSceneControlRigSpaceBaseKey ExistingValue, PreviousValue;
 						using namespace UE::MovieScene;
-						EvaluateChannel(&(Channel->SpaceCurve), Time - 1, PreviousValue);
-						EvaluateChannel(&(Channel->SpaceCurve), Time, ExistingValue);
+						EvaluateChannel(&(Channel->SpaceCurve), TimeToCompensate, PreviousValue);
+						EvaluateChannel(&(Channel->SpaceCurve), TimeToCompare, ExistingValue);
 
 						if (ExistingValue != PreviousValue) //if they are the same no need to do anything
 						{
@@ -1340,7 +1369,7 @@ void FControlRigSpaceChannelHelpers::CompensateIfNeeded(
 							FControlRigSnapper Snapper;
 							Snapper.GetControlRigControlTransforms(
 								Sequencer, ControlRig, Control->GetFName(),
-								{Time},
+								{ TimeToCompare },
 								ControlRigParentWorldTransforms, ControlWorldTransforms);
 
 							//set space to previous space value that's different.
@@ -1360,7 +1389,7 @@ void FControlRigSpaceChannelHelpers::CompensateIfNeeded(
 							
 							//now set time -1 frame value
 							ControlRig->Evaluate_AnyThread();
-							KeyframeContext.LocalTime = TickResolution.AsSeconds(FFrameTime(Time - 1));
+							KeyframeContext.LocalTime = TickResolution.AsSeconds(FFrameTime(TimeToCompensate));
 							ControlRig->SetControlGlobalTransform(Control->GetFName(), ControlWorldTransforms[0], true, KeyframeContext, false /*undo*/, false /*bPrintPython*/, true/* bFixEulerFlips*/);
 							
 							bDidIt = true;
diff --git a/Engine/Plugins/Animation/ControlRig/Source/ControlRigEditor/Public/ControlRigSequencerEditorLibrary.h b/Engine/Plugins/Animation/ControlRig/Source/ControlRigEditor/Public/ControlRigSequencerEditorLibrary.h
--- a/Engine/Plugins/Animation/ControlRig/Source/ControlRigEditor/Public/ControlRigSequencerEditorLibrary.h
+++ b/Engine/Plugins/Animation/ControlRig/Source/ControlRigEditor/Public/ControlRigSequencerEditorLibrary.h
@@ -1009,6 +1009,21 @@ public:
 	UFUNCTION(BlueprintCallable, Category = "Editor Scripting | Sequencer Tools | Control Rig")
 	static bool BakeControlRigSpace(ULevelSequence* InSequence, UControlRig* InControlRig, const TArray<FName>& InControlNames, FRigSpacePickerBakeSettings InSettings, EMovieSceneTimeUnit TimeUnit = EMovieSceneTimeUnit::DisplayRate);
 	
+	/** Perform compensation for any spaces at the specified time for the specified control rig
+	* @param InControlRig Control Rig to compensate
+	* @param InTime  The time to look for a space key to compensate
+	* @param TimeUnit Unit for the InTime
+	* @return Will return false if function fails
+	*/
+	UFUNCTION(BlueprintCallable, Category = "Editor Scripting | Sequencer Tools | Control Rig")
+	static bool SpaceCompensate(UControlRig* InControlRig, FFrameNumber InTime, EMovieSceneTimeUnit TimeUnit = EMovieSceneTimeUnit::DisplayRate);
+
+	/** Perform compensation for all spaces for the specified control rig
+	* @param InControlRig Control Rig to compensate
+	* @return Will return false if function fails
+	*/
+	static bool SpaceCompensateAll(UControlRig* InControlRig);
+
 	/** Delete the Control Rig Space Key for the Control at the specified time. This will delete any attached Control Rig keys at this time and will perform any needed compensation to the new space.
 	*
 	* @param InSequence Sequence to set the space
diff --git a/Engine/Plugins/Animation/ControlRig/Source/ControlRigEditor/Public/ControlRigSpaceChannelEditors.h b/Engine/Plugins/Animation/ControlRig/Source/ControlRigEditor/Public/ControlRigSpaceChannelEditors.h
--- a/Engine/Plugins/Animation/ControlRig/Source/ControlRigEditor/Public/ControlRigSpaceChannelEditors.h
+++ b/Engine/Plugins/Animation/ControlRig/Source/ControlRigEditor/Public/ControlRigSpaceChannelEditors.h
@@ -46,7 +46,8 @@ struct FControlRigSpaceChannelHelpers
 {
 	static FKeyHandle SequencerKeyControlRigSpaceChannel(UControlRig* ControlRig, ISequencer* Sequencer, FMovieSceneControlRigSpaceChannel* Channel, UMovieSceneSection* SectionToKey, FFrameNumber Time, URigHierarchy* RigHierarchy, const FRigElementKey& ControlKey, const FRigElementKey& SpaceKey);
 	static void SequencerSpaceChannelKeyDeleted(UControlRig* ControlRig, ISequencer* Sequencer, FName ControlName, FMovieSceneControlRigSpaceChannel* Channel, UMovieSceneControlRigParameterSection* SectionToKey, FFrameNumber TimeOfDeletion);
-	static void CompensateIfNeeded(UControlRig* ControlRig, ISequencer* Sequencer, UMovieSceneControlRigParameterSection* Section, FName ControlName, TOptional<FFrameNumber>& Time);
+	static UMovieSceneControlRigParameterSection* GetControlRigSection(ISequencer* Sequencer, const UControlRig* ControlRig);
+	static void CompensateIfNeeded(UControlRig* ControlRig, ISequencer* Sequencer, UMovieSceneControlRigParameterSection* Section, TOptional<FFrameNumber>& Time, bool bCompPreviousTick);
 	static FSpaceChannelAndSection FindSpaceChannelAndSectionForControl(UControlRig* ControlRig, FName ControlName, ISequencer* Sequencer, bool bCreateIfNeeded);
 	static void SequencerBakeControlInSpace(UControlRig* ControlRig, ISequencer* Sequencer, FMovieSceneControlRigSpaceChannel* Channel, UMovieSceneSection* SectionToKey,
 		URigHierarchy* RigHierarchy, const FRigElementKey& ControlKey, FRigSpacePickerBakeSettings InSettings);
diff --git a/Engine/Source/Editor/MovieSceneTools/Private/Constraints/MovieSceneConstraintChannelHelper.cpp b/Engine/Source/Editor/MovieSceneTools/Private/Constraints/MovieSceneConstraintChannelHelper.cpp
--- a/Engine/Source/Editor/MovieSceneTools/Private/Constraints/MovieSceneConstraintChannelHelper.cpp
+++ b/Engine/Source/Editor/MovieSceneTools/Private/Constraints/MovieSceneConstraintChannelHelper.cpp
@@ -970,7 +970,8 @@ bool FMovieSceneConstraintChannelHelper::SmartConstraintKey(
 void FMovieSceneConstraintChannelHelper::Compensate(
 	const TSharedPtr<ISequencer>& InSequencer,
 	const UTickableTransformConstraint* InConstraint,
-	const TOptional<FFrameNumber>& InOptTime)
+	const TOptional<FFrameNumber>& InOptTime,
+	const bool bCompPreviousTick)
 {
 	if (!InSequencer.IsValid() || !InSequencer->GetFocusedMovieSceneSequence())
 	{
@@ -993,13 +994,14 @@ void FMovieSceneConstraintChannelHelper::Compensate(
 		return;
 	}
 
-	CompensateIfNeeded(InSequencer, Section, InOptTime, Handle->GetHash());
+	CompensateIfNeeded(InSequencer, Section, InOptTime, bCompPreviousTick, Handle->GetHash());
 }
 
 void FMovieSceneConstraintChannelHelper::CompensateIfNeeded(
 	const TSharedPtr<ISequencer>& InSequencer,
 	IMovieSceneConstrainedSection* ConstraintSection,
 	const TOptional<FFrameNumber>& OptionalTime,
+	const bool bCompPreviousTick,
 	const int32 InChildHash)
 {
 	if (bDoNotCompensate)
@@ -1101,14 +1103,15 @@ void FMovieSceneConstraintChannelHelper::CompensateIfNeeded(
 			const EMovieSceneTransformChannel ChannelsToKey = Constraint->GetChannelsToKey();
 			for (const FFrameNumber& Time : Data.Value)
 			{
+				const FFrameNumber EvalTime = bCompPreviousTick ? Time : Time - 1;
+				const FFrameNumber SetTime = bCompPreviousTick ? Time - 1 : Time;
 				// compute transform to set
 				// if switching from active to inactive then we must add a key at T-1 in the constraint space
 				// if switching from inactive to active then we must add a key at T-1 in the previous constraint or parent space
-				Evaluator.ComputeCompensation(World, InSequencer, Time);
+				Evaluator.ComputeCompensation(World, InSequencer, EvalTime);
 				const TArray<FTransform>& LocalTransforms = Evaluator.ChildLocals;
 
-				const FFrameNumber TimeMinusOne(Time - 1);
-				Interface->AddHandleTransformKeys(InSequencer, Handle, { TimeMinusOne }, LocalTransforms, ChannelsToKey);
+				Interface->AddHandleTransformKeys(InSequencer, Handle, { SetTime }, LocalTransforms, ChannelsToKey);
 
 				bNeedsEvaluation = true;
 			}
diff --git a/Engine/Source/Editor/MovieSceneTools/Private/TrackEditors/TransformTrackEditor.cpp b/Engine/Source/Editor/MovieSceneTools/Private/TrackEditors/TransformTrackEditor.cpp
--- a/Engine/Source/Editor/MovieSceneTools/Private/TrackEditors/TransformTrackEditor.cpp
+++ b/Engine/Source/Editor/MovieSceneTools/Private/TrackEditors/TransformTrackEditor.cpp
@@ -1014,7 +1014,7 @@ void F3DTransformTrackEditor::AddTransformKeys( UObject* ObjectToKey, const TOpt
 			{
 				if (UMovieScene3DTransformSection* Section = Cast< UMovieScene3DTransformSection>(WeakSection.Get()))
 				{
-					FMovieSceneConstraintChannelHelper::CompensateIfNeeded(GetSequencer(), Section, Time);
+					FMovieSceneConstraintChannelHelper::CompensateIfNeeded(GetSequencer(), Section, Time, true /*bCompPreviousTick*/);
 				}
 			}
 		}
@@ -1429,7 +1429,7 @@ void F3DTransformTrackEditor::OnTransformPropertyChanged(const FPropertyChangedP
 		{
 			if (UMovieScene3DTransformSection* Section = Cast<UMovieScene3DTransformSection>(WeakSection.Get()))
 			{
-				FMovieSceneConstraintChannelHelper::CompensateIfNeeded(GetSequencer(), Section, KeyTime);
+				FMovieSceneConstraintChannelHelper::CompensateIfNeeded(GetSequencer(), Section, KeyTime, true /*bCompPreviousTick*/);
 			}
 		}
 
diff --git a/Engine/Source/Editor/MovieSceneTools/Public/Constraints/MovieSceneConstraintChannelHelper.h b/Engine/Source/Editor/MovieSceneTools/Public/Constraints/MovieSceneConstraintChannelHelper.h
--- a/Engine/Source/Editor/MovieSceneTools/Public/Constraints/MovieSceneConstraintChannelHelper.h
+++ b/Engine/Source/Editor/MovieSceneTools/Public/Constraints/MovieSceneConstraintChannelHelper.h
@@ -73,12 +73,14 @@ public:
 	static void Compensate(
 		const TSharedPtr<ISequencer>& InSequencer,
 		const UTickableTransformConstraint* InConstraint,
-		const TOptional<FFrameNumber>& InOptTime);
+		const TOptional<FFrameNumber>& InOptTime,
+		const bool bCompPreviousTick);
 	
 	static void CompensateIfNeeded(
 		const TSharedPtr<ISequencer>& InSequencer,
 		IMovieSceneConstrainedSection* Section,
 		const TOptional<FFrameNumber>& OptionalTime,
+		const bool bCompPreviousTick,
 		const int32 InChildHash = INDEX_NONE);
 		
 	static void HandleConstraintRemoved(
