Control Rig: Show tags in the module details panel. Also warn the user when using a module with an invalid tag.
#rnx

diff --git a/Engine/Plugins/Animation/ControlRig/Source/ControlRigEditor/Private/ControlRigModuleDetails.cpp b/Engine/Plugins/Animation/ControlRig/Source/ControlRigEditor/Private/ControlRigModuleDetails.cpp
--- a/Engine/Plugins/Animation/ControlRig/Source/ControlRigEditor/Private/ControlRigModuleDetails.cpp
+++ b/Engine/Plugins/Animation/ControlRig/Source/ControlRigEditor/Private/ControlRigModuleDetails.cpp
@@ -27,6 +27,7 @@
 #include "ModularRigRuleManager.h"
 #include "ScopedTransaction.h"
 #include "Editor/SRigHierarchyTreeView.h"
+#include "Widgets/SRigVMVariantTagWidget.h"
 
 #define LOCTEXT_NAMESPACE "ControlRigModuleDetails"
 
@@ -236,6 +237,63 @@ void FRigModuleInstanceDetails::CustomizeDetails(IDetailLayoutBuilder& DetailBui
 			.Text(this, &FRigModuleInstanceDetails::GetRigClassPath)
 			.IsEnabled(true)
 		];
+
+		GeneralCategory.AddCustomRow(FText::FromString(TEXT("Variant Tags")))
+		.NameContent()
+		[
+			SNew(STextBlock)
+			.Text(FText::FromString(TEXT("Variant Tags")))
+			.Font(IDetailLayoutBuilder::GetDetailFont())
+			.IsEnabled(true)
+		]
+		.ValueContent()
+		[
+			SNew(SRigVMVariantTagWidget)
+			.Orientation(EOrientation::Orient_Horizontal)
+			.CanAddTags(false)
+			.EnableContextMenu(false)
+			.OnGetTags_Lambda([this]() -> TArray<FRigVMTag>
+			{
+				TArray<FRigVMTag> Tags;
+				for (int32 InfoIndex=0; InfoIndex<PerModuleInfos.Num(); ++InfoIndex)
+				{
+					const FPerModuleInfo& ModuleInfo = PerModuleInfos[InfoIndex]; 
+					if(ModuleInfo.Module.IsValid())
+					{
+						if (const FRigModuleInstance* Module = ModuleInfo.GetModule())
+						{
+							if (const UControlRigBlueprint* ModuleBlueprint = Cast<UControlRigBlueprint>(Module->GetRig()->GetClass()->ClassGeneratedBy))
+							{
+								if (InfoIndex == 0)
+								{
+									Tags = ModuleBlueprint->GetAssetVariant().Tags;
+								}
+								else
+								{
+									const TArray<FRigVMTag>& OtherTags = ModuleBlueprint->GetAssetVariant().Tags;
+									bool bSameArray = Tags.Num() == OtherTags.Num();
+									if (bSameArray)
+									{
+										for (const FRigVMTag& OtherTag : OtherTags)
+										{
+											if (!Tags.ContainsByPredicate([OtherTag](const FRigVMTag& Tag) { return OtherTag.Name == Tag.Name; }))
+											{
+												return {};
+											}
+										}
+									}
+									else
+									{
+										return {};
+									}
+								}
+							}
+						}
+					}
+				}
+				return Tags;
+			})
+		];
 	}
 
 	IDetailCategoryBuilder& ConnectionsCategory = DetailBuilder.EditCategory(TEXT("Connections"), LOCTEXT("Connections", "Connections"));
diff --git a/Engine/Plugins/Animation/ControlRig/Source/ControlRigEditor/Private/Editor/SModularRigModel.cpp b/Engine/Plugins/Animation/ControlRig/Source/ControlRigEditor/Private/Editor/SModularRigModel.cpp
--- a/Engine/Plugins/Animation/ControlRig/Source/ControlRigEditor/Private/Editor/SModularRigModel.cpp
+++ b/Engine/Plugins/Animation/ControlRig/Source/ControlRigEditor/Private/Editor/SModularRigModel.cpp
@@ -117,6 +117,14 @@ void SModularRigModel::Construct(const FArguments& InArgs, TSharedRef<FControlRi
 		.HAlignHeader(HAlign_Left)
 		.VAlignCell(VAlign_Fill)
 	);
+	HeaderRowWidget->AddColumn(
+		SHeaderRow::Column(SModularRigTreeView::Column_Tags)
+		.DefaultLabel(FText())
+		.HAlignCell(HAlign_Fill)
+		.HAlignHeader(HAlign_Fill)
+		.FixedWidth(16.f)
+		.VAlignCell(VAlign_Center)
+	);
 	HeaderRowWidget->AddColumn(
 		SHeaderRow::Column(SModularRigTreeView::Column_Connector)
 		.DefaultLabel(FText::FromName(SModularRigTreeView::Column_Connector))
diff --git a/Engine/Plugins/Animation/ControlRig/Source/ControlRigEditor/Private/Editor/SModularRigTreeView.cpp b/Engine/Plugins/Animation/ControlRig/Source/ControlRigEditor/Private/Editor/SModularRigTreeView.cpp
--- a/Engine/Plugins/Animation/ControlRig/Source/ControlRigEditor/Private/Editor/SModularRigTreeView.cpp
+++ b/Engine/Plugins/Animation/ControlRig/Source/ControlRigEditor/Private/Editor/SModularRigTreeView.cpp
@@ -29,6 +29,7 @@
 #include "Widgets/Input/SButton.h"
 #include "ScopedTransaction.h"
 #include "DetailLayoutBuilder.h"
+#include "Widgets/SRigVMVariantTagWidget.h"
 
 #define LOCTEXT_NAMESPACE "SModularRigTreeView"
 
@@ -351,6 +352,70 @@ TSharedRef<SWidget> SModularRigModelItem::GenerateWidgetForColumn(const FName& C
 
 		return Widget;
 	}
+	if (ColumnName == SModularRigTreeView::Column_Tags)
+	{
+		TSharedRef<SWidget> Widget = SNew(SImage)
+		.Visibility_Lambda([this]() -> EVisibility
+		{
+			if (WeakRigTreeElement.IsValid())
+			{
+				if (!WeakRigTreeElement.Pin()->bIsPrimary)
+				{
+					return EVisibility::Hidden;
+				}
+			
+				if (const UModularRig* ModularRig = Delegates.GetModularRig())
+				{
+					if (const FRigModuleInstance* Module = ModularRig->FindModule(WeakRigTreeElement.Pin()->ModulePath))
+					{
+						if (const UControlRigBlueprint* ModuleBlueprint = Cast<UControlRigBlueprint>(Module->GetRig()->GetClass()->ClassGeneratedBy))
+						{
+							for (const FRigVMTag& Tag : ModuleBlueprint->GetAssetVariant().Tags)
+							{
+								if (Tag.bMarksSubjectAsInvalid)
+								{
+									return EVisibility::Visible;
+								}
+							}
+						}
+					}
+				}
+			}
+			return EVisibility::Hidden;
+		})
+		.ToolTipText_Lambda([this]() -> FText
+		{
+			TArray<FString> ToolTip;
+			if (WeakRigTreeElement.IsValid())
+			{
+				if (const UModularRig* ModularRig = Delegates.GetModularRig())
+				{
+					if (const FRigModuleInstance* Module = ModularRig->FindModule(WeakRigTreeElement.Pin()->ModulePath))
+					{
+						if (const UControlRigBlueprint* ModuleBlueprint = Cast<UControlRigBlueprint>(Module->GetRig()->GetClass()->ClassGeneratedBy))
+						{
+							for (const FRigVMTag& Tag : ModuleBlueprint->GetAssetVariant().Tags)
+							{
+								if (Tag.bMarksSubjectAsInvalid)
+								{
+									ToolTip.Add(FString::Printf(TEXT("%s: %s"), *Tag.Label, *Tag.ToolTip.ToString()));
+								}
+							}
+						}
+					}
+				}
+			}
+			return FText::FromString(FString::Join(ToolTip, TEXT("\n")));
+		})
+		.Image_Lambda([this]() -> const FSlateBrush*
+		{
+			const FSlateBrush* WarningBrush = FAppStyle::Get().GetBrush("Icons.WarningWithColor");
+			return WarningBrush;
+		})
+		.DesiredSizeOverride(FVector2D(16, 16));
+
+		return Widget;
+	}
 	if(ColumnName == SModularRigTreeView::Column_Connector)
 	{
 		TSharedPtr<SVerticalBox> ComboButtonBox;
@@ -585,6 +650,7 @@ FText SModularRigModelItem::GetItemTooltip() const
 ///////////////////////////////////////////////////////////
 
 const FName SModularRigTreeView::Column_Module = TEXT("Module");
+const FName SModularRigTreeView::Column_Tags = TEXT("Tags");
 const FName SModularRigTreeView::Column_Connector = TEXT("Connector");
 const FName SModularRigTreeView::Column_Buttons = TEXT("Actions");
 
diff --git a/Engine/Plugins/Animation/ControlRig/Source/ControlRigEditor/Private/Editor/SModularRigTreeView.h b/Engine/Plugins/Animation/ControlRig/Source/ControlRigEditor/Private/Editor/SModularRigTreeView.h
--- a/Engine/Plugins/Animation/ControlRig/Source/ControlRigEditor/Private/Editor/SModularRigTreeView.h
+++ b/Engine/Plugins/Animation/ControlRig/Source/ControlRigEditor/Private/Editor/SModularRigTreeView.h
@@ -174,6 +174,7 @@ class SModularRigTreeView : public STreeView<TSharedPtr<FModularRigTreeElement>>
 public:
 
 	static const FName Column_Module;
+	static const FName Column_Tags;
 	static const FName Column_Connector;
 	static const FName Column_Buttons;
 
diff --git a/Engine/Plugins/Runtime/RigVM/Source/RigVMDeveloper/Public/RigVMBlueprint.h b/Engine/Plugins/Runtime/RigVM/Source/RigVMDeveloper/Public/RigVMBlueprint.h
--- a/Engine/Plugins/Runtime/RigVM/Source/RigVMDeveloper/Public/RigVMBlueprint.h
+++ b/Engine/Plugins/Runtime/RigVM/Source/RigVMDeveloper/Public/RigVMBlueprint.h
@@ -560,6 +560,8 @@ public:
 
 	UFUNCTION(BlueprintCallable, Category = "Variables")
 	bool ChangeMemberVariableType(const FName& InName, const FString& InCPPType, bool bIsPublic = false, bool bIsReadOnly = false, FString InDefaultValue = TEXT(""));
+
+	const FRigVMVariant& GetAssetVariant() const { return AssetVariant; }
 #endif
 
 private:
diff --git a/Engine/Plugins/Runtime/RigVM/Source/RigVMEditor/Public/Widgets/SRigVMVariantTagWidget.h b/Engine/Plugins/Runtime/RigVM/Source/RigVMEditor/Public/Widgets/SRigVMVariantTagWidget.h
--- a/Engine/Plugins/Runtime/RigVM/Source/RigVMEditor/Public/Widgets/SRigVMVariantTagWidget.h
+++ b/Engine/Plugins/Runtime/RigVM/Source/RigVMEditor/Public/Widgets/SRigVMVariantTagWidget.h
@@ -52,7 +52,7 @@ private:
 	TAttribute<bool> EnableContextMenu;
 };
 
-class SRigVMVariantTagWidget : public SBox
+class RIGVMEDITOR_API SRigVMVariantTagWidget : public SBox
 {
 public:
 
