Control Rig: Fix control rig subobjects not being able to be collected by GC when created in game
#rnx

diff --git a/Engine/Plugins/Animation/ControlRig/Source/ControlRig/Private/ControlRig.cpp b/Engine/Plugins/Animation/ControlRig/Source/ControlRig/Private/ControlRig.cpp
--- a/Engine/Plugins/Animation/ControlRig/Source/ControlRig/Private/ControlRig.cpp
+++ b/Engine/Plugins/Animation/ControlRig/Source/ControlRig/Private/ControlRig.cpp
@@ -1650,6 +1650,13 @@ UAnimationDataSourceRegistry* UControlRig::GetDataSourceRegistry()
 	{
 		DataSourceRegistry = NewObject<UAnimationDataSourceRegistry>(this, NAME_None, RF_Transient);
 
+		if (!IsInGameThread())
+		{
+			// If the object was created on a non-game thread, clear the async flag immediately, so that it can be
+			// garbage collected in the future. 
+			(void)DataSourceRegistry->AtomicallyClearInternalFlags(EInternalObjectFlags::Async);
+		}
+
 		if (HasAnyFlags(RF_ClassDefaultObject) && GetClass()->IsNative())
 		{
 			DataSourceRegistry->AddToRoot();
@@ -3374,6 +3381,13 @@ void UControlRig::PostInitInstance(URigVMHost* InCDO)
 		if(!IsRigModuleInstance())
 		{
 			DynamicHierarchy = NewObject<URigHierarchy>(this, TEXT("DynamicHierarchy"), SubObjectFlags);
+
+			if (!IsInGameThread())
+			{
+				// If the object was created on a non-game thread, clear the async flag immediately, so that it can be
+				// garbage collected in the future. 
+				(void)DynamicHierarchy->AtomicallyClearInternalFlags(EInternalObjectFlags::Async);
+			}
 		}
 	}
 
@@ -3408,6 +3422,13 @@ void UControlRig::PostInitInstance(URigVMHost* InCDO)
 		if (VM == nullptr)
 		{
 			VM = NewObject<URigVM>(this, TEXT("ControlRig_VM"), SubObjectFlags);
+
+			if (!IsInGameThread())
+			{
+				// If the object was created on a non-game thread, clear the async flag immediately, so that it can be
+				// garbage collected in the future. 
+				(void)VM->AtomicallyClearInternalFlags(EInternalObjectFlags::Async);
+			}
 		}
 
 		// for default objects we need to check if the CDO is rooted. specialized Control Rigs
