Delete this autortfm usage

diff --git a/Engine/Plugins/Animation/ControlRig/Source/ControlRig/ControlRig.Build.cs b/Engine/Plugins/Animation/ControlRig/Source/ControlRig/ControlRig.Build.cs
--- a/Engine/Plugins/Animation/ControlRig/Source/ControlRig/ControlRig.Build.cs
+++ b/Engine/Plugins/Animation/ControlRig/Source/ControlRig/ControlRig.Build.cs
@@ -63,8 +63,6 @@ namespace UnrealBuildTool.Rules
                 PrivateIncludePathModuleNames.Add("ControlRigEditor");
                 DynamicallyLoadedModuleNames.Add("ControlRigEditor");
             }
-
-			bAllowAutoRTFMInstrumentation = true;
         }
     }
 }
diff --git a/Engine/Plugins/Animation/DeformerGraph/Source/OptimusCore/OptimusCore.Build.cs b/Engine/Plugins/Animation/DeformerGraph/Source/OptimusCore/OptimusCore.Build.cs
--- a/Engine/Plugins/Animation/DeformerGraph/Source/OptimusCore/OptimusCore.Build.cs
+++ b/Engine/Plugins/Animation/DeformerGraph/Source/OptimusCore/OptimusCore.Build.cs
@@ -28,6 +28,8 @@ namespace UnrealBuildTool.Rules
 					"RHI",
 				}
 			);
+
+			bDisableAutoRTFMInstrumentation = true;
         }
     }
 }
diff --git a/Engine/Plugins/Animation/LiveLink/Source/LiveLink/LiveLink.Build.cs b/Engine/Plugins/Animation/LiveLink/Source/LiveLink/LiveLink.Build.cs
--- a/Engine/Plugins/Animation/LiveLink/Source/LiveLink/LiveLink.Build.cs
+++ b/Engine/Plugins/Animation/LiveLink/Source/LiveLink/LiveLink.Build.cs
@@ -38,6 +38,8 @@ namespace UnrealBuildTool.Rules
 					"UnrealEd",
 				});
 			}
+
+			bDisableAutoRTFMInstrumentation = true;
 		}
 	}
 }
diff --git a/Engine/Plugins/Experimental/ControlFlows/Source/ControlFlows/ControlFlows.Build.cs b/Engine/Plugins/Experimental/ControlFlows/Source/ControlFlows/ControlFlows.Build.cs
--- a/Engine/Plugins/Experimental/ControlFlows/Source/ControlFlows/ControlFlows.Build.cs
+++ b/Engine/Plugins/Experimental/ControlFlows/Source/ControlFlows/ControlFlows.Build.cs
@@ -14,7 +14,5 @@ public class ControlFlows : ModuleRules
 				// ... add private dependencies that you statically link with here ...	
 			}
 			);
-
-		bAllowAutoRTFMInstrumentation = true;
 	}
 }
