Control Rig: Avoid drifting when manipulating rig in a layered control rig
#rnx

diff --git a/Engine/Plugins/Animation/ControlRig/Source/ControlRigEditor/Private/EditMode/ControlRigEditMode.cpp b/Engine/Plugins/Animation/ControlRig/Source/ControlRigEditor/Private/EditMode/ControlRigEditMode.cpp
--- a/Engine/Plugins/Animation/ControlRig/Source/ControlRigEditor/Private/EditMode/ControlRigEditMode.cpp
+++ b/Engine/Plugins/Animation/ControlRig/Source/ControlRigEditor/Private/EditMode/ControlRigEditMode.cpp
@@ -2754,7 +2754,9 @@ void FControlRigEditMode::BindCommands()
 		FExecuteAction::CreateRaw(this, &FControlRigEditMode::ToggleManipulators));
 	CommandBindings->MapAction(
 		Commands.ToggleAllManipulators,
-		FExecuteAction::CreateRaw(this, &FControlRigEditMode::ToggleAllManipulators));
+		FExecuteAction::CreateRaw(this, &FControlRigEditMode::ToggleAllManipulators),
+		FCanExecuteAction(),
+		FIsActionChecked::CreateSP(this, &FControlRigEditMode::AreControlsVisible));
 	CommandBindings->MapAction(
 		Commands.ZeroTransforms,
 		FExecuteAction::CreateRaw(this, &FControlRigEditMode::ZeroTransforms, true));
@@ -3283,6 +3285,12 @@ void FControlRigEditMode::ToggleAllManipulators()
 	}
 }
 
+bool FControlRigEditMode::AreControlsVisible() const
+{
+	UControlRigEditModeSettings* Settings = GetMutableDefault<UControlRigEditModeSettings>();
+	return Settings->bHideControlShapes;
+}
+
 void FControlRigEditMode::ZeroTransforms(bool bSelectionOnly)
 {
 	// Gather up the control rigs for the selected controls
@@ -4764,7 +4772,8 @@ void FControlRigEditMode::MoveControlShape(AControlRigShapeActor* ShapeActor, co
 	}
 	if(!bTransformChanged) //not local or doing scale.
 	{
-		FTransform CurrentTransform = GetControlShapeTransform(ShapeActor) * ToWorldTransform;
+		// Get the global transform from shape actor to avoid drifting
+		FTransform CurrentTransform = ShapeActor->GetGlobalTransform() * ToWorldTransform;
 
 		if (bRotation)
 		{
