Control Rig: Remove visible controls command
#rnx

diff --git a/Engine/Plugins/Animation/ControlRig/Source/ControlRigEditor/Private/EditMode/ControlRigEditMode.cpp b/Engine/Plugins/Animation/ControlRig/Source/ControlRigEditor/Private/EditMode/ControlRigEditMode.cpp
--- a/Engine/Plugins/Animation/ControlRig/Source/ControlRigEditor/Private/EditMode/ControlRigEditMode.cpp
+++ b/Engine/Plugins/Animation/ControlRig/Source/ControlRigEditor/Private/EditMode/ControlRigEditMode.cpp
@@ -2754,9 +2754,7 @@ void FControlRigEditMode::BindCommands()
 		FExecuteAction::CreateRaw(this, &FControlRigEditMode::ToggleManipulators));
 	CommandBindings->MapAction(
 		Commands.ToggleAllManipulators,
-		FExecuteAction::CreateRaw(this, &FControlRigEditMode::ToggleAllManipulators),
-		FCanExecuteAction(),
-		FIsActionChecked::CreateSP(this, &FControlRigEditMode::AreControlsVisible));
+		FExecuteAction::CreateRaw(this, &FControlRigEditMode::ToggleAllManipulators));
 	CommandBindings->MapAction(
 		Commands.ZeroTransforms,
 		FExecuteAction::CreateRaw(this, &FControlRigEditMode::ZeroTransforms, true));
@@ -3285,12 +3283,6 @@ void FControlRigEditMode::ToggleAllManipulators()
 	}
 }
 
-bool FControlRigEditMode::AreControlsVisible() const
-{
-	UControlRigEditModeSettings* Settings = GetMutableDefault<UControlRigEditModeSettings>();
-	return Settings->bHideControlShapes;
-}
-
 void FControlRigEditMode::ZeroTransforms(bool bSelectionOnly)
 {
 	// Gather up the control rigs for the selected controls
