Sequencer: Control Rig: Don't refresh bones transforms if post evaluating when restoring state.

diff --git a/Engine/Plugins/Animation/ControlRig/Source/ControlRig/Private/Sequencer/MovieSceneControlRigParameterTemplate.cpp b/Engine/Plugins/Animation/ControlRig/Source/ControlRig/Private/Sequencer/MovieSceneControlRigParameterTemplate.cpp
--- a/Engine/Plugins/Animation/ControlRig/Source/ControlRig/Private/Sequencer/MovieSceneControlRigParameterTemplate.cpp
+++ b/Engine/Plugins/Animation/ControlRig/Source/ControlRig/Private/Sequencer/MovieSceneControlRigParameterTemplate.cpp
@@ -916,14 +916,16 @@ struct FControlRigParameterPreAnimatedTokenProducer : IMovieScenePreAnimatedToke
 								// Restore pose after unbinding to force the restored pose
 								SkeletalMeshComponent->SetUpdateAnimationInEditor(true);
 								SkeletalMeshComponent->SetUpdateClothInEditor(true);
-								SkeletalMeshComponent->TickAnimation(0.f, false);
-
-								SkeletalMeshComponent->RefreshBoneTransforms();
-								SkeletalMeshComponent->RefreshFollowerComponents();
-								SkeletalMeshComponent->UpdateComponentToWorld();
-								SkeletalMeshComponent->FinalizeBoneTransform();
-								SkeletalMeshComponent->MarkRenderTransformDirty();
-								SkeletalMeshComponent->MarkRenderDynamicDataDirty();
+								if (!SkeletalMeshComponent->IsPostEvaluatingAnimation())
+								{
+									SkeletalMeshComponent->TickAnimation(0.f, false);
+									SkeletalMeshComponent->RefreshBoneTransforms();
+									SkeletalMeshComponent->RefreshFollowerComponents();
+									SkeletalMeshComponent->UpdateComponentToWorld();
+									SkeletalMeshComponent->FinalizeBoneTransform();
+									SkeletalMeshComponent->MarkRenderTransformDirty();
+									SkeletalMeshComponent->MarkRenderDynamicDataDirty();
+								}
 								SkeletalMeshRestoreState.RestoreState(SkeletalMeshComponent);
 
 								if (SkeletalMeshComponent->GetAnimationMode() != AnimationMode)
