Control Rig: Add toggle control visibility menu entry in HUD control rig editor viewport
#rnx

diff --git a/Engine/Plugins/Animation/ControlRig/Source/ControlRigEditor/Private/EditMode/ControlRigEditMode.cpp b/Engine/Plugins/Animation/ControlRig/Source/ControlRigEditor/Private/EditMode/ControlRigEditMode.cpp
--- a/Engine/Plugins/Animation/ControlRig/Source/ControlRigEditor/Private/EditMode/ControlRigEditMode.cpp
+++ b/Engine/Plugins/Animation/ControlRig/Source/ControlRigEditor/Private/EditMode/ControlRigEditMode.cpp
@@ -3283,6 +3283,31 @@ void FControlRigEditMode::ToggleAllManipulators()
 	}
 }
 
+bool FControlRigEditMode::AreControlsVisible() const
+{
+	if (!AreEditingControlRigDirectly())
+	{
+		TMap<UControlRig*, TArray<FRigElementKey>> SelectedControls;
+		GetAllSelectedControls(SelectedControls);
+		TArray<UControlRig*> ControlRigs;
+		SelectedControls.GenerateKeyArray(ControlRigs);
+		for (UControlRig* ControlRig : ControlRigs)
+		{
+			if (ControlRig)
+			{
+				if (!ControlRig->bControlsVisible)
+				{
+					return false;
+				}
+			}
+		}
+		return true;
+	}
+	
+	UControlRigEditModeSettings* Settings = GetMutableDefault<UControlRigEditModeSettings>();
+	return !Settings->bHideControlShapes;
+}
+
 void FControlRigEditMode::ZeroTransforms(bool bSelectionOnly)
 {
 	// Gather up the control rigs for the selected controls
diff --git a/Engine/Plugins/Animation/ControlRig/Source/ControlRigEditor/Private/Editor/ControlRigEditor.cpp b/Engine/Plugins/Animation/ControlRig/Source/ControlRigEditor/Private/Editor/ControlRigEditor.cpp
--- a/Engine/Plugins/Animation/ControlRig/Source/ControlRigEditor/Private/Editor/ControlRigEditor.cpp
+++ b/Engine/Plugins/Animation/ControlRig/Source/ControlRigEditor/Private/Editor/ControlRigEditor.cpp
@@ -2313,58 +2313,10 @@ void FControlRigEditor::HandleViewportCreated(const TSharedRef<class IPersonaVie
 			InMenuBuilder.AddMenuSeparator(TEXT("Control Rig"));
 			InMenuBuilder.BeginSection("ControlRig", LOCTEXT("ControlRig_Label", "Control Rig"));
 			{
-				InMenuBuilder.AddWidget(
-					SNew(SBox)
-					.HAlign(HAlign_Right)
-					[
-						SNew(SBox)
-						.Padding(FMargin(4.0f, 0.0f, 0.0f, 0.0f))
-						.WidthOverride(100.0f)
-						.IsEnabled(this, &FControlRigEditor::IsToolbarDrawNullsEnabled)
-						[
-							SNew(SCheckBox)
-							.IsChecked(this, &FControlRigEditor::GetToolbarDrawNulls)
-							.OnCheckStateChanged(this, &FControlRigEditor::OnToolbarDrawNullsChanged)
-							.ToolTipText(LOCTEXT("ControlRigDrawNullsToolTip", "If checked all nulls are drawn as axes."))
-						]
-					],
-					LOCTEXT("ControlRigDisplayNulls", "Display Nulls")
-				);
-
-				InMenuBuilder.AddWidget(
-					SNew(SBox)
-					.HAlign(HAlign_Right)
-					[
-						SNew(SBox)
-						.Padding(FMargin(4.0f, 0.0f, 0.0f, 0.0f))
-						.WidthOverride(100.0f)
-						.IsEnabled(this, &FControlRigEditor::IsToolbarDrawSocketsEnabled)
-						[
-							SNew(SCheckBox)
-							.IsChecked(this, &FControlRigEditor::GetToolbarDrawSockets)
-							.OnCheckStateChanged(this, &FControlRigEditor::OnToolbarDrawSocketsChanged)
-							.ToolTipText(LOCTEXT("ControlRigDrawSocketsToolTip", "If checked all sockets are drawn."))
-						]
-					],
-					LOCTEXT("ControlRigDisplaySockets", "Display Sockets")
-				);
-
-				InMenuBuilder.AddWidget(
-					SNew(SBox)
-					.HAlign(HAlign_Right)
-					[
-						SNew(SBox)
-						.Padding(FMargin(4.0f, 0.0f, 0.0f, 0.0f))
-						.WidthOverride(100.0f)
-						[
-							SNew(SCheckBox)
-							.IsChecked(this, &FControlRigEditor::GetToolbarDrawAxesOnSelection)
-							.OnCheckStateChanged(this, &FControlRigEditor::OnToolbarDrawAxesOnSelectionChanged)
-							.ToolTipText(LOCTEXT("ControlRigDisplayAxesOnSelectionToolTip", "If checked axes will be drawn for all selected rig elements."))
-						]
-					],
-					LOCTEXT("ControlRigDisplayAxesOnSelection", "Display Axes On Selection")
-				);
+				InMenuBuilder.AddMenuEntry(FControlRigEditorCommands::Get().ToggleControlVisibility);
+				InMenuBuilder.AddMenuEntry(FControlRigEditorCommands::Get().ToggleDrawNulls);
+				InMenuBuilder.AddMenuEntry(FControlRigEditorCommands::Get().ToggleDrawSockets);
+				InMenuBuilder.AddMenuEntry(FControlRigEditorCommands::Get().ToggleDrawAxesOnSelection);
 
 				InMenuBuilder.AddWidget(
 					SNew(SBox)
@@ -2519,6 +2471,23 @@ void FControlRigEditor::OnToolbarAxesScaleChanged(float InValue)
 	}
 }
 
+void FControlRigEditor::HandleToggleControlVisibility()
+{
+	if (FControlRigEditMode* EditMode = GetEditMode())
+	{
+		EditMode->ToggleAllManipulators();
+	}
+}
+
+bool FControlRigEditor::AreControlsVisible() const
+{
+	if (FControlRigEditMode* EditMode = GetEditMode())
+	{
+		return EditMode->AreControlsVisible();
+	}
+	return false;
+}
+
 void FControlRigEditor::HandleToggleSchematicViewport()
 {
 	if(SchematicViewport.IsValid())
@@ -2537,20 +2506,20 @@ bool FControlRigEditor::IsSchematicViewportActive() const
 	return false;
 }
 
-ECheckBoxState FControlRigEditor::GetToolbarDrawAxesOnSelection() const
+bool FControlRigEditor::GetToolbarDrawAxesOnSelection() const
 {
 	if (const UControlRigEditModeSettings* Settings = GetDefault<UControlRigEditModeSettings>())
 	{
-		return Settings->bDisplayAxesOnSelection ? ECheckBoxState::Checked : ECheckBoxState::Unchecked;
+		return Settings->bDisplayAxesOnSelection;
 	}
-	return ECheckBoxState::Unchecked;
+	return false;
 }
 
-void FControlRigEditor::OnToolbarDrawAxesOnSelectionChanged(ECheckBoxState InNewValue)
+void FControlRigEditor::HandleToggleToolbarDrawAxesOnSelection()
 {
 	if (UControlRigEditModeSettings* Settings = GetMutableDefault<UControlRigEditModeSettings>())
 	{
-		Settings->bDisplayAxesOnSelection = InNewValue == ECheckBoxState::Checked;
+		Settings->bDisplayAxesOnSelection = !Settings->bDisplayAxesOnSelection;
 	}
 }
 
@@ -2566,20 +2535,20 @@ bool FControlRigEditor::IsToolbarDrawNullsEnabled() const
 	return false;
 }
 
-ECheckBoxState FControlRigEditor::GetToolbarDrawNulls() const
+bool FControlRigEditor::GetToolbarDrawNulls() const
 {
 	if (const UControlRigEditModeSettings* Settings = GetDefault<UControlRigEditModeSettings>())
 	{
-		return Settings->bDisplayNulls ? ECheckBoxState::Checked : ECheckBoxState::Unchecked;
+		return Settings->bDisplayNulls;
 	}
-	return ECheckBoxState::Unchecked;
+	return false;
 }
 
-void FControlRigEditor::OnToolbarDrawNullsChanged(ECheckBoxState InNewValue)
+void FControlRigEditor::HandleToggleToolbarDrawNulls()
 {
 	if (UControlRigEditModeSettings* Settings = GetMutableDefault<UControlRigEditModeSettings>())
 	{
-		Settings->bDisplayNulls = InNewValue == ECheckBoxState::Checked;
+		Settings->bDisplayNulls = !Settings->bDisplayNulls;
 	}
 }
 
@@ -2595,20 +2564,20 @@ bool FControlRigEditor::IsToolbarDrawSocketsEnabled() const
 	return false;
 }
 
-ECheckBoxState FControlRigEditor::GetToolbarDrawSockets() const
+bool FControlRigEditor::GetToolbarDrawSockets() const
 {
 	if (const UControlRigEditModeSettings* Settings = GetDefault<UControlRigEditModeSettings>())
 	{
-		return Settings->bDisplaySockets ? ECheckBoxState::Checked : ECheckBoxState::Unchecked;
+		return Settings->bDisplaySockets;
 	}
-	return ECheckBoxState::Unchecked;
+	return false;
 }
 
-void FControlRigEditor::OnToolbarDrawSocketsChanged(ECheckBoxState InNewValue)
+void FControlRigEditor::HandleToggleToolbarDrawSockets()
 {
 	if (UControlRigEditModeSettings* Settings = GetMutableDefault<UControlRigEditModeSettings>())
 	{
-		Settings->bDisplaySockets = InNewValue == ECheckBoxState::Checked;
+		Settings->bDisplaySockets = !Settings->bDisplaySockets;
 	}
 }
 
@@ -3673,6 +3642,30 @@ void FControlRigEditor::BindCommands()
 		FExecuteAction::CreateSP(this, &FRigVMEditor::SetEventQueue, TArray<FName>(BackwardsAndForwardsSolveEventQueue)),
 		FCanExecuteAction());
 
+	GetToolkitCommands()->MapAction(
+		FControlRigEditorCommands::Get().ToggleControlVisibility,
+		FExecuteAction::CreateSP(this, &FControlRigEditor::HandleToggleControlVisibility),
+		FCanExecuteAction(),
+		FIsActionChecked::CreateSP(this, &FControlRigEditor::AreControlsVisible));
+
+	GetToolkitCommands()->MapAction(
+		FControlRigEditorCommands::Get().ToggleDrawNulls,
+		FExecuteAction::CreateSP(this, &FControlRigEditor::HandleToggleToolbarDrawNulls),
+		FCanExecuteAction::CreateSP(this, &FControlRigEditor::IsToolbarDrawNullsEnabled),
+		FIsActionChecked::CreateSP(this, &FControlRigEditor::GetToolbarDrawNulls));
+
+	GetToolkitCommands()->MapAction(
+		FControlRigEditorCommands::Get().ToggleDrawSockets,
+		FExecuteAction::CreateSP(this, &FControlRigEditor::HandleToggleToolbarDrawSockets),
+		FCanExecuteAction::CreateSP(this, &FControlRigEditor::IsToolbarDrawSocketsEnabled),
+		FIsActionChecked::CreateSP(this, &FControlRigEditor::GetToolbarDrawSockets));
+
+	GetToolkitCommands()->MapAction(
+		FControlRigEditorCommands::Get().ToggleDrawAxesOnSelection,
+		FExecuteAction::CreateSP(this, &FControlRigEditor::HandleToggleToolbarDrawAxesOnSelection),
+		FCanExecuteAction(),
+		FIsActionChecked::CreateSP(this, &FControlRigEditor::GetToolbarDrawAxesOnSelection));
+
 	GetToolkitCommands()->MapAction(
 		FControlRigEditorCommands::Get().ToggleSchematicViewportVisibility,
 		FExecuteAction::CreateSP(this, &FControlRigEditor::HandleToggleSchematicViewport),
diff --git a/Engine/Plugins/Animation/ControlRig/Source/ControlRigEditor/Private/Editor/ControlRigEditor.h b/Engine/Plugins/Animation/ControlRig/Source/ControlRigEditor/Private/Editor/ControlRigEditor.h
--- a/Engine/Plugins/Animation/ControlRig/Source/ControlRigEditor/Private/Editor/ControlRigEditor.h
+++ b/Engine/Plugins/Animation/ControlRig/Source/ControlRigEditor/Private/Editor/ControlRigEditor.h
@@ -240,14 +240,16 @@ private:
 	/** Handle preview scene setup */
 	void HandlePreviewSceneCreated(const TSharedRef<IPersonaPreviewScene>& InPersonaPreviewScene);
 	void HandleViewportCreated(const TSharedRef<class IPersonaViewport>& InViewport);
+	void HandleToggleControlVisibility();
+	bool AreControlsVisible() const;
 	bool IsToolbarDrawNullsEnabled() const;
-	ECheckBoxState GetToolbarDrawNulls() const;
-	void OnToolbarDrawNullsChanged(ECheckBoxState InNewValue);
+	bool GetToolbarDrawNulls() const;
+	void HandleToggleToolbarDrawNulls();
 	bool IsToolbarDrawSocketsEnabled() const;
-	ECheckBoxState GetToolbarDrawSockets() const;
-	void OnToolbarDrawSocketsChanged(ECheckBoxState InNewValue);
-	ECheckBoxState GetToolbarDrawAxesOnSelection() const;
-	void OnToolbarDrawAxesOnSelectionChanged(ECheckBoxState InNewValue);
+	bool GetToolbarDrawSockets() const;
+	void HandleToggleToolbarDrawSockets();
+	bool GetToolbarDrawAxesOnSelection() const;
+	void HandleToggleToolbarDrawAxesOnSelection();
 	TOptional<float> GetToolbarAxesScale() const;
 	void OnToolbarAxesScaleChanged(float InValue);
 	void HandleToggleSchematicViewport();
diff --git a/Engine/Plugins/Animation/ControlRig/Source/ControlRigEditor/Private/Editor/ControlRigEditorCommands.cpp b/Engine/Plugins/Animation/ControlRig/Source/ControlRigEditor/Private/Editor/ControlRigEditorCommands.cpp
--- a/Engine/Plugins/Animation/ControlRig/Source/ControlRigEditor/Private/Editor/ControlRigEditorCommands.cpp
+++ b/Engine/Plugins/Animation/ControlRig/Source/ControlRigEditor/Private/Editor/ControlRigEditorCommands.cpp
@@ -13,6 +13,10 @@ void FControlRigEditorCommands::RegisterCommands()
 	UI_COMMAND(RequestDirectManipulationPosition, "Request Direct Manipulation for Position", "Request per node direct manipulation on a position", EUserInterfaceActionType::Button, FInputChord(EKeys::W));
 	UI_COMMAND(RequestDirectManipulationRotation, "Request Direct Manipulation for Rotation", "Request per node direct manipulation on a rotation", EUserInterfaceActionType::Button, FInputChord(EKeys::E));
 	UI_COMMAND(RequestDirectManipulationScale, "Request Direct Manipulation for Scale", "Request per node direct manipulation on a scale", EUserInterfaceActionType::Button, FInputChord(EKeys::R));
+	UI_COMMAND(ToggleControlVisibility, "Show Controls", "Toggles the visibility of the controls.", EUserInterfaceActionType::ToggleButton, FInputChord(EKeys::T));
+	UI_COMMAND(ToggleDrawNulls, "Display Nulls", "If checked all nulls are drawn as axes.", EUserInterfaceActionType::ToggleButton, FInputChord());
+	UI_COMMAND(ToggleDrawSockets, "Display Sockets", "If checked all sockets are drawn.", EUserInterfaceActionType::ToggleButton, FInputChord());
+	UI_COMMAND(ToggleDrawAxesOnSelection, "Display Axes On Selection", "If checked axes will be drawn for all selected rig elements.", EUserInterfaceActionType::ToggleButton, FInputChord());
 	UI_COMMAND(ToggleSchematicViewportVisibility, "Show Schematic Viewport", "Toggles the visibility of the viewport schematic", EUserInterfaceActionType::ToggleButton, FInputChord(EKeys::Y));
 	UI_COMMAND(SwapModuleWithinAsset, "Swap Module (Asset)", "Swaps a module for all occurrences within this asset.", EUserInterfaceActionType::Button, FInputChord());
 	UI_COMMAND(SwapModuleAcrossProject, "Swap Module (Project)", "Swaps a module for all occurrences in the project.", EUserInterfaceActionType::Button, FInputChord());
diff --git a/Engine/Plugins/Animation/ControlRig/Source/ControlRigEditor/Private/Editor/ControlRigEditorCommands.h b/Engine/Plugins/Animation/ControlRig/Source/ControlRigEditor/Private/Editor/ControlRigEditorCommands.h
--- a/Engine/Plugins/Animation/ControlRig/Source/ControlRigEditor/Private/Editor/ControlRigEditorCommands.h
+++ b/Engine/Plugins/Animation/ControlRig/Source/ControlRigEditor/Private/Editor/ControlRigEditorCommands.h
@@ -39,6 +39,18 @@ public:
 	/** Request per node direct manipulation on a scale */
 	TSharedPtr< FUICommandInfo > RequestDirectManipulationScale;
 
+	/** Toggle visibility of the controls */
+	TSharedPtr< FUICommandInfo > ToggleControlVisibility;
+
+	/** Toggle visibility of nulls */
+	TSharedPtr< FUICommandInfo > ToggleDrawNulls;
+
+	/** Toggle visibility of sockets */
+	TSharedPtr< FUICommandInfo > ToggleDrawSockets;
+
+	/** Toggle visibility of axes on selection */
+	TSharedPtr< FUICommandInfo > ToggleDrawAxesOnSelection;
+
 	/** Toggle visibility of the schematic */
 	TSharedPtr< FUICommandInfo > ToggleSchematicViewportVisibility;
 
diff --git a/Engine/Plugins/Animation/ControlRig/Source/ControlRigEditor/Public/EditMode/ControlRigEditMode.h b/Engine/Plugins/Animation/ControlRig/Source/ControlRigEditor/Public/EditMode/ControlRigEditMode.h
--- a/Engine/Plugins/Animation/ControlRig/Source/ControlRigEditor/Public/EditMode/ControlRigEditMode.h
+++ b/Engine/Plugins/Animation/ControlRig/Source/ControlRigEditor/Public/EditMode/ControlRigEditMode.h
@@ -263,6 +263,9 @@ protected:
 	/** Toggles visibility of all  control rig shapes in the viewport */
 	void ToggleAllManipulators();
 
+	/** Returns true if all control rig shapes are visible in the viewport */
+	bool AreControlsVisible() const;
+
 	/** If Anim Slider is open, got to the next tool*/
 	void ChangeAnimSliderTool();
 
