Control Rig: Hide default property buttons outside for Rig Element Key customization for arrays

#rnx

diff --git a/Engine/Plugins/Animation/ControlRig/Source/ControlRigEditor/Private/ControlRigElementDetails.cpp b/Engine/Plugins/Animation/ControlRig/Source/ControlRigEditor/Private/ControlRigElementDetails.cpp
--- a/Engine/Plugins/Animation/ControlRig/Source/ControlRigEditor/Private/ControlRigElementDetails.cpp
+++ b/Engine/Plugins/Animation/ControlRig/Source/ControlRigEditor/Private/ControlRigElementDetails.cpp
@@ -344,7 +344,7 @@ void FRigElementKeyDetails::CustomizeHeader(TSharedRef<IPropertyHandle> InStruct
 		]
 		.ValueContent()
 		[
-			InStructPropertyHandle->CreatePropertyValueWidget()
+			InStructPropertyHandle->CreatePropertyValueWidget(false)
 		];
 	}
 	else
