ControlRig: fix problem with graph editor pin factory that prevented any other factory from working if they were registered later.

diff --git a/Engine/Plugins/Animation/ControlRig/Source/ControlRigEditor/Private/Graph/ControlRigGraphPanelPinFactory.cpp b/Engine/Plugins/Animation/ControlRig/Source/ControlRigEditor/Private/Graph/ControlRigGraphPanelPinFactory.cpp
--- a/Engine/Plugins/Animation/ControlRig/Source/ControlRigEditor/Private/Graph/ControlRigGraphPanelPinFactory.cpp
+++ b/Engine/Plugins/Animation/ControlRig/Source/ControlRigEditor/Private/Graph/ControlRigGraphPanelPinFactory.cpp
@@ -50,68 +50,67 @@ TSharedPtr<SGraphPin> FControlRigGraphPanelPinFactory::CreatePin(UEdGraphPin* In
 		}
 	}
 	
-	TSharedPtr<SGraphPin> K2PinWidget = FNodeFactory::CreateK2PinWidget(InPin);
-	if(K2PinWidget.IsValid())
+	if (InPin->Direction == EEdGraphPinDirection::EGPD_Input)
 	{
-		if(InPin->Direction == EEdGraphPinDirection::EGPD_Input)
+		// if we are an enum pin - and we are inside a RigElementKey,
+		// let's remove the "all" entry.
+		if(InPin->PinType.PinSubCategoryObject == StaticEnum<ERigElementType>())
 		{
-			// if we are an enum pin - and we are inside a RigElementKey,
-			// let's remove the "all" entry.
-			if(InPin->PinType.PinSubCategoryObject == StaticEnum<ERigElementType>())
+			TSharedPtr<SGraphPin> K2PinWidget = FNodeFactory::CreateK2PinWidget(InPin);
+			if(InPin->ParentPin)
 			{
-				if(InPin->ParentPin)
+				if(InPin->ParentPin->PinType.PinSubCategoryObject == FRigElementKey::StaticStruct())
 				{
-					if(InPin->ParentPin->PinType.PinSubCategoryObject == FRigElementKey::StaticStruct())
+					TSharedPtr<SWidget> ValueWidget = K2PinWidget->GetValueWidget();
+					if(ValueWidget.IsValid())
 					{
-						TSharedPtr<SWidget> ValueWidget = K2PinWidget->GetValueWidget();
-						if(ValueWidget.IsValid())
+						if(TSharedPtr<SPinComboBox> EnumCombo = StaticCastSharedPtr<SPinComboBox>(ValueWidget))
 						{
-							if(TSharedPtr<SPinComboBox> EnumCombo = StaticCastSharedPtr<SPinComboBox>(ValueWidget))
+							if(EnumCombo.IsValid())
 							{
-								if(EnumCombo.IsValid())
-								{
-									EnumCombo->RemoveItemByIndex(StaticEnum<ERigElementType>()->GetIndexByValue((int64)ERigElementType::All));
-								}
+								EnumCombo->RemoveItemByIndex(StaticEnum<ERigElementType>()->GetIndexByValue((int64)ERigElementType::All));
 							}
 						}
 					}
 				}
 			}
+			return K2PinWidget;
+		}
 
-			const UEnum* RigControlTransformChannelEnum = StaticEnum<ERigControlTransformChannel>();
-			if (InPin->PinType.PinSubCategoryObject == RigControlTransformChannelEnum)
+		const UEnum* RigControlTransformChannelEnum = StaticEnum<ERigControlTransformChannel>();
+		if (InPin->PinType.PinSubCategoryObject == RigControlTransformChannelEnum)
+		{
+			TSharedPtr<SGraphPin> K2PinWidget = FNodeFactory::CreateK2PinWidget(InPin);
+			TSharedPtr<SWidget> ValueWidget = K2PinWidget->GetValueWidget();
+			if(ValueWidget.IsValid())
 			{
-				TSharedPtr<SWidget> ValueWidget = K2PinWidget->GetValueWidget();
-				if(ValueWidget.IsValid())
+				if(TSharedPtr<SPinComboBox> EnumCombo = StaticCastSharedPtr<SPinComboBox>(ValueWidget))
 				{
-					if(TSharedPtr<SPinComboBox> EnumCombo = StaticCastSharedPtr<SPinComboBox>(ValueWidget))
+					if(EnumCombo.IsValid())
 					{
-						if(EnumCombo.IsValid())
+						if (const UControlRigGraphNode* RigNode = Cast<UControlRigGraphNode>(InPin->GetOwningNode()))
 						{
-							if (const UControlRigGraphNode* RigNode = Cast<UControlRigGraphNode>(InPin->GetOwningNode()))
+							if (const URigVMPin* ModelPin = RigNode->GetModelPinFromPinPath(InPin->GetName()))
 							{
-								if (const URigVMPin* ModelPin = RigNode->GetModelPinFromPinPath(InPin->GetName()))
+								if(const URigVMUnitNode* UnitNode = Cast<URigVMUnitNode>(ModelPin->GetNode()))
 								{
-									if(const URigVMUnitNode* UnitNode = Cast<URigVMUnitNode>(ModelPin->GetNode()))
+									if(UnitNode->GetScriptStruct() &&
+										UnitNode->GetScriptStruct()->IsChildOf(FRigUnit_HierarchyAddControlElement::StaticStruct()))
 									{
-										if(UnitNode->GetScriptStruct() &&
-											UnitNode->GetScriptStruct()->IsChildOf(FRigUnit_HierarchyAddControlElement::StaticStruct()))
-										{
-											const TSharedPtr<FStructOnScope> StructInstanceScope = UnitNode->ConstructStructInstance();
-											const FRigUnit_HierarchyAddControlElement* StructInstance = 
-												(const FRigUnit_HierarchyAddControlElement*)StructInstanceScope->GetStructMemory();
+										const TSharedPtr<FStructOnScope> StructInstanceScope = UnitNode->ConstructStructInstance();
+										const FRigUnit_HierarchyAddControlElement* StructInstance = 
+											(const FRigUnit_HierarchyAddControlElement*)StructInstanceScope->GetStructMemory();
 
-											if(const TArray<ERigControlTransformChannel>* VisibleChannels =
-												FRigControlTransformChannelDetails::GetVisibleChannelsForControlType(StructInstance->GetControlTypeToSpawn()))
+										if(const TArray<ERigControlTransformChannel>* VisibleChannels =
+											FRigControlTransformChannelDetails::GetVisibleChannelsForControlType(StructInstance->GetControlTypeToSpawn()))
+										{
+											for(int32 Index = 0; Index < RigControlTransformChannelEnum->NumEnums(); Index++)
 											{
-												for(int32 Index = 0; Index < RigControlTransformChannelEnum->NumEnums(); Index++)
+												const ERigControlTransformChannel Value =
+													(ERigControlTransformChannel)RigControlTransformChannelEnum->GetValueByIndex(Index);
+												if(!VisibleChannels->Contains(Value))
 												{
-													const ERigControlTransformChannel Value =
-														(ERigControlTransformChannel)RigControlTransformChannelEnum->GetValueByIndex(Index);
-													if(!VisibleChannels->Contains(Value))
-													{
-														EnumCombo->RemoveItemByIndex(Index);
-													}
+													EnumCombo->RemoveItemByIndex(Index);
 												}
 											}
 										}
@@ -122,9 +121,11 @@ TSharedPtr<SGraphPin> FControlRigGraphPanelPinFactory::CreatePin(UEdGraphPin* In
 					}
 				}
 			}
+			return K2PinWidget;
 		}
 	}
-	return K2PinWidget;
+	
+	return nullptr;
 }
 
 FName FControlRigGraphPanelPinFactory::GetFactoryName() const
