Control Rig: Avoid drifting in layered control rigs by not setting the global transform on the shape actors
#rnx

diff --git a/Engine/Plugins/Animation/ControlRig/Source/ControlRigEditor/Private/EditMode/ControlRigEditMode.cpp b/Engine/Plugins/Animation/ControlRig/Source/ControlRigEditor/Private/EditMode/ControlRigEditMode.cpp
--- a/Engine/Plugins/Animation/ControlRig/Source/ControlRigEditor/Private/EditMode/ControlRigEditMode.cpp
+++ b/Engine/Plugins/Animation/ControlRig/Source/ControlRigEditor/Private/EditMode/ControlRigEditMode.cpp
@@ -4863,7 +4863,10 @@ void FControlRigEditMode::MoveControlShape(AControlRigShapeActor* ShapeActor, co
 						ControlRig->Evaluate_AnyThread();
 					}
 				}
-				ShapeActor->SetGlobalTransform(CurrentTransform);
+				
+				// Don't set the global transform to the shape actor to avoid drifting
+				//ShapeActor->SetGlobalTransform(CurrentTransform);
+
 				if (bCalcLocal)
 				{
 					FTransform NewLocal = ControlRig->GetControlLocalTransform(ShapeActor->ControlName);
