Control Rig: Add controls as overlay option in the control rig editor preview menu
#rnx

diff --git a/Engine/Plugins/Animation/ControlRig/Source/ControlRigEditor/Private/Editor/ControlRigEditor.cpp b/Engine/Plugins/Animation/ControlRig/Source/ControlRigEditor/Private/Editor/ControlRigEditor.cpp
--- a/Engine/Plugins/Animation/ControlRig/Source/ControlRigEditor/Private/Editor/ControlRigEditor.cpp
+++ b/Engine/Plugins/Animation/ControlRig/Source/ControlRigEditor/Private/Editor/ControlRigEditor.cpp
@@ -2314,6 +2314,7 @@ void FControlRigEditor::HandleViewportCreated(const TSharedRef<class IPersonaVie
 			InMenuBuilder.BeginSection("ControlRig", LOCTEXT("ControlRig_Label", "Control Rig"));
 			{
 				InMenuBuilder.AddMenuEntry(FControlRigEditorCommands::Get().ToggleControlVisibility);
+				InMenuBuilder.AddMenuEntry(FControlRigEditorCommands::Get().ToggleControlsAsOverlay);
 				InMenuBuilder.AddMenuEntry(FControlRigEditorCommands::Get().ToggleDrawNulls);
 				InMenuBuilder.AddMenuEntry(FControlRigEditorCommands::Get().ToggleDrawSockets);
 				InMenuBuilder.AddMenuEntry(FControlRigEditorCommands::Get().ToggleDrawAxesOnSelection);
@@ -2488,6 +2489,25 @@ bool FControlRigEditor::AreControlsVisible() const
 	return false;
 }
 
+void FControlRigEditor::HandleToggleControlsAsOverlay()
+{
+	if (FControlRigEditMode* EditMode = GetEditMode())
+	{
+		EditMode->bShowControlsAsOverlay = !EditMode->bShowControlsAsOverlay;
+		EditMode->UpdateSelectabilityOnSkeletalMeshes(GetControlRig(), !EditMode->bShowControlsAsOverlay);
+		EditMode->RequestToRecreateControlShapeActors();
+	}
+}
+
+bool FControlRigEditor::AreControlsAsOverlay() const
+{
+	if (FControlRigEditMode* EditMode = GetEditMode())
+	{
+		return EditMode->bShowControlsAsOverlay;
+	}
+	return false;
+}
+
 void FControlRigEditor::HandleToggleSchematicViewport()
 {
 	if(SchematicViewport.IsValid())
@@ -3648,6 +3668,12 @@ void FControlRigEditor::BindCommands()
 		FCanExecuteAction(),
 		FIsActionChecked::CreateSP(this, &FControlRigEditor::AreControlsVisible));
 
+	GetToolkitCommands()->MapAction(
+		FControlRigEditorCommands::Get().ToggleControlsAsOverlay,
+		FExecuteAction::CreateSP(this, &FControlRigEditor::HandleToggleControlsAsOverlay),
+		FCanExecuteAction(),
+		FIsActionChecked::CreateSP(this, &FControlRigEditor::AreControlsAsOverlay));
+
 	GetToolkitCommands()->MapAction(
 		FControlRigEditorCommands::Get().ToggleDrawNulls,
 		FExecuteAction::CreateSP(this, &FControlRigEditor::HandleToggleToolbarDrawNulls),
diff --git a/Engine/Plugins/Animation/ControlRig/Source/ControlRigEditor/Private/Editor/ControlRigEditor.h b/Engine/Plugins/Animation/ControlRig/Source/ControlRigEditor/Private/Editor/ControlRigEditor.h
--- a/Engine/Plugins/Animation/ControlRig/Source/ControlRigEditor/Private/Editor/ControlRigEditor.h
+++ b/Engine/Plugins/Animation/ControlRig/Source/ControlRigEditor/Private/Editor/ControlRigEditor.h
@@ -242,6 +242,8 @@ private:
 	void HandleViewportCreated(const TSharedRef<class IPersonaViewport>& InViewport);
 	void HandleToggleControlVisibility();
 	bool AreControlsVisible() const;
+	void HandleToggleControlsAsOverlay();
+	bool AreControlsAsOverlay() const;
 	bool IsToolbarDrawNullsEnabled() const;
 	bool GetToolbarDrawNulls() const;
 	void HandleToggleToolbarDrawNulls();
diff --git a/Engine/Plugins/Animation/ControlRig/Source/ControlRigEditor/Private/Editor/ControlRigEditorCommands.cpp b/Engine/Plugins/Animation/ControlRig/Source/ControlRigEditor/Private/Editor/ControlRigEditorCommands.cpp
--- a/Engine/Plugins/Animation/ControlRig/Source/ControlRigEditor/Private/Editor/ControlRigEditorCommands.cpp
+++ b/Engine/Plugins/Animation/ControlRig/Source/ControlRigEditor/Private/Editor/ControlRigEditorCommands.cpp
@@ -14,6 +14,7 @@ void FControlRigEditorCommands::RegisterCommands()
 	UI_COMMAND(RequestDirectManipulationRotation, "Request Direct Manipulation for Rotation", "Request per node direct manipulation on a rotation", EUserInterfaceActionType::Button, FInputChord(EKeys::E));
 	UI_COMMAND(RequestDirectManipulationScale, "Request Direct Manipulation for Scale", "Request per node direct manipulation on a scale", EUserInterfaceActionType::Button, FInputChord(EKeys::R));
 	UI_COMMAND(ToggleControlVisibility, "Show Controls", "Toggles the visibility of the controls.", EUserInterfaceActionType::ToggleButton, FInputChord(EKeys::T));
+	UI_COMMAND(ToggleControlsAsOverlay, "Show Controls as Overlay", "If checked controls will be rendered on top of other controls.", EUserInterfaceActionType::ToggleButton, FInputChord());
 	UI_COMMAND(ToggleDrawNulls, "Display Nulls", "If checked all nulls are drawn as axes.", EUserInterfaceActionType::ToggleButton, FInputChord());
 	UI_COMMAND(ToggleDrawSockets, "Display Sockets", "If checked all sockets are drawn.", EUserInterfaceActionType::ToggleButton, FInputChord());
 	UI_COMMAND(ToggleDrawAxesOnSelection, "Display Axes On Selection", "If checked axes will be drawn for all selected rig elements.", EUserInterfaceActionType::ToggleButton, FInputChord());
diff --git a/Engine/Plugins/Animation/ControlRig/Source/ControlRigEditor/Private/Editor/ControlRigEditorCommands.h b/Engine/Plugins/Animation/ControlRig/Source/ControlRigEditor/Private/Editor/ControlRigEditorCommands.h
--- a/Engine/Plugins/Animation/ControlRig/Source/ControlRigEditor/Private/Editor/ControlRigEditorCommands.h
+++ b/Engine/Plugins/Animation/ControlRig/Source/ControlRigEditor/Private/Editor/ControlRigEditorCommands.h
@@ -42,6 +42,9 @@ public:
 	/** Toggle visibility of the controls */
 	TSharedPtr< FUICommandInfo > ToggleControlVisibility;
 
+	/** Toggle if controls should be rendered on top of other controls */
+	TSharedPtr< FUICommandInfo > ToggleControlsAsOverlay;
+
 	/** Toggle visibility of nulls */
 	TSharedPtr< FUICommandInfo > ToggleDrawNulls;
 
