Control Rig: Set metadata failing when working in a rig module editor
#rnx

diff --git a/Engine/Plugins/Animation/ControlRig/Source/ControlRig/Private/Units/Hierarchy/RigUnit_Metadata.cpp b/Engine/Plugins/Animation/ControlRig/Source/ControlRig/Private/Units/Hierarchy/RigUnit_Metadata.cpp
--- a/Engine/Plugins/Animation/ControlRig/Source/ControlRig/Private/Units/Hierarchy/RigUnit_Metadata.cpp
+++ b/Engine/Plugins/Animation/ControlRig/Source/ControlRig/Private/Units/Hierarchy/RigUnit_Metadata.cpp
@@ -861,7 +861,7 @@ FRigBaseMetadata* FRigDispatch_GetModuleMetadata::FindMetadata(const FRigVMExten
 			if(Connector->IsPrimary())
 			{
 				const FName Name = ExecuteContext.AdaptMetadataName(InNameSpace, InName);
-				return ExecuteContext.Hierarchy->FindMetadataForElement(Connector, InName, InType);
+				return ExecuteContext.Hierarchy->FindMetadataForElement(Connector, Name, InType);
 			}
 		}
 	}
@@ -989,7 +989,7 @@ FRigBaseMetadata* FRigDispatch_SetModuleMetadata::FindOrAddMetadata(const FContr
 			if(Connector->IsPrimary())
 			{
 				const FName Name = InContext.AdaptMetadataName(InNameSpace, InName);
-				return InContext.Hierarchy->GetMetadataForElement(const_cast<FRigConnectorElement*>(Connector), InName, InType, bNotify);
+				return InContext.Hierarchy->GetMetadataForElement(const_cast<FRigConnectorElement*>(Connector), Name, InType, bNotify);
 			}
 		}
 	}
diff --git a/Engine/Plugins/Animation/ControlRig/Source/ControlRigEditor/Private/Editor/SRigHierarchy.cpp b/Engine/Plugins/Animation/ControlRig/Source/ControlRigEditor/Private/Editor/SRigHierarchy.cpp
--- a/Engine/Plugins/Animation/ControlRig/Source/ControlRigEditor/Private/Editor/SRigHierarchy.cpp
+++ b/Engine/Plugins/Animation/ControlRig/Source/ControlRigEditor/Private/Editor/SRigHierarchy.cpp
@@ -1889,6 +1889,26 @@ void SRigHierarchy::HandleDeleteItem()
 
  		TArray<FRigElementKey> SelectedKeys = GetSelectedKeys();
 
+ 		if (ControlRigBlueprint.IsValid() && ControlRigBlueprint->IsControlRigModule())
+ 		{
+ 			SelectedKeys.RemoveAll([Hierarchy, Controller](const FRigElementKey& Selected)
+			{
+				if (const FRigBaseElement* Element = Hierarchy->Find(Selected))
+				{
+				   if (const FRigConnectorElement* Connector = Cast<FRigConnectorElement>(Element))
+				   {
+					   if (Connector->IsPrimary())
+					   {
+						   static constexpr TCHAR Format[] = TEXT("Cannot delete primary connector: %s");
+						   Controller->ReportAndNotifyErrorf(Format, *Connector->GetName());
+						   return true;
+					   }
+				   }
+			   }
+				return false;
+			});
+ 		}
+
  		// clear selection early here to make sure ControlRigEditMode can react to this deletion
  		// it cannot react to it during Controller->RemoveElement() later because bSuspendAllNotifications is true
  		Controller->ClearSelection();
@@ -2097,7 +2117,9 @@ void SRigHierarchy::HandleNewItem(ERigElementType InElementType, bool bIsAnimati
 							}
 						}
 
-						const bool bIsPrimary = Hierarchy->GetConnectorKeys(false).Num() == 0;
+						const TArray<FRigConnectorElement*> Connectors = Hierarchy->GetConnectors(false);
+						const bool bIsPrimary = !Connectors.ContainsByPredicate([](const FRigConnectorElement* Connector) { return Connector->IsPrimary(); });
+						
 						FRigConnectorSettings Settings;
 						Settings.Type = bIsPrimary ? EConnectorType::Primary : EConnectorType::Secondary;
 							if(!bIsPrimary)
