Control Rig: Maintain selection when reparenting rig elements
#rnx

diff --git a/Engine/Plugins/Animation/ControlRig/Source/ControlRigEditor/Private/Editor/SRigHierarchy.cpp b/Engine/Plugins/Animation/ControlRig/Source/ControlRigEditor/Private/Editor/SRigHierarchy.cpp
--- a/Engine/Plugins/Animation/ControlRig/Source/ControlRigEditor/Private/Editor/SRigHierarchy.cpp
+++ b/Engine/Plugins/Animation/ControlRig/Source/ControlRigEditor/Private/Editor/SRigHierarchy.cpp
@@ -3454,6 +3454,8 @@ FReply SRigHierarchy::ReparentOrMatchTransform(const TArray<FRigElementKey>& Dra
 	URigHierarchy* DebuggedHierarchy = GetHierarchy();
 	URigHierarchy* Hierarchy = GetDefaultHierarchy();
 
+	const TArray<FRigElementKey> SelectedKeys = (Hierarchy) ? Hierarchy->GetSelectedKeys() : TArray<FRigElementKey>();
+
 	if (Hierarchy && ControlRigBlueprint.IsValid())
 	{
 		URigHierarchyController* Controller = Hierarchy->GetController(true);
@@ -3614,7 +3616,11 @@ FReply SRigHierarchy::ReparentOrMatchTransform(const TArray<FRigElementKey>& Dra
 		RefreshTreeView();
 	}
 
-		
+	if (Hierarchy)
+	{
+		Hierarchy->GetController()->SetSelection(SelectedKeys);
+	}
+	
 	return FReply::Handled();
 
 }
