Control Rig: Maintain selection when unparenting rig element
#rnx

diff --git a/Engine/Plugins/Animation/ControlRig/Source/ControlRigEditor/Private/Editor/SRigHierarchy.cpp b/Engine/Plugins/Animation/ControlRig/Source/ControlRigEditor/Private/Editor/SRigHierarchy.cpp
--- a/Engine/Plugins/Animation/ControlRig/Source/ControlRigEditor/Private/Editor/SRigHierarchy.cpp
+++ b/Engine/Plugins/Animation/ControlRig/Source/ControlRigEditor/Private/Editor/SRigHierarchy.cpp
@@ -3393,6 +3393,12 @@ void SRigHierarchy::HandleUnparent()
 	ControlRigBlueprint->PropagateHierarchyFromBPToInstances();
 	ControlRigEditor.Pin()->OnHierarchyChanged();
 	RefreshTreeView();
+
+	if (URigHierarchy* Hierarchy = GetDefaultHierarchy())
+	{
+		Hierarchy->GetController()->SetSelection(SelectedKeys);
+	}
+	
 	FSlateApplication::Get().DismissAllMenus();
 }
 
