Fix the control rig graph pin factory: it needs to create K2 pin widgets for all pins inside a control rig graph

diff --git a/Engine/Plugins/Animation/ControlRig/Source/ControlRigEditor/Private/Graph/ControlRigGraphPanelPinFactory.cpp b/Engine/Plugins/Animation/ControlRig/Source/ControlRigEditor/Private/Graph/ControlRigGraphPanelPinFactory.cpp
--- a/Engine/Plugins/Animation/ControlRig/Source/ControlRigEditor/Private/Graph/ControlRigGraphPanelPinFactory.cpp
+++ b/Engine/Plugins/Animation/ControlRig/Source/ControlRigEditor/Private/Graph/ControlRigGraphPanelPinFactory.cpp
@@ -29,88 +29,85 @@ TSharedPtr<SGraphPin> FControlRigGraphPanelPinFactory::CreatePin(UEdGraphPin* In
 	{
 		return nullptr;
 	}
-	
-	TSharedPtr<SGraphPin> InternalResult = CreatePin_Internal(InPin);
-	if(InternalResult.IsValid())
-	{
-		return InternalResult;
-	}
 
 	// if the graph we are looking at is not a control rig graph - let's not do this
 	if (const UEdGraphNode* OwningNode = InPin->GetOwningNode())
 	{
-		// only create pins within control rig graphs
-		if(const URigVMEdGraph* EdGraph = Cast<URigVMEdGraph>(OwningNode->GetGraph()))
+		const UControlRigGraph* EdGraph = Cast<UControlRigGraph>(OwningNode->GetGraph());
+		if(!EdGraph)
 		{
-			if ((Cast<UControlRigGraph>(EdGraph) == nullptr) &&
-				!EdGraph->IsPreviewGraph())
-			{
-				return nullptr;
-			}
+			return nullptr;
 		}
 	}
 	
-	if (InPin->Direction == EEdGraphPinDirection::EGPD_Input)
+	TSharedPtr<SGraphPin> InternalResult = CreatePin_Internal(InPin);
+	if(InternalResult.IsValid())
+	{
+		return InternalResult;
+	}
+
+	TSharedPtr<SGraphPin> K2PinWidget = FNodeFactory::CreateK2PinWidget(InPin);
+	if(K2PinWidget.IsValid())
 	{
-		// if we are an enum pin - and we are inside a RigElementKey,
-		// let's remove the "all" entry.
-		if(InPin->PinType.PinSubCategoryObject == StaticEnum<ERigElementType>())
+		if(InPin->Direction == EEdGraphPinDirection::EGPD_Input)
 		{
-			TSharedPtr<SGraphPin> K2PinWidget = FNodeFactory::CreateK2PinWidget(InPin);
-			if(InPin->ParentPin)
+			// if we are an enum pin - and we are inside a RigElementKey,
+			// let's remove the "all" entry.
+			if(InPin->PinType.PinSubCategoryObject == StaticEnum<ERigElementType>())
 			{
-				if(InPin->ParentPin->PinType.PinSubCategoryObject == FRigElementKey::StaticStruct())
+				if(InPin->ParentPin)
 				{
-					TSharedPtr<SWidget> ValueWidget = K2PinWidget->GetValueWidget();
-					if(ValueWidget.IsValid())
+					if(InPin->ParentPin->PinType.PinSubCategoryObject == FRigElementKey::StaticStruct())
 					{
-						if(TSharedPtr<SPinComboBox> EnumCombo = StaticCastSharedPtr<SPinComboBox>(ValueWidget))
+						TSharedPtr<SWidget> ValueWidget = K2PinWidget->GetValueWidget();
+						if(ValueWidget.IsValid())
 						{
-							if(EnumCombo.IsValid())
+							if(TSharedPtr<SPinComboBox> EnumCombo = StaticCastSharedPtr<SPinComboBox>(ValueWidget))
 							{
-								EnumCombo->RemoveItemByIndex(StaticEnum<ERigElementType>()->GetIndexByValue((int64)ERigElementType::All));
+								if(EnumCombo.IsValid())
+								{
+									EnumCombo->RemoveItemByIndex(StaticEnum<ERigElementType>()->GetIndexByValue((int64)ERigElementType::All));
+								}
 							}
 						}
 					}
 				}
 			}
-			return K2PinWidget;
-		}
 
-		const UEnum* RigControlTransformChannelEnum = StaticEnum<ERigControlTransformChannel>();
-		if (InPin->PinType.PinSubCategoryObject == RigControlTransformChannelEnum)
-		{
-			TSharedPtr<SGraphPin> K2PinWidget = FNodeFactory::CreateK2PinWidget(InPin);
-			TSharedPtr<SWidget> ValueWidget = K2PinWidget->GetValueWidget();
-			if(ValueWidget.IsValid())
+			const UEnum* RigControlTransformChannelEnum = StaticEnum<ERigControlTransformChannel>();
+			if (InPin->PinType.PinSubCategoryObject == RigControlTransformChannelEnum)
 			{
-				if(TSharedPtr<SPinComboBox> EnumCombo = StaticCastSharedPtr<SPinComboBox>(ValueWidget))
+				TSharedPtr<SWidget> ValueWidget = K2PinWidget->GetValueWidget();
+				if(ValueWidget.IsValid())
 				{
-					if(EnumCombo.IsValid())
+					if(TSharedPtr<SPinComboBox> EnumCombo = StaticCastSharedPtr<SPinComboBox>(ValueWidget))
 					{
-						if (const UControlRigGraphNode* RigNode = Cast<UControlRigGraphNode>(InPin->GetOwningNode()))
+						if(EnumCombo.IsValid())
 						{
-							if (const URigVMPin* ModelPin = RigNode->GetModelPinFromPinPath(InPin->GetName()))
+							if (const UControlRigGraphNode* RigNode = Cast<UControlRigGraphNode>(InPin->GetOwningNode()))
 							{
-								if(const URigVMUnitNode* UnitNode = Cast<URigVMUnitNode>(ModelPin->GetNode()))
+								if (const URigVMPin* ModelPin = RigNode->GetModelPinFromPinPath(InPin->GetName()))
 								{
-									if(UnitNode->GetScriptStruct() &&
-										UnitNode->GetScriptStruct()->IsChildOf(FRigUnit_HierarchyAddControlElement::StaticStruct()))
+									if(const URigVMUnitNode* UnitNode = Cast<URigVMUnitNode>(ModelPin->GetNode()))
 									{
-										const TSharedPtr<FStructOnScope> StructInstanceScope = UnitNode->ConstructStructInstance();
-										const FRigUnit_HierarchyAddControlElement* StructInstance = 
-											(const FRigUnit_HierarchyAddControlElement*)StructInstanceScope->GetStructMemory();
-
-										if(const TArray<ERigControlTransformChannel>* VisibleChannels =
-											FRigControlTransformChannelDetails::GetVisibleChannelsForControlType(StructInstance->GetControlTypeToSpawn()))
+										if(UnitNode->GetScriptStruct() &&
+											UnitNode->GetScriptStruct()->IsChildOf(FRigUnit_HierarchyAddControlElement::StaticStruct()))
 										{
-											for(int32 Index = 0; Index < RigControlTransformChannelEnum->NumEnums(); Index++)
+											const TSharedPtr<FStructOnScope> StructInstanceScope = UnitNode->ConstructStructInstance();
+											const FRigUnit_HierarchyAddControlElement* StructInstance = 
+												(const FRigUnit_HierarchyAddControlElement*)StructInstanceScope->GetStructMemory();
+
+											if(const TArray<ERigControlTransformChannel>* VisibleChannels =
+												FRigControlTransformChannelDetails::GetVisibleChannelsForControlType(StructInstance->GetControlTypeToSpawn()))
 											{
-												const ERigControlTransformChannel Value =
-													(ERigControlTransformChannel)RigControlTransformChannelEnum->GetValueByIndex(Index);
-												if(!VisibleChannels->Contains(Value))
+												for(int32 Index = 0; Index < RigControlTransformChannelEnum->NumEnums(); Index++)
 												{
-													EnumCombo->RemoveItemByIndex(Index);
+													const ERigControlTransformChannel Value =
+														(ERigControlTransformChannel)RigControlTransformChannelEnum->GetValueByIndex(Index);
+													if(!VisibleChannels->Contains(Value))
+													{
+														EnumCombo->RemoveItemByIndex(Index);
+													}
 												}
 											}
 										}
@@ -121,11 +118,9 @@ TSharedPtr<SGraphPin> FControlRigGraphPanelPinFactory::CreatePin(UEdGraphPin* In
 					}
 				}
 			}
-			return K2PinWidget;
 		}
 	}
-	
-	return nullptr;
+	return K2PinWidget;
 }
 
 FName FControlRigGraphPanelPinFactory::GetFactoryName() const
