Control Rig: If possible, recreate manually added elements on modular rig hierarchies after changing the preview mesh
#rnx

diff --git a/Engine/Plugins/Animation/ControlRig/Source/ControlRigEditor/Private/Editor/ControlRigEditor.cpp b/Engine/Plugins/Animation/ControlRig/Source/ControlRigEditor/Private/Editor/ControlRigEditor.cpp
--- a/Engine/Plugins/Animation/ControlRig/Source/ControlRigEditor/Private/Editor/ControlRigEditor.cpp
+++ b/Engine/Plugins/Animation/ControlRig/Source/ControlRigEditor/Private/Editor/ControlRigEditor.cpp
@@ -2843,6 +2843,7 @@ void FControlRigEditor::HandlePreviewMeshChanged(USkeletalMesh* InOldSkeletalMes
 		if (UControlRigBlueprint* ControlRigBP = GetControlRigBlueprint())
 		{
 			ControlRigBP->SetPreviewMesh(InNewSkeletalMesh);
+			URigHierarchy* BPHierarchy = ControlRigBP->GetHierarchy();
 
 			FModularRigConnections PreviousConnections;
 			if(IsModularRig())
@@ -2854,16 +2855,44 @@ void FControlRigEditor::HandlePreviewMeshChanged(USkeletalMesh* InOldSkeletalMes
 					{
 						// remove all connectors / sockets. keeping them around may mess up the order of the elements
 						// in the hierarchy, such as [bone,bone,bone,connector,connector,bone,bone,bone].
+						// if the element is manually created, remember it to create it after importing the skeleton element
 						TArray<FRigElementKey> ConnectorsAndSockets = Controller->GetHierarchy()->GetConnectorKeys();
 						ConnectorsAndSockets.Append(Controller->GetHierarchy()->GetSocketKeys());
+
+						TArray<TTuple<FRigElementKey, FRigElementKey, FTransform>> ConnectorsAndSocketsToParents;
+						ConnectorsAndSocketsToParents.Reserve(ConnectorsAndSockets.Num());
+						
 						for(const FRigElementKey& Key : ConnectorsAndSockets)
 						{
+							// Remember manually created elements to apply them again
+							if (BPHierarchy->GetNameSpace(Key).IsEmpty())
+							{
+								const FRigElementKey& Parent = BPHierarchy->GetDefaultParent(Key);
+								ConnectorsAndSocketsToParents.Emplace(Key, Parent, BPHierarchy->GetLocalTransform(Key));
+							}
 							(void)Controller->RemoveElement(Key, true, true);
 						}
 						
 						USkeleton* Skeleton = InNewSkeletalMesh ? InNewSkeletalMesh->GetSkeleton() : nullptr;
 						Controller->ImportBones(Skeleton, NAME_None, true, true, false, true, true);
 						Controller->ImportCurves(Skeleton, NAME_None, false, true, true);
+
+						// Recreate manually created elements
+						for (const TTuple<FRigElementKey, FRigElementKey, FTransform>& Tuple : ConnectorsAndSocketsToParents)
+						{
+							const FRigElementKey& Key = Tuple.Get<0>();
+							const FRigElementKey& Parent = Tuple.Get<1>();
+							const FTransform& Transform = Tuple.Get<2>();
+							
+							if (!Parent.IsValid() || BPHierarchy->Contains(Parent))
+							{
+								switch (Key.Type)
+								{
+									case ERigElementType::Socket: Controller->AddSocket(Key.Name, Parent, Transform, false); break;
+									case ERigElementType::Connector: Controller->AddConnector(Key.Name); break;
+								}
+							}
+						}
 					}
 				}
 				ControlRigBP->PropagateHierarchyFromBPToInstances();
