Control Rig: Do not use control shape global transform when there is no a proper shape (e.g. FKControlRig)
#rnx

diff --git a/Engine/Plugins/Animation/ControlRig/Source/ControlRigEditor/Private/EditMode/ControlRigEditMode.cpp b/Engine/Plugins/Animation/ControlRig/Source/ControlRigEditor/Private/EditMode/ControlRigEditMode.cpp
--- a/Engine/Plugins/Animation/ControlRig/Source/ControlRigEditor/Private/EditMode/ControlRigEditMode.cpp
+++ b/Engine/Plugins/Animation/ControlRig/Source/ControlRigEditor/Private/EditMode/ControlRigEditMode.cpp
@@ -4790,7 +4790,18 @@ void FControlRigEditMode::MoveControlShape(AControlRigShapeActor* ShapeActor, co
 	if(!bTransformChanged) //not local or doing scale.
 	{
 		// Get the global transform from shape actor to avoid drifting
-		FTransform CurrentTransform = ShapeActor->GetGlobalTransform() * ToWorldTransform;
+		FTransform CurrentTransform;
+		if (ShapeActor->StaticMeshComponent && ShapeActor->StaticMeshComponent->GetStaticMesh())
+		{
+			CurrentTransform = ShapeActor->GetGlobalTransform() * ToWorldTransform;
+		}
+		else
+		{
+			// If the static mesh is not valid, we cannot rely on the shape's transform.
+			// This happens for FKControlRigs (and other control types)
+			// We will need to rely on the information we have in the rig hierarchy
+			CurrentTransform = GetControlShapeTransform(ShapeActor) * ToWorldTransform;
+		}
 
 		if (bRotation)
 		{
