Constraints: handle constraints keyframes between controls from different control rigs acting on the same skeletal mesh

#rnx

diff --git a/Engine/Plugins/Animation/ControlRig/Source/ControlRigEditor/Private/Sequencer/ControlConstraintChannelInterface.cpp b/Engine/Plugins/Animation/ControlRig/Source/ControlRigEditor/Private/Sequencer/ControlConstraintChannelInterface.cpp
--- a/Engine/Plugins/Animation/ControlRig/Source/ControlRigEditor/Private/Sequencer/ControlConstraintChannelInterface.cpp
+++ b/Engine/Plugins/Animation/ControlRig/Source/ControlRigEditor/Private/Sequencer/ControlConstraintChannelInterface.cpp
@@ -248,11 +248,13 @@ UMovieSceneControlRigParameterSection* FControlConstraintChannelInterface::GetCo
 			const TArray<FMovieSceneBinding>& Bindings = MovieScene->GetBindings();
 			for (const FMovieSceneBinding& Binding : Bindings)
 			{
-				UMovieSceneTrack* Track = MovieScene->FindTrack(UMovieSceneControlRigParameterTrack::StaticClass(), Binding.GetObjectGuid());
-				UMovieSceneControlRigParameterTrack* ControlRigTrack = Cast<UMovieSceneControlRigParameterTrack>(Track);
-				if (ControlRigTrack && ControlRigTrack->GetControlRig() == ControlRig)
+				for (UMovieSceneTrack* Track : MovieScene->FindTracks(UMovieSceneControlRigParameterTrack::StaticClass(), Binding.GetObjectGuid()))
 				{
-					return ControlRigTrack;
+					UMovieSceneControlRigParameterTrack* ControlRigTrack = Cast<UMovieSceneControlRigParameterTrack>(Track);
+					if (ControlRigTrack && ControlRigTrack->GetControlRig() == ControlRig)
+					{
+						return ControlRigTrack;
+					}
 				}
 			}
 		}
