Sequencer: Control Rig: Collapse All sections now uses a modeless dialog as a speculative fix

diff --git a/Engine/Plugins/Animation/ControlRig/Source/ControlRigEditor/Private/Sequencer/ControlRigParameterTrackEditor.cpp b/Engine/Plugins/Animation/ControlRig/Source/ControlRigEditor/Private/Sequencer/ControlRigParameterTrackEditor.cpp
--- a/Engine/Plugins/Animation/ControlRig/Source/ControlRigEditor/Private/Sequencer/ControlRigParameterTrackEditor.cpp
+++ b/Engine/Plugins/Animation/ControlRig/Source/ControlRigEditor/Private/Sequencer/ControlRigParameterTrackEditor.cpp
@@ -4553,7 +4553,7 @@ public:
 		: _Sequencer(nullptr), _OwnerTrack(nullptr)
 	{}
 
-	SLATE_ARGUMENT(TSharedPtr<ISequencer>, Sequencer)
+	SLATE_ARGUMENT(TWeakPtr<ISequencer>, Sequencer)
 	SLATE_ARGUMENT(UMovieSceneTrack*, OwnerTrack)
 	SLATE_END_ARGS()
 
@@ -4567,7 +4567,7 @@ public:
 private:
 	void Collapse();
 
-	TSharedPtr<ISequencer> Sequencer;
+	TWeakPtr<ISequencer> Sequencer;
 	TWeakObjectPtr<UMovieSceneTrack> OwnerTrack;
 	//static to be reused
 	static TOptional<FBakingAnimationKeySettings> CollapseControlsSettings;
@@ -4582,18 +4582,18 @@ TOptional<FBakingAnimationKeySettings> SCollapseControlsWidget::CollapseControls
 
 void SCollapseControlsWidget::Construct(const FArguments& InArgs)
 {
-	check(InArgs._Sequencer);
 	Sequencer = InArgs._Sequencer;
 	OwnerTrack = InArgs._OwnerTrack;
 
 	if (CollapseControlsSettings.IsSet() == false)
 	{
+		TSharedPtr<ISequencer> SequencerPtr = Sequencer.Pin();
 		CollapseControlsSettings = FBakingAnimationKeySettings();
-		const FFrameRate TickResolution = Sequencer->GetFocusedTickResolution();
-		const FFrameTime FrameTime = Sequencer->GetLocalTime().ConvertTo(TickResolution);
+		const FFrameRate TickResolution = SequencerPtr->GetFocusedTickResolution();
+		const FFrameTime FrameTime = SequencerPtr->GetLocalTime().ConvertTo(TickResolution);
 		FFrameNumber CurrentTime = FrameTime.GetFrame();
 
-		TRange<FFrameNumber> Range = Sequencer->GetFocusedMovieSceneSequence()->GetMovieScene()->GetPlaybackRange();
+		TRange<FFrameNumber> Range = SequencerPtr->GetFocusedMovieSceneSequence()->GetMovieScene()->GetPlaybackRange();
 		TArray<FFrameNumber> Keys;
 		TArray < FKeyHandle> KeyHandles;
 
@@ -4619,7 +4619,7 @@ void SCollapseControlsWidget::Construct(const FArguments& InArgs)
 	FPropertyEditorModule& PropertyEditor = FModuleManager::Get().LoadModuleChecked<FPropertyEditorModule>(TEXT("PropertyEditor"));
 
 	DetailsView = PropertyEditor.CreateStructureDetailView(ViewArgs, StructureViewArgs, TSharedPtr<FStructOnScope>());
-	TSharedPtr<INumericTypeInterface<double>> NumericTypeInterface = Sequencer->GetNumericTypeInterface();
+	TSharedPtr<INumericTypeInterface<double>> NumericTypeInterface = Sequencer.Pin()->GetNumericTypeInterface();
 	DetailsView->GetDetailsView()->RegisterInstancedCustomPropertyTypeLayout("FrameNumber",
 		FOnGetPropertyTypeCustomizationInstance::CreateLambda([=]() {return MakeShared<FFrameNumberDetailsCustomization>(NumericTypeInterface); }));
 	DetailsView->SetStructureData(Settings);
@@ -4680,7 +4680,8 @@ void SCollapseControlsWidget::Construct(const FArguments& InArgs)
 void  SCollapseControlsWidget::Collapse()
 {
 	FBakingAnimationKeySettings* BakeSettings = Settings->Get();
-	FControlRigParameterTrackEditor::CollapseAllLayers(Sequencer, OwnerTrack.Get(), *BakeSettings);
+	TSharedPtr<ISequencer> SequencerPtr = Sequencer.Pin();
+	FControlRigParameterTrackEditor::CollapseAllLayers(SequencerPtr, OwnerTrack.Get(), *BakeSettings);
 
 	CollapseControlsSettings = *BakeSettings;
 }
@@ -4927,7 +4928,7 @@ void FControlRigParameterSection::CollapseAllLayers()
 			.Sequencer(Sequencer)
 			.OwnerTrack(OwnerTrack);
 
-		BakeWidget->OpenDialog(true);
+		BakeWidget->OpenDialog(false);
 	}
 }
 
