Control Rig: Control creation from bones creates controls at root
#rnx

diff --git a/Engine/Plugins/Animation/ControlRig/Source/ControlRigEditor/Private/Editor/SRigHierarchy.cpp b/Engine/Plugins/Animation/ControlRig/Source/ControlRigEditor/Private/Editor/SRigHierarchy.cpp
--- a/Engine/Plugins/Animation/ControlRig/Source/ControlRigEditor/Private/Editor/SRigHierarchy.cpp
+++ b/Engine/Plugins/Animation/ControlRig/Source/ControlRigEditor/Private/Editor/SRigHierarchy.cpp
@@ -1981,7 +1981,8 @@ void SRigHierarchy::HandleNewItem(ERigElementType InElementType, bool bIsAnimati
 
 		const bool bAllowMultipleItems =
 			InElementType == ERigElementType::Socket ||
-			InElementType == ERigElementType::Null;
+			InElementType == ERigElementType::Null ||
+			(InElementType == ERigElementType::Control && !bIsAnimationChannel);
 
 		URigHierarchyController* Controller = Hierarchy->GetController(true);
 		check(Controller);
@@ -1998,6 +1999,7 @@ void SRigHierarchy::HandleNewItem(ERigElementType InElementType, bool bIsAnimati
 			SelectedKeys = {FRigElementKey()};
 		}
 
+		TMap<FRigElementKey, FRigElementKey> SelectedToCreated;
 		for(const FRigElementKey& SelectedKey : SelectedKeys)
 		{
 			FRigElementKey ParentKey;
@@ -2080,7 +2082,17 @@ void SRigHierarchy::HandleNewItem(ERigElementType InElementType, bool bIsAnimati
 							Settings.MinimumValue = ValueToSet;
 							Settings.MaximumValue = ValueToSet;
 
-							NewItemKey = Controller->AddControl(NewElementName, ParentKey, Settings, Settings.GetIdentityValue(), FTransform::Identity, FTransform::Identity, true, true);
+							FRigElementKey NewParentKey;
+							FTransform OffsetTransform = ParentTransform;
+							if (FRigElementKey* CreatedParentKey = SelectedToCreated.Find(Hierarchy->GetDefaultParent(ParentKey)))
+							{
+								NewParentKey = *CreatedParentKey;
+								OffsetTransform = ParentTransform.GetRelativeTransform(Hierarchy->GetGlobalTransform(NewParentKey, true));
+							}
+
+							NewItemKey = Controller->AddControl(NewElementName, NewParentKey, Settings, Settings.GetIdentityValue(), OffsetTransform, FTransform::Identity, true, true);
+
+							SelectedToCreated.Add(SelectedKey, NewItemKey);
 						}
 						break;
 					}
