Control Rig: Make sure the shape actor is attached to a scene component when moving it
#rnx

diff --git a/Engine/Plugins/Animation/ControlRig/Source/ControlRigEditor/Private/EditMode/ControlRigEditMode.cpp b/Engine/Plugins/Animation/ControlRig/Source/ControlRigEditor/Private/EditMode/ControlRigEditMode.cpp
--- a/Engine/Plugins/Animation/ControlRig/Source/ControlRigEditor/Private/EditMode/ControlRigEditMode.cpp
+++ b/Engine/Plugins/Animation/ControlRig/Source/ControlRigEditor/Private/EditMode/ControlRigEditMode.cpp
@@ -4703,6 +4703,21 @@ void FControlRigEditMode::MoveControlShape(AControlRigShapeActor* ShapeActor, co
 {
 	bool bTransformChanged = false;
 
+	// In case for some reason the shape actor was detached, make sure to attach it again to the scene component
+	if (!ShapeActor->GetAttachParentActor())
+	{
+		if (UControlRig* ControlRig = ShapeActor->ControlRig.Get())
+		{
+			if (USceneComponent* SceneComponent = GetHostingSceneComponent(ControlRig))
+			{
+				if (AActor* OwnerActor = SceneComponent->GetOwner())
+				{
+					ShapeActor->AttachToActor(OwnerActor, FAttachmentTransformRules::KeepWorldTransform);
+				}
+			}
+		}
+	}
+
 	auto RotatorToStr = [](const FRotator& InRotator)
 	{
 		FRigPreferredEulerAngles EulerAngles;
