Control Rig: Fix new module creation when a connector is selected
#rnx

diff --git a/Engine/Plugins/Animation/ControlRig/Source/ControlRigEditor/Private/Editor/SModularRigModel.cpp b/Engine/Plugins/Animation/ControlRig/Source/ControlRigEditor/Private/Editor/SModularRigModel.cpp
--- a/Engine/Plugins/Animation/ControlRig/Source/ControlRigEditor/Private/Editor/SModularRigModel.cpp
+++ b/Engine/Plugins/Animation/ControlRig/Source/ControlRigEditor/Private/Editor/SModularRigModel.cpp
@@ -529,7 +529,11 @@ void SModularRigModel::HandleNewItem()
 	FString ParentPath;
 	if (IsSingleSelected())
 	{
-		ParentPath = GetSelectedKeys()[0];
+		const TSharedPtr<FModularRigTreeElement> ParentElement = TreeView->FindElement(GetSelectedKeys()[0]);
+		if (ParentElement.IsValid())
+		{
+			ParentPath = ParentElement->ModulePath;
+		}
 	}
 	
 	FClassViewerInitializationOptions Options;
