Initial version of Animation Mode extension of the new level editor viewport toolbar

When entering Animation Mode, the Transforms, Select, and Show menus of the new viewport toolbar will be extended with more entries.

diff --git a/Engine/Plugins/Animation/ControlRig/Source/ControlRigEditor/Private/EditMode/ControlRigEditMode.cpp b/Engine/Plugins/Animation/ControlRig/Source/ControlRigEditor/Private/EditMode/ControlRigEditMode.cpp
--- a/Engine/Plugins/Animation/ControlRig/Source/ControlRigEditor/Private/EditMode/ControlRigEditMode.cpp
+++ b/Engine/Plugins/Animation/ControlRig/Source/ControlRigEditor/Private/EditMode/ControlRigEditMode.cpp
@@ -74,6 +74,7 @@
 #include "DragTool_FrustumSelect.h"
 #include "AnimationEditorViewportClient.h"
 #include "EditorInteractiveGizmoManager.h"
+#include "Editor/ControlRigViewportToolbarExtensions.h"
 #include "Tools/BakingHelper.h"
 
 #include UE_INLINE_GENERATED_CPP_BY_NAME(ControlRigEditMode)
@@ -441,6 +442,12 @@ void FControlRigEditMode::Enter()
 			Settings->OnSettingsChange.AddSP(this, &FControlRigEditMode::OnSettingsChanged);
 		}
 	}
+
+	UE::ControlRig::PopulateControlRigViewportToolbarTransformsSubmenu("LevelEditor.ViewportToolbar.Transforms");
+	UE::ControlRig::PopulateControlRigViewportToolbarSelectionSubmenu("LevelEditor.ViewportToolbar.Select");
+	UE::ControlRig::PopulateControlRigViewportToolbarShowSubmenu(
+		"LevelEditor.ViewportToolbar.Show", GetToolkit()->GetToolkitCommands()
+	);
 }
 
 //todo get working with Persona
@@ -463,6 +470,8 @@ static void ClearOutAnyActiveTools()
 
 void FControlRigEditMode::Exit()
 {
+	UE::ControlRig::RemoveControlRigViewportToolbarExtensions();
+
 	ClearOutAnyActiveTools();
 	OnControlRigAddedOrRemovedDelegate.Clear();
 	OnControlRigSelectedDelegate.Clear();
diff --git a/Engine/Plugins/Animation/ControlRig/Source/ControlRigEditor/Private/Editor/ControlRigViewportToolbarExtensions.cpp b/Engine/Plugins/Animation/ControlRig/Source/ControlRigEditor/Private/Editor/ControlRigViewportToolbarExtensions.cpp
new file mode 100644
--- /dev/null
+++ b/Engine/Plugins/Animation/ControlRig/Source/ControlRigEditor/Private/Editor/ControlRigViewportToolbarExtensions.cpp
@@ -0,0 +1,368 @@
+﻿// Copyright Epic Games, Inc. All Rights Reserved.
+
+#include "Editor/ControlRigViewportToolbarExtensions.h"
+
+#include "ControlRigEditorCommands.h"
+#include "DetailCategoryBuilder.h"
+#include "EditMode/ControlRigEditModeSettings.h"
+#include "ToolMenus.h"
+#include "Tools/MotionTrailOptions.h"
+#include "Widgets/Input/SSpinBox.h"
+
+#include "LevelEditorActions.h"
+
+#define LOCTEXT_NAMESPACE "ControlRigViewportToolbar"
+
+namespace UE::ControlRig::Private
+{
+
+const FName ControlRigOwnerName = "ControlRigViewportToolbar";
+
+} // namespace UE::ControlRig::Private
+
+void UE::ControlRig::PopulateControlRigViewportToolbarTransformsSubmenu(const FName InMenuName)
+{
+	FToolMenuOwnerScoped ScopeOwner(UE::ControlRig::Private::ControlRigOwnerName);
+
+	UToolMenu* const Menu = UToolMenus::Get()->ExtendMenu(InMenuName);
+
+	{
+		FToolMenuSection& GizmoSection = Menu->FindOrAddSection("Gizmo");
+
+		UControlRigEditModeSettings* const ViewportSettings = GetMutableDefault<UControlRigEditModeSettings>();
+
+		// Add "Local Transforms in Each Local Space" checkbox.
+		{
+			FUIAction Action;
+
+			Action.ExecuteAction = FExecuteAction::CreateLambda(
+				[ViewportSettings]() -> void
+				{
+					ViewportSettings->bLocalTransformsInEachLocalSpace = !ViewportSettings->bLocalTransformsInEachLocalSpace;
+					ViewportSettings->PostEditChange();
+				}
+			);
+			Action.GetActionCheckState = FGetActionCheckState::CreateLambda(
+				[ViewportSettings]() -> ECheckBoxState
+				{
+					return ViewportSettings->bLocalTransformsInEachLocalSpace ? ECheckBoxState::Checked
+																			  : ECheckBoxState::Unchecked;
+				}
+			);
+
+			FToolMenuEntry Entry = FToolMenuEntry::InitMenuEntry(
+				"LocalTransformsInEachLocalSpace",
+				LOCTEXT("LocalTransformsInEachLocalSpaceLabel", "Local Transforms in Each Local Space"),
+				LOCTEXT(
+					"LocalTransformsInEachLocalSpaceTooltip", "When multiple objects are selected, whether or not to transform each invidual object along its own local transform space."
+				),
+				FSlateIcon(),
+				Action,
+				EUserInterfaceActionType::ToggleButton
+			);
+			// We want to appear early in the section.
+			Entry.InsertPosition.Position = EToolMenuInsertType::First;
+			Entry.SetShowInToolbarTopLevel(true);
+			GizmoSection.AddEntry(Entry);
+		}
+
+		// Add "Restore Coordinate Space on Switch" checkbox.
+		{
+			FUIAction Action;
+
+			Action.ExecuteAction = FExecuteAction::CreateLambda(
+				[ViewportSettings]() -> void
+				{
+					ViewportSettings->bCoordSystemPerWidgetMode = !ViewportSettings->bCoordSystemPerWidgetMode;
+					ViewportSettings->PostEditChange();
+				}
+			);
+			Action.GetActionCheckState = FGetActionCheckState::CreateLambda(
+				[ViewportSettings]() -> ECheckBoxState
+				{
+					return ViewportSettings->bCoordSystemPerWidgetMode ? ECheckBoxState::Checked
+																	   : ECheckBoxState::Unchecked;
+				}
+			);
+
+			FToolMenuEntry Entry = FToolMenuEntry::InitMenuEntry(
+				"RestoreCoordinateSpaceOnSwitch",
+				LOCTEXT("RestoreCoordinateSpaceOnSwitchLabel", "Restore Coordinate Space on Switch"),
+
+				LOCTEXT(
+					"RestoreCoordinateSpaceOnSwitchTooltip",
+					"Whether to restore the coordinate space when changing Widget Modes in the Viewport."
+				),
+				FSlateIcon(),
+				Action,
+				EUserInterfaceActionType::ToggleButton
+			);
+			// We want to appear early in the section.
+			Entry.InsertPosition.Position = EToolMenuInsertType::First;
+			Entry.SetShowInToolbarTopLevel(true);
+			GizmoSection.AddEntry(Entry);
+		}
+	}
+}
+
+void UE::ControlRig::PopulateControlRigViewportToolbarSelectionSubmenu(const FName InMenuName)
+{
+	FToolMenuOwnerScoped ScopeOwner(UE::ControlRig::Private::ControlRigOwnerName);
+
+	UToolMenu* const Menu = UToolMenus::Get()->ExtendMenu(InMenuName);
+	FToolMenuSection& OptionsSection = Menu->FindOrAddSection("Options");
+	UControlRigEditModeSettings* const Settings = GetMutableDefault<UControlRigEditModeSettings>();
+
+	// Add "Restore Coordinate Space on Switch" checkbox.
+	{
+		FUIAction Action;
+
+		Action.ExecuteAction = FExecuteAction::CreateLambda(
+			[Settings]() -> void
+			{
+				Settings->bOnlySelectRigControls = !Settings->bOnlySelectRigControls;
+				Settings->PostEditChange();
+			}
+		);
+		Action.GetActionCheckState = FGetActionCheckState::CreateLambda(
+			[Settings]() -> ECheckBoxState
+			{
+				return Settings->bOnlySelectRigControls ? ECheckBoxState::Checked : ECheckBoxState::Unchecked;
+			}
+		);
+
+		FToolMenuEntry Entry = FToolMenuEntry::InitMenuEntry(
+			"OnlySelectRigControls",
+			LOCTEXT("OnlySelectRigControlsLabel", "Select Only Control Rig Controls"),
+			LOCTEXT("OnlySelectRigControlsTooltip", "Whether or not only Rig Controls can be selected."),
+			FSlateIcon(),
+			Action,
+			EUserInterfaceActionType::ToggleButton
+		);
+		// We want to appear late in the section.
+		Entry.InsertPosition.Position = EToolMenuInsertType::Last;
+		Entry.SetShowInToolbarTopLevel(true);
+		OptionsSection.AddEntry(Entry);
+	}
+}
+
+void UE::ControlRig::PopulateControlRigViewportToolbarShowSubmenu(
+	const FName InMenuName, const TSharedPtr<const FUICommandList>& InCommandList
+)
+{
+	FToolMenuOwnerScoped ScopeOwner(UE::ControlRig::Private::ControlRigOwnerName);
+
+	UToolMenu* const Menu = UToolMenus::Get()->ExtendMenu(InMenuName);
+	FToolMenuSection& AllShowFlagsSection = Menu->FindOrAddSection("AllShowFlags");
+
+	FToolMenuEntry AnimationSubmenu = FToolMenuEntry::InitSubMenu(
+		"Animation",
+		LOCTEXT("AnimationLabel", "Animation"),
+		LOCTEXT("AnimationTooltip", "Animation-related show flags"),
+		FNewToolMenuDelegate::CreateLambda(
+			[WeakCommandList = InCommandList.ToWeakPtr()](UToolMenu* AnimationShowFlagsSubmenu)
+			{
+				FToolMenuSection& UnnamedSection = AnimationShowFlagsSubmenu->FindOrAddSection(NAME_None);
+
+				// TODO: These don't work because the commands have not yet been bound.
+				if (const TSharedPtr<const FUICommandList>& CommandList = WeakCommandList.Pin())
+				{
+					UnnamedSection.AddMenuEntryWithCommandList(
+						FControlRigEditorCommands::Get().ToggleDrawAxesOnSelection, CommandList
+					);
+					UnnamedSection.AddMenuEntryWithCommandList(
+						FControlRigEditorCommands::Get().ToggleControlsAsOverlay, CommandList
+					);
+					UnnamedSection.AddMenuEntryWithCommandList(
+						FControlRigEditorCommands::Get().ToggleControlVisibility, CommandList
+					);
+					UnnamedSection.AddMenuEntryWithCommandList(FControlRigEditorCommands::Get().ToggleDrawNulls, CommandList);
+					UnnamedSection.AddMenuEntryWithCommandList(FControlRigEditorCommands::Get().ToggleDrawSockets, CommandList);
+				}
+
+				FNewToolMenuDelegate MakeMenuDelegate = FNewToolMenuDelegate::CreateLambda(
+					[](UToolMenu* Submenu)
+					{
+						UMotionTrailToolOptions* Settings = GetMutableDefault<UMotionTrailToolOptions>();
+
+						// TODO: This section needs to be completed to show all radio buttons. Perhaps the settings
+						// object should be updated to better match the design, or the design updated to match what the
+						// settings object can express.
+						{
+							FToolMenuSection& PathModeSection =
+								Submenu->FindOrAddSection("PathMode", LOCTEXT("PathModeLabel", "Path Mode"));
+
+							// Add "Full Trail" radio button.
+							{
+								FUIAction Action;
+
+								Action.ExecuteAction = FExecuteAction::CreateLambda(
+									[Settings]() -> void
+									{
+										Settings->bShowFullTrail = true;
+										Settings->PostEditChange();
+									}
+								);
+								Action.GetActionCheckState = FGetActionCheckState::CreateLambda(
+									[Settings]() -> ECheckBoxState
+									{
+										return Settings->bShowFullTrail ? ECheckBoxState::Checked
+																		: ECheckBoxState::Unchecked;
+									}
+								);
+
+								FToolMenuEntry Entry = FToolMenuEntry::InitMenuEntry(
+									"FullTrail",
+									LOCTEXT("FullTrailLabel", "Full trail"),
+									LOCTEXT("FullTrailTooltip", "Whether or not to show the full motion trail."),
+									FSlateIcon(),
+									Action,
+									EUserInterfaceActionType::RadioButton
+								);
+								PathModeSection.AddEntry(Entry);
+							}
+						}
+
+						{
+							FToolMenuSection& PathOptionsSection =
+								Submenu->FindOrAddSection("PathOptions", LOCTEXT("PathOptionsLabel", "Path Options"));
+
+							// Add "Show Keys" checkbox.
+							{
+								FUIAction Action;
+
+								Action.ExecuteAction = FExecuteAction::CreateLambda(
+									[Settings]() -> void
+									{
+										Settings->bShowKeys = !Settings->bShowKeys;
+										Settings->PostEditChange();
+									}
+								);
+								Action.GetActionCheckState = FGetActionCheckState::CreateLambda(
+									[Settings]() -> ECheckBoxState
+									{
+										return Settings->bShowKeys ? ECheckBoxState::Checked : ECheckBoxState::Unchecked;
+									}
+								);
+
+								FToolMenuEntry Entry = FToolMenuEntry::InitMenuEntry(
+									"ShowKeys",
+									LOCTEXT("ShowKeysLabel", "Show Keys"),
+									LOCTEXT("ShowKeysTooltip", "Whether or not to show keys on the motion trail."),
+									FSlateIcon(),
+									Action,
+									EUserInterfaceActionType::ToggleButton
+								);
+								PathOptionsSection.AddEntry(Entry);
+							}
+
+							// TODO: This currently only uses FramesBefore. This should be updated to support all cases.
+							TSharedRef<SWidget> FramesShownWidget =
+								// clang-format off
+								SNew(SHorizontalBox)
+								+ SHorizontalBox::Slot()
+								.FillWidth(0.9f)
+								[
+									SNew(SSpinBox<int32>)
+										.MinValue(0)
+										.MaxValue(100)
+										.MinDesiredWidth(100)
+										.ToolTipText_Lambda(
+											[Settings]() -> FText
+											{
+												return FText::AsNumber(
+													
+													Settings->FramesBefore
+												);
+											}
+										)
+										.Value_Lambda(
+											[Settings]() -> float
+											{
+												return Settings->FramesBefore;
+											}
+										)
+										.OnValueChanged_Lambda(
+											[Settings](float InValue)
+											{
+												Settings->FramesBefore = InValue;
+											}
+										)
+								]
+								+ SHorizontalBox::Slot()
+								.FillWidth(0.1f);
+							// clang-format on
+							PathOptionsSection.AddEntry(FToolMenuEntry::InitWidget(
+								"FramesShown", FramesShownWidget, LOCTEXT("FramesShownLabel", "Frames Shown")
+							));
+						}
+					}
+				);
+
+				UMotionTrailToolOptions* Settings = GetMutableDefault<UMotionTrailToolOptions>();
+
+				// Create the checkbox actions for the MotionPaths submenu itself.
+				FToolUIAction CheckboxMenuAction;
+				{
+					CheckboxMenuAction.ExecuteAction = FToolMenuExecuteAction::CreateLambda(
+						[Settings](const FToolMenuContext& Context) -> void
+						{
+							Settings->bShowTrails = !Settings->bShowTrails;
+						}
+					);
+					CheckboxMenuAction.GetActionCheckState = FToolMenuGetActionCheckState::CreateLambda(
+						[Settings](const FToolMenuContext& Context) -> ECheckBoxState
+						{
+							if (Settings->bShowTrails)
+							{
+								return ECheckBoxState::Checked;
+							}
+							else
+							{
+								return ECheckBoxState::Unchecked;
+							}
+						}
+					);
+				}
+
+				FToolMenuEntry MotionPathsSubmenu = FToolMenuEntry::InitSubMenu(
+					"MotionPaths",
+					LOCTEXT("MotionPathsLabel", "Motion Paths"),
+					// TODO: Update this and other labels/tooltips in this file.
+					LOCTEXT("MotionPathsTooltip", "Check to enable motion paths. Submenu contains settings for motion paths."),
+					MakeMenuDelegate,
+					CheckboxMenuAction,
+					EUserInterfaceActionType::ToggleButton,
+					false,
+					// TODO: Temporary icon.
+					FSlateIcon(FAppStyle::GetAppStyleSetName(), "AnimGraph.Attribute.InertialBlending.Icon")
+				);
+				MotionPathsSubmenu.SetShowInToolbarTopLevel(true);
+				// Override the label when this item is raised to the top-level toolbar.
+				MotionPathsSubmenu.ToolbarLabelOverride = TAttribute<FText>::CreateLambda(
+					[Settings]() -> FText
+					{
+						// TODO: This only uses the FramesBefore value at the moment. It should be updated.
+						const int32 FramesShown = Settings->FramesBefore;
+						return FText::AsNumber(FramesShown);
+					}
+				);
+
+				UnnamedSection.AddEntry(MotionPathsSubmenu);
+			}
+		),
+		false,
+		FSlateIcon(FAppStyle::GetAppStyleSetName(), "GraphEditor.Animation_16x")
+	);
+	// Show this in the top-level to highlight it for Animation Mode users.
+	AnimationSubmenu.SetShowInToolbarTopLevel(true);
+	AnimationSubmenu.InsertPosition.Position = EToolMenuInsertType::First;
+	AllShowFlagsSection.AddEntry(AnimationSubmenu);
+}
+
+void UE::ControlRig::RemoveControlRigViewportToolbarExtensions()
+{
+	UToolMenus::Get()->UnregisterOwnerByName(UE::ControlRig::Private::ControlRigOwnerName);
+}
+
+#undef LOCTEXT_NAMESPACE
diff --git a/Engine/Plugins/Animation/ControlRig/Source/ControlRigEditor/Private/Editor/ControlRigViewportToolbarExtensions.h b/Engine/Plugins/Animation/ControlRig/Source/ControlRigEditor/Private/Editor/ControlRigViewportToolbarExtensions.h
new file mode 100644
--- /dev/null
+++ b/Engine/Plugins/Animation/ControlRig/Source/ControlRigEditor/Private/Editor/ControlRigViewportToolbarExtensions.h
@@ -0,0 +1,19 @@
+﻿// Copyright Epic Games, Inc. All Rights Reserved.
+
+#pragma once
+
+#include "Templates/SharedPointerFwd.h"
+#include "UObject/NameTypes.h"
+
+class FUICommandList;
+
+namespace UE::ControlRig
+{
+
+void PopulateControlRigViewportToolbarTransformsSubmenu(const FName InMenuName);
+void PopulateControlRigViewportToolbarSelectionSubmenu(const FName InMenuName);
+void PopulateControlRigViewportToolbarShowSubmenu(const FName InMenuName, const TSharedPtr<const FUICommandList>& InCommandList);
+
+void RemoveControlRigViewportToolbarExtensions();
+
+} // namespace UE::ControlRig
