Constraints: removed dead code + fixed constraint creation notification

#rnx

diff --git a/Engine/Plugins/Animation/ControlRig/Source/ControlRigEditor/Private/Editor/Constraints/SConstraintsWidget.cpp b/Engine/Plugins/Animation/ControlRig/Source/ControlRigEditor/Private/Editor/Constraints/SConstraintsWidget.cpp
--- a/Engine/Plugins/Animation/ControlRig/Source/ControlRigEditor/Private/Editor/Constraints/SConstraintsWidget.cpp
+++ b/Engine/Plugins/Animation/ControlRig/Source/ControlRigEditor/Private/Editor/Constraints/SConstraintsWidget.cpp
@@ -417,6 +417,18 @@ void SConstraintMenuEntry::NotifyPostChange(const FPropertyChangedEvent& InPrope
 
 FReply SConstraintMenuEntry::CreateSelectionPicker(const bool bUseDefault) const
 {
+	struct FCreationArgs
+	{
+		FOnConstraintCreated CreationDelegate;
+		ETransformConstraintType ConstraintType = ETransformConstraintType::Parent;
+		bool bUseDefault = false;
+	};
+
+	FCreationArgs Args;
+	Args.CreationDelegate = OnConstraintCreated;
+	Args.ConstraintType = ConstraintType;
+	Args.bUseDefault = bUseDefault;
+	
 	// FIXME temp approach for selecting the parent
 	FSlateApplication::Get().DismissAllMenus();
 	
@@ -425,9 +437,9 @@ FReply SConstraintMenuEntry::CreateSelectionPicker(const bool bUseDefault) const
 	ActorPickerMode.BeginActorPickingMode(
 		FOnGetAllowedClasses(), 
 		FOnShouldFilterActor(), 
-		FOnActorSelected::CreateLambda([CreationDelegate = OnConstraintCreated, Type = ConstraintType, bUseDefault](AActor* InActor)
+		FOnActorSelected::CreateLambda([Args](AActor* InActor)
 		{
-			CreateConstraint(InActor, CreationDelegate, Type, bUseDefault);
+			CreateConstraint(InActor, Args.CreationDelegate, Args.ConstraintType, Args.bUseDefault);
 		}) );
 	
 	return FReply::Handled();
diff --git a/Engine/Plugins/Animation/ControlRig/Source/ControlRigEditor/Private/Editor/Constraints/SConstraintsWidget.h b/Engine/Plugins/Animation/ControlRig/Source/ControlRigEditor/Private/Editor/Constraints/SConstraintsWidget.h
--- a/Engine/Plugins/Animation/ControlRig/Source/ControlRigEditor/Private/Editor/Constraints/SConstraintsWidget.h
+++ b/Engine/Plugins/Animation/ControlRig/Source/ControlRigEditor/Private/Editor/Constraints/SConstraintsWidget.h
@@ -125,40 +125,6 @@ private:
 	TObjectPtr<UTickableConstraint> ConfigurableConstraint = nullptr;
 };
 
-/**
- * SConstraintsCreationWidget
- */
-
-class CONTROLRIGEDITOR_API SConstraintsCreationWidget : public SCompoundWidget
-{
-public:
-
-	SLATE_BEGIN_ARGS(SConstraintsCreationWidget)	{}
-
-		SLATE_EVENT(FOnConstraintCreated, OnConstraintCreated)
-	
-	SLATE_END_ARGS()
-
-	FOnConstraintCreated OnConstraintCreated;
-	
-	/** Constructs this widget with InArgs */
-	void Construct(const FArguments& InArgs);
-
-private:
-	/** Types */
-	using ItemSharedPtr = TSharedPtr<FDroppableConstraintItem>;
-	using ConstraintItemListView = SListView< ItemSharedPtr >;
-
-	/** Generates a widget for the specified item */
-	TSharedRef<ITableRow> OnGenerateWidgetForItem(ItemSharedPtr InItem, const TSharedRef<STableViewBase>& OwnerTable);
-	
-	/** List view that shows constraint types*/
-	TSharedPtr< ConstraintItemListView > ListView;
-
-	/** Static list of constraint types */
-	static TArray< ItemSharedPtr > ListItems;
-};
-
 /**
  * The classes below are used to display a list of constraints. 
  */
