SpawnControl: fix initial prefered euler angles

#rnx

diff --git a/Engine/Plugins/Animation/ControlRig/Source/ControlRig/Private/Rigs/RigHierarchyController.cpp b/Engine/Plugins/Animation/ControlRig/Source/ControlRig/Private/Rigs/RigHierarchyController.cpp
--- a/Engine/Plugins/Animation/ControlRig/Source/ControlRig/Private/Rigs/RigHierarchyController.cpp
+++ b/Engine/Plugins/Animation/ControlRig/Source/ControlRig/Private/Rigs/RigHierarchyController.cpp
@@ -377,7 +377,8 @@ FRigElementKey URigHierarchyController::AddControl(
 		NewElement->Shape.Set(ERigTransformType::InitialLocal, InShapeTransform);  
 		Hierarchy->SetControlValue(NewElement, InValue, ERigControlValueType::Initial, false);
 		const FTransform LocalTransform = Hierarchy->GetTransform(NewElement, ERigTransformType::InitialLocal);
-		Hierarchy->SetControlPreferredEulerAngles(NewElement, LocalTransform);
+		static constexpr bool bInitial = true;
+		Hierarchy->SetControlPreferredEulerAngles(NewElement, LocalTransform, bInitial);
 
 		NewElement->Offset.MarkDirty(ERigTransformType::InitialGlobal);
 		NewElement->Pose.MarkDirty(ERigTransformType::InitialGlobal);
