ControlRig: Apply Pose Cache improvement
    - avoid switching space when useless (only controls & nulls support it)
    - reorganized the SwitchToParent function to exit as early as possible

diff --git a/Engine/Plugins/Animation/ControlRig/Source/ControlRig/Private/Rigs/RigHierarchy.cpp b/Engine/Plugins/Animation/ControlRig/Source/ControlRig/Private/Rigs/RigHierarchy.cpp
--- a/Engine/Plugins/Animation/ControlRig/Source/ControlRig/Private/Rigs/RigHierarchy.cpp
+++ b/Engine/Plugins/Animation/ControlRig/Source/ControlRig/Private/Rigs/RigHierarchy.cpp
@@ -2497,6 +2497,19 @@ bool URigHierarchy::SwitchToParent(FRigBaseElement* InChild, FRigBaseElement* In
 {
 	FRigHierarchyEnableControllerBracket EnableController(this, true);
 
+	// Exit early if switching to the same parent 
+	if (InChild)
+	{
+		const FRigElementKey ChildKey = InChild->GetKey();
+		const FRigElementKey ParentKey = InParent ? InParent->GetKey() : GetDefaultParent(ChildKey);
+		const FRigElementKey ActiveParentKey = GetActiveParent(ChildKey);
+		if(ActiveParentKey == ParentKey ||
+			(ActiveParentKey == URigHierarchy::GetDefaultParentKey() && GetDefaultParent(ChildKey) == ParentKey))
+		{
+			return true;
+		}
+	}
+
 	// rely on the VM's dependency map if there's currently an available context.
 	const TElementDependencyMap* DependencyMapPtr = &InDependencyMap;
 
@@ -2525,19 +2538,6 @@ bool URigHierarchy::SwitchToParent(FRigBaseElement* InChild, FRigBaseElement* In
 		}
 	}
 
-	// Exit early if switching to the same parent 
-	if (InChild)
-	{
-		const FRigElementKey ChildKey = InChild->GetKey();
-		const FRigElementKey ParentKey = InParent ? InParent->GetKey() : GetDefaultParent(ChildKey);
-		const FRigElementKey ActiveParentKey = GetActiveParent(ChildKey);
-		if(ActiveParentKey == ParentKey ||
-			(ActiveParentKey == URigHierarchy::GetDefaultParentKey() && GetDefaultParent(ChildKey) == ParentKey))
-		{
-			return true;
-		}
-	}
-	
 	if(const FRigMultiParentElement* MultiParentElement = Cast<FRigMultiParentElement>(InChild))
 	{
 		int32 ParentIndex = INDEX_NONE;
@@ -3060,6 +3060,10 @@ void URigHierarchy::SetPose(const FRigPose& InPose, ERigTransformType::Type InTr
 		return;
 	}
 
+	const bool bBlend = U < 1.f - SMALL_NUMBER;
+	const bool bLocal = IsLocal(InTransformType);
+	static constexpr bool bAffectChildren = true;
+
 	for(const FRigPoseElement& PoseElement : InPose)
 	{
 		FCachedRigElement Index = PoseElement.Index;
@@ -3084,22 +3088,28 @@ void URigHierarchy::SetPose(const FRigPose& InPose, ERigTransformType::Type InTr
 			FRigBaseElement* Element = Get(Index.GetIndex());
 			if(FRigTransformElement* TransformElement = Cast<FRigTransformElement>(Element))
 			{
-				FTransform TransformToSet =
-					ERigTransformType::IsLocal(InTransformType) ?
-						PoseElement.LocalTransform :
-						PoseElement.GlobalTransform;
-				
-				if(U < 1.f - SMALL_NUMBER)
+				// only nulls and controls can switch parent (cf. FRigUnit_SwitchParent)
+				const bool bCanSwitch = TransformElement->IsA<FRigMultiParentElement>() && PoseElement.ActiveParent.IsValid();
+					
+				const FTransform& PoseTransform = bLocal ? PoseElement.LocalTransform : PoseElement.GlobalTransform;
+				if (bBlend)
 				{
 					const FTransform PreviousTransform = GetTransform(TransformElement, InTransformType);
-					TransformToSet = FControlRigMathLibrary::LerpTransform(PreviousTransform, TransformToSet, U);
+					const FTransform TransformToSet = FControlRigMathLibrary::LerpTransform(PreviousTransform, PoseTransform, U);
+					if (bCanSwitch)
+					{
+						SwitchToParent(Element->GetKey(), PoseElement.ActiveParent);
+					}
+					SetTransform(TransformElement, TransformToSet, InTransformType, bAffectChildren);
 				}
-
-				if (PoseElement.ActiveParent.IsValid())
+				else
 				{
-					SwitchToParent(Element->GetKey(), PoseElement.ActiveParent);
+					if (bCanSwitch)
+					{
+						SwitchToParent(Element->GetKey(), PoseElement.ActiveParent);
+					}
+					SetTransform(TransformElement, PoseTransform, InTransformType, bAffectChildren);
 				}
-				SetTransform(TransformElement, TransformToSet, InTransformType, true);
 			}
 			else if(FRigCurveElement* CurveElement = Cast<FRigCurveElement>(Element))
 			{
