Control Rig: Fix control shape actors flying out
#rnx

diff --git a/Engine/Plugins/Animation/ControlRig/Source/ControlRigEditor/Private/EditMode/ControlRigEditMode.cpp b/Engine/Plugins/Animation/ControlRig/Source/ControlRigEditor/Private/EditMode/ControlRigEditMode.cpp
--- a/Engine/Plugins/Animation/ControlRig/Source/ControlRigEditor/Private/EditMode/ControlRigEditMode.cpp
+++ b/Engine/Plugins/Animation/ControlRig/Source/ControlRigEditor/Private/EditMode/ControlRigEditMode.cpp
@@ -4811,7 +4811,8 @@ void FControlRigEditMode::MoveControlShape(AControlRigShapeActor* ShapeActor, co
 		FTransform CurrentTransform;
 		if (ShapeActor->StaticMeshComponent && ShapeActor->StaticMeshComponent->GetStaticMesh())
 		{
-			CurrentTransform = ShapeActor->GetGlobalTransform() * ToWorldTransform;
+			const FTransform& ParentTransform = ShapeActor->GetAttachParentActor()->GetTransform();
+			CurrentTransform = ShapeActor->GetGlobalTransform() * ParentTransform;
 		}
 		else
 		{
