Control Rig: Fix GetParent rig unit not returning the active parent
#rnx

diff --git a/Engine/Plugins/Animation/ControlRig/Source/ControlRig/Private/Rigs/RigHierarchy.cpp b/Engine/Plugins/Animation/ControlRig/Source/ControlRig/Private/Rigs/RigHierarchy.cpp
--- a/Engine/Plugins/Animation/ControlRig/Source/ControlRig/Private/Rigs/RigHierarchy.cpp
+++ b/Engine/Plugins/Animation/ControlRig/Source/ControlRig/Private/Rigs/RigHierarchy.cpp
@@ -2159,12 +2159,40 @@ TArray<FRigElementWeight> URigHierarchy::GetParentWeightArray(const FRigBaseElem
 	return Weights;
 }
 
-FRigElementKey URigHierarchy::GetActiveParent(const FRigElementKey& InKey) const
+FRigElementKey URigHierarchy::GetActiveParent(const FRigElementKey& InKey, bool bReferenceKey) const
 {
-	const TArray<FRigElementWeight> ParentWeights = GetParentWeightArray(InKey);
+	if(FRigBaseElement* Parent = GetActiveParent(Find(InKey)))
+	{
+		if (bReferenceKey && Parent->GetKey() == GetDefaultParent(InKey))
+		{
+			return URigHierarchy::GetDefaultParentKey();
+		}
+		return Parent->Key;
+	}
+
+	if (bReferenceKey)
+	{
+		return URigHierarchy::GetWorldSpaceReferenceKey();
+	}
+
+	return FRigElementKey();
+}
+
+int32 URigHierarchy::GetActiveParent(int32 InIndex) const
+{
+	if(FRigBaseElement* Parent = GetActiveParent(Get(InIndex)))
+	{
+		return Parent->Index;
+	}
+	return INDEX_NONE;
+}
+
+FRigBaseElement* URigHierarchy::GetActiveParent(const FRigBaseElement* InElement) const
+{
+	const TArray<FRigElementWeight> ParentWeights = GetParentWeightArray(InElement);
 	if (ParentWeights.Num() > 0)
 	{
-		const TArray<FRigElementKey> ParentKeys = GetParents(InKey);
+		const FRigBaseElementParentArray ParentKeys = GetParents(InElement);
 		check(ParentKeys.Num() == ParentWeights.Num());
 		for (int32 ParentIndex = 0; ParentIndex < ParentKeys.Num(); ParentIndex++)
 		{
@@ -2172,20 +2200,14 @@ FRigElementKey URigHierarchy::GetActiveParent(const FRigElementKey& InKey) const
 			{
 				continue;
 			}
-			if (ParentIndex == 0)
+			if (Elements.IsValidIndex(ParentKeys[ParentIndex]->GetIndex()))
 			{
-				if (!(ParentKeys[ParentIndex] == URigHierarchy::GetDefaultParentKey() || ParentKeys[ParentIndex] == URigHierarchy::GetWorldSpaceReferenceKey()))
-				{
-					if(ParentKeys[ParentIndex] == GetDefaultParent(InKey))
-					{
-						return URigHierarchy::GetDefaultParentKey();
-					}
-				}
+				return Elements[ParentKeys[ParentIndex]->GetIndex()];
 			}
-			return ParentKeys[ParentIndex];
 		}
 	}
-	return URigHierarchy::GetDefaultParentKey();
+
+	return nullptr;
 }
 
 
diff --git a/Engine/Plugins/Animation/ControlRig/Source/ControlRig/Private/Units/Execution/RigUnit_Hierarchy.cpp b/Engine/Plugins/Animation/ControlRig/Source/ControlRig/Private/Units/Execution/RigUnit_Hierarchy.cpp
--- a/Engine/Plugins/Animation/ControlRig/Source/ControlRig/Private/Units/Execution/RigUnit_Hierarchy.cpp
+++ b/Engine/Plugins/Animation/ControlRig/Source/ControlRig/Private/Units/Execution/RigUnit_Hierarchy.cpp
@@ -11,29 +11,39 @@ FRigUnit_HierarchyGetParent_Execute()
 {
     DECLARE_SCOPE_HIERARCHICAL_COUNTER_RIGUNIT()
 
-	if(CachedChild.IsIdentical(Child, ExecuteContext.Hierarchy))
+	if (bDefaultParent)
 	{
-		Parent = CachedParent.GetKey();
-	}
-	else
-	{
-		Parent.Reset();
-		CachedParent.Reset();
-
-		if(CachedChild.UpdateCache(Child, ExecuteContext.Hierarchy))
+		if(CachedChild.IsIdentical(Child, ExecuteContext.Hierarchy))
 		{
-			Parent = ExecuteContext.Hierarchy->GetFirstParent(CachedChild.GetResolvedKey());
-			if(Parent.IsValid())
+			Parent = CachedParent.GetKey();
+		}
+		else
+		{
+			Parent.Reset();
+			CachedParent.Reset();
+
+			if(CachedChild.UpdateCache(Child, ExecuteContext.Hierarchy))
 			{
-				CachedParent.UpdateCache(Parent, ExecuteContext.Hierarchy);
+				Parent = ExecuteContext.Hierarchy->GetFirstParent(CachedChild.GetResolvedKey());
+				if(Parent.IsValid())
+				{
+					CachedParent.UpdateCache(Parent, ExecuteContext.Hierarchy);
+				}
 			}
 		}
 	}
+	else
+	{
+		CachedChild.UpdateCache(Child, ExecuteContext.Hierarchy);
+		const FRigElementKey& ChildKey = CachedChild.GetResolvedKey();
+
+		Parent = ExecuteContext.Hierarchy->GetActiveParent(ChildKey, false);
+	}
 }
 
 FRigUnit_HierarchyGetParents_Execute()
 {
-	FRigUnit_HierarchyGetParentsItemArray::StaticExecute(ExecuteContext, Child, bIncludeChild, bReverse, Parents.Keys, CachedChild, CachedParents);
+	FRigUnit_HierarchyGetParentsItemArray::StaticExecute(ExecuteContext, Child, bIncludeChild, bReverse, true, Parents.Keys, CachedChild, CachedParents);
 }
 
 FRigVMStructUpgradeInfo FRigUnit_HierarchyGetParents::GetUpgradeInfo() const
@@ -50,33 +60,61 @@ FRigUnit_HierarchyGetParentsItemArray_Execute()
 {
     DECLARE_SCOPE_HIERARCHICAL_COUNTER_RIGUNIT()
 
-	if(!CachedChild.IsIdentical(Child, ExecuteContext.Hierarchy))
+	if (bDefaultParent)
 	{
-		CachedParents.Reset();
-
-		if(CachedChild.UpdateCache(Child, ExecuteContext.Hierarchy))
+		if(!CachedChild.IsIdentical(Child, ExecuteContext.Hierarchy))
 		{
-			TArray<FRigElementKey> Keys;
-			FRigElementKey Parent = CachedChild.GetResolvedKey();
-			do
+			CachedParents.Reset();
+
+			if(CachedChild.UpdateCache(Child, ExecuteContext.Hierarchy))
 			{
-				if(bIncludeChild || Parent != Child)
+				TArray<FRigElementKey> Keys;
+				FRigElementKey Parent = CachedChild.GetResolvedKey();
+				do
 				{
-					Keys.Add(Parent);
+					if(bIncludeChild || Parent != Child)
+					{
+						Keys.Add(Parent);
+					}
+					Parent = ExecuteContext.Hierarchy->GetFirstParent(Parent);
+				}
+				while(Parent.IsValid());
+
+				CachedParents = FRigElementKeyCollection(Keys);
+				if(bReverse)
+				{
+					CachedParents = FRigElementKeyCollection::MakeReversed(CachedParents);
 				}
-				Parent = ExecuteContext.Hierarchy->GetFirstParent(Parent);
 			}
-			while(Parent.IsValid());
+		}
 
-			CachedParents = FRigElementKeyCollection(Keys);
-			if(bReverse)
+		Parents = CachedParents.Keys;
+	}
+	else
+	{
+		CachedChild.UpdateCache(Child, ExecuteContext.Hierarchy);
+
+		TArray<FRigElementKey> Keys;
+		FRigElementKey Parent = CachedChild.GetResolvedKey();
+		do
+		{
+			if(bIncludeChild || Parent != Child)
 			{
-				CachedParents = FRigElementKeyCollection::MakeReversed(CachedParents);
+				Keys.Add(Parent);
 			}
+			const FRigElementKey PreviousParent = Parent;
+
+			Parent = ExecuteContext.Hierarchy->GetActiveParent(Parent, false);
 		}
-	}
+		while(Parent.IsValid());
 
-	Parents = CachedParents.Keys;
+		if(bReverse)
+		{
+			Algo::Reverse(Keys);
+		}
+
+		Parents = Keys;
+	}
 }
 
 FRigUnit_HierarchyGetChildren_Execute()
diff --git a/Engine/Plugins/Animation/ControlRig/Source/ControlRig/Public/Rigs/RigHierarchy.h b/Engine/Plugins/Animation/ControlRig/Source/ControlRig/Public/Rigs/RigHierarchy.h
--- a/Engine/Plugins/Animation/ControlRig/Source/ControlRig/Public/Rigs/RigHierarchy.h
+++ b/Engine/Plugins/Animation/ControlRig/Source/ControlRig/Public/Rigs/RigHierarchy.h
@@ -3156,11 +3156,26 @@ public:
 	TArray<FRigElementWeight> GetParentWeightArray(const FRigBaseElement* InChild, bool bInitial = false) const;
 
 	/**
-	 * Get the current active for the passed in key. This is only valid when only one parent has a weight value and the other parents have zero weights
+	 * Get the current active parent for the passed in key. This is only valid when only one parent has a weight value and the other parents have zero weights
 	 * @param InKey The multi parented element
+	 * @param bReferenceKey Whether or not to return a reference key
 	 * @return Returns the first parent with a non-zero weight
 	 */
-	FRigElementKey GetActiveParent(const FRigElementKey& InKey) const;
+	FRigElementKey GetActiveParent(const FRigElementKey& InKey, bool bReferenceKey = true) const;
+
+	/**
+	 * Get the current active parent for a given element index. This is only valid when only one parent has a weight value and the other parents have zero weights
+	 * @param InIndex The index of the element to retrieve the parent for
+	 * @return Returns the first parent index (or INDEX_NONE) with a non-zero weight
+	 */
+	int32 GetActiveParent(int32 InIndex) const;
+
+	/**
+	 * Get the current active parent for the passed in key. This is only valid when only one parent has a weight value and the other parents have zero weights
+	 * @param InElement The element to retrieve the parents for
+	 * @return Returns the first parent with a non-zero weight
+	 */
+	FRigBaseElement* GetActiveParent(const FRigBaseElement* InElement) const;
 
 	/**
 	 * Sets the weight of a parent below a multi parent element
diff --git a/Engine/Plugins/Animation/ControlRig/Source/ControlRig/Public/Units/Execution/RigUnit_Hierarchy.h b/Engine/Plugins/Animation/ControlRig/Source/ControlRig/Public/Units/Execution/RigUnit_Hierarchy.h
--- a/Engine/Plugins/Animation/ControlRig/Source/ControlRig/Public/Units/Execution/RigUnit_Hierarchy.h
+++ b/Engine/Plugins/Animation/ControlRig/Source/ControlRig/Public/Units/Execution/RigUnit_Hierarchy.h
@@ -29,6 +29,7 @@ struct CONTROLRIG_API FRigUnit_HierarchyGetParent : public FRigUnit_HierarchyBas
 	FRigUnit_HierarchyGetParent()
 	{
 		Child = Parent = FRigElementKey(NAME_None, ERigElementType::Bone);
+		bDefaultParent = true;
 		CachedChild = CachedParent = FCachedRigElement();
 	}
 
@@ -38,6 +39,9 @@ struct CONTROLRIG_API FRigUnit_HierarchyGetParent : public FRigUnit_HierarchyBas
 	UPROPERTY(meta = (Input, ExpandByDefault))
 	FRigElementKey Child;
 
+	UPROPERTY(meta = (Input))
+	bool bDefaultParent;
+
 	UPROPERTY(meta = (Output))
 	FRigElementKey Parent;
 
@@ -109,6 +113,7 @@ struct CONTROLRIG_API FRigUnit_HierarchyGetParentsItemArray : public FRigUnit_Hi
 		CachedParents = FRigElementKeyCollection();
 		bIncludeChild = false;
 		bReverse = false;
+		bDefaultParent = true;
 	}
 
 	RIGVM_METHOD()
@@ -123,6 +128,9 @@ struct CONTROLRIG_API FRigUnit_HierarchyGetParentsItemArray : public FRigUnit_Hi
 	UPROPERTY(meta = (Input))
 	bool bReverse;
 
+	UPROPERTY(meta = (Input))
+	bool bDefaultParent;
+
 	UPROPERTY(meta = (Output))
 	TArray<FRigElementKey> Parents;
 
diff --git a/Engine/Plugins/Animation/ControlRig/Source/ControlRigEditor/Private/Editor/SRigHierarchyTreeView.cpp b/Engine/Plugins/Animation/ControlRig/Source/ControlRigEditor/Private/Editor/SRigHierarchyTreeView.cpp
--- a/Engine/Plugins/Animation/ControlRig/Source/ControlRigEditor/Private/Editor/SRigHierarchyTreeView.cpp
+++ b/Engine/Plugins/Animation/ControlRig/Source/ControlRigEditor/Private/Editor/SRigHierarchyTreeView.cpp
@@ -669,7 +669,7 @@ bool SRigHierarchyTreeView::AddElement(const FRigBaseElement* InElement)
 				AddConnectorResolveWarningTag(ElementMap.FindChecked(InElement->GetKey()), InElement, Hierarchy);
 			}
 			
-			FRigElementKey ParentKey = Hierarchy->GetFirstParent(InElement->GetKey());
+			FRigElementKey ParentKey = Hierarchy->GetActiveParent(InElement->GetKey());
 			if(InElement->GetType() == ERigElementType::Connector)
 			{
 				ParentKey = Delegates.GetResolvedKey(InElement->GetKey());
@@ -834,7 +834,7 @@ void SRigHierarchyTreeView::RefreshTreeView(bool bRebuildContent)
 			Hierarchy->Traverse([&](FRigBaseElement* Element, bool& bContinue)
 			{
 				int32& Depth = ElementDepth.Add(Element, 0);
-				if(const FRigBaseElement* ParentElement = Hierarchy->GetFirstParent(Element))
+				if(const FRigBaseElement* ParentElement = Hierarchy->GetActiveParent(Element))
 				{
 					if (int32* ParentDepth = ElementDepth.Find(ParentElement))
 					{
