Control Rig: Rename rigid body element to physics element

#prefligt https://horde.devtools.epicgames.com/job/66963a31206ca881ae05f4ef
#rnx

diff --git a/Engine/Plugins/Animation/ControlRig/Source/ControlRig/Private/ControlRig.cpp b/Engine/Plugins/Animation/ControlRig/Source/ControlRig/Private/ControlRig.cpp
--- a/Engine/Plugins/Animation/ControlRig/Source/ControlRig/Private/ControlRig.cpp
+++ b/Engine/Plugins/Animation/ControlRig/Source/ControlRig/Private/ControlRig.cpp
@@ -563,6 +563,9 @@ void UControlRig::InitializeFromCDO()
 			UControlRig* CDO = GetClass()->GetDefaultObject<UControlRig>();
 			URigHierarchy* Hierarchy = GetHierarchy();
 
+			// copy physics solvers
+			PhysicsSolvers = CDO->PhysicsSolvers;
+
 			// copy hierarchy
 			{
 				FRigHierarchyValidityBracket ValidityBracketA(Hierarchy);
@@ -573,6 +576,25 @@ void UControlRig::InitializeFromCDO()
 				Hierarchy->ResetPoseToInitial(ERigElementType::All);
 			}
 
+			// update the physics solvers with new unique IDs and remap the physics solver IDs within hierarchy
+			{
+				TMap<FRigPhysicsSolverID, FRigPhysicsSolverID> PhysicsSolverMap;
+				for(FRigPhysicsSolverDescription& Solver : PhysicsSolvers)
+				{
+					const FRigPhysicsSolverID OldID = Solver.ID;
+					Solver.ID = FRigPhysicsSolverDescription::MakeID(GetPathName(), Solver.Name);
+					PhysicsSolverMap.Add(OldID, Solver.ID);
+				}
+				const TArray<FRigPhysicsElement*> PhysicsElements = Hierarchy->GetPhysicsElements();
+				for(FRigPhysicsElement* PhysicsElement : PhysicsElements)
+				{
+					if(const FRigPhysicsSolverID* NewID = PhysicsSolverMap.Find(PhysicsElement->Solver))
+					{
+						PhysicsElement->Solver = *NewID;
+					}
+				}
+			}
+
 #if WITH_EDITOR
 			// current hierarchy should always mirror CDO's hierarchy whenever a change of interest happens
 			CDO->GetHierarchy()->RegisterListeningHierarchy(Hierarchy);
@@ -1738,6 +1760,15 @@ void UControlRig::Serialize(FArchive& Ar)
 	Super::Serialize(Ar);
 
 	Ar.UsingCustomVersion(FControlRigObjectVersion::GUID);
+
+	if (Ar.CustomVer(FControlRigObjectVersion::GUID) >= FControlRigObjectVersion::ControlRigStoresPhysicsSolvers)
+	{
+		Ar << PhysicsSolvers;
+	}
+	else
+	{
+		PhysicsSolvers.Reset();
+	}
 }
 
 void UControlRig::PostLoad()
@@ -3517,6 +3548,60 @@ void UControlRig::OnHierarchyTransformUndoRedo(URigHierarchy* InHierarchy, const
 	}
 }
 
+int32 UControlRig::NumPhysicsSolvers() const
+{
+	return PhysicsSolvers.Num();
+}
+
+const FRigPhysicsSolverDescription* UControlRig::GetPhysicsSolver(int32 InIndex) const
+{
+	if(PhysicsSolvers.IsValidIndex(InIndex))
+	{
+		return &PhysicsSolvers[InIndex];
+	}
+	return nullptr;
+}
+
+const FRigPhysicsSolverDescription* UControlRig::FindPhysicsSolver(const FRigPhysicsSolverID& InID) const
+{
+	return PhysicsSolvers.FindByPredicate([InID](const FRigPhysicsSolverDescription& Solver) -> bool
+	{
+		return Solver.ID == InID;
+	});
+}
+
+const FRigPhysicsSolverDescription* UControlRig::FindPhysicsSolverByName(const FName& InName) const
+{
+	return PhysicsSolvers.FindByPredicate([InName](const FRigPhysicsSolverDescription& Solver) -> bool
+	{
+		return Solver.Name == InName;
+	});
+}
+
+FRigPhysicsSolverID UControlRig::AddPhysicsSolver(FName InName, bool bSetupUndo, bool bPrintPythonCommand)
+{
+	if(CVarControlRigHierarchyEnablePhysics.GetValueOnAnyThread() == false)
+	{
+		return FRigPhysicsSolverID();
+	}
+
+	if(InName.IsNone())
+	{
+		return FRigPhysicsSolverID();
+	}
+
+	const FName Name = UtilityHelpers::CreateUniqueName(InName, [this](const FName& InName) -> bool
+	{
+		return FindPhysicsSolverByName(InName) == nullptr;
+	});
+
+	FRigPhysicsSolverDescription Solver;
+	Solver.ID = FRigPhysicsSolverDescription::MakeID(GetPathName(), Name);
+	Solver.Name = Name;
+	PhysicsSolvers.Add(Solver);
+	return Solver.ID;
+}
+
 UControlRig::FPoseScope::FPoseScope(UControlRig* InControlRig, ERigElementType InFilter, const TArray<FRigElementKey>& InElements, const ERigTransformType::Type InTransformType)
 : ControlRig(InControlRig)
 , Filter(InFilter)
diff --git a/Engine/Plugins/Animation/ControlRig/Source/ControlRig/Private/Rigs/RigHierarchy.cpp b/Engine/Plugins/Animation/ControlRig/Source/ControlRig/Private/Rigs/RigHierarchy.cpp
--- a/Engine/Plugins/Animation/ControlRig/Source/ControlRig/Private/Rigs/RigHierarchy.cpp
+++ b/Engine/Plugins/Animation/ControlRig/Source/ControlRig/Private/Rigs/RigHierarchy.cpp
@@ -49,6 +49,7 @@ static TAutoConsoleVariable<int32> CVarControlRigHierarchyTracePrecision(TEXT("C
 static TAutoConsoleVariable<int32> CVarControlRigHierarchyTraceOnSpawn(TEXT("ControlRig.Hierarchy.TraceOnSpawn"), 0, TEXT("sets the number of frames to trace when a new hierarchy is spawned"));
 TAutoConsoleVariable<bool> CVarControlRigHierarchyEnableRotationOrder(TEXT("ControlRig.Hierarchy.EnableRotationOrder"), true, TEXT("enables the rotation order for controls"));
 TAutoConsoleVariable<bool> CVarControlRigHierarchyEnableModules(TEXT("ControlRig.Hierarchy.Modules"), true, TEXT("enables the modular rigging functionality"));
+TAutoConsoleVariable<bool> CVarControlRigHierarchyEnablePhysics(TEXT("ControlRig.Hierarchy.Physics"), false, TEXT("enables physics support for the rig hierarchy"));
 static int32 sRigHierarchyLastTrace = INDEX_NONE;
 static TCHAR sRigHierarchyTraceFormat[16];
 
@@ -240,7 +241,6 @@ void URigHierarchy::Save(FArchive& Ar)
 		
 		Ar << ElementMetadataToSave;
 	}
-	
 }
 
 void URigHierarchy::Load(FArchive& Ar)
@@ -540,7 +540,7 @@ void URigHierarchy::CopyHierarchy(URigHierarchy* InHierarchy)
 			sizeof(FRigNullElement),
 			sizeof(FRigControlElement),
 			sizeof(FRigCurveElement),
-			sizeof(FRigRigidBodyElement),
+			sizeof(FRigPhysicsElement),
 			sizeof(FRigReferenceElement),
 			sizeof(FRigConnectorElement),
 			sizeof(FRigSocketElement),
@@ -1144,6 +1144,24 @@ TArray<FRigElementKey> URigHierarchy::RestoreConnectorsFromStates(TArray<FRigCon
 	return Keys;
 }
 
+const FRigPhysicsSolverDescription* URigHierarchy::FindPhysicsSolver(const FRigPhysicsSolverID& InID) const
+{
+	if(const UControlRig* ControlRig = Cast<UControlRig>(GetOuter()))
+	{
+		return ControlRig->FindPhysicsSolver(InID);
+	}
+	return nullptr;
+}
+
+const FRigPhysicsSolverDescription* URigHierarchy::FindPhysicsSolverByName(const FName& InName) const
+{
+	if(const UControlRig* ControlRig = Cast<UControlRig>(GetOuter()))
+	{
+		return ControlRig->FindPhysicsSolverByName(InName);
+	}
+	return nullptr;
+}
+
 TArray<FName> URigHierarchy::GetMetadataNames(FRigElementKey InItem) const
 {
 	TArray<FName> Names;
@@ -5125,13 +5143,13 @@ FRigBaseElement* URigHierarchy::MakeElement(ERigElementType InElementType, int32
 			Element = NewElement<FRigCurveElement>(InCount);
 			break;
 		}
-		case ERigElementType::RigidBody:
+		case ERigElementType::Physics:
 		{
 			if(OutStructureSize)
 			{
-				*OutStructureSize = sizeof(FRigRigidBodyElement);
+				*OutStructureSize = sizeof(FRigPhysicsElement);
 			}
-			Element = NewElement<FRigRigidBodyElement>(InCount);
+			Element = NewElement<FRigPhysicsElement>(InCount);
 			break;
 		}
 		case ERigElementType::Reference:
@@ -5223,13 +5241,13 @@ void URigHierarchy::DestroyElement(FRigBaseElement*& InElement)
 			}
 			break;
 		}
-		case ERigElementType::RigidBody:
+		case ERigElementType::Physics:
 		{
-			FRigRigidBodyElement* ExistingElements = Cast<FRigRigidBodyElement>(InElement);
+			FRigPhysicsElement* ExistingElements = Cast<FRigPhysicsElement>(InElement);
 			for(int32 Index=0;Index<Count;Index++)
 			{
 				TGuardValue<const FRigBaseElement*> DestroyGuard(ElementBeingDestroyed, &ExistingElements[Index]);
-				ExistingElements[Index].~FRigRigidBodyElement(); 
+				ExistingElements[Index].~FRigPhysicsElement(); 
 			}
 			break;
 		}
diff --git a/Engine/Plugins/Animation/ControlRig/Source/ControlRig/Private/Rigs/RigHierarchyController.cpp b/Engine/Plugins/Animation/ControlRig/Source/ControlRig/Private/Rigs/RigHierarchyController.cpp
--- a/Engine/Plugins/Animation/ControlRig/Source/ControlRig/Private/Rigs/RigHierarchyController.cpp
+++ b/Engine/Plugins/Animation/ControlRig/Source/ControlRig/Private/Rigs/RigHierarchyController.cpp
@@ -5,6 +5,7 @@
 #include "AnimationCoreLibrary.h"
 #include "UObject/Package.h"
 #include "ModularRig.h"
+#include "HelperUtil.h"
 
 #if WITH_EDITOR
 #include "Framework/Notifications/NotificationManager.h"
@@ -484,30 +485,48 @@ FRigElementKey URigHierarchyController::AddCurve(FName InName, float InValue, bo
 	return NewElement->Key;
 }
 
-FRigElementKey URigHierarchyController::AddRigidBody(FName InName, FRigElementKey InParent,
-                                                     FRigRigidBodySettings InSettings, FTransform InLocalTransform, bool bSetupUndo, bool bPrintPythonCommand)
+FRigElementKey URigHierarchyController::AddPhysicsElement(FName InName, FRigElementKey InParent, FRigPhysicsSolverID InSolver, 
+                                                          FRigPhysicsSettings InSettings, FTransform InLocalTransform, bool bSetupUndo, bool bPrintPythonCommand)
 {
 	if(!IsValid())
 	{
 		return FRigElementKey();
 	}
 
+	if(CVarControlRigHierarchyEnablePhysics.GetValueOnAnyThread() == false)
+	{
+		return FRigElementKey();
+	}
+
 	URigHierarchy* Hierarchy = GetHierarchy();
 
+	if(!InSolver.IsValid())
+	{
+		ReportError(TEXT("Physics Solver Guid is not valid"));
+		return FRigElementKey();
+	}
+
+	if(Hierarchy->FindPhysicsSolver(InSolver) == nullptr)
+	{
+		ReportErrorf(TEXT("Physics Solver with guid '%s' cannot be found."), *InSolver.ToString());
+		return FRigElementKey();
+	}
+
 #if WITH_EDITOR
 	TSharedPtr<FScopedTransaction> TransactionPtr;
 	if(bSetupUndo)
 	{
-		TransactionPtr = MakeShared<FScopedTransaction>(NSLOCTEXT("RigHierarchyController", "Add RigidBody", "Add RigidBody"));
+		TransactionPtr = MakeShared<FScopedTransaction>(NSLOCTEXT("RigHierarchyController", "Add Physics Element", "Add Physics Element"));
 		Hierarchy->Modify();
 	}
 #endif
 
-	FRigRigidBodyElement* NewElement = MakeElement<FRigRigidBodyElement>();
+	FRigPhysicsElement* NewElement = MakeElement<FRigPhysicsElement>();
 	{
 		TGuardValue<bool> DisableCacheValidityChecks(Hierarchy->bEnableCacheValidityCheck, false);		
-		NewElement->Key.Type = ERigElementType::RigidBody;
+		NewElement->Key.Type = ERigElementType::Physics;
 		NewElement->Key.Name = GetSafeNewName(InName, NewElement->Key.Type);
+		NewElement->Solver = InSolver;
 		NewElement->Settings = InSettings;
 		AddElement(NewElement, Hierarchy->Get(Hierarchy->GetIndex(InParent)), true, InName);
 
@@ -523,7 +542,7 @@ FRigElementKey URigHierarchyController::AddRigidBody(FName InName, FRigElementKe
 		UBlueprint* Blueprint = GetTypedOuter<UBlueprint>();
 		if (Blueprint)
 		{
-			TArray<FString> Commands = GetAddRigidBodyPythonCommands(NewElement);
+			TArray<FString> Commands = GetAddPhysicsElementPythonCommands(NewElement);
 			for (const FString& Command : Commands)
 			{
 				RigVMPythonUtils::Print(Blueprint->GetFName().ToString(),
@@ -1225,10 +1244,10 @@ FString URigHierarchyController::ExportToText(TArray<FRigElementKey> InKeys) con
 				FRigCurveElement::StaticStruct()->ExportText(PerElementData.Content, Element, &DefaultElement, nullptr, PPF_None, nullptr);
 				break;
 			}
-			case ERigElementType::RigidBody:
+			case ERigElementType::Physics:
 			{
-				FRigRigidBodyElement DefaultElement;
-				FRigRigidBodyElement::StaticStruct()->ExportText(PerElementData.Content, Element, &DefaultElement, nullptr, PPF_None, nullptr);
+				FRigPhysicsElement DefaultElement;
+				FRigPhysicsElement::StaticStruct()->ExportText(PerElementData.Content, Element, &DefaultElement, nullptr, PPF_None, nullptr);
 				break;
 			}
 			case ERigElementType::Reference:
@@ -1399,10 +1418,10 @@ TArray<FRigElementKey> URigHierarchyController::ImportFromText(FString InContent
 				FRigCurveElement::StaticStruct()->ImportText(*PerElementData.Content, NewElement, nullptr, EPropertyPortFlags::PPF_None, &ErrorPipe, FRigCurveElement::StaticStruct()->GetName(), true);
 				break;
 			}
-			case ERigElementType::RigidBody:
+			case ERigElementType::Physics:
 			{
-				NewElement = MakeElement<FRigRigidBodyElement>();
-				FRigRigidBodyElement::StaticStruct()->ImportText(*PerElementData.Content, NewElement, nullptr, EPropertyPortFlags::PPF_None, &ErrorPipe, FRigRigidBodyElement::StaticStruct()->GetName(), true);
+				NewElement = MakeElement<FRigPhysicsElement>();
+				FRigPhysicsElement::StaticStruct()->ImportText(*PerElementData.Content, NewElement, nullptr, EPropertyPortFlags::PPF_None, &ErrorPipe, FRigPhysicsElement::StaticStruct()->GetName(), true);
 				break;
 			}
 			case ERigElementType::Reference:
@@ -1727,9 +1746,9 @@ TArray<FString> URigHierarchyController::GetAddElementPythonCommands(FRigBaseEle
 	{
 		return GetAddCurvePythonCommands(CurveElement);
 	}
-	else if(FRigRigidBodyElement* RigidBodyElement = Cast<FRigRigidBodyElement>(Element))
+	else if(FRigPhysicsElement* PhysicsElement = Cast<FRigPhysicsElement>(Element))
 	{
-		return GetAddRigidBodyPythonCommands(RigidBodyElement);
+		return GetAddPhysicsElementPythonCommands(PhysicsElement);
 	}
 	else if(FRigReferenceElement* ReferenceElement = Cast<FRigReferenceElement>(Element))
 	{
@@ -1841,35 +1860,35 @@ TArray<FString> URigHierarchyController::GetAddCurvePythonCommands(FRigCurveElem
 		Hierarchy->GetCurveValue(Curve))};
 }
 
-TArray<FString> URigHierarchyController::GetAddRigidBodyPythonCommands(FRigRigidBodyElement* RigidBody) const
+TArray<FString> URigHierarchyController::GetAddPhysicsElementPythonCommands(FRigPhysicsElement* PhysicsElement) const
 {
 	TArray<FString> Commands;
-	FString TransformStr = RigVMPythonUtils::TransformToPythonString(RigidBody->Pose.Initial.Local.Transform);
+	FString TransformStr = RigVMPythonUtils::TransformToPythonString(PhysicsElement->Pose.Initial.Local.Transform);
 	
 	FString ParentKeyStr = "''";
-	if (RigidBody->ParentElement)
+	if (PhysicsElement->ParentElement)
 	{
-		ParentKeyStr = RigidBody->ParentElement->GetKey().ToPythonString();
+		ParentKeyStr = PhysicsElement->ParentElement->GetKey().ToPythonString();
 	}
 
-	FString RigidBodyNamePythonized = RigVMPythonUtils::PythonizeName(RigidBody->GetName());
-	FRigRigidBodySettings& Settings = RigidBody->Settings;
+	FString PhysicsElementNamePythonized = RigVMPythonUtils::PythonizeName(PhysicsElement->GetName());
+	FRigPhysicsSettings& Settings = PhysicsElement->Settings;
 	FString SettingsStr;
 	{
-		SettingsStr =FString::Printf(TEXT("rigid_body_settings_%s"),
-			*RigidBodyNamePythonized);
+		SettingsStr =FString::Printf(TEXT("physics_settings%s"),
+			*PhysicsElementNamePythonized);
 			
-		Commands.Add(FString::Printf(TEXT("rigid_body_settings_%s = unreal.RigRigidBodySettings()"),
-			*RigidBodyNamePythonized));
+		Commands.Add(FString::Printf(TEXT("physics_settings%s = unreal.RigPhysicsSettings()"),
+			*PhysicsElementNamePythonized));
 
 		Commands.Add(FString::Printf(TEXT("control_settings_%s.mass = %f"),
-			*RigidBodyNamePythonized,
+			*PhysicsElementNamePythonized,
 			Settings.Mass));	
 	}
 	
-	// FRigElementKey AddRigidBody(FName InName, FRigElementKey InParent, FRigRigidBodySettings InSettings, FTransform InLocalTransform, bool bSetupUndo = false);
-	Commands.Add(FString::Printf(TEXT("hierarchy_controller.add_rigid_body('%s', %s, %s, %s)"),
-		*RigidBody->GetName(),
+	// FRigElementKey AddPhysicsElement(FName InName, FRigElementKey InParent, FRigPhysicsSettings InSettings, FTransform InLocalTransform, bool bSetupUndo = false);
+	Commands.Add(FString::Printf(TEXT("hierarchy_controller.add_physics_element('%s', %s, %s, %s)"),
+		*PhysicsElement->GetName(),
 		*ParentKeyStr,
 		*SettingsStr,
 		*TransformStr));
diff --git a/Engine/Plugins/Animation/ControlRig/Source/ControlRig/Private/Rigs/RigHierarchyElements.cpp b/Engine/Plugins/Animation/ControlRig/Source/ControlRig/Private/Rigs/RigHierarchyElements.cpp
--- a/Engine/Plugins/Animation/ControlRig/Source/ControlRig/Private/Rigs/RigHierarchyElements.cpp
+++ b/Engine/Plugins/Animation/ControlRig/Source/ControlRig/Private/Rigs/RigHierarchyElements.cpp
@@ -47,9 +47,9 @@ UScriptStruct* FRigBaseElement::GetElementStruct() const
 		{
 			return FRigReferenceElement::StaticStruct();
 		}
-		case ERigElementType::RigidBody:
+		case ERigElementType::Physics:
 		{
-			return FRigRigidBodyElement::StaticStruct();
+			return FRigPhysicsElement::StaticStruct();
 		}
 		case ERigElementType::Connector:
 		{
@@ -1308,55 +1308,106 @@ void FRigCurveElement::CopyFrom(const FRigBaseElement* InOther)
 }
 
 ////////////////////////////////////////////////////////////////////////////////
-// FRigRigidBodySettings
+// FRigPhysicsSolver
 ////////////////////////////////////////////////////////////////////////////////
 
-FRigRigidBodySettings::FRigRigidBodySettings()
+void FRigPhysicsSolverDescription::Serialize(FArchive& Ar)
+{
+	if (Ar.IsSaving() || Ar.IsObjectReferenceCollector() || Ar.IsCountingMemory())
+	{
+		Save(Ar);
+	}
+	else if (Ar.IsLoading())
+	{
+		Load(Ar);
+	}
+}
+
+void FRigPhysicsSolverDescription::Save(FArchive& Ar)
+{
+	Ar << ID;
+	Ar << Name;
+}
+
+void FRigPhysicsSolverDescription::Load(FArchive& Ar)
+{
+	Ar << ID;
+	Ar << Name;
+}
+
+FGuid FRigPhysicsSolverDescription::MakeGuid(const FString& InObjectPath, const FName& InSolverName)
+{
+	const FString CompletePath = FString::Printf(TEXT("%s|%s"), *InObjectPath, *InSolverName.ToString());
+	return FGuid::NewDeterministicGuid(CompletePath);
+}
+
+FRigPhysicsSolverID FRigPhysicsSolverDescription::MakeID(const FString& InObjectPath, const FName& InSolverName)
+{
+	return FRigPhysicsSolverID(MakeGuid(InObjectPath, InSolverName));
+}
+
+void FRigPhysicsSolverDescription::CopyFrom(const FRigPhysicsSolverDescription* InOther)
+{
+	if(InOther)
+	{
+		ID = InOther->ID;
+		Name = InOther->Name;
+	}
+}
+
+////////////////////////////////////////////////////////////////////////////////
+// FRigPhysicsSettings
+////////////////////////////////////////////////////////////////////////////////
+
+FRigPhysicsSettings::FRigPhysicsSettings()
 	: Mass(1.f)
 {
 }
 
-void FRigRigidBodySettings::Save(FArchive& Ar)
+void FRigPhysicsSettings::Save(FArchive& Ar)
 {
 	Ar << Mass;
 }
 
-void FRigRigidBodySettings::Load(FArchive& Ar)
+void FRigPhysicsSettings::Load(FArchive& Ar)
 {
 	Ar << Mass;
 }
 
 ////////////////////////////////////////////////////////////////////////////////
-// FRigRigidBodyElement
+// FRigPhysicsElement
 ////////////////////////////////////////////////////////////////////////////////
 
-const FRigBaseElement::EElementIndex FRigRigidBodyElement::ElementTypeIndex = RigidBodyElement;
+const FRigBaseElement::EElementIndex FRigPhysicsElement::ElementTypeIndex = PhysicsElement;
 
-void FRigRigidBodyElement::Save(FArchive& Ar, ESerializationPhase SerializationPhase)
+void FRigPhysicsElement::Save(FArchive& Ar, ESerializationPhase SerializationPhase)
 {
 	Super::Save(Ar, SerializationPhase);
 
 	if(SerializationPhase == ESerializationPhase::StaticData)
 	{
+		Ar << Solver;
 		Settings.Save(Ar);
 	}
 }
 
-void FRigRigidBodyElement::Load(FArchive& Ar, ESerializationPhase SerializationPhase)
+void FRigPhysicsElement::Load(FArchive& Ar, ESerializationPhase SerializationPhase)
 {
 	Super::Load(Ar, SerializationPhase);
 
 	if(SerializationPhase == ESerializationPhase::StaticData)
 	{
+		Ar << Solver;
 		Settings.Load(Ar);
 	}
 }
 
-void FRigRigidBodyElement::CopyFrom(const FRigBaseElement* InOther)
+void FRigPhysicsElement::CopyFrom(const FRigBaseElement* InOther)
 {
 	Super::CopyFrom(InOther);
 	
-	const FRigRigidBodyElement* Source = CastChecked<FRigRigidBodyElement>(InOther);
+	const FRigPhysicsElement* Source = CastChecked<FRigPhysicsElement>(InOther);
+	Solver = Source->Solver;
 	Settings = Source->Settings;
 }
 
diff --git a/Engine/Plugins/Animation/ControlRig/Source/ControlRig/Private/Units/Execution/RigUnit_Physics.cpp b/Engine/Plugins/Animation/ControlRig/Source/ControlRig/Private/Units/Execution/RigUnit_Physics.cpp
new file mode 100644
--- /dev/null
+++ b/Engine/Plugins/Animation/ControlRig/Source/ControlRig/Private/Units/Execution/RigUnit_Physics.cpp
@@ -0,0 +1,49 @@
+// Copyright Epic Games, Inc. All Rights Reserved.
+
+#include "Units/Execution/RigUnit_Physics.h"
+#include "Rigs/RigHierarchyController.h"
+#include "Units/RigUnitContext.h"
+#include "ControlRig.h"
+
+#include UE_INLINE_GENERATED_CPP_BY_NAME(RigUnit_Physics)
+
+FRigUnit_HierarchyAddPhysicsSolver_Execute()
+{
+	FString ErrorMessage;
+	if(!FRigUnit_DynamicHierarchyBase::IsValidToRunInContext(ExecuteContext, true, &ErrorMessage))
+	{
+		if(!ErrorMessage.IsEmpty())
+		{
+			UE_CONTROLRIG_RIGUNIT_REPORT_ERROR(TEXT("%s"), *ErrorMessage);
+		}
+		return;
+	}
+
+	Solver = FRigPhysicsSolverID();
+
+	if(UControlRig* ControlRig = Cast<UControlRig>(ExecuteContext.Hierarchy->GetOuter()))
+	{
+		Solver = ControlRig->AddPhysicsSolver(Name, false, false);
+	}
+}
+
+FRigUnit_HierarchyAddPhysicsJoint_Execute()
+{
+	FString ErrorMessage;
+	if(!FRigUnit_DynamicHierarchyBase::IsValidToRunInContext(ExecuteContext, true, &ErrorMessage))
+	{
+		if(!ErrorMessage.IsEmpty())
+		{
+			UE_CONTROLRIG_RIGUNIT_REPORT_ERROR(TEXT("%s"), *ErrorMessage);
+		}
+		return;
+	}
+
+	Item.Reset();
+
+	if(URigHierarchyController* Controller = ExecuteContext.Hierarchy->GetController(true))
+	{
+		FRigHierarchyControllerInstructionBracket InstructionBracket(Controller, ExecuteContext.GetInstructionIndex());
+		Item = Controller->AddPhysicsElement(Name, Parent, Solver, Settings, Transform, false, false);
+	}
+}
diff --git a/Engine/Plugins/Animation/ControlRig/Source/ControlRig/Public/ControlRig.h b/Engine/Plugins/Animation/ControlRig/Source/ControlRig/Public/ControlRig.h
--- a/Engine/Plugins/Animation/ControlRig/Source/ControlRig/Public/ControlRig.h
+++ b/Engine/Plugins/Animation/ControlRig/Source/ControlRig/Public/ControlRig.h
@@ -699,9 +699,52 @@ protected:
 
 private:
 	float DebugBoneRadiusMultiplier;
-	
+
+	// Physics Solvers
+	TArray<FRigPhysicsSolverDescription> PhysicsSolvers;
+
 public:
-	
+
+	/**
+	 * Returns the number of physics solvers
+	 * @return The number of physics solvers
+	 */
+	int32 NumPhysicsSolvers() const; 
+
+	/**
+	 * Returns a physics solver by index (or nullptr)
+	 * @param InIndex The index of the physics solver to return
+	 * @return The physics solver
+	 */
+	const FRigPhysicsSolverDescription* GetPhysicsSolver(int32 InIndex) const; 
+
+	/**
+	 * Finds a new physics solver given its guid
+	 * @param InID The id identifying the physics solver
+	 * @return The physics solver
+	 */
+	const FRigPhysicsSolverDescription* FindPhysicsSolver(const FRigPhysicsSolverID& InID) const; 
+
+	/**
+	 * Finds a new physics solver given its name
+	 * @param InName The name identifying the physics solver in the scope of this control rig
+	 * @return The physics solver
+	 */
+	const FRigPhysicsSolverDescription* FindPhysicsSolverByName(const FName& InName) const; 
+
+	/**
+	 * Adds a physics solver to the hierarchy
+	 * @param InName The suggested name of the new physics solver - will eventually be corrected by the namespace
+	 * @param bSetupUndo If set to true the stack will record the change for undo / redo
+	 * @param bPrintPythonCommand If set to true a python command equivalent to this call will be printed out
+	 * @return The ID for the newly created physics solver.
+	 */
+	UFUNCTION(BlueprintCallable, Category = ControlRig)
+	FRigPhysicsSolverID AddPhysicsSolver(
+		FName InName,
+		bool bSetupUndo = false,
+		bool bPrintPythonCommand = false);
+
 #if WITH_EDITOR	
 
 	void ToggleControlsVisible() { bControlsVisible = !bControlsVisible; }
diff --git a/Engine/Plugins/Animation/ControlRig/Source/ControlRig/Public/ControlRigObjectVersion.h b/Engine/Plugins/Animation/ControlRig/Source/ControlRig/Public/ControlRigObjectVersion.h
--- a/Engine/Plugins/Animation/ControlRig/Source/ControlRig/Public/ControlRigObjectVersion.h
+++ b/Engine/Plugins/Animation/ControlRig/Source/ControlRig/Public/ControlRigObjectVersion.h
@@ -101,6 +101,9 @@ struct CONTROLRIG_API FControlRigObjectVersion
 		// Add parent key to control rig pose
 		RigPoseWithParentKey,
 
+		// Physics solvers stored on hierarchy
+		ControlRigStoresPhysicsSolvers,
+
 		// -----<new versions can be added above this line>-------------------------------------------------
 		VersionPlusOne,
 		LatestVersion = VersionPlusOne - 1
diff --git a/Engine/Plugins/Animation/ControlRig/Source/ControlRig/Public/Rigs/RigHierarchy.h b/Engine/Plugins/Animation/ControlRig/Source/ControlRig/Public/Rigs/RigHierarchy.h
--- a/Engine/Plugins/Animation/ControlRig/Source/ControlRig/Public/Rigs/RigHierarchy.h
+++ b/Engine/Plugins/Animation/ControlRig/Source/ControlRig/Public/Rigs/RigHierarchy.h
@@ -31,6 +31,7 @@ DECLARE_MULTICAST_DELEGATE_ThreeParams(FRigHierarchyMetadataTagChangedDelegate,
 
 extern CONTROLRIG_API TAutoConsoleVariable<bool> CVarControlRigHierarchyEnableRotationOrder;
 extern CONTROLRIG_API TAutoConsoleVariable<bool> CVarControlRigHierarchyEnableModules;
+extern CONTROLRIG_API TAutoConsoleVariable<bool> CVarControlRigHierarchyEnablePhysics;
 
 UENUM()
 enum ERigTransformStackEntryType : int
@@ -661,7 +662,7 @@ public:
 	{
 		return Get<T>(GetIndex(InKey));
 	}
-	
+
 private:
 	/**
 	* Returns bone element for a given key, for scripting purpose only, for cpp usage, use Find<FRigBoneElement>()
@@ -923,22 +924,22 @@ public:
 	}
 
 	/**
-	 * Returns all RigidBody elements
+	 * Returns all Physics elements
 	 * @param bTraverse Returns the elements in order of a depth first traversal
 	 */
-	TArray<FRigRigidBodyElement*> GetRigidBodies(bool bTraverse = false) const
+	TArray<FRigPhysicsElement*> GetPhysicsElements(bool bTraverse = false) const
 	{
-		return GetElementsOfType<FRigRigidBodyElement>(bTraverse);
+		return GetElementsOfType<FRigPhysicsElement>(bTraverse);
 	}
 
 	/**
-	 * Returns all RigidBody elements
+	 * Returns all Physics elements
 	 * @param bTraverse Returns the elements in order of a depth first traversal
 	 */
-	UFUNCTION(BlueprintCallable, Category = URigHierarchy, meta = (DisplayName = "Get RigidBodies", ScriptName = "GetRigidBodies"))
-    TArray<FRigElementKey> GetRigidBodyKeys(bool bTraverse = true) const
+	UFUNCTION(BlueprintCallable, Category = URigHierarchy, meta = (DisplayName = "Get Physics Keys", ScriptName = "GetPhysicsKeys"))
+    TArray<FRigElementKey> GetPhysicsKeys(bool bTraverse = true) const
 	{
-		return GetKeysOfType<FRigRigidBodyElement>(bTraverse);
+		return GetKeysOfType<FRigPhysicsElement>(bTraverse);
 	}
 
 	/**
@@ -1045,6 +1046,20 @@ public:
 		}, false);
 	}
 
+	/**
+	 * Finds a new physics solver given its guid
+	 * @param InID The id identifying the physics solver
+	 * @return The physics solver
+	 */
+	const FRigPhysicsSolverDescription* FindPhysicsSolver(const FRigPhysicsSolverID& InID) const; 
+
+	/**
+	 * Finds a new physics solver given its name
+	 * @param InName The name identifying the physics solver in the scope of this hierarchy
+	 * @return The physics solver
+	 */
+	const FRigPhysicsSolverDescription* FindPhysicsSolverByName(const FName& InName) const; 
+
 	/**
 	 * Returns the name of metadata for a given element
 	 * @param InItem The element key to return the metadata keys for
@@ -4453,7 +4468,7 @@ protected:
 			{
 				return 3;
 			}
-			case ERigElementType::RigidBody:
+			case ERigElementType::Physics:
 			{
 				return 4;
 			}
@@ -4501,7 +4516,7 @@ protected:
 			}
 			case 4:
 			{
-				return ERigElementType::RigidBody;
+				return ERigElementType::Physics;
 			}
 			case 5:
 			{
diff --git a/Engine/Plugins/Animation/ControlRig/Source/ControlRig/Public/Rigs/RigHierarchyController.h b/Engine/Plugins/Animation/ControlRig/Source/ControlRig/Public/Rigs/RigHierarchyController.h
--- a/Engine/Plugins/Animation/ControlRig/Source/ControlRig/Public/Rigs/RigHierarchyController.h
+++ b/Engine/Plugins/Animation/ControlRig/Source/ControlRig/Public/Rigs/RigHierarchyController.h
@@ -201,20 +201,22 @@ public:
         );
 
 	/**
-	* Adds a rigidbody to the hierarchy
-	* @param InName The suggested name of the new rigidbody - will eventually be corrected by the namespace
-	* @param InParent The (optional) parent of the new rigidbody. If you don't need a parent, pass FRigElementKey()
-	* @param InSettings All of the rigidbody's settings
-	* @param InLocalTransform The transform for the new rigidbody - in the space of the provided parent
+	* Adds a physics element to the hierarchy
+	* @param InName The suggested name of the new physics element - will eventually be corrected by the namespace
+	* @param InParent The (optional) parent of the new physics element. If you don't need a parent, pass FRigElementKey()
+	* @param InSolver The guid identifying the solver to use
+	* @param InSettings All of the physics element's settings
+	* @param InLocalTransform The transform for the new physics element - in the space of the provided parent
 	* @param bSetupUndo If set to true the stack will record the change for undo / redo
 	* @param bPrintPythonCommand If set to true a python command equivalent to this call will be printed out
-	* @return The key for the newly created rigidbody.
+	* @return The key for the newly created physics element.
 	*/
 	UFUNCTION(BlueprintCallable, Category = URigHierarchyController)
-    FRigElementKey AddRigidBody(
+    FRigElementKey AddPhysicsElement(
     	FName InName,
     	FRigElementKey InParent,
-        FRigRigidBodySettings InSettings,
+    	FRigPhysicsSolverID InSolver,
+        FRigPhysicsSettings InSettings,
     	FTransform InLocalTransform,
     	bool bSetupUndo = false,
 		bool bPrintPythonCommand = false);
@@ -611,7 +613,7 @@ public:
 
 	TArray<FString> GetAddCurvePythonCommands(FRigCurveElement* Curve) const;
 
-	TArray<FString> GetAddRigidBodyPythonCommands(FRigRigidBodyElement* RigidBody) const;
+	TArray<FString> GetAddPhysicsElementPythonCommands(FRigPhysicsElement* PhysicsElement) const;
 
 	TArray<FString> GetAddConnectorPythonCommands(FRigConnectorElement* Connector) const;
 
diff --git a/Engine/Plugins/Animation/ControlRig/Source/ControlRig/Public/Rigs/RigHierarchyDefines.h b/Engine/Plugins/Animation/ControlRig/Source/ControlRig/Public/Rigs/RigHierarchyDefines.h
--- a/Engine/Plugins/Animation/ControlRig/Source/ControlRig/Public/Rigs/RigHierarchyDefines.h
+++ b/Engine/Plugins/Animation/ControlRig/Source/ControlRig/Public/Rigs/RigHierarchyDefines.h
@@ -26,14 +26,14 @@ enum class ERigElementType : uint8
 	Space = Null UMETA(Hidden),
 	Control = 0x004,
 	Curve = 0x008,
-	RigidBody = 0x010 UMETA(Hidden), 
+	Physics = 0x010, 
 	Reference = 0x020,
 	Connector = 0x040,
 	Socket = 0x080,
 	
 	First = Bone UMETA(Hidden), 
 	Last = Socket UMETA(Hidden), 
-	All = Bone | Null | Control | Curve | RigidBody | Reference | Connector | Socket,
+	All = Bone | Null | Control | Curve | Physics | Reference | Connector | Socket,
 	ToResetAfterConstructionEvent = Bone | Control | Curve | Socket UMETA(Hidden),
 };
 
@@ -1580,9 +1580,9 @@ public:
 			{
 				return FString::Printf(TEXT("Curve(%s)"), *Name.ToString());
 			}
-			case ERigElementType::RigidBody:
+			case ERigElementType::Physics:
 			{
-				return FString::Printf(TEXT("RigidBody(%s)"), *Name.ToString());
+				return FString::Printf(TEXT("Physics(%s)"), *Name.ToString());
 			}
 			case ERigElementType::Reference:
 			{
diff --git a/Engine/Plugins/Animation/ControlRig/Source/ControlRig/Public/Rigs/RigHierarchyElements.h b/Engine/Plugins/Animation/ControlRig/Source/ControlRig/Public/Rigs/RigHierarchyElements.h
--- a/Engine/Plugins/Animation/ControlRig/Source/ControlRig/Public/Rigs/RigHierarchyElements.h
+++ b/Engine/Plugins/Animation/ControlRig/Source/ControlRig/Public/Rigs/RigHierarchyElements.h
@@ -566,7 +566,7 @@ public:
 		NullElement,
 		ControlElement,
 		CurveElement,
-		RigidBodyElement,
+		PhysicsElement,
 		ReferenceElement,
 		ConnectorElement,
 		SocketElement,
@@ -822,7 +822,7 @@ protected:
 		return InElement->GetType() == ERigElementType::Bone ||
 			InElement->GetType() == ERigElementType::Null ||
 			InElement->GetType() == ERigElementType::Control ||
-			InElement->GetType() == ERigElementType::RigidBody ||
+			InElement->GetType() == ERigElementType::Physics ||
 			InElement->GetType() == ERigElementType::Reference ||
 			InElement->GetType() == ERigElementType::Socket;
 	}
@@ -860,7 +860,7 @@ protected:
 	static bool IsClassOf(const FRigBaseElement* InElement)
 	{
 		return InElement->GetType() == ERigElementType::Bone ||
-			InElement->GetType() == ERigElementType::RigidBody ||
+			InElement->GetType() == ERigElementType::Physics ||
 			InElement->GetType() == ERigElementType::Reference ||
 			InElement->GetType() == ERigElementType::Socket;
 	}
@@ -1521,11 +1521,109 @@ private:
 };
 
 USTRUCT(BlueprintType)
-struct CONTROLRIG_API FRigRigidBodySettings
+struct CONTROLRIG_API FRigPhysicsSolverID
 {
+public:
+	
 	GENERATED_BODY()
+	
+	FRigPhysicsSolverID()
+		: Guid()
+	{
+	}
+
+	explicit FRigPhysicsSolverID(const FGuid& InGuid)
+		: Guid(InGuid)
+	{
+	}
+
+	UPROPERTY(BlueprintReadOnly, EditAnywhere, Category = Physics, meta=(ShowOnlyInnerProperties))
+	FGuid Guid;
 
-	FRigRigidBodySettings();
+	bool IsValid() const
+	{
+		return Guid.IsValid();
+	}
+
+	FString ToString() const
+	{
+		return Guid.ToString();
+	}
+
+	bool operator == (const FRigPhysicsSolverID& InOther) const
+	{
+		return Guid == InOther.Guid;
+	}
+
+	bool operator != (const FRigPhysicsSolverID& InOther) const
+	{
+		return !(*this == InOther);
+	}
+
+	friend FArchive& operator <<(FArchive& Ar, FRigPhysicsSolverID& ID)
+	{
+		Ar << ID.Guid;
+		return Ar;
+	}
+
+	friend uint32 GetTypeHash(const FRigPhysicsSolverID& InID)
+	{
+		return GetTypeHash(InID.Guid);
+	}
+};
+
+USTRUCT(BlueprintType)
+struct CONTROLRIG_API FRigPhysicsSolverDescription
+{
+public:
+	
+	GENERATED_BODY()
+
+	UPROPERTY(BlueprintReadOnly, EditAnywhere, Category = Physics, meta=(ShowOnlyInnerProperties))
+	FRigPhysicsSolverID ID;
+
+	UPROPERTY(BlueprintReadOnly, EditAnywhere, Category = Physics, meta=(ShowOnlyInnerProperties))
+	FName Name;
+
+	FRigPhysicsSolverDescription()
+	: ID()
+	, Name(NAME_None)
+	{}
+	
+	FRigPhysicsSolverDescription(const FRigPhysicsSolverDescription& InOther)
+	{
+		*this = InOther;
+	}
+	
+	FRigPhysicsSolverDescription& operator=(const FRigPhysicsSolverDescription& InOther)
+	{
+		CopyFrom(&InOther);
+		return *this;
+	}
+
+	void Serialize(FArchive& Ar);
+	void Save(FArchive& Ar);
+	void Load(FArchive& Ar);
+	friend FArchive& operator<<(FArchive& Ar, FRigPhysicsSolverDescription& P)
+	{
+		P.Serialize(Ar);
+		return Ar;
+	}
+
+	static FGuid MakeGuid(const FString& InObjectPath, const FName& InSolverName);
+	static FRigPhysicsSolverID MakeID(const FString& InObjectPath, const FName& InSolverName);
+
+private:
+
+	void CopyFrom(const FRigPhysicsSolverDescription* InOther);
+};
+
+USTRUCT(BlueprintType)
+struct CONTROLRIG_API FRigPhysicsSettings
+{
+	GENERATED_BODY()
+
+	FRigPhysicsSettings();
 
 	void Save(FArchive& Ar);
 	void Load(FArchive& Ar);
@@ -1536,49 +1634,53 @@ struct CONTROLRIG_API FRigRigidBodySettings
 };
 
 USTRUCT(BlueprintType)
-struct CONTROLRIG_API FRigRigidBodyElement : public FRigSingleParentElement
+struct CONTROLRIG_API FRigPhysicsElement : public FRigSingleParentElement
 {
 public:
 	
 	GENERATED_BODY()
-	DECLARE_RIG_ELEMENT_METHODS(FRigRigidBodyElement)
+	DECLARE_RIG_ELEMENT_METHODS(FRigPhysicsElement)
 
 	static const EElementIndex ElementTypeIndex;
 
-	UPROPERTY(BlueprintReadOnly, EditAnywhere, Category = Control, meta=(ShowOnlyInnerProperties))
-	FRigRigidBodySettings Settings;
+	UPROPERTY(BlueprintReadOnly, EditAnywhere, Category = Physics, meta=(ShowOnlyInnerProperties))
+	FRigPhysicsSolverID Solver;
+
+	UPROPERTY(BlueprintReadOnly, EditAnywhere, Category = Physics, meta=(ShowOnlyInnerProperties))
+	FRigPhysicsSettings Settings;
 	
-	FRigRigidBodyElement()
-        : FRigRigidBodyElement(nullptr)
+	FRigPhysicsElement()
+        : FRigPhysicsElement(nullptr)
 	{ }
 	
-	FRigRigidBodyElement(const FRigRigidBodyElement& InOther)
+	FRigPhysicsElement(const FRigPhysicsElement& InOther)
 	{
 		*this = InOther;
 	}
 	
-	FRigRigidBodyElement& operator=(const FRigRigidBodyElement& InOther)
+	FRigPhysicsElement& operator=(const FRigPhysicsElement& InOther)
 	{
 		Super::operator=(InOther);
+		Solver = InOther.Solver;
 		Settings = InOther.Settings;
 		return *this;
 	}
 	
-	virtual ~FRigRigidBodyElement() override {}
+	virtual ~FRigPhysicsElement() override {}
 
 	virtual void Save(FArchive& A, ESerializationPhase SerializationPhase) override;
 	virtual void Load(FArchive& Ar, ESerializationPhase SerializationPhase) override;
 
 private:
-	explicit FRigRigidBodyElement(URigHierarchy* InOwner)
-		: FRigSingleParentElement(InOwner, ERigElementType::RigidBody)
+	explicit FRigPhysicsElement(URigHierarchy* InOwner)
+		: FRigSingleParentElement(InOwner, ERigElementType::Physics)
 	{ }
 	
 	virtual void CopyFrom(const FRigBaseElement* InOther) override;
 
 	static bool IsClassOf(const FRigBaseElement* InElement)
 	{
-		return InElement->GetType() == ERigElementType::RigidBody;
+		return InElement->GetType() == ERigElementType::Physics;
 	}
 
 	friend class URigHierarchy;
diff --git a/Engine/Plugins/Animation/ControlRig/Source/ControlRig/Public/Units/Execution/RigUnit_Physics.h b/Engine/Plugins/Animation/ControlRig/Source/ControlRig/Public/Units/Execution/RigUnit_Physics.h
new file mode 100644
--- /dev/null
+++ b/Engine/Plugins/Animation/ControlRig/Source/ControlRig/Public/Units/Execution/RigUnit_Physics.h
@@ -0,0 +1,80 @@
+// Copyright Epic Games, Inc. All Rights Reserved.
+
+#pragma once
+
+#include "RigUnit_DynamicHierarchy.h"
+#include "RigUnit_Physics.generated.h"
+
+/**
+ * Adds a new physics solver to the hierarchy
+ * Note: This node only runs as part of the construction event.
+ */
+USTRUCT(meta=(DisplayName="Spawn Physics Solver", Keywords="Construction,Create,New,Simulation", Varying))
+struct CONTROLRIG_API FRigUnit_HierarchyAddPhysicsSolver : public FRigUnit_DynamicHierarchyBaseMutable
+{
+	GENERATED_BODY()
+
+	FRigUnit_HierarchyAddPhysicsSolver()
+	{
+		Name = TEXT("Solver");
+		Solver = FRigPhysicsSolverID();
+	}
+
+	/*
+	 * The name of the new solver to add
+	 */
+	UPROPERTY(meta = (Input))
+	FName Name;
+
+	/*
+	 * The identifier of the solver
+	 */
+	UPROPERTY(meta = (Output))
+	FRigPhysicsSolverID Solver;
+
+	RIGVM_METHOD()
+	virtual void Execute() override;
+};
+
+/**
+ * Adds a new physics joint to the hierarchy
+ * Note: This node only runs as part of the construction event.
+ */
+USTRUCT(meta=(DisplayName="Spawn Physics Joint", Keywords="Construction,Create,New,Joint,Skeleton", Varying))
+struct CONTROLRIG_API FRigUnit_HierarchyAddPhysicsJoint : public FRigUnit_HierarchyAddElement
+{
+	GENERATED_BODY()
+
+	FRigUnit_HierarchyAddPhysicsJoint()
+	{
+		Name = TEXT("NewPhysicsJoint");
+		Parent = FRigElementKey(NAME_None, ERigElementType::Bone);
+		Transform = FTransform::Identity;
+		Solver = FRigPhysicsSolverID();
+	}
+
+	virtual ERigElementType GetElementTypeToSpawn() const override { return ERigElementType::Physics; }
+
+
+	RIGVM_METHOD()
+	virtual void Execute() override;
+
+	/*
+	 * The initial global transform of the spawned element
+	 */
+	UPROPERTY(meta = (Input))
+	FTransform Transform;
+
+	/*
+	 * The solver to relate this new physics element to
+	 */
+	UPROPERTY(meta = (Input))
+	FRigPhysicsSolverID Solver;
+
+	/*
+	 * The settings of the new physics element
+	 */
+	UPROPERTY(meta = (Input))
+	FRigPhysicsSettings Settings;
+};
+
diff --git a/Engine/Plugins/Animation/ControlRig/Source/ControlRigDeveloper/Private/ControlRigSchema.cpp b/Engine/Plugins/Animation/ControlRig/Source/ControlRigDeveloper/Private/ControlRigSchema.cpp
--- a/Engine/Plugins/Animation/ControlRig/Source/ControlRigDeveloper/Private/ControlRigSchema.cpp
+++ b/Engine/Plugins/Animation/ControlRig/Source/ControlRigDeveloper/Private/ControlRigSchema.cpp
@@ -35,6 +35,10 @@ bool UControlRigSchema::ShouldUnfoldStruct(URigVMController* InController, const
 	{
 		return false;
 	}
+	if (InStruct == FRigPhysicsSolverID::StaticStruct())
+	{
+		return false;
+	}
 	
 	return true;
 }
diff --git a/Engine/Plugins/Animation/ControlRig/Source/ControlRigDeveloper/Private/Graph/ControlRigGraph.cpp b/Engine/Plugins/Animation/ControlRig/Source/ControlRigDeveloper/Private/Graph/ControlRigGraph.cpp
--- a/Engine/Plugins/Animation/ControlRig/Source/ControlRigDeveloper/Private/Graph/ControlRigGraph.cpp
+++ b/Engine/Plugins/Animation/ControlRig/Source/ControlRigDeveloper/Private/Graph/ControlRigGraph.cpp
@@ -82,7 +82,7 @@ void UControlRigGraph::CacheNameLists(URigHierarchy* InHierarchy, const FRigVMDr
 		ElementNameLists.FindOrAdd(ERigElementType::Null);
 		ElementNameLists.FindOrAdd(ERigElementType::Control);
 		ElementNameLists.FindOrAdd(ERigElementType::Curve);
-		ElementNameLists.FindOrAdd(ERigElementType::RigidBody);
+		ElementNameLists.FindOrAdd(ERigElementType::Physics);
 		ElementNameLists.FindOrAdd(ERigElementType::Reference);
 		ElementNameLists.FindOrAdd(ERigElementType::Connector);
 		ElementNameLists.FindOrAdd(ERigElementType::Socket);
@@ -92,7 +92,7 @@ void UControlRigGraph::CacheNameLists(URigHierarchy* InHierarchy, const FRigVMDr
 		TArray<TSharedPtr<FRigVMStringWithTag>>& NullNameList = ElementNameLists.FindChecked(ERigElementType::Null);
 		TArray<TSharedPtr<FRigVMStringWithTag>>& ControlNameList = ElementNameLists.FindChecked(ERigElementType::Control);
 		TArray<TSharedPtr<FRigVMStringWithTag>>& CurveNameList = ElementNameLists.FindChecked(ERigElementType::Curve);
-		TArray<TSharedPtr<FRigVMStringWithTag>>& RigidBodyNameList = ElementNameLists.FindChecked(ERigElementType::RigidBody);
+		TArray<TSharedPtr<FRigVMStringWithTag>>& PhysicsNameList = ElementNameLists.FindChecked(ERigElementType::Physics);
 		TArray<TSharedPtr<FRigVMStringWithTag>>& ReferenceNameList = ElementNameLists.FindChecked(ERigElementType::Reference);
 		TArray<TSharedPtr<FRigVMStringWithTag>>& ConnectorNameList = ElementNameLists.FindChecked(ERigElementType::Connector);
 		TArray<TSharedPtr<FRigVMStringWithTag>>& SocketNameList = ElementNameLists.FindChecked(ERigElementType::Socket);
@@ -103,7 +103,7 @@ void UControlRigGraph::CacheNameLists(URigHierarchy* InHierarchy, const FRigVMDr
 		CacheNameListForHierarchy<FRigControlElement>(ControlRig, InHierarchy, ControlNameList, false);
 		CacheNameListForHierarchy<FRigControlElement>(ControlRig, InHierarchy, ControlNameListWithoutAnimationChannels, true);
 		CacheNameListForHierarchy<FRigCurveElement>(ControlRig, InHierarchy, CurveNameList, false);
-		CacheNameListForHierarchy<FRigRigidBodyElement>(ControlRig, InHierarchy, RigidBodyNameList, false);
+		CacheNameListForHierarchy<FRigPhysicsElement>(ControlRig, InHierarchy, PhysicsNameList, false);
 		CacheNameListForHierarchy<FRigReferenceElement>(ControlRig, InHierarchy, ReferenceNameList, false);
 		CacheNameListForHierarchy<FRigConnectorElement>(ControlRig, InHierarchy, ConnectorNameList, false);
 		CacheNameListForHierarchy<FRigSocketElement>(ControlRig, InHierarchy, SocketNameList, false);
diff --git a/Engine/Plugins/Animation/ControlRig/Source/ControlRigEditor/Private/ControlRigElementDetails.cpp b/Engine/Plugins/Animation/ControlRig/Source/ControlRigEditor/Private/ControlRigElementDetails.cpp
--- a/Engine/Plugins/Animation/ControlRig/Source/ControlRigEditor/Private/ControlRigElementDetails.cpp
+++ b/Engine/Plugins/Animation/ControlRig/Source/ControlRigEditor/Private/ControlRigElementDetails.cpp
@@ -866,12 +866,12 @@ void FRigBaseElementDetails::CustomizeDetails(IDetailLayoutBuilder& DetailBuilde
 	if(!bAllControls || !bAllAnimationChannels)
 	{
 		GeneralCategory.AddCustomRow(FText::FromString(TEXT("Name")))
+		.IsEnabled(!bIsProcedural)
 		.NameContent()
 		[
 			SNew(STextBlock)
 			.Text(FText::FromString(TEXT("Name")))
 			.Font(IDetailLayoutBuilder::GetDetailFont())
-			.IsEnabled(!bIsProcedural)
 		]
 		.ValueContent()
 		[
@@ -5675,4 +5675,97 @@ void FRigConnectionRuleDetails::OnRuleContentChanged()
 	StructPropertyHandle->GetParentHandle()->NotifyPostChange(EPropertyChangeType::ValueSet);
 }
 
+void FRigPhysicsElementDetails::CustomizeDetails(IDetailLayoutBuilder& DetailBuilder)
+{
+	FRigTransformElementDetails::CustomizeDetails(DetailBuilder);
+	CustomizeSettings(DetailBuilder);
+	CustomizeTransform(DetailBuilder);
+	CustomizeMetadata(DetailBuilder);
+}
+
+void FRigPhysicsElementDetails::CustomizeSettings(IDetailLayoutBuilder& DetailBuilder)
+{
+	if(PerElementInfos.IsEmpty())
+	{
+		return;
+	}
+
+	if(IsAnyElementNotOfType(ERigElementType::Physics))
+	{
+		return;
+	}
+
+	const TSharedPtr<IPropertyHandle> SolverHandle = DetailBuilder.GetProperty(GET_MEMBER_NAME_CHECKED(FRigPhysicsElement, Solver));
+	DetailBuilder.HideProperty(SolverHandle);
+
+	const TSharedPtr<IPropertyHandle> SettingsHandle = DetailBuilder.GetProperty(GET_MEMBER_NAME_CHECKED(FRigPhysicsElement, Settings));
+	DetailBuilder.HideProperty(SettingsHandle);
+
+	IDetailCategoryBuilder& PhysicsCategory = DetailBuilder.EditCategory(TEXT("Physics"), LOCTEXT("Physics", "Physics"));
+
+	FDetailWidgetRow& SolverNameRow = PhysicsCategory.AddCustomRow(LOCTEXT("SolverName", "Solver Name"));
+	SolverNameRow.IsEnabled(false);
+	SolverNameRow.NameContent()
+	[
+		SNew(STextBlock)
+		.Text(LOCTEXT("SolverName", "Solver Name"))
+		.Font(IDetailLayoutBuilder::GetDetailFont())
+	];
+	SolverNameRow.ValueContent()
+	[
+		SNew(SEditableText)
+		.Text(this, &FRigPhysicsElementDetails::GetSolverNameText)
+		.Font(IDetailLayoutBuilder::GetDetailFont())
+	];
+
+	PhysicsCategory
+	.AddProperty(SolverHandle)
+	.DisplayName(LOCTEXT("SolverGuid", "Solver Guid"))
+	.IsEnabled(false);
+
+	const bool bIsProcedural = IsAnyElementProcedural();
+
+	PhysicsCategory
+		.AddProperty(SettingsHandle->GetChildHandle(GET_MEMBER_NAME_CHECKED(FRigPhysicsSettings, Mass)))
+		.IsEnabled(!bIsProcedural);
+}
+
+FText FRigPhysicsElementDetails::GetSolverNameText() const
+{
+	if(PerElementInfos.Num() > 0)
+	{
+		if(const URigHierarchy* Hierarchy = PerElementInfos[0].GetHierarchy())
+		{
+			TOptional<FText> FirstValue;
+			
+			for(const FPerElementInfo& Info : PerElementInfos)
+			{
+				if(const FRigPhysicsElement* PhysicsElement = Info.Element.Get<FRigPhysicsElement>())
+				{
+					if(const FRigPhysicsSolverDescription* Solver = Hierarchy->FindPhysicsSolver(PhysicsElement->Solver))
+					{
+						const FText SolverName = FText::FromName(Solver->Name);
+						if(FirstValue.IsSet())
+						{
+							if(!FirstValue.GetValue().EqualTo(SolverName))
+							{
+								return ControlRigDetailsMultipleValues;	
+							}
+						}
+						else
+						{
+							FirstValue = SolverName;
+						}
+					}
+				}
+			}
+
+			if(FirstValue.IsSet())
+			{
+				return FirstValue.GetValue();
+			}
+		}
+	}
+	return FText();
+}
 #undef LOCTEXT_NAMESPACE
diff --git a/Engine/Plugins/Animation/ControlRig/Source/ControlRigEditor/Private/ControlRigElementDetails.h b/Engine/Plugins/Animation/ControlRig/Source/ControlRigEditor/Private/ControlRigElementDetails.h
--- a/Engine/Plugins/Animation/ControlRig/Source/ControlRigEditor/Private/ControlRigElementDetails.h
+++ b/Engine/Plugins/Animation/ControlRig/Source/ControlRigEditor/Private/ControlRigElementDetails.h
@@ -868,3 +868,23 @@ protected:
 	TSharedPtr<IPropertyUtilities> PropertyUtilities;
 	TAttribute<bool> EnabledAttribute;
 };
+
+class FRigPhysicsElementDetails : public FRigTransformElementDetails
+{
+public:
+
+	// Makes a new instance of this detail layout class for a specific detail view requesting it
+	static TSharedRef<IDetailCustomization> MakeInstance()
+	{
+		return MakeShareable(new FRigPhysicsElementDetails);
+	}
+
+	/** IDetailCustomization interface */
+	virtual void CustomizeDetails(IDetailLayoutBuilder& DetailBuilder) override;
+
+	void CustomizeSettings(IDetailLayoutBuilder& DetailBuilder);
+
+private:
+
+	FText GetSolverNameText() const;
+};
diff --git a/Engine/Plugins/Animation/ControlRig/Source/ControlRigEditor/Private/ControlRigHierarchyCommands.cpp b/Engine/Plugins/Animation/ControlRig/Source/ControlRigEditor/Private/ControlRigHierarchyCommands.cpp
--- a/Engine/Plugins/Animation/ControlRig/Source/ControlRigEditor/Private/ControlRigHierarchyCommands.cpp
+++ b/Engine/Plugins/Animation/ControlRig/Source/ControlRigEditor/Private/ControlRigHierarchyCommands.cpp
@@ -37,7 +37,7 @@ void FControlRigHierarchyCommands::RegisterCommands()
 	UI_COMMAND(ShowBones, "Show Bones", "Whether to show or hide bones", EUserInterfaceActionType::ToggleButton, FInputChord());
 	UI_COMMAND(ShowControls, "Show Controls", "Whether to show or hide controls", EUserInterfaceActionType::ToggleButton, FInputChord());
 	UI_COMMAND(ShowNulls, "Show Nulls", "Whether to show or hide nulls", EUserInterfaceActionType::ToggleButton, FInputChord());
-	UI_COMMAND(ShowRigidBodies, "Show RigidBodies", "Whether to show or hide rigidbodies", EUserInterfaceActionType::ToggleButton, FInputChord());
+	UI_COMMAND(ShowPhysics, "Show Physics", "Whether to show or hide physics", EUserInterfaceActionType::ToggleButton, FInputChord());
 	UI_COMMAND(ShowReferences, "Show References", "Whether to show or hide references", EUserInterfaceActionType::ToggleButton, FInputChord());
 	UI_COMMAND(ShowSockets, "Show Sockets", "Whether to show or hide sockets", EUserInterfaceActionType::ToggleButton, FInputChord());
 	UI_COMMAND(ToggleControlShapeTransformEdit, "Toggle Shape Transform Edit", "Toggle Editing Selected Control's Shape Transform", EUserInterfaceActionType::ToggleButton, FInputChord(EKeys::Period, EModifierKey::Control)); 
diff --git a/Engine/Plugins/Animation/ControlRig/Source/ControlRigEditor/Private/ControlRigHierarchyCommands.h b/Engine/Plugins/Animation/ControlRig/Source/ControlRigEditor/Private/ControlRigHierarchyCommands.h
--- a/Engine/Plugins/Animation/ControlRig/Source/ControlRigEditor/Private/ControlRigHierarchyCommands.h
+++ b/Engine/Plugins/Animation/ControlRig/Source/ControlRigEditor/Private/ControlRigHierarchyCommands.h
@@ -111,8 +111,8 @@ public:
 	/* Show spaces */
 	TSharedPtr< FUICommandInfo > ShowNulls;
 
-	/* Show rigidbodies */
-	TSharedPtr< FUICommandInfo > ShowRigidBodies;
+	/* Show physics */
+	TSharedPtr< FUICommandInfo > ShowPhysics;
 
 	/* Show references */
 	TSharedPtr< FUICommandInfo > ShowReferences;
diff --git a/Engine/Plugins/Animation/ControlRig/Source/ControlRigEditor/Private/EditMode/ControlRigEditMode.cpp b/Engine/Plugins/Animation/ControlRig/Source/ControlRigEditor/Private/EditMode/ControlRigEditMode.cpp
--- a/Engine/Plugins/Animation/ControlRig/Source/ControlRigEditor/Private/EditMode/ControlRigEditMode.cpp
+++ b/Engine/Plugins/Animation/ControlRig/Source/ControlRigEditor/Private/EditMode/ControlRigEditMode.cpp
@@ -4292,7 +4292,7 @@ void FControlRigEditMode::OnHierarchyModified(ERigHierarchyNotification InNotif,
             	case ERigElementType::Null:
             	case ERigElementType::Curve:
             	case ERigElementType::Control:
-            	case ERigElementType::RigidBody:
+            	case ERigElementType::Physics:
             	case ERigElementType::Reference:
             	case ERigElementType::Connector:
             	case ERigElementType::Socket:
diff --git a/Engine/Plugins/Animation/ControlRig/Source/ControlRigEditor/Private/EditMode/SControlRigEditModeTools.cpp b/Engine/Plugins/Animation/ControlRig/Source/ControlRigEditor/Private/EditMode/SControlRigEditModeTools.cpp
--- a/Engine/Plugins/Animation/ControlRig/Source/ControlRigEditor/Private/EditMode/SControlRigEditModeTools.cpp
+++ b/Engine/Plugins/Animation/ControlRig/Source/ControlRigEditor/Private/EditMode/SControlRigEditModeTools.cpp
@@ -219,7 +219,7 @@ void SControlRigEditModeTools::Construct(const FArguments& InArgs, TSharedPtr<FC
 	DisplaySettings.bShowNulls = false;
 	DisplaySettings.bShowReferences = false;
 	DisplaySettings.bShowSockets = false;
-	DisplaySettings.bShowRigidBodies = false;
+	DisplaySettings.bShowPhysics = false;
 	DisplaySettings.bHideParentsOnFilter = true;
 	DisplaySettings.bFlattenHierarchyOnFilter = true;
 	DisplaySettings.bShowIconColors = true;
diff --git a/Engine/Plugins/Animation/ControlRig/Source/ControlRigEditor/Private/EditMode/SControlRigOutliner.cpp b/Engine/Plugins/Animation/ControlRig/Source/ControlRigEditor/Private/EditMode/SControlRigOutliner.cpp
--- a/Engine/Plugins/Animation/ControlRig/Source/ControlRigEditor/Private/EditMode/SControlRigOutliner.cpp
+++ b/Engine/Plugins/Animation/ControlRig/Source/ControlRigEditor/Private/EditMode/SControlRigOutliner.cpp
@@ -522,9 +522,13 @@ bool SMultiRigHierarchyTreeView::AddElement(UControlRig* InControlRig, const FRi
 				}
 				break;
 			}
-			case ERigElementType::RigidBody:
+			case ERigElementType::Physics:
 			{
-				if (!Settings.bShowRigidBodies)
+				if (!Settings.bShowPhysics)
+				{
+					return false;
+				}
+				if(CVarControlRigHierarchyEnablePhysics.GetValueOnAnyThread() == false)
 				{
 					return false;
 				}
@@ -977,7 +981,7 @@ void SControlRigOutliner::Construct(const FArguments& InArgs, FControlRigEditMod
 	DisplaySettings.bShowNulls = false;
 	DisplaySettings.bShowReferences = false;
 	DisplaySettings.bShowSockets = false;
-	DisplaySettings.bShowRigidBodies = false;
+	DisplaySettings.bShowPhysics = false;
 	DisplaySettings.bHideParentsOnFilter = true;
 	DisplaySettings.bFlattenHierarchyOnFilter = true;
 	DisplaySettings.bShowConnectors = false;
diff --git a/Engine/Plugins/Animation/ControlRig/Source/ControlRigEditor/Private/Editor/ControlRigWrapperObject.cpp b/Engine/Plugins/Animation/ControlRig/Source/ControlRigEditor/Private/Editor/ControlRigWrapperObject.cpp
--- a/Engine/Plugins/Animation/ControlRig/Source/ControlRigEditor/Private/Editor/ControlRigWrapperObject.cpp
+++ b/Engine/Plugins/Animation/ControlRig/Source/ControlRigEditor/Private/Editor/ControlRigWrapperObject.cpp
@@ -45,6 +45,10 @@ UClass* UControlRigWrapperObject::GetClassForStruct(UScriptStruct* InStruct, boo
 			{
 				PropertyEditorModule.RegisterCustomClassLayout(WrapperClassName, FOnGetDetailCustomizationInstance::CreateStatic(&FRigSocketElementDetails::MakeInstance));
 			}
+			else if(InStruct == FRigPhysicsElement::StaticStruct())
+			{
+				PropertyEditorModule.RegisterCustomClassLayout(WrapperClassName, FOnGetDetailCustomizationInstance::CreateStatic(&FRigPhysicsElementDetails::MakeInstance));
+			}
 		}
 	}
 #endif
diff --git a/Engine/Plugins/Animation/ControlRig/Source/ControlRigEditor/Private/Editor/SRigHierarchy.cpp b/Engine/Plugins/Animation/ControlRig/Source/ControlRigEditor/Private/Editor/SRigHierarchy.cpp
--- a/Engine/Plugins/Animation/ControlRig/Source/ControlRigEditor/Private/Editor/SRigHierarchy.cpp
+++ b/Engine/Plugins/Animation/ControlRig/Source/ControlRigEditor/Private/Editor/SRigHierarchy.cpp
@@ -418,10 +418,10 @@ void SRigHierarchy::BindCommands()
 		FIsActionChecked::CreateLambda([this]() { return DisplaySettings.bShowNulls; }));
 
 	CommandList->MapAction(
-		Commands.ShowRigidBodies,
-		FExecuteAction::CreateLambda([this]() { DisplaySettings.bShowRigidBodies = !DisplaySettings.bShowRigidBodies; RefreshTreeView(); }),
+		Commands.ShowPhysics,
+		FExecuteAction::CreateLambda([this]() { DisplaySettings.bShowPhysics = !DisplaySettings.bShowPhysics; RefreshTreeView(); }),
 		FCanExecuteAction(),
-		FIsActionChecked::CreateLambda([this]() { return DisplaySettings.bShowRigidBodies; }));
+		FIsActionChecked::CreateLambda([this]() { return DisplaySettings.bShowPhysics; }));
 
 	CommandList->MapAction(
 		Commands.ShowReferences,
@@ -2292,7 +2292,7 @@ bool SRigHierarchy::CanRenameItem() const
 	if(IsSingleSelected(false))
 	{
 		const FRigElementKey Key = GetSelectedKeys()[0];
-		if(Key.Type == ERigElementType::RigidBody ||
+		if(Key.Type == ERigElementType::Physics ||
 			Key.Type == ERigElementType::Reference)
 		{
 			return false;
@@ -2639,7 +2639,7 @@ TOptional<EItemDropZone> SRigHierarchy::OnCanAcceptDrop(const FDragDropEvent& Dr
 			}
 			case ERigElementType::Control:
 			case ERigElementType::Null:
-			case ERigElementType::RigidBody:
+			case ERigElementType::Physics:
 			case ERigElementType::Reference:
 			{
 				for (const FRigElementKey& DraggedKey : RigDragDropOp->GetElements())
@@ -2648,7 +2648,7 @@ TOptional<EItemDropZone> SRigHierarchy::OnCanAcceptDrop(const FDragDropEvent& Dr
 					{
 						case ERigElementType::Control:
 						case ERigElementType::Null:
-						case ERigElementType::RigidBody:
+						case ERigElementType::Physics:
 						case ERigElementType::Reference:
 						case ERigElementType::Connector:
 						case ERigElementType::Socket:
diff --git a/Engine/Plugins/Animation/ControlRig/Source/ControlRigEditor/Private/Editor/SRigHierarchyTreeView.cpp b/Engine/Plugins/Animation/ControlRig/Source/ControlRigEditor/Private/Editor/SRigHierarchyTreeView.cpp
--- a/Engine/Plugins/Animation/ControlRig/Source/ControlRigEditor/Private/Editor/SRigHierarchyTreeView.cpp
+++ b/Engine/Plugins/Animation/ControlRig/Source/ControlRigEditor/Private/Editor/SRigHierarchyTreeView.cpp
@@ -543,9 +543,13 @@ bool SRigHierarchyTreeView::AddElement(const FRigBaseElement* InElement)
 			}
 			break;
 		}
-		case ERigElementType::RigidBody:
+		case ERigElementType::Physics:
 		{
-			if(!Settings.bShowRigidBodies)
+			if(!Settings.bShowPhysics)
+			{
+				return false;
+			}
+			if(CVarControlRigHierarchyEnablePhysics.GetValueOnAnyThread() == false)
 			{
 				return false;
 			}
@@ -1114,7 +1118,7 @@ TPair<const FSlateBrush*, FSlateColor> SRigHierarchyItem::GetBrushForElementType
 	static const FSlateBrush* NullBrush = FControlRigEditorStyle::Get().GetBrush("ControlRig.Tree.Null");
 	static const FSlateBrush* BoneImportedBrush = FControlRigEditorStyle::Get().GetBrush("ControlRig.Tree.BoneImported");
 	static const FSlateBrush* BoneUserBrush = FControlRigEditorStyle::Get().GetBrush("ControlRig.Tree.BoneUser");
-	static const FSlateBrush* RigidBodyBrush = FControlRigEditorStyle::Get().GetBrush("ControlRig.Tree.RigidBody");
+	static const FSlateBrush* PhysicsBrush = FControlRigEditorStyle::Get().GetBrush("ControlRig.Tree.RigidBody");
 	static const FSlateBrush* SocketOpenBrush = FControlRigEditorStyle::Get().GetBrush("ControlRig.Tree.Socket_Open");
 	static const FSlateBrush* SocketClosedBrush = FControlRigEditorStyle::Get().GetBrush("ControlRig.Tree.Socket_Closed");
 	static const FSlateBrush* PrimaryConnectorBrush = FControlRigEditorStyle::Get().GetBrush("ControlRig.ConnectorPrimary");
@@ -1195,9 +1199,9 @@ TPair<const FSlateBrush*, FSlateColor> SRigHierarchyItem::GetBrushForElementType
 
 			break;
 		}
-		case ERigElementType::RigidBody:
+		case ERigElementType::Physics:
 		{
-			Brush = RigidBodyBrush;
+			Brush = PhysicsBrush;
 			break;
 		}
 		case ERigElementType::Reference:
diff --git a/Engine/Plugins/Animation/ControlRig/Source/ControlRigEditor/Private/Editor/SRigHierarchyTreeView.h b/Engine/Plugins/Animation/ControlRig/Source/ControlRigEditor/Private/Editor/SRigHierarchyTreeView.h
--- a/Engine/Plugins/Animation/ControlRig/Source/ControlRigEditor/Private/Editor/SRigHierarchyTreeView.h
+++ b/Engine/Plugins/Animation/ControlRig/Source/ControlRigEditor/Private/Editor/SRigHierarchyTreeView.h
@@ -25,7 +25,7 @@ struct CONTROLRIGEDITOR_API FRigTreeDisplaySettings
 		bShowBones = true;
 		bShowControls = true;
 		bShowNulls = true;
-		bShowRigidBodies = true;
+		bShowPhysics = true;
 		bShowReferences = true;
 		bShowSockets = true;
 		bShowConnectors = true;
@@ -55,8 +55,8 @@ struct CONTROLRIGEDITOR_API FRigTreeDisplaySettings
 	/** Whether or not to show spaces in the hierarchy */
 	bool bShowNulls;
 
-	/** Whether or not to show rigidbodies in the hierarchy */
-	bool bShowRigidBodies;
+	/** Whether or not to show physics elements in the hierarchy */
+	bool bShowPhysics;
 
 	/** Whether or not to show references in the hierarchy */
 	bool bShowReferences;
