Follow up for Control Rig: Rename rigid body element to physics element

#rnx

diff --git a/Engine/Plugins/Animation/ControlRig/Source/ControlRig/Private/ControlRig.cpp b/Engine/Plugins/Animation/ControlRig/Source/ControlRig/Private/ControlRig.cpp
--- a/Engine/Plugins/Animation/ControlRig/Source/ControlRig/Private/ControlRig.cpp
+++ b/Engine/Plugins/Animation/ControlRig/Source/ControlRig/Private/ControlRig.cpp
@@ -3580,10 +3580,12 @@ const FRigPhysicsSolverDescription* UControlRig::FindPhysicsSolverByName(const F
 
 FRigPhysicsSolverID UControlRig::AddPhysicsSolver(FName InName, bool bSetupUndo, bool bPrintPythonCommand)
 {
+#if WITH_EDITOR
 	if(CVarControlRigHierarchyEnablePhysics.GetValueOnAnyThread() == false)
 	{
 		return FRigPhysicsSolverID();
 	}
+#endif
 
 	if(InName.IsNone())
 	{
diff --git a/Engine/Plugins/Animation/ControlRig/Source/ControlRig/Private/Rigs/RigHierarchyController.cpp b/Engine/Plugins/Animation/ControlRig/Source/ControlRig/Private/Rigs/RigHierarchyController.cpp
--- a/Engine/Plugins/Animation/ControlRig/Source/ControlRig/Private/Rigs/RigHierarchyController.cpp
+++ b/Engine/Plugins/Animation/ControlRig/Source/ControlRig/Private/Rigs/RigHierarchyController.cpp
@@ -493,10 +493,12 @@ FRigElementKey URigHierarchyController::AddPhysicsElement(FName InName, FRigElem
 		return FRigElementKey();
 	}
 
+#if WITH_EDITOR
 	if(CVarControlRigHierarchyEnablePhysics.GetValueOnAnyThread() == false)
 	{
 		return FRigElementKey();
 	}
+#endif
 
 	URigHierarchy* Hierarchy = GetHierarchy();
 
