Control Rig: Fix GetChildren and GetSiblings nodes not getting active children
#rnx

diff --git a/Engine/Plugins/Animation/ControlRig/Source/ControlRig/Private/Rigs/RigHierarchy.cpp b/Engine/Plugins/Animation/ControlRig/Source/ControlRig/Private/Rigs/RigHierarchy.cpp
--- a/Engine/Plugins/Animation/ControlRig/Source/ControlRig/Private/Rigs/RigHierarchy.cpp
+++ b/Engine/Plugins/Animation/ControlRig/Source/ControlRig/Private/Rigs/RigHierarchy.cpp
@@ -1877,6 +1877,54 @@ TArray<FRigElementKey> URigHierarchy::GetChildren(FRigElementKey InKey, bool bRe
 	}
 }
 
+FRigBaseElementChildrenArray URigHierarchy::GetActiveChildren(const FRigBaseElement* InElement, bool bRecursive) const
+{
+	TArray<FRigBaseElement*> Children;
+
+	TArray<const FRigBaseElement*> ToProcess;
+	ToProcess.Push(InElement);
+
+	while (!ToProcess.IsEmpty())
+	{
+		const FRigBaseElement* CurrentParent = ToProcess.Pop();
+		TArray<FRigBaseElement*> CurrentChildren;
+		if (CurrentParent)
+		{
+			CurrentChildren = GetChildren(CurrentParent);
+			if (CurrentParent->GetKey() == URigHierarchy::GetWorldSpaceReferenceKey())
+			{
+				CurrentChildren.Append(GetFilteredElements<FRigBaseElement>([this](const FRigBaseElement* Element)
+				{
+					return !Element->GetKey().IsTypeOf(ERigElementType::Reference) && GetActiveParent(Element) == nullptr;
+				}));
+			}
+		}
+		else
+		{
+			CurrentChildren = GetFilteredElements<FRigBaseElement>([this](const FRigBaseElement* Element)
+			{
+				return !Element->GetKey().IsTypeOf(ERigElementType::Reference) && GetActiveParent(Element) == nullptr;
+			});
+		}
+		const FRigElementKey ParentKey = CurrentParent ? CurrentParent->GetKey() : URigHierarchy::GetWorldSpaceReferenceKey();
+		for (const FRigBaseElement* Child : CurrentChildren)
+		{
+			const FRigBaseElement* ThisParent = GetActiveParent(Child);
+			const FRigElementKey ThisParentKey = ThisParent ? ThisParent->GetKey() : URigHierarchy::GetWorldSpaceReferenceKey();
+			if (ThisParentKey == ParentKey)
+			{
+				Children.Add(const_cast<FRigBaseElement*>(Child));
+				if (bRecursive)
+				{
+					ToProcess.Push(Child);
+				}
+			}
+		}
+	}
+	
+	return FRigBaseElementChildrenArray(Children);
+}
+
 TArray<int32> URigHierarchy::GetChildren(int32 InIndex, bool bRecursive) const
 {
 	LLM_SCOPE_BYNAME(TEXT("Animation/ControlRig"));
diff --git a/Engine/Plugins/Animation/ControlRig/Source/ControlRig/Private/Units/Execution/RigUnit_Collection.cpp b/Engine/Plugins/Animation/ControlRig/Source/ControlRig/Private/Units/Execution/RigUnit_Collection.cpp
--- a/Engine/Plugins/Animation/ControlRig/Source/ControlRig/Private/Units/Execution/RigUnit_Collection.cpp
+++ b/Engine/Plugins/Animation/ControlRig/Source/ControlRig/Private/Units/Execution/RigUnit_Collection.cpp
@@ -88,7 +88,7 @@ FRigUnit_CollectionNameSearchArray_Execute()
 
 FRigUnit_CollectionChildren_Execute()
 {
-	FRigUnit_CollectionChildrenArray::StaticExecute(ExecuteContext, Parent, bIncludeParent, bRecursive, TypeToSearch, Collection.Keys);
+	FRigUnit_CollectionChildrenArray::StaticExecute(ExecuteContext, Parent, bIncludeParent, bRecursive, true, TypeToSearch, Collection.Keys);
 }
 
 FRigVMStructUpgradeInfo FRigUnit_CollectionChildren::GetUpgradeInfo() const
@@ -113,32 +113,58 @@ FRigUnit_CollectionChildrenArray_Execute()
 		Items.Reset();
 		return;
 	}
-	
-	uint32 Hash = GetTypeHash(StaticStruct()) + ExecuteContext.Hierarchy->GetTopologyVersion() * 17;
-	Hash = HashCombine(Hash, GetTypeHash(ExecuteContext.Hierarchy->GetResolvedTarget(Parent)));
-	Hash = HashCombine(Hash, bRecursive ? 2 : 0);
-	Hash = HashCombine(Hash, bIncludeParent ? 1 : 0);
-	Hash = HashCombine(Hash, (int32)TypeToSearch * 8);
 
-	FRigElementKeyCollection Collection;
-	if(const FRigElementKeyCollection* Cache = ExecuteContext.Hierarchy->FindCachedCollection(Hash))
-	{
-		Collection = *Cache;
-	}
-	else
+	if (bDefaultChildren)
 	{
-		Collection = FRigElementKeyCollection::MakeFromChildren(ExecuteContext.Hierarchy, Parent, bRecursive, bIncludeParent, (uint8)TypeToSearch);
-		if (Collection.IsEmpty())
+		uint32 Hash = GetTypeHash(StaticStruct()) + ExecuteContext.Hierarchy->GetTopologyVersion() * 17;
+		Hash = HashCombine(Hash, GetTypeHash(ExecuteContext.Hierarchy->GetResolvedTarget(Parent)));
+		Hash = HashCombine(Hash, bRecursive ? 2 : 0);
+		Hash = HashCombine(Hash, bIncludeParent ? 1 : 0);
+		Hash = HashCombine(Hash, (int32)TypeToSearch * 8);
+
+		FRigElementKeyCollection Collection;
+		if(const FRigElementKeyCollection* Cache = ExecuteContext.Hierarchy->FindCachedCollection(Hash))
+		{
+			Collection = *Cache;
+		}
+		else
 		{
-			if (ExecuteContext.Hierarchy->GetIndex(Parent) == INDEX_NONE)
+			Collection = FRigElementKeyCollection::MakeFromChildren(ExecuteContext.Hierarchy, Parent, bRecursive, bIncludeParent, (uint8)TypeToSearch);
+			if (Collection.IsEmpty())
 			{
-				UE_CONTROLRIG_RIGUNIT_REPORT_WARNING(TEXT("Parent '%s' is not valid."), *Parent.ToString());
+				if (ExecuteContext.Hierarchy->GetIndex(Parent) == INDEX_NONE)
+				{
+					UE_CONTROLRIG_RIGUNIT_REPORT_WARNING(TEXT("Parent '%s' is not valid."), *Parent.ToString());
+				}
 			}
+			ExecuteContext.Hierarchy->AddCachedCollection(Hash, Collection);
 		}
-		ExecuteContext.Hierarchy->AddCachedCollection(Hash, Collection);
+
+		Items = Collection.Keys;
 	}
+	else
+	{
+		const URigHierarchy* Hierarchy = ExecuteContext.Hierarchy;
+		Items.Reset();
 
-	Items = Collection.Keys;
+		if (bIncludeParent)
+		{
+			Items.Add(Parent);
+		}
+
+		FRigBaseElementChildrenArray Children = Hierarchy->GetActiveChildren(Hierarchy->Find(Parent), bRecursive);
+
+		Children = Children.FilterByPredicate([TypeToSearch](const FRigBaseElement* Child)
+		{
+			return ((uint8)Child->GetType() & (uint8)TypeToSearch) != 0;
+		});
+
+		Items.Reserve(Items.Num() + Children.Num());
+		for (FRigBaseElement* Child : Children)
+		{
+			Items.Add(Child->GetKey());
+		}
+	}
 }
 
 FRigUnit_CollectionGetAll_Execute()
diff --git a/Engine/Plugins/Animation/ControlRig/Source/ControlRig/Private/Units/Execution/RigUnit_Hierarchy.cpp b/Engine/Plugins/Animation/ControlRig/Source/ControlRig/Private/Units/Execution/RigUnit_Hierarchy.cpp
--- a/Engine/Plugins/Animation/ControlRig/Source/ControlRig/Private/Units/Execution/RigUnit_Hierarchy.cpp
+++ b/Engine/Plugins/Animation/ControlRig/Source/ControlRig/Private/Units/Execution/RigUnit_Hierarchy.cpp
@@ -157,7 +157,7 @@ FRigVMStructUpgradeInfo FRigUnit_HierarchyGetChildren::GetUpgradeInfo() const
 
 FRigUnit_HierarchyGetSiblings_Execute()
 {
-	FRigUnit_HierarchyGetSiblingsItemArray::StaticExecute(ExecuteContext, Item, bIncludeItem, Siblings.Keys, CachedItem, CachedSiblings);
+	FRigUnit_HierarchyGetSiblingsItemArray::StaticExecute(ExecuteContext, Item, bIncludeItem, true, Siblings.Keys, CachedItem, CachedSiblings);
 }
 
 FRigVMStructUpgradeInfo FRigUnit_HierarchyGetSiblings::GetUpgradeInfo() const
@@ -172,37 +172,60 @@ FRigUnit_HierarchyGetSiblingsItemArray_Execute()
 {
     DECLARE_SCOPE_HIERARCHICAL_COUNTER_RIGUNIT()
 
-	if(!CachedItem.IsIdentical(Item, ExecuteContext.Hierarchy))
+	if (bDefaultSiblings)
 	{
-		CachedSiblings.Reset();
-
-		if(CachedItem.UpdateCache(Item, ExecuteContext.Hierarchy))
+		if(!CachedItem.IsIdentical(Item, ExecuteContext.Hierarchy))
 		{
-			TArray<FRigElementKey> Keys;
+			CachedSiblings.Reset();
 
-			FRigElementKey Parent = ExecuteContext.Hierarchy->GetFirstParent(CachedItem.GetResolvedKey());
-			if(Parent.IsValid())
+			if(CachedItem.UpdateCache(Item, ExecuteContext.Hierarchy))
 			{
-				TArray<FRigElementKey> Children = ExecuteContext.Hierarchy->GetChildren(Parent, false);
-				for(FRigElementKey Child : Children)
+				TArray<FRigElementKey> Keys;
+
+				FRigElementKey Parent = ExecuteContext.Hierarchy->GetFirstParent(CachedItem.GetResolvedKey());
+				if(Parent.IsValid())
 				{
-					if(bIncludeItem || Child != Item)
+					TArray<FRigElementKey> Children = ExecuteContext.Hierarchy->GetChildren(Parent, false);
+					for(FRigElementKey Child : Children)
 					{
-						Keys.Add(Child);
+						if(bIncludeItem || Child != Item)
+						{
+							Keys.Add(Child);
+						}
 					}
 				}
-			}
 
-			if(Keys.Num() == 0 && bIncludeItem)
-			{
-				Keys.Add(Item);
-			}
+				if(Keys.Num() == 0 && bIncludeItem)
+				{
+					Keys.Add(Item);
+				}
 
-			CachedSiblings = FRigElementKeyCollection(Keys);
+				CachedSiblings = FRigElementKeyCollection(Keys);
+			}
 		}
+
+		Siblings = CachedSiblings.Keys;
 	}
+	else
+	{
+		Siblings.Reset();
+		
+		const URigHierarchy* Hierarchy = ExecuteContext.Hierarchy;
+		FRigElementKey ParentKey = Hierarchy->GetActiveParent(Item, false);
 
-	Siblings = CachedSiblings.Keys;
+		if (!ParentKey.IsValid())
+		{
+			ParentKey = URigHierarchy::GetWorldSpaceReferenceKey();
+		}
+
+		FRigBaseElementChildrenArray SiblingsPtrs = Hierarchy->GetActiveChildren(Hierarchy->Find(ParentKey), false);
+		Algo::Transform(SiblingsPtrs, Siblings, [](const FRigBaseElement* Element) { return Element->GetKey(); });
+
+		if (!bIncludeItem)
+		{
+			Siblings.Remove(Item);
+		}
+	}
 }
 
 FRigUnit_HierarchyGetChainItemArray_Execute()
diff --git a/Engine/Plugins/Animation/ControlRig/Source/ControlRig/Public/Rigs/RigHierarchy.h b/Engine/Plugins/Animation/ControlRig/Source/ControlRig/Public/Rigs/RigHierarchy.h
--- a/Engine/Plugins/Animation/ControlRig/Source/ControlRig/Public/Rigs/RigHierarchy.h
+++ b/Engine/Plugins/Animation/ControlRig/Source/ControlRig/Public/Rigs/RigHierarchy.h
@@ -3024,6 +3024,14 @@ public:
 	UFUNCTION(BlueprintCallable, Category = URigHierarchy)
 	TArray<FRigElementKey> GetChildren(FRigElementKey InKey, bool bRecursive = false) const;
 
+	/**
+	 * Returns the child elements of a given element key
+	 * @param InKey The key of the element to retrieve the children for
+	 * @param bRecursive If set to true grand-children will also be returned etc
+	 * @return Returns the child elements
+	 */
+	FRigBaseElementChildrenArray GetActiveChildren(const FRigBaseElement* InElement, bool bRecursive = false) const;
+
 	/**
 	 * Returns the child elements of a given element index
 	 * @param InIndex The index of the element to retrieve the children for
diff --git a/Engine/Plugins/Animation/ControlRig/Source/ControlRig/Public/Units/Execution/RigUnit_Collection.h b/Engine/Plugins/Animation/ControlRig/Source/ControlRig/Public/Units/Execution/RigUnit_Collection.h
--- a/Engine/Plugins/Animation/ControlRig/Source/ControlRig/Public/Units/Execution/RigUnit_Collection.h
+++ b/Engine/Plugins/Animation/ControlRig/Source/ControlRig/Public/Units/Execution/RigUnit_Collection.h
@@ -198,6 +198,7 @@ struct CONTROLRIG_API FRigUnit_CollectionChildrenArray : public FRigUnit_Collect
 		Parent = FRigElementKey(NAME_None, ERigElementType::Bone);
 		bIncludeParent = false;
 		bRecursive = false;
+		bDefaultChildren = true;
 		TypeToSearch = ERigElementType::All;
 	}
 
@@ -213,6 +214,11 @@ struct CONTROLRIG_API FRigUnit_CollectionChildrenArray : public FRigUnit_Collect
 	UPROPERTY(meta = (Input))
 	bool bRecursive;
 
+	/** When true, it will return all children, regardless of whether the parent is active or not.
+	 * When false, will return only the children which are influenced by this parent */
+	UPROPERTY(meta = (Input))
+	bool bDefaultChildren;
+
 	UPROPERTY(meta = (Input))
 	ERigElementType TypeToSearch;
 
diff --git a/Engine/Plugins/Animation/ControlRig/Source/ControlRig/Public/Units/Execution/RigUnit_Hierarchy.h b/Engine/Plugins/Animation/ControlRig/Source/ControlRig/Public/Units/Execution/RigUnit_Hierarchy.h
--- a/Engine/Plugins/Animation/ControlRig/Source/ControlRig/Public/Units/Execution/RigUnit_Hierarchy.h
+++ b/Engine/Plugins/Animation/ControlRig/Source/ControlRig/Public/Units/Execution/RigUnit_Hierarchy.h
@@ -39,6 +39,7 @@ struct CONTROLRIG_API FRigUnit_HierarchyGetParent : public FRigUnit_HierarchyBas
 	UPROPERTY(meta = (Input, ExpandByDefault))
 	FRigElementKey Child;
 
+	/** When true, it will return the default parent, regardless of whether the parent incluences the element or not  */
 	UPROPERTY(meta = (Input))
 	bool bDefaultParent;
 
@@ -242,6 +243,7 @@ struct CONTROLRIG_API FRigUnit_HierarchyGetSiblingsItemArray : public FRigUnit_H
 		CachedItem = FCachedRigElement();
 		CachedSiblings = FRigElementKeyCollection();
 		bIncludeItem = false;
+		bDefaultSiblings = true;
 	}
 
 	RIGVM_METHOD()
@@ -253,6 +255,11 @@ struct CONTROLRIG_API FRigUnit_HierarchyGetSiblingsItemArray : public FRigUnit_H
 	UPROPERTY(meta = (Input))
 	bool bIncludeItem;
 
+	/** When true, it will return all siblings, regardless of whether the parent is active or not.
+	 * When false, will return only the siblings which are influenced by the same parent */
+	UPROPERTY(meta = (Input))
+	bool bDefaultSiblings;
+
 	UPROPERTY(meta = (Output))
 	TArray<FRigElementKey> Siblings;
 
