Control Rig: Fix flattening of hierarchy
#rnx

diff --git a/Engine/Plugins/Animation/ControlRig/Source/ControlRigEditor/Private/Editor/SRigHierarchyTreeView.cpp b/Engine/Plugins/Animation/ControlRig/Source/ControlRigEditor/Private/Editor/SRigHierarchyTreeView.cpp
--- a/Engine/Plugins/Animation/ControlRig/Source/ControlRigEditor/Private/Editor/SRigHierarchyTreeView.cpp
+++ b/Engine/Plugins/Animation/ControlRig/Source/ControlRigEditor/Private/Editor/SRigHierarchyTreeView.cpp
@@ -673,7 +673,7 @@ bool SRigHierarchyTreeView::AddElement(const FRigBaseElement* InElement)
 				AddConnectorResolveWarningTag(ElementMap.FindChecked(InElement->GetKey()), InElement, Hierarchy);
 			}
 			
-			FRigElementKey ParentKey = Hierarchy->GetActiveParent(InElement->GetKey());
+			FRigElementKey ParentKey = Hierarchy->GetFirstParent(InElement->GetKey());
 			if(InElement->GetType() == ERigElementType::Connector)
 			{
 				ParentKey = Delegates.GetResolvedKey(InElement->GetKey());
@@ -838,7 +838,7 @@ void SRigHierarchyTreeView::RefreshTreeView(bool bRebuildContent)
 			Hierarchy->Traverse([&](FRigBaseElement* Element, bool& bContinue)
 			{
 				int32& Depth = ElementDepth.Add(Element, 0);
-				if(const FRigBaseElement* ParentElement = Hierarchy->GetActiveParent(Element))
+				if(const FRigBaseElement* ParentElement = Hierarchy->GetFirstParent(Element))
 				{
 					if (int32* ParentDepth = ElementDepth.Find(ParentElement))
 					{
