[Control Rig] fixed a regression in loading CR containing get anim attribute node. When new types are registered, we have to update the templates. It used to be that we call onnewargumenttype for all arguments using category, but there was a perf optimization to only do it for the primary argument.

diff --git a/Engine/Plugins/Animation/ControlRig/Source/ControlRig/Private/Units/Animation/RigUnit_AnimAttribute.cpp b/Engine/Plugins/Animation/ControlRig/Source/ControlRig/Private/Units/Animation/RigUnit_AnimAttribute.cpp
--- a/Engine/Plugins/Animation/ControlRig/Source/ControlRig/Private/Units/Animation/RigUnit_AnimAttribute.cpp
+++ b/Engine/Plugins/Animation/ControlRig/Source/ControlRig/Private/Units/Animation/RigUnit_AnimAttribute.cpp
@@ -386,10 +386,13 @@ FRigVMTemplateTypeMap FRigDispatch_GetAnimAttribute::OnNewArgumentType(const FNa
 {
 	FRigVMTemplateTypeMap Types;
 
-	// similar pattern to URigVMArrayNode's FRigVMTemplate_NewArgumentTypeDelegate to avoid double registration
-	// this is needed since a single type is called for both DefaultArg and ValueArg but we should only
-	// register one permutation
-	if (InArgumentName == ValueArgName)
+	// similar pattern to FRigVMDispatch_ArrayGetAtIndex::OnNewArgumentType
+	// we used to need the if check below to avoid double registration for both default and value arguments, but since CL 29742624
+	// only the primary argument is used to add new permutations and thus this check becomes optional.
+	// primary argument is defined as the first argument that uses category, see FRigVMTemplate::UpdateArgumentTypes(), and in this case
+	// default is the first argument instead of the value
+	
+	if (ensure(InArgumentName == DefaultArgName))
 	{
 		if (IsTypeSupported(InTypeIndex))
 		{
