Deprecate UContentBrowserAssetContextMenuContext::GetSelectedObjects and reworked code that use it.

When possible the existing code as been change to not load the selected assets on the contextual menu opening.

diff --git a/Engine/Plugins/Animation/ControlRig/Source/ControlRigEditor/Private/ControlRigBlueprintActions.cpp b/Engine/Plugins/Animation/ControlRig/Source/ControlRigEditor/Private/ControlRigBlueprintActions.cpp
--- a/Engine/Plugins/Animation/ControlRig/Source/ControlRigEditor/Private/ControlRigBlueprintActions.cpp
+++ b/Engine/Plugins/Animation/ControlRig/Source/ControlRigEditor/Private/ControlRigBlueprintActions.cpp
@@ -13,6 +13,7 @@
 #include "Styling/AppStyle.h" 
 #include "Subsystems/AssetEditorSubsystem.h"
 
+#include "AssetRegistry/AssetData.h"
 #include "AssetRegistry/AssetRegistryModule.h"
 #include "ToolMenus.h"
 #include "ContentBrowserMenuContexts.h"
@@ -155,8 +156,7 @@ void FControlRigBlueprintActions::ExtendSketalMeshToolMenu()
 			UContentBrowserAssetContextMenuContext* Context = InSection.FindContext<UContentBrowserAssetContextMenuContext>();
 			if (Context)
 			{
-				TArray<UObject*> SelectedObjects = Context->GetSelectedObjects();
-				if (SelectedObjects.Num() > 0)
+				if (Context->SelectedAssets.Num() > 0)
 				{
 					static constexpr bool bModularRig = true;
 					InSection.AddMenuEntry(
@@ -164,8 +164,20 @@ void FControlRigBlueprintActions::ExtendSketalMeshToolMenu()
 						LOCTEXT("CreateControlRig", "Control Rig"),
 						LOCTEXT("CreateControlRig_ToolTip", "Creates a control rig and preconfigures it for this asset"),
 						FSlateIcon(FRigVMEditorStyle::Get().GetStyleSetName(), "RigVM", "RigVM.Unit"),
-						FExecuteAction::CreateLambda([SelectedObjects]()
+						FExecuteAction::CreateLambda([InSelectedAssets = Context->SelectedAssets]()
 						{
+							TArray<UObject*> SelectedObjects;
+							SelectedObjects.Reserve(InSelectedAssets.Num());
+
+							for (const FAssetData& Asset : InSelectedAssets)
+							{
+								if (UObject* LoadedAsset = Asset.GetAsset())
+								{
+									SelectedObjects.Add(LoadedAsset);
+								}
+							}
+
+
 							for (UObject* SelectedObject : SelectedObjects)
 							{
 								CreateControlRigFromSkeletalMeshOrSkeleton(SelectedObject, !bModularRig);
@@ -177,8 +189,19 @@ void FControlRigBlueprintActions::ExtendSketalMeshToolMenu()
 						LOCTEXT("CreateModularRig", "Modular Rig"),
 						LOCTEXT("CreateModularRig_ToolTip", "Creates a modular rig and preconfigures it for this asset"),
 						FSlateIcon(FRigVMEditorStyle::Get().GetStyleSetName(), "RigVM", "RigVM.Unit"),
-						FExecuteAction::CreateLambda([SelectedObjects]()
+						FExecuteAction::CreateLambda([InSelectedAssets = Context->SelectedAssets]()
 						{
+							TArray<UObject*> SelectedObjects;
+							SelectedObjects.Reserve(InSelectedAssets.Num());
+
+							for (const FAssetData& Asset : InSelectedAssets)
+							{
+								if (UObject* LoadedAsset = Asset.GetAsset())
+								{
+									SelectedObjects.Add(LoadedAsset);
+								}
+							}
+
 							for (UObject* SelectedObject : SelectedObjects)
 							{
 								CreateControlRigFromSkeletalMeshOrSkeleton(SelectedObject, bModularRig);
diff --git a/Engine/Plugins/Animation/IKRig/Source/IKRigEditor/Private/RetargetEditor/AssetTypeActions_IKRetargeter.cpp b/Engine/Plugins/Animation/IKRig/Source/IKRigEditor/Private/RetargetEditor/AssetTypeActions_IKRetargeter.cpp
--- a/Engine/Plugins/Animation/IKRig/Source/IKRigEditor/Private/RetargetEditor/AssetTypeActions_IKRetargeter.cpp
+++ b/Engine/Plugins/Animation/IKRig/Source/IKRigEditor/Private/RetargetEditor/AssetTypeActions_IKRetargeter.cpp
@@ -56,16 +56,26 @@ void FAssetTypeActions_IKRetargeter::ExtendIKRigMenuToMakeRetargeter()
 		UContentBrowserAssetContextMenuContext* Context = InSection.FindContext<UContentBrowserAssetContextMenuContext>();
 		if (Context)
 		{
-			TArray<UObject*> SelectedObjects = Context->GetSelectedObjects();
-			if (SelectedObjects.Num() > 0)
+			if (Context->SelectedAssets.Num() > 0)
 			{
 				InSection.AddMenuEntry(
 					"CreateRetargeter",
 					LOCTEXT("CreateIKRetargeter", "Create IK Retargeter"),
 					LOCTEXT("CreateIKRetargeter_ToolTip", "Creates an IK Retargeter using this IK Rig as the source."),
 					FSlateIcon(FIKRigEditorStyle::Get().GetStyleSetName(), "IKRig", "ClassIcon.IKRigDefinition"),
-					FExecuteAction::CreateLambda([SelectedObjects]()
+					FExecuteAction::CreateLambda([InSelectedAssets = Context->SelectedAssets]()
 					{
+						TArray<UObject*> SelectedObjects;
+						SelectedObjects.Reserve(InSelectedAssets.Num());
+
+						for (const FAssetData& Asset : InSelectedAssets)
+						{
+							if (UObject* LoadedAsset = Asset.GetAsset())
+							{
+								SelectedObjects.Add(LoadedAsset);
+							}
+						}
+
 						for (UObject* SelectedObject : SelectedObjects)
 						{
 							CreateNewIKRetargeterFromIKRig(SelectedObject);
diff --git a/Engine/Plugins/Animation/IKRig/Source/IKRigEditor/Private/RigEditor/AssetTypeActions_IKRigDefinition.cpp b/Engine/Plugins/Animation/IKRig/Source/IKRigEditor/Private/RigEditor/AssetTypeActions_IKRigDefinition.cpp
--- a/Engine/Plugins/Animation/IKRig/Source/IKRigEditor/Private/RigEditor/AssetTypeActions_IKRigDefinition.cpp
+++ b/Engine/Plugins/Animation/IKRig/Source/IKRigEditor/Private/RigEditor/AssetTypeActions_IKRigDefinition.cpp
@@ -9,6 +9,7 @@
 #include "RigEditor/IKRigEditorStyle.h"
 #include "Engine/SkeletalMesh.h"
 #include "ToolMenus.h"
+#include "AssetRegistry/AssetData.h"
 
 #define LOCTEXT_NAMESPACE "AssetTypeActions"
 
@@ -38,16 +39,26 @@ void FAssetTypeActions_IKRigDefinition::ExtendSkeletalMeshMenuToMakeIKRig()
 		UContentBrowserAssetContextMenuContext* Context = InSection.FindContext<UContentBrowserAssetContextMenuContext>();
 		if (Context)
 		{
-			TArray<UObject*> SelectedObjects = Context->GetSelectedObjects();
-			if (SelectedObjects.Num() > 0)
+			if (Context->SelectedAssets.Num() > 0)
 			{
 				InSection.AddMenuEntry(
 					"CreateIKRig",
 					LOCTEXT("CreateIKRig", "IK Rig"),
 					LOCTEXT("CreateIKRig_ToolTip", "Creates an IK rig for this skeletal mesh."),
 					FSlateIcon(FIKRigEditorStyle::Get().GetStyleSetName(), "IKRig", "ClassIcon.IKRigDefinition"),
-					FExecuteAction::CreateLambda([SelectedObjects]()
+					FExecuteAction::CreateLambda([InSelectedAssets = Context->SelectedAssets]()
 					{
+						TArray<UObject*> SelectedObjects;
+						SelectedObjects.Reserve(InSelectedAssets.Num());
+
+						for (const FAssetData& Asset : InSelectedAssets)
+						{
+							if (UObject* LoadedAsset = Asset.GetAsset())
+							{
+								SelectedObjects.Add(LoadedAsset);
+							}
+						}
+
 						for (UObject* SelectedObject : SelectedObjects)
 						{
 							CreateNewIKRigFromSkeletalMesh(SelectedObject);
diff --git a/Engine/Plugins/Enterprise/DatasmithImporter/Source/DatasmithImporter/Private/DatasmithImporterModule.cpp b/Engine/Plugins/Enterprise/DatasmithImporter/Source/DatasmithImporter/Private/DatasmithImporterModule.cpp
--- a/Engine/Plugins/Enterprise/DatasmithImporter/Source/DatasmithImporter/Private/DatasmithImporterModule.cpp
+++ b/Engine/Plugins/Enterprise/DatasmithImporter/Source/DatasmithImporter/Private/DatasmithImporterModule.cpp
@@ -268,24 +268,21 @@ void FDatasmithImporterModule::AddDataprepMenuEntryForDatasmithSceneAsset()
 					{
 						if (ContentBrowserMenuContext->CommonClass == UDatasmithScene::StaticClass())
 						{
-							TArray<UObject*> SelectedObjects = ContentBrowserMenuContext->GetSelectedObjects();
+							TArray<UDatasmithScene*> SelectedObjects = ContentBrowserMenuContext->LoadSelectedObjects<UDatasmithScene>({});
 							DataprepAssetInterfacesPtr.Reserve( SelectedObjects.Num() );
-							for (UObject* SelectedObject : SelectedObjects)
+							for (UDatasmithScene* SelectedDatasmithScene : SelectedObjects)
 							{
-								if (UDatasmithScene* SelectedDatasmithScene = Cast<UDatasmithScene>(SelectedObject))
+								if (UDataprepAssetUserData* DataprepAssetUserData = SelectedDatasmithScene->GetAssetUserData<UDataprepAssetUserData>())
 								{
-									if (UDataprepAssetUserData* DataprepAssetUserData = SelectedDatasmithScene->GetAssetUserData<UDataprepAssetUserData>())
+									if (UDataprepAssetInterface* DataprepAsset = DataprepAssetUserData->DataprepAssetPtr.LoadSynchronous())
 									{
-										if (UDataprepAssetInterface* DataprepAsset = DataprepAssetUserData->DataprepAssetPtr.LoadSynchronous())
+										if (UDatasmithConsumer* DatasmithConsumer = Cast<UDatasmithConsumer>(DataprepAsset->GetConsumer()))
 										{
-											if (UDatasmithConsumer* DatasmithConsumer = Cast<UDatasmithConsumer>(DataprepAsset->GetConsumer()))
+											if (DatasmithConsumer->GetDatasmithScene() == SelectedDatasmithScene)
 											{
-												if (DatasmithConsumer->GetDatasmithScene() == SelectedDatasmithScene)
-												{
-													// A Dataprep asset was found and it will regenerate this scene
-													DataprepAssetInterfacesPtr.Emplace( DataprepAsset );
-													continue;
-												}
+												// A Dataprep asset was found and it will regenerate this scene
+												DataprepAssetInterfacesPtr.Emplace( DataprepAsset );
+												continue;
 											}
 										}
 									}
diff --git a/Engine/Source/Editor/ContentBrowser/Public/ContentBrowserMenuContexts.h b/Engine/Source/Editor/ContentBrowser/Public/ContentBrowserMenuContexts.h
--- a/Engine/Source/Editor/ContentBrowser/Public/ContentBrowserMenuContexts.h
+++ b/Engine/Source/Editor/ContentBrowser/Public/ContentBrowserMenuContexts.h
@@ -47,9 +47,6 @@ public:
 	
 	UPROPERTY()
 	TObjectPtr<const UAssetDefinition> CommonAssetDefinition;
-	
-	UE_DEPRECATED(5.1, "Use SelectedAssets now, this field will not contain any objects.  You should call LoadSelectedObjects() based on what you need.")
-	TArray<TWeakObjectPtr<UObject>> SelectedObjects;
 
 	/**
 	 * The currently selected assets in the content browser.
@@ -72,7 +69,7 @@ public:
 	UPROPERTY(BlueprintReadOnly, Category = ContentBrowser)
 	bool bContainsUnsupportedAssets = true;
 
-	//UE_DEPRECATED(5.2, "GetSelectedObjects has been deprecated.  We no longer automatically load assets on right click.  Please use SelectedAssets and determine whatever you need for your context menu options without actually loading the assets.  When you finally need all or a subset of the selected assets use LoadSelectedAssets or LoadSelectedAssetsIf")
+	UE_DEPRECATED(5.5, "GetSelectedObjects has been deprecated.  We no longer automatically load assets on right click.  Please use SelectedAssets and determine whatever you need for your context menu options without actually loading the assets.  When you finally need all or a subset of the selected assets use LoadSelectedAssets or LoadSelectedAssetsIf")
 	UFUNCTION(BlueprintCallable, Category="Tool Menus", meta=(DeprecatedFunction, DeprecationMessage = "GetSelectedObjects has been deprecated.  We no longer automatically load assets on right click.  If you can work without loading the assets, please use SelectedAssets.  Otherwise call LoadSelectedObjects"))
 	TArray<UObject*> GetSelectedObjects() const
 	{
