Control Rig: Add optional item filter to draw hierarchy and draw pose

#rnx

diff --git a/Engine/Plugins/Animation/ControlRig/Source/ControlRig/Private/Units/Debug/RigUnit_DebugHierarchy.cpp b/Engine/Plugins/Animation/ControlRig/Source/ControlRig/Private/Units/Debug/RigUnit_DebugHierarchy.cpp
--- a/Engine/Plugins/Animation/ControlRig/Source/ControlRig/Private/Units/Debug/RigUnit_DebugHierarchy.cpp
+++ b/Engine/Plugins/Animation/ControlRig/Source/ControlRig/Private/Units/Debug/RigUnit_DebugHierarchy.cpp
@@ -13,7 +13,7 @@ FRigUnit_DebugHierarchy_Execute()
 	URigHierarchy* Hierarchy = ExecuteContext.Hierarchy;
 	if (Hierarchy)
 	{
-		DrawHierarchy(ExecuteContext, WorldOffset, Hierarchy, EControlRigDrawHierarchyMode::Axes, Scale, Color, Thickness, nullptr);
+		DrawHierarchy(ExecuteContext, WorldOffset, Hierarchy, EControlRigDrawHierarchyMode::Axes, Scale, Color, Thickness, nullptr, &Items);
 	}
 }
 
@@ -24,11 +24,11 @@ FRigUnit_DebugPose_Execute()
 	URigHierarchy* Hierarchy = ExecuteContext.Hierarchy;
 	if (Hierarchy)
 	{
-		FRigUnit_DebugHierarchy::DrawHierarchy(ExecuteContext, WorldOffset, Hierarchy, EControlRigDrawHierarchyMode::Axes, Scale, Color, Thickness, &Pose);
+		FRigUnit_DebugHierarchy::DrawHierarchy(ExecuteContext, WorldOffset, Hierarchy, EControlRigDrawHierarchyMode::Axes, Scale, Color, Thickness, &Pose, &Items);
 	}
 }
 
-void FRigUnit_DebugHierarchy::DrawHierarchy(const FRigVMExecuteContext& InContext, const FTransform& WorldOffset, URigHierarchy* Hierarchy, EControlRigDrawHierarchyMode::Type Mode, float Scale, const FLinearColor& Color, float Thickness, const FRigPose* InPose)
+void FRigUnit_DebugHierarchy::DrawHierarchy(const FRigVMExecuteContext& InContext, const FTransform& WorldOffset, URigHierarchy* Hierarchy, EControlRigDrawHierarchyMode::Type Mode, float Scale, const FLinearColor& Color, float Thickness, const FRigPose* InPose, const TArrayView<const FRigElementKey>* InItems)
 {
 	FRigVMDrawInterface* DrawInterface = InContext.GetDrawInterface();
 	if(DrawInterface == nullptr)
@@ -41,6 +41,18 @@ void FRigUnit_DebugHierarchy::DrawHierarchy(const FRigVMExecuteContext& InContex
 		return;
 	}
 
+	TMap<const FRigBaseElement*,  bool> ElementMap;
+	if(InItems)
+	{
+		for(const FRigElementKey& Item : *InItems)
+		{
+			if(const FRigBaseElement* Element = Hierarchy->Find(Item))
+			{
+				ElementMap.Add(Element, true);
+			}
+		}
+	}
+
 	switch (Mode)
 	{
 		case EControlRigDrawHierarchyMode::Axes:
@@ -76,6 +88,12 @@ void FRigUnit_DebugHierarchy::DrawHierarchy(const FRigVMExecuteContext& InContex
 					Transform = Hierarchy->GetTransform(InElement, ERigTransformType::CurrentGlobal);
 					return bValid;
 				};
+
+				// use the optional filter
+				if(!ElementMap.IsEmpty() && !ElementMap.Contains(Child))
+				{
+					return true;
+				}
 				
 				FTransform Transform = FTransform::Identity;
 				if(!GetTransformLambda(Child, Transform))
@@ -101,6 +119,12 @@ void FRigUnit_DebugHierarchy::DrawHierarchy(const FRigVMExecuteContext& InContex
 				{
 					if(FRigTransformElement* ParentTransformElement = Cast<FRigTransformElement>(Parents[ParentIndex]))
 					{
+						// use the optional filter
+						if(!ElementMap.IsEmpty() && !ElementMap.Contains(ParentTransformElement))
+						{
+							continue;
+						}
+
 						if(Weights.IsValidIndex(ParentIndex))
 						{
 							if(Weights[ParentIndex].IsAlmostZero())
diff --git a/Engine/Plugins/Animation/ControlRig/Source/ControlRig/Public/Units/Debug/RigUnit_DebugHierarchy.h b/Engine/Plugins/Animation/ControlRig/Source/ControlRig/Public/Units/Debug/RigUnit_DebugHierarchy.h
--- a/Engine/Plugins/Animation/ControlRig/Source/ControlRig/Public/Units/Debug/RigUnit_DebugHierarchy.h
+++ b/Engine/Plugins/Animation/ControlRig/Source/ControlRig/Public/Units/Debug/RigUnit_DebugHierarchy.h
@@ -43,6 +43,11 @@ struct CONTROLRIG_API FRigUnit_DebugHierarchy : public FRigVMFunction_DebugBase
 	UPROPERTY(DisplayName = "Execute", Transient, meta = (Input, Output))
 	FControlRigExecuteContext ExecuteContext;
 
+	// the items to draw the pose for.
+	// if this is empty we'll draw the whole hierarchy
+	UPROPERTY(meta = (Input))
+	TArray<FRigElementKey> Items;
+
 	UPROPERTY(meta = (Input))
 	float Scale;
 
@@ -58,7 +63,7 @@ struct CONTROLRIG_API FRigUnit_DebugHierarchy : public FRigVMFunction_DebugBase
 	UPROPERTY(meta = (Input))
 	bool bEnabled;
 
-	static void DrawHierarchy(const FRigVMExecuteContext& InContext, const FTransform& WorldOffset, URigHierarchy* Hierarchy, EControlRigDrawHierarchyMode::Type Mode, float Scale, const FLinearColor& Color, float Thickness, const FRigPose* InPose);
+	static void DrawHierarchy(const FRigVMExecuteContext& InContext, const FTransform& WorldOffset, URigHierarchy* Hierarchy, EControlRigDrawHierarchyMode::Type Mode, float Scale, const FLinearColor& Color, float Thickness, const FRigPose* InPose, const TArrayView<const FRigElementKey>* InItems);
 };
 
 /**
@@ -87,6 +92,11 @@ struct CONTROLRIG_API FRigUnit_DebugPose : public FRigVMFunction_DebugBase
 	UPROPERTY(meta = (Input))
 	FRigPose Pose;
 
+	// the items to draw the pose cache for.
+	// if this is empty we'll draw the whole pose cache
+	UPROPERTY(meta = (Input))
+	TArray<FRigElementKey> Items;
+
 	UPROPERTY(meta = (Input))
 	float Scale;
 
