Control Rig: New node to allow to create a pose cache

#rnx

diff --git a/Engine/Plugins/Animation/ControlRig/Source/ControlRig/Private/Units/Execution/RigUnit_Hierarchy.cpp b/Engine/Plugins/Animation/ControlRig/Source/ControlRig/Private/Units/Execution/RigUnit_Hierarchy.cpp
--- a/Engine/Plugins/Animation/ControlRig/Source/ControlRig/Private/Units/Execution/RigUnit_Hierarchy.cpp
+++ b/Engine/Plugins/Animation/ControlRig/Source/ControlRig/Private/Units/Execution/RigUnit_Hierarchy.cpp
@@ -710,3 +710,48 @@ FRigUnit_PoseLoop_Execute()
 	Ratio = GetRatioFromIndex(Index, Count);
 }
 
+FRigUnit_HierarchyCreatePoseItemArray_Execute()
+{
+	const URigHierarchy* Hierarchy = ExecuteContext.Hierarchy;
+	if(Hierarchy == nullptr)
+	{
+		Pose.Reset();
+		return;
+	}
+	
+	Pose.Elements.Reset(Entries.Num());
+	Pose.HierarchyTopologyVersion = Hierarchy->GetTopologyVersion();
+	Pose.PoseHash = Pose.HierarchyTopologyVersion;
+
+	for(const FRigUnit_HierarchyCreatePoseItemArray_Entry& Entry : Entries)
+	{
+		if(const FRigTransformElement* TransformElement = Hierarchy->Find<FRigTransformElement>(Entry.Item))
+		{
+			const FRigElementKey& Key = Entry.Item;
+			FRigPoseElement Element;
+			Element.Index = FCachedRigElement(Key, Hierarchy);
+			Element.ActiveParent = Hierarchy->GetActiveParent(Key);
+			Element.LocalTransform = Entry.LocalTransform;
+			Element.GlobalTransform = Entry.GlobalTransform;
+			Element.CurveValue = Entry.CurveValue;
+			Element.PreferredEulerAngle = FVector::ZeroVector;
+
+			if(const FRigControlElement* ControlElement = Cast<FRigControlElement>(TransformElement))
+			{
+				if(Entry.UseEulerAngles)
+				{
+					Element.PreferredEulerAngle = Entry.EulerAngles;
+				}
+				else
+				{
+					FRigPreferredEulerAngles EulerAngles = ControlElement->PreferredEulerAngles;
+					EulerAngles.SetRotator(Entry.LocalTransform.Rotator(), false, false);
+					Element.PreferredEulerAngle = EulerAngles.GetAngles(false, EulerAngles.RotationOrder); 
+				}
+			}
+
+			Pose.Elements.Add(Element);
+			Pose.PoseHash = HashCombine(Pose.PoseHash, GetTypeHash(Key));
+		}
+	}
+}
\ No newline at end of file
diff --git a/Engine/Plugins/Animation/ControlRig/Source/ControlRig/Public/Units/Execution/RigUnit_Hierarchy.h b/Engine/Plugins/Animation/ControlRig/Source/ControlRig/Public/Units/Execution/RigUnit_Hierarchy.h
--- a/Engine/Plugins/Animation/ControlRig/Source/ControlRig/Public/Units/Execution/RigUnit_Hierarchy.h
+++ b/Engine/Plugins/Animation/ControlRig/Source/ControlRig/Public/Units/Execution/RigUnit_Hierarchy.h
@@ -824,3 +824,65 @@ struct CONTROLRIG_API FRigUnit_PoseLoop : public FRigUnit_HierarchyBaseMutable
 	UPROPERTY(meta = (Output))
 	FControlRigExecuteContext Completed;
 };
+
+USTRUCT(BlueprintType)
+struct FRigUnit_HierarchyCreatePoseItemArray_Entry
+{
+	GENERATED_BODY()
+
+	FRigUnit_HierarchyCreatePoseItemArray_Entry()
+	: Item(NAME_None, ERigElementType::Bone)
+	, LocalTransform(FTransform::Identity)
+	, GlobalTransform(FTransform::Identity)
+	, UseEulerAngles(false)
+	, EulerAngles(FVector::ZeroVector)
+	, CurveValue(0.f)
+	{
+	}
+
+	UPROPERTY(BlueprintReadWrite, Category = "Entry")
+	FRigElementKey Item;
+
+	UPROPERTY(BlueprintReadWrite, Category = "Entry")
+	FTransform LocalTransform;
+
+	UPROPERTY(BlueprintReadWrite, Category = "Entry")
+	FTransform GlobalTransform;
+
+	// in case of a control this can be used to drive the preferred euler angles
+	UPROPERTY(BlueprintReadWrite, Category = "Entry")
+	bool UseEulerAngles;
+
+	// in case of a control this can be used to drive the preferred euler angles
+	UPROPERTY(BlueprintReadWrite, Category = "Entry")
+	FVector EulerAngles;
+
+	// in case of a curve this can be used to drive the curve value
+	UPROPERTY(BlueprintReadWrite, Category = "Entry")
+	float CurveValue;
+};
+
+
+/**
+ * Creates the hierarchy's pose
+ */
+USTRUCT(meta=(DisplayName="Create Pose Cache", Keywords="Hierarchy,Pose,State,MakePoseCache,NewPoseCache,EmptyPoseCache", Varying, Category = "Pose Cache"))
+struct CONTROLRIG_API FRigUnit_HierarchyCreatePoseItemArray : public FRigUnit_HierarchyBase
+{
+	GENERATED_BODY()
+
+	FRigUnit_HierarchyCreatePoseItemArray()
+	{
+		Pose = FRigPose();
+	}
+
+	RIGVM_METHOD()
+	virtual void Execute() override;
+
+	// The entries to create
+	UPROPERTY(meta = (Input))
+	TArray<FRigUnit_HierarchyCreatePoseItemArray_Entry> Entries;
+
+	UPROPERTY(meta = (Output))
+	FRigPose Pose;
+};
