Control Rig: Rotation Orders:  Make sure to set the rotation orders when flip eulers flag is on versus using the override use preferred flag. But then make sure to set this to false still for additives. This seems to fix a bunch of issues.

diff --git a/Engine/Plugins/Animation/ControlRig/Source/ControlRig/Private/Rigs/RigHierarchy.cpp b/Engine/Plugins/Animation/ControlRig/Source/ControlRig/Private/Rigs/RigHierarchy.cpp
--- a/Engine/Plugins/Animation/ControlRig/Source/ControlRig/Private/Rigs/RigHierarchy.cpp
+++ b/Engine/Plugins/Animation/ControlRig/Source/ControlRig/Private/Rigs/RigHierarchy.cpp
@@ -4301,7 +4301,7 @@ void URigHierarchy::SetPreferredEulerAnglesFromValue(FRigControlElement* InContr
 		{
 			FEulerTransform EulerTransform = GetEulerTransformFromControlValue(InValue);
 			FQuat Quat = EulerTransform.GetRotation();
-			const FVector Angle = GetControlAnglesFromQuat(InControlElement, Quat);
+			const FVector Angle = GetControlAnglesFromQuat(InControlElement, Quat, bFixEulerFlips);
 			InControlElement->PreferredEulerAngles.SetAngles(Angle, bInitial, InControlElement->PreferredEulerAngles.RotationOrder, bFixEulerFlips);
 			break;
 		}
@@ -4309,7 +4309,7 @@ void URigHierarchy::SetPreferredEulerAnglesFromValue(FRigControlElement* InContr
 		{
 			FTransform Transform = GetTransformFromControlValue(InValue);
 			FQuat Quat = Transform.GetRotation();
-			const FVector Angle = GetControlAnglesFromQuat(InControlElement, Quat);
+			const FVector Angle = GetControlAnglesFromQuat(InControlElement, Quat, bFixEulerFlips);
 			InControlElement->PreferredEulerAngles.SetAngles(Angle, bInitial, InControlElement->PreferredEulerAngles.RotationOrder, bFixEulerFlips);
 			break;
 		}
@@ -4317,7 +4317,7 @@ void URigHierarchy::SetPreferredEulerAnglesFromValue(FRigControlElement* InContr
 		{
 			FTransform Transform = GetTransformNoScaleFromControlValue(InValue);
 			FQuat Quat = Transform.GetRotation();
-			const FVector Angle = GetControlAnglesFromQuat(InControlElement, Quat);
+			const FVector Angle = GetControlAnglesFromQuat(InControlElement, Quat, bFixEulerFlips);
 			InControlElement->PreferredEulerAngles.SetAngles(Angle, bInitial, InControlElement->PreferredEulerAngles.RotationOrder, bFixEulerFlips);
 			break;
 		}
diff --git a/Engine/Plugins/Animation/ControlRig/Source/ControlRig/Public/Rigs/RigHierarchy.h b/Engine/Plugins/Animation/ControlRig/Source/ControlRig/Public/Rigs/RigHierarchy.h
--- a/Engine/Plugins/Animation/ControlRig/Source/ControlRig/Public/Rigs/RigHierarchy.h
+++ b/Engine/Plugins/Animation/ControlRig/Source/ControlRig/Public/Rigs/RigHierarchy.h
@@ -2568,13 +2568,13 @@ public:
 		return FQuat();
 	}
 
-	FVector GetControlAnglesFromQuat(const FRigControlElement* InControlElement, const FQuat& InQuat) const
+	FVector GetControlAnglesFromQuat(const FRigControlElement* InControlElement, const FQuat& InQuat, bool bUseRotationOrder) const
 	{
 		FVector Angle(0, 0, 0);
 		if (InControlElement)
 		{
 
-			if (GetUsePreferredRotationOrder(InControlElement))
+			if (bUseRotationOrder && InControlElement->Settings.bUsePreferredRotationOrder)
 			{
 				FRotator Rotator = InControlElement->PreferredEulerAngles.GetRotatorFromQuat(InQuat);
 				Angle = Rotator.Euler();
diff --git a/Engine/Plugins/Animation/ControlRig/Source/ControlRigEditor/Private/EditMode/ControlRigEditMode.cpp b/Engine/Plugins/Animation/ControlRig/Source/ControlRigEditor/Private/EditMode/ControlRigEditMode.cpp
--- a/Engine/Plugins/Animation/ControlRig/Source/ControlRigEditor/Private/EditMode/ControlRigEditMode.cpp
+++ b/Engine/Plugins/Animation/ControlRig/Source/ControlRigEditor/Private/EditMode/ControlRigEditMode.cpp
@@ -4883,7 +4883,9 @@ void FControlRigEditMode::MoveControlShape(AControlRigShapeActor* ShapeActor, co
 				}
 				
 				ControlRig->Evaluate_AnyThread();
-				SetControlShapeTransform(ShapeActor, NewTransform, ToWorldTransform, Context, bPrintPythonCommands, /*fix flips*/ true);
+				//fix flips and do rotation orders only if not additive
+				const bool bFixEulerFlips = ControlRig->IsAdditive() == false ? true : false;
+				SetControlShapeTransform(ShapeActor, NewTransform, ToWorldTransform, Context, bPrintPythonCommands, bFixEulerFlips);
 				//UpdatePreferredEulerAngles(ControlRig);
 				NotifyDrivenControls(ControlRig, ShapeActor->GetElementKey(),Context);
 				if(const FRigControlElement* ControlElement = ControlRig->FindControl(ShapeActor->ControlName))
@@ -5538,6 +5540,7 @@ void FControlRigEditMode::NotifyDrivenControls(UControlRig* InControlRig, const
 	{
 		if(ControlElement->CanDriveControls())
 		{
+			const bool bFixEulerFlips = InControlRig->IsAdditive() == false ? true : false;
 			FRigControlModifiedContext Context(InContext);
 			Context.EventName = FRigUnit_BeginExecution::EventName;
 
@@ -5546,7 +5549,7 @@ void FControlRigEditMode::NotifyDrivenControls(UControlRig* InControlRig, const
 				if(DrivenKey.Type == ERigElementType::Control)
 				{
 					const FTransform DrivenTransform = InControlRig->GetControlLocalTransform(DrivenKey.Name);
-					InControlRig->SetControlLocalTransform(DrivenKey.Name, DrivenTransform, true, Context, false /*undo*/, true/* bFixEulerFlips*/);
+					InControlRig->SetControlLocalTransform(DrivenKey.Name, DrivenTransform, true, Context, false /*undo*/, bFixEulerFlips);
 				}
 			}
 		}
