Sequencer: Control Rigs cannot be added anymore in  custom Sequences derived from UMovieSceneSequence which is effecting licenses with custom Sequencers.  This is a regression caused by some refactoring when layered control rigs were added . To fix we go back and just use UMovieSceneSequence instead of ULevelSequence. Unfortunately since we can't change headers we do this directly in the .cpp file but for 5.5 we will fix using updated functions.
#tests Tested creating some Control Rig's from the editor via Sequencer.

diff --git a/Engine/Plugins/Animation/ControlRig/Source/ControlRigEditor/Private/Sequencer/ControlRigParameterTrackEditor.cpp b/Engine/Plugins/Animation/ControlRig/Source/ControlRigEditor/Private/Sequencer/ControlRigParameterTrackEditor.cpp
--- a/Engine/Plugins/Animation/ControlRig/Source/ControlRigEditor/Private/Sequencer/ControlRigParameterTrackEditor.cpp
+++ b/Engine/Plugins/Animation/ControlRig/Source/ControlRigEditor/Private/Sequencer/ControlRigParameterTrackEditor.cpp
@@ -1623,14 +1623,152 @@ void FControlRigParameterTrackEditor::OnControlRigAssetSelected(const FAssetData
 		}
 	}
 */
+static UMovieSceneControlRigParameterTrack* AddControlRig(TSharedPtr<ISequencer> SharedSequencer , UMovieSceneSequence* Sequence, const UClass* InClass, FGuid ObjectBinding, UControlRig* InExistingControlRig, bool bIsAdditiveControlRig)
+{
+	FSlateApplication::Get().DismissAllMenus();
+	if (!InClass || !InClass->IsChildOf(UControlRig::StaticClass()) ||
+		!Sequence || !Sequence->GetMovieScene())
+	{
+		return nullptr;
+	}
+
+	UMovieScene* OwnerMovieScene = Sequence->GetMovieScene();
+	ISequencer* Sequencer = nullptr; // will be valid  if we have a ISequencer AND it's focused.
+	if (SharedSequencer.IsValid() && SharedSequencer->GetFocusedMovieSceneSequence() == Sequence)
+	{
+		Sequencer = SharedSequencer.Get();
+	}
+	Sequence->Modify();
+	OwnerMovieScene->Modify();
+
+	if (bIsAdditiveControlRig && InClass != UFKControlRig::StaticClass() && !InClass->GetDefaultObject<UControlRig>()->SupportsEvent(FRigUnit_InverseExecution::EventName))
+	{
+		UE_LOG(LogControlRigEditor, Error, TEXT("Cannot add an additive control rig which does not contain a backwards solve event."));
+		return nullptr;
+	}
+
+	FScopedTransaction AddControlRigTrackTransaction(LOCTEXT("AddControlRigTrack", "Add Control Rig Track"));
+
+	UMovieSceneControlRigParameterTrack* Track = Cast<UMovieSceneControlRigParameterTrack>(OwnerMovieScene->AddTrack(UMovieSceneControlRigParameterTrack::StaticClass(), ObjectBinding));
+	if (Track)
+	{
+		FString ObjectName = InClass->GetName(); //GetDisplayNameText().ToString();
+		ObjectName.RemoveFromEnd(TEXT("_C"));
+
+		bool bSequencerOwnsControlRig = false;
+		UControlRig* ControlRig = InExistingControlRig;
+		if (ControlRig == nullptr)
+		{
+			ControlRig = NewObject<UControlRig>(Track, InClass, FName(*ObjectName), RF_Transactional);
+			bSequencerOwnsControlRig = true;
+		}
+
+		ControlRig->Modify();
+		if (UFKControlRig* FKControlRig = Cast<UFKControlRig>(Cast<UControlRig>(ControlRig)))
+		{
+			if (bIsAdditiveControlRig)
+			{
+				FKControlRig->SetApplyMode(EControlRigFKRigExecuteMode::Additive);
+			}
+		}
+		else
+		{
+			ControlRig->SetIsAdditive(bIsAdditiveControlRig);
+		}
+		ControlRig->SetObjectBinding(MakeShared<FControlRigObjectBinding>());
+		// Do not re-initialize existing control rig
+		if (!InExistingControlRig)
+		{
+			ControlRig->Initialize();
+		}
+		ControlRig->Evaluate_AnyThread();
+
+		if (SharedSequencer.IsValid())
+		{
+			SharedSequencer->NotifyMovieSceneDataChanged(EMovieSceneDataChangeType::MovieSceneStructureItemsChanged);
+		}
+
+		Track->Modify();
+		UMovieSceneSection* NewSection = Track->CreateControlRigSection(0, ControlRig, bSequencerOwnsControlRig);
+		NewSection->Modify();
+
+		if (bIsAdditiveControlRig)
+		{
+			const FString AdditiveObjectName = ObjectName + TEXT(" (Layered)");
+			Track->SetTrackName(FName(*ObjectName));
+			Track->SetDisplayName(FText::FromString(AdditiveObjectName));
+			Track->SetColorTint(UMovieSceneControlRigParameterTrack::LayeredRigTrackColor);
+		}
+		else
+		{
+			//mz todo need to have multiple rigs with same class
+			Track->SetTrackName(FName(*ObjectName));
+			Track->SetDisplayName(FText::FromString(ObjectName));
+			Track->SetColorTint(UMovieSceneControlRigParameterTrack::AbsoluteRigTrackColor);
+		}
+
+		if (SharedSequencer.IsValid())
+		{
+			SharedSequencer->EmptySelection();
+			SharedSequencer->SelectSection(NewSection);
+			SharedSequencer->ThrobSectionSelection();
+			SharedSequencer->NotifyMovieSceneDataChanged(EMovieSceneDataChangeType::MovieSceneStructureItemAdded);
+			SharedSequencer->ObjectImplicitlyAdded(ControlRig);
+		}
+
+		FControlRigEditMode* ControlRigEditMode = static_cast<FControlRigEditMode*>(GLevelEditorModeTools().GetActiveMode(FControlRigEditMode::ModeName));
+		if (!ControlRigEditMode)
+		{
+			GLevelEditorModeTools().ActivateMode(FControlRigEditMode::ModeName);
+			ControlRigEditMode = static_cast<FControlRigEditMode*>(GLevelEditorModeTools().GetActiveMode(FControlRigEditMode::ModeName));
+
+		}
+		if (ControlRigEditMode)
+		{
+			ControlRigEditMode->AddControlRigObject(ControlRig, SharedSequencer);
+		}
+		return Track;
+	}
+	return nullptr;
+}
+static UMovieSceneTrack* FindOrCreateControlRigTrack(TSharedPtr<ISequencer>& Sequencer, UWorld* World, const UClass* ControlRigClass, const FMovieSceneBindingProxy& InBinding, bool bIsLayeredControlRig)
+{
+	UMovieScene* MovieScene = InBinding.Sequence ? InBinding.Sequence->GetMovieScene() : nullptr;
+	UMovieSceneTrack* BaseTrack = nullptr;
+	if (MovieScene && InBinding.BindingID.IsValid())
+	{
+		if (const FMovieSceneBinding* Binding = MovieScene->FindBinding(InBinding.BindingID))
+		{
+			TArray<UMovieSceneTrack*> Tracks = MovieScene->FindTracks(UMovieSceneControlRigParameterTrack::StaticClass(), Binding->GetObjectGuid(), NAME_None);
+			for (UMovieSceneTrack* AnyOleTrack : Tracks)
+			{
+				UMovieSceneControlRigParameterTrack* Track = Cast<UMovieSceneControlRigParameterTrack>(AnyOleTrack);
+				if (Track && Track->GetControlRig() && Track->GetControlRig()->GetClass() == ControlRigClass)
+				{
+					return Track;
+				}
+			}
+
+			UMovieSceneControlRigParameterTrack* Track = AddControlRig(Sequencer, InBinding.Sequence, ControlRigClass, InBinding.BindingID, nullptr, bIsLayeredControlRig);
+
+			if (Track)
+			{
+				BaseTrack = Track;
+			}
+		}
+	}
+	return BaseTrack;
+}
 
 void FControlRigParameterTrackEditor::AddControlRig(const UClass* InClass, UObject* BoundActor, FGuid ObjectBinding, UControlRig* InExistingControlRig)
 {
 	UWorld* World = GCurrentLevelEditingViewportClient ? GCurrentLevelEditingViewportClient->GetWorld() : nullptr;
-	ULevelSequence* LevelSequence = Cast<ULevelSequence>(GetSequencer()->GetFocusedMovieSceneSequence());
 	UMovieSceneSequence* Sequence = GetSequencer()->GetFocusedMovieSceneSequence();
 	FMovieSceneBindingProxy BindingProxy(ObjectBinding, Sequence);
-	if (UMovieSceneTrack* Track = UControlRigSequencerEditorLibrary::FindOrCreateControlRigTrack(World, LevelSequence, InClass, BindingProxy, bIsLayeredControlRig))
+	TSharedPtr<ISequencer> SequencerPtr = GetSequencer();
+	//UControlRigSequencerEditorLibrary::FindOrCreateControlRigTrack.. in 5.5 we will redo this but for 
+	//5.4.4 we can't change headers so for now we just make the change here locally
+	if (UMovieSceneTrack* Track = FindOrCreateControlRigTrack(SequencerPtr, World, InClass, BindingProxy, bIsLayeredControlRig))
 	{
 		BindControlRig(CastChecked<UMovieSceneControlRigParameterTrack>(Track)->GetControlRig());
 	}
