ControlRig: fixed marquee selection with hidden control shapes

#rnx

diff --git a/Engine/Plugins/Animation/ControlRig/Source/ControlRigEditor/Private/EditMode/ControlRigEditMode.cpp b/Engine/Plugins/Animation/ControlRig/Source/ControlRigEditor/Private/EditMode/ControlRigEditMode.cpp
--- a/Engine/Plugins/Animation/ControlRig/Source/ControlRigEditor/Private/EditMode/ControlRigEditMode.cpp
+++ b/Engine/Plugins/Animation/ControlRig/Source/ControlRigEditor/Private/EditMode/ControlRigEditMode.cpp
@@ -1838,14 +1838,16 @@ bool FControlRigEditMode::FrustumSelect(const FConvexVolume& InFrustum, FEditorV
 	{
 		for (AControlRigShapeActor* ShapeActor : Pairs.Value)
 		{
-			for (UActorComponent* Component : ShapeActor->GetComponents())
+			const bool bTreatShape = ShapeActor && ShapeActor->ControlRig.IsValid() && ShapeActor->ControlRig->GetControlsVisible() &&
+				ShapeActor->IsSelectable() && !ShapeActor->IsTemporarilyHiddenInEditor();
+			if (bTreatShape)
 			{
-				UPrimitiveComponent* PrimitiveComponent = Cast<UPrimitiveComponent>(Component);
-				if (PrimitiveComponent && PrimitiveComponent->IsRegistered() && PrimitiveComponent->IsVisibleInEditor())
+				for (UActorComponent* Component : ShapeActor->GetComponents())
 				{
-					if (PrimitiveComponent->IsShown(InViewportClient->EngineShowFlags) && PrimitiveComponent->ComponentIsTouchingSelectionFrustum(InFrustum, false /*only bsp*/, false/*encompass entire*/))
+					UPrimitiveComponent* PrimitiveComponent = Cast<UPrimitiveComponent>(Component);
+					if (PrimitiveComponent && PrimitiveComponent->IsRegistered() && PrimitiveComponent->IsVisibleInEditor())
 					{
-						if (ShapeActor->IsSelectable() && ShapeActor->ControlRig.IsValid() &&  ShapeActor->ControlRig->GetControlsVisible() )
+						if (PrimitiveComponent->IsShown(InViewportClient->EngineShowFlags) && PrimitiveComponent->ComponentIsTouchingSelectionFrustum(InFrustum, false /*only bsp*/, false/*encompass entire*/))
 						{
 							bSomethingSelected = true;
 							const FName& ControlName = ShapeActor->ControlName;
