Animation mode & Sequencer: focus improvements
    Set the focus on the viewport, sequencer or the curve editor if the mouse hovers the widget AND a key is pressed when the mouse is over the widget.
    This is disabled by default but can be enabled using 'AnimMode.PendingFocusMode true' to enable it in animation mode only

+ improved fps when playing animations in animation mode (mainly due to synthesized mouse moves)

TODO: add a user preference instead of the cvar if the feedback is positive
NOTE: this could probably be something provided as a core slate functionality but we'd like this to be tested at the animation level first.

diff --git a/Engine/Plugins/Animation/ControlRig/Source/ControlRigEditor/ControlRigEditor.Build.cs b/Engine/Plugins/Animation/ControlRig/Source/ControlRigEditor/ControlRigEditor.Build.cs
--- a/Engine/Plugins/Animation/ControlRig/Source/ControlRigEditor/ControlRigEditor.Build.cs
+++ b/Engine/Plugins/Animation/ControlRig/Source/ControlRigEditor/ControlRigEditor.Build.cs
@@ -81,7 +81,8 @@ namespace UnrealBuildTool.Rules
                     "AnimationEditorWidgets",
                     "ActorPickerMode",
                     "Constraints",
-                    "AnimationEditMode"
+                    "AnimationEditMode",
+					"SequencerWidgets"
 				}
             );
 
diff --git a/Engine/Plugins/Animation/ControlRig/Source/ControlRigEditor/Private/EditMode/ControlRigEditMode.cpp b/Engine/Plugins/Animation/ControlRig/Source/ControlRigEditor/Private/EditMode/ControlRigEditMode.cpp
--- a/Engine/Plugins/Animation/ControlRig/Source/ControlRigEditor/Private/EditMode/ControlRigEditMode.cpp
+++ b/Engine/Plugins/Animation/ControlRig/Source/ControlRigEditor/Private/EditMode/ControlRigEditMode.cpp
@@ -75,12 +75,33 @@
 #include "AnimationEditorViewportClient.h"
 #include "EditorInteractiveGizmoManager.h"
 #include "Editor/ControlRigViewportToolbarExtensions.h"
+#include "Editor/Sequencer/Private/SSequencer.h"
+#include "Slate/SceneViewport.h"
 #include "Tools/BakingHelper.h"
 
 #include UE_INLINE_GENERATED_CPP_BY_NAME(ControlRigEditMode)
 
 TAutoConsoleVariable<bool> CVarClickSelectThroughGizmo(TEXT("ControlRig.Sequencer.ClickSelectThroughGizmo"), false, TEXT("When false you can't click through a gizmo and change selection if you will select the gizmo when in Animation Mode, default to false."));
 
+namespace CRModeLocals
+{
+
+	const TCHAR* FocusModeName = TEXT("AnimMode.PendingFocusMode");
+
+static bool bFocusMode = false;
+static FAutoConsoleVariableRef CVarSetFocusOnHover(
+   FocusModeName,
+   bFocusMode,
+   TEXT("Force setting focus on the hovered viewport when entering a key.")
+   );
+
+IConsoleVariable* GetFocusModeVariable()
+{
+	return IConsoleManager::Get().FindConsoleVariable(FocusModeName);
+}
+	
+}
+
 void UControlRigEditModeDelegateHelper::OnPoseInitialized()
 {
 	if (EditMode)
@@ -214,6 +235,14 @@ bool FControlRigEditMode::SetSequencer(TWeakPtr<ISequencer> InSequencer)
 {
 	if (InSequencer != WeakSequencer)
 	{
+		if (WeakSequencer.IsValid())
+		{
+			static constexpr bool bDisable = false;
+			TSharedPtr<ISequencer> PreviousSequencer = WeakSequencer.Pin();
+			TSharedRef<SSequencer> PreviousSequencerWidget = StaticCastSharedRef<SSequencer>(PreviousSequencer->GetSequencerWidget());
+			PreviousSequencerWidget->EnablePendingFocusOnHovering(bDisable);
+		}
+		
 		WeakSequencer = InSequencer;
 
 		DetailKeyFrameCache->UnsetDelegates();
@@ -245,6 +274,11 @@ bool FControlRigEditMode::SetSequencer(TWeakPtr<ISequencer> InSequencer)
 			LastMovieSceneSig = Sequencer->GetFocusedMovieSceneSequence()->GetMovieScene()->GetSignature();
 			DetailKeyFrameCache->SetDelegates(WeakSequencer, this);
 			ControlProxy->SetSequencer(WeakSequencer);
+
+			{
+				TSharedRef<SSequencer> SequencerWidget = StaticCastSharedRef<SSequencer>(Sequencer->GetSequencerWidget());
+				SequencerWidget->EnablePendingFocusOnHovering(CRModeLocals::bFocusMode);
+			}
 		}
 		SetObjects_Internal();
 		if (FControlRigEditModeToolkit::Details.IsValid())
@@ -448,6 +482,8 @@ void FControlRigEditMode::Enter()
 	UE::ControlRig::PopulateControlRigViewportToolbarShowSubmenu(
 		"LevelEditor.ViewportToolbar.Show", GetToolkit()->GetToolkitCommands()
 	);
+
+	RegisterPendingFocusMode();
 }
 
 //todo get working with Persona
@@ -472,6 +508,8 @@ void FControlRigEditMode::Exit()
 {
 	UE::ControlRig::RemoveControlRigViewportToolbarExtensions();
 
+	UnregisterPendingFocusMode();
+	
 	ClearOutAnyActiveTools();
 	OnControlRigAddedOrRemovedDelegate.Clear();
 	OnControlRigSelectedDelegate.Clear();
@@ -3596,19 +3634,44 @@ void FControlRigEditMode::InvertInputPose(bool bSelectionOnly)
 	}
 }
 
-bool FControlRigEditMode::MouseMove(FEditorViewportClient* ViewportClient, FViewport* Viewport, int32 x, int32 y)
+bool FControlRigEditMode::MouseMove(FEditorViewportClient* InViewportClient, FViewport* InViewport, int32 InX, int32 InY)
 {
-	// Inform units of hover state
-	HActor* ActorHitProxy = HitProxyCast<HActor>(Viewport->GetHitProxy(x, y));
-	if(ActorHitProxy && ActorHitProxy->Actor)
+	// avoid hit proxy cast as much as possible
+	// NOTE: with synthesized mouse moves, this is being called a lot sadly so playing in sequencer with the mouse over the viewport leads to fps drop
+	auto HasAnyHoverableShapeActor = [this, InViewportClient]()
 	{
-		if (ActorHitProxy->Actor->IsA<AControlRigShapeActor>())
+		if (!InViewportClient || InViewportClient->IsInGameView())
 		{
-			for (auto& ShapeActors : ControlRigShapeActors)
+			return false;
+		}
+		
+		for (auto&[ControlRig, ShapeActors]: ControlRigShapeActors)
+		{
+			for (const AControlRigShapeActor* ShapeActor : ShapeActors)
+			{
+				if (ShapeActor && ShapeActor->IsSelectable() && !ShapeActor->IsTemporarilyHiddenInEditor())
+				{
+					return true;
+				}
+			}
+		}
+		return false;
+	};
+	
+	if (HasAnyHoverableShapeActor())
+	{	
+		HActor* ActorHitProxy = HitProxyCast<HActor>(InViewport->GetHitProxy(InX, InY));
+		if (ActorHitProxy && ActorHitProxy->Actor)
+		{
+			if (ActorHitProxy->Actor->IsA<AControlRigShapeActor>())
 			{
-				for (AControlRigShapeActor* ShapeActor : ShapeActors.Value)
+				for (auto& ShapeActors : ControlRigShapeActors)
 				{
-					ShapeActor->SetHovered(ShapeActor == ActorHitProxy->Actor);
+					for (AControlRigShapeActor* ShapeActor : ShapeActors.Value)
+					{
+						ShapeActor->SetHovered(ShapeActor == ActorHitProxy->Actor);
+						return false;
+					}
 				}
 			}
 		}
@@ -3617,8 +3680,32 @@ bool FControlRigEditMode::MouseMove(FEditorViewportClient* ViewportClient, FView
 	return false;
 }
 
-bool FControlRigEditMode::MouseLeave(FEditorViewportClient* ViewportClient, FViewport* Viewport)
+bool FControlRigEditMode::MouseEnter(FEditorViewportClient* InViewportClient, FViewport* InViewport, int32 InX, int32 InY)
+{
+	if (PendingFocus.IsEnabled() && InViewportClient)
+	{
+		const FEditorModeTools* ModeTools = GetModeManager();  
+		if (ModeTools && ModeTools == &GLevelEditorModeTools())
+		{
+			FEditorViewportClient* HoveredVPC = ModeTools->GetHoveredViewportClient();
+			if (HoveredVPC == InViewportClient)
+			{
+				const TWeakPtr<SViewport> ViewportWidget = InViewportClient->GetEditorViewportWidget()->GetSceneViewport()->GetViewportWidget();
+				if (ViewportWidget.IsValid())
+				{
+					PendingFocus.SetPendingFocusIfNeeded(ViewportWidget.Pin()->GetContent());
+				}
+			}
+		}
+	}
+
+	return IPersonaEditMode::MouseEnter(InViewportClient, InViewport, InX, InY);
+}
+
+bool FControlRigEditMode::MouseLeave(FEditorViewportClient* /*InViewportClient*/, FViewport* /*InViewport*/)
 {
+	PendingFocus.ResetPendingFocus();
+	
 	for (auto& ShapeActors : ControlRigShapeActors)
 	{
 		for (AControlRigShapeActor* ShapeActor : ShapeActors.Value)
@@ -3630,6 +3717,56 @@ bool FControlRigEditMode::MouseLeave(FEditorViewportClient* ViewportClient, FVie
 	return false;
 }
 
+void FControlRigEditMode::RegisterPendingFocusMode()
+{
+	if (!IsInLevelEditor())
+	{
+		return;
+	}
+	
+	static IConsoleVariable* UseFocusMode = CRModeLocals::GetFocusModeVariable();
+	if (ensure(UseFocusMode))
+	{
+		auto OnFocusModeChanged = [this](IConsoleVariable*)
+		{
+			PendingFocus.Enable(CRModeLocals::bFocusMode);
+			if (WeakSequencer.IsValid())
+			{
+				TSharedPtr<ISequencer> PreviousSequencer = WeakSequencer.Pin();
+				TSharedRef<SSequencer> PreviousSequencerWidget = StaticCastSharedRef<SSequencer>(PreviousSequencer->GetSequencerWidget());
+				PreviousSequencerWidget->EnablePendingFocusOnHovering(CRModeLocals::bFocusMode);
+			}
+		};
+		if (!PendingFocusHandle.IsValid())
+		{
+			PendingFocusHandle = UseFocusMode->OnChangedDelegate().AddLambda(OnFocusModeChanged);
+		}		
+		OnFocusModeChanged(UseFocusMode);
+	}
+}
+
+void FControlRigEditMode::UnregisterPendingFocusMode()
+{
+	static constexpr bool bDisable = false;
+	if (WeakSequencer.IsValid())
+	{
+		TSharedRef<SSequencer> SequencerWidget = StaticCastSharedRef<SSequencer>(WeakSequencer.Pin()->GetSequencerWidget());
+		SequencerWidget->EnablePendingFocusOnHovering(bDisable);
+	}
+
+	PendingFocus.Enable(bDisable);
+	
+	if (PendingFocusHandle.IsValid())
+	{
+		static IConsoleVariable* UseFocusMode = CRModeLocals::GetFocusModeVariable();
+		if (ensure(UseFocusMode))
+		{
+			UseFocusMode->OnChangedDelegate().Remove(PendingFocusHandle);
+		}
+		PendingFocusHandle.Reset();
+	}
+}
+
 bool FControlRigEditMode::CheckMovieSceneSig()
 {
 	bool bSomethingChanged = false;
diff --git a/Engine/Plugins/Animation/ControlRig/Source/ControlRigEditor/Public/EditMode/ControlRigEditMode.h b/Engine/Plugins/Animation/ControlRig/Source/ControlRigEditor/Public/EditMode/ControlRigEditMode.h
--- a/Engine/Plugins/Animation/ControlRig/Source/ControlRigEditor/Public/EditMode/ControlRigEditMode.h
+++ b/Engine/Plugins/Animation/ControlRig/Source/ControlRigEditor/Public/EditMode/ControlRigEditMode.h
@@ -17,6 +17,7 @@
 #include "Materials/Material.h"
 #include "Materials/MaterialInstanceDynamic.h"
 #include "IDetailKeyframeHandler.h"
+#include "WidgetFocusUtils.h"
 #include "ControlRigEditMode.generated.h"
 
 class FEditorViewportClient;
@@ -175,6 +176,7 @@ public:
 	virtual bool ShouldDrawWidget() const override;
 	virtual bool IsCompatibleWith(FEditorModeID OtherModeID) const override;
 	virtual bool MouseMove(FEditorViewportClient* ViewportClient, FViewport* Viewport, int32 x, int32 y) override;
+	virtual bool MouseEnter( FEditorViewportClient* ViewportClient,FViewport* Viewport,int32 x, int32 y ) override;
 	virtual bool MouseLeave(FEditorViewportClient* ViewportClient, FViewport* Viewport) override;
 	virtual void PostUndo() override;
 
@@ -318,6 +320,14 @@ private:
 	/** Reset Shape Size */
 	void ResetControlShapeSize();
 
+	/** Pending focus handler */
+	FPendingWidgetFocus PendingFocus;
+
+	/** Pending focus cvar binding functions to enable/disable pending focus mode */
+	void RegisterPendingFocusMode();
+	void UnregisterPendingFocusMode();
+	FDelegateHandle PendingFocusHandle;
+	
 public:
 	
 	/** Toggle Shape Transform Edit*/
diff --git a/Engine/Source/Editor/CurveEditor/Private/SCurveEditorPanel.cpp b/Engine/Source/Editor/CurveEditor/Private/SCurveEditorPanel.cpp
--- a/Engine/Source/Editor/CurveEditor/Private/SCurveEditorPanel.cpp
+++ b/Engine/Source/Editor/CurveEditor/Private/SCurveEditorPanel.cpp
@@ -1005,6 +1005,23 @@ FReply SCurveEditorPanel::OnKeyDown(const FGeometry& MyGeometry, const FKeyEvent
 	return FReply::Unhandled();
 }
 
+void SCurveEditorPanel::OnMouseEnter(const FGeometry& MyGeometry, const FPointerEvent& MouseEvent)
+{
+	SCompoundWidget::OnMouseEnter(MyGeometry, MouseEvent);
+	PendingFocus.SetPendingFocusIfNeeded(AsWeak());
+}
+
+void SCurveEditorPanel::OnMouseLeave(const FPointerEvent& MouseEvent)
+{
+	SCompoundWidget::OnMouseLeave(MouseEvent);
+	PendingFocus.ResetPendingFocus();
+}
+
+void SCurveEditorPanel::EnablePendingFocusOnHovering(const bool InEnabled)
+{
+	PendingFocus.Enable(InEnabled);
+}
+
 TSharedRef<SWidget> SCurveEditorPanel::MakeCurveEditorCurveViewOptionsMenu()
 {
 	// This builds the dropdown menu when looking at the Curve View Options combobox.
diff --git a/Engine/Source/Editor/CurveEditor/Public/SCurveEditorPanel.h b/Engine/Source/Editor/CurveEditor/Public/SCurveEditorPanel.h
--- a/Engine/Source/Editor/CurveEditor/Public/SCurveEditorPanel.h
+++ b/Engine/Source/Editor/CurveEditor/Public/SCurveEditorPanel.h
@@ -9,6 +9,7 @@
 #include "CurveDrawInfo.h"
 #include "CurveEditor.h"
 #include "CurveEditorTypes.h"
+#include "WidgetFocusUtils.h"
 #include "Curves/RealCurve.h"
 #include "Curves/RichCurve.h"
 #include "HAL/Platform.h"
@@ -175,6 +176,9 @@ class CURVEEDITOR_API SCurveEditorPanel : public SCompoundWidget
 	FSimpleDelegate OnFilterClassChanged;
 	void FilterClassChanged();
 
+	/** Enable/disable pending focus */
+	void EnablePendingFocusOnHovering(const bool InEnabled);
+	
 private:
 	// SWidget Interface
 	virtual void Tick(const FGeometry& AllottedGeometry, const double InCurrentTime, const float InDeltaTime) override;
@@ -183,6 +187,8 @@ private:
 	/*~ Keyboard interaction */
 	virtual bool SupportsKeyboardFocus() const override { return true; }
 	virtual FReply OnKeyDown(const FGeometry& MyGeometry, const FKeyEvent& InKeyEvent) override;
+	virtual void OnMouseEnter(const FGeometry& MyGeometry, const FPointerEvent& MouseEvent) override;
+    virtual void OnMouseLeave(const FPointerEvent& MouseEvent) override;
 
 	TSharedRef<SWidget> MakeTimeSnapMenu();
 	FText GetTimeSnapMenuTooltip() const;
@@ -384,6 +390,9 @@ private:
 
 private:
 
+	/** Pending focus handler */
+	FPendingWidgetFocus PendingFocus;
+	
 	/** Whether to explicitly refresh the views for this panel */
 	bool bNeedsRefresh;
 
diff --git a/Engine/Source/Editor/Sequencer/Private/SSequencer.cpp b/Engine/Source/Editor/Sequencer/Private/SSequencer.cpp
--- a/Engine/Source/Editor/Sequencer/Private/SSequencer.cpp
+++ b/Engine/Source/Editor/Sequencer/Private/SSequencer.cpp
@@ -3378,6 +3378,18 @@ void SSequencer::OnFocusChanging( const FWeakWidgetPath& PreviousFocusPath, cons
 	}
 }
 
+void SSequencer::OnMouseEnter(const FGeometry& MyGeometry, const FPointerEvent& MouseEvent)
+{
+	SCompoundWidget::OnMouseEnter(MyGeometry, MouseEvent);
+	PendingFocus.SetPendingFocusIfNeeded(AsWeak());
+}
+
+void SSequencer::OnMouseLeave(const FPointerEvent& MouseEvent)
+{
+	SCompoundWidget::OnMouseLeave(MouseEvent);
+	PendingFocus.ResetPendingFocus();
+}
+
 void SSequencer::OnAssetsDropped( const FAssetDragDropOp& DragDropOp )
 {
 	using namespace UE::Sequencer;
@@ -4310,4 +4322,36 @@ void SSequencer::ToggleSidebarDrawerDock()
 	}
 }
 
+void SSequencer::EnablePendingFocusOnHovering(const bool InEnabled)
+{
+	PendingFocus.Enable(InEnabled);
+	EnableCurveEditorPendingFocusOnHovering(InEnabled);
+}
+
+void SSequencer::EnableCurveEditorPendingFocusOnHovering(const bool InEnabled) const
+{
+	using namespace UE::Sequencer;
+
+	TSharedPtr<FSequencer> Sequencer = SequencerPtr.Pin();
+	if (!Sequencer.IsValid())
+	{
+		return;
+	}
+
+	if (!Sequencer->GetHostCapabilities().bSupportsCurveEditor)
+	{
+		return;
+	}
+
+	const FCurveEditorExtension* CurveEditorExtension = Sequencer->GetViewModel()->CastDynamic<FCurveEditorExtension>();
+	const TSharedPtr<FCurveEditor> CurveEditor = CurveEditorExtension ? CurveEditorExtension->GetCurveEditor() : nullptr;
+	const TSharedPtr<SCurveEditorPanel> CurveEditorPanel = CurveEditor.IsValid() ? CurveEditor->GetPanel() : nullptr;
+	if (!CurveEditorPanel)
+	{
+		return;
+	}
+	
+	CurveEditorPanel->EnablePendingFocusOnHovering(InEnabled);
+}
+
 #undef LOCTEXT_NAMESPACE
diff --git a/Engine/Source/Editor/Sequencer/Private/SSequencer.h b/Engine/Source/Editor/Sequencer/Private/SSequencer.h
--- a/Engine/Source/Editor/Sequencer/Private/SSequencer.h
+++ b/Engine/Source/Editor/Sequencer/Private/SSequencer.h
@@ -22,6 +22,7 @@
 #include "Sequencer.h"
 #include "SequencerWidgetsDelegates.h"
 #include "STemporarilyFocusedSpinBox.h"
+#include "WidgetFocusUtils.h"
 
 class FActorDragDropOp;
 class FFolderDragDropOp;
@@ -384,6 +385,8 @@ protected:
 	virtual FReply OnDrop( const FGeometry& MyGeometry, const FDragDropEvent& DragDropEvent ) override;
 	virtual FReply OnKeyDown( const FGeometry& MyGeometry, const FKeyEvent& InKeyEvent ) override;
 	virtual void OnFocusChanging( const FWeakWidgetPath& PreviousFocusPath, const FWidgetPath& NewWidgetPath, const FFocusEvent& InFocusEvent ) override;
+	virtual void OnMouseEnter(const FGeometry& MyGeometry, const FPointerEvent& MouseEvent) override;
+	virtual void OnMouseLeave(const FPointerEvent& MouseEvent) override;
 
 private:
 	
@@ -700,7 +703,17 @@ public:
 	/** Undocks the docked sidebar drawer if docked or docks the sidebar drawer if there is one open and no currently docked drawer. */
 	void ToggleSidebarDrawerDock();
 
+	/** Enable/disable pending focus in sequencer */
+	SEQUENCER_API void EnablePendingFocusOnHovering(const bool InEnabled);
+	
 private:
+
+	/** Pending focus handler */
+	FPendingWidgetFocus PendingFocus;
+
+	/** Enable/disable pending focus in the curve editor */
+	void EnableCurveEditorPendingFocusOnHovering(const bool InEnabled) const;
+	
 	/** Applies a single customization. */
 	void ApplySequencerCustomization(const FSequencerCustomizationInfo& Customization);
 
diff --git a/Engine/Source/Editor/SequencerWidgets/Private/WidgetFocusUtils.cpp b/Engine/Source/Editor/SequencerWidgets/Private/WidgetFocusUtils.cpp
new file mode 100644
--- /dev/null
+++ b/Engine/Source/Editor/SequencerWidgets/Private/WidgetFocusUtils.cpp
@@ -0,0 +1,73 @@
+// Copyright Epic Games, Inc. All Rights Reserved.
+
+#include "WidgetFocusUtils.h"
+
+#include "Framework/Application/SlateApplication.h"
+#include "Framework/Application/SlateUser.h"
+
+FPendingWidgetFocus::~FPendingWidgetFocus()
+{
+	PendingFocusFunction.Reset();
+	if (PreInputKeyDownHandle.IsValid())
+	{
+		FSlateApplication& SlateApplication = FSlateApplication::Get();
+		SlateApplication.OnApplicationPreInputKeyDownListener().Remove(PreInputKeyDownHandle);
+		PreInputKeyDownHandle.Reset();
+	}
+}
+
+void FPendingWidgetFocus::SetPendingFocusIfNeeded(const TWeakPtr<SWidget>& InWidget)
+{
+	if (!PreInputKeyDownHandle.IsValid())
+	{
+		return;
+	}
+
+	PendingFocusFunction = [WidgetFocus = InWidget]()
+	{
+		if (WidgetFocus.IsValid())
+		{
+			TSharedPtr<SWidget> Widget = WidgetFocus.Pin();
+			FSlateApplication::Get().ForEachUser([&Widget](FSlateUser& User) 
+			{
+				User.SetFocus(Widget.ToSharedRef());
+			});
+		}
+	};
+}
+
+void FPendingWidgetFocus::ResetPendingFocus()
+{
+	PendingFocusFunction.Reset();
+}
+
+void FPendingWidgetFocus::Enable(const bool InEnabled)
+{
+	FSlateApplication& SlateApplication = FSlateApplication::Get();
+	if (PreInputKeyDownHandle.IsValid())
+	{
+		SlateApplication.OnApplicationPreInputKeyDownListener().Remove(PreInputKeyDownHandle);
+		PreInputKeyDownHandle.Reset();
+	}
+	
+	PendingFocusFunction.Reset();
+	
+	if (InEnabled)
+	{
+		PreInputKeyDownHandle = SlateApplication.OnApplicationPreInputKeyDownListener().AddRaw(this, &FPendingWidgetFocus::OnPreInputKeyDown);
+	}
+}
+
+bool FPendingWidgetFocus::IsEnabled() const
+{
+	return PreInputKeyDownHandle.IsValid();
+}
+	
+void FPendingWidgetFocus::OnPreInputKeyDown(const FKeyEvent&)
+{
+	if (PendingFocusFunction)
+	{
+		PendingFocusFunction();
+		PendingFocusFunction.Reset();
+	}
+}
diff --git a/Engine/Source/Editor/SequencerWidgets/Public/WidgetFocusUtils.h b/Engine/Source/Editor/SequencerWidgets/Public/WidgetFocusUtils.h
new file mode 100644
--- /dev/null
+++ b/Engine/Source/Editor/SequencerWidgets/Public/WidgetFocusUtils.h
@@ -0,0 +1,37 @@
+// Copyright Epic Games, Inc. All Rights Reserved.
+
+#pragma once
+
+#include "Templates/SharedPointer.h"
+#include "Templates/Function.h"
+#include "Delegates/IDelegateInstance.h"
+
+class SWidget;
+struct FKeyEvent;
+
+/**
+ * FPendingWidgetFocus is a utility class that stores a pending focus function when a SWidget is hovered over.
+ * It aims to provide a way to focus on a widget without having to actually click on it. 
+ * The focus function is stored on the mouse enter event and will only be executed if a key down event is sent while the widget is hovered.
+ * If the mouse leave event is called without any key down event having been called, the function is reset and the focus is not modified at all.
+ */
+
+class SEQUENCERWIDGETS_API FPendingWidgetFocus : public TSharedFromThis<FPendingWidgetFocus>
+{
+public:
+	
+	~FPendingWidgetFocus();
+
+	void Enable(const bool InEnabled);
+	bool IsEnabled() const;
+
+	void SetPendingFocusIfNeeded(const TWeakPtr<SWidget>& InWidget);
+	void ResetPendingFocus();
+
+private:
+	
+	void OnPreInputKeyDown(const FKeyEvent&);
+
+	TFunction<void()> PendingFocusFunction;
+	FDelegateHandle PreInputKeyDownHandle;
+};
