Control Rig: Guard against nullptr in modular rig model

#rnx

diff --git a/Engine/Plugins/Animation/ControlRig/Source/ControlRigEditor/Private/Editor/SModularRigModel.cpp b/Engine/Plugins/Animation/ControlRig/Source/ControlRigEditor/Private/Editor/SModularRigModel.cpp
--- a/Engine/Plugins/Animation/ControlRig/Source/ControlRigEditor/Private/Editor/SModularRigModel.cpp
+++ b/Engine/Plugins/Animation/ControlRig/Source/ControlRigEditor/Private/Editor/SModularRigModel.cpp
@@ -268,9 +268,16 @@ void SModularRigModel::RefreshTreeView(bool bRebuildContent)
 	TreeView->RefreshTreeView(bRebuildContent);
 }
 
-TArray<FString> SModularRigModel::GetSelectedKeys() const
+TArray<TSharedPtr<FModularRigTreeElement>> SModularRigModel::GetSelectedItems() const
 {
 	TArray<TSharedPtr<FModularRigTreeElement>> SelectedItems = TreeView->GetSelectedItems();
+	SelectedItems.Remove(TSharedPtr<FModularRigTreeElement>(nullptr));
+	return SelectedItems;
+}
+
+TArray<FString> SModularRigModel::GetSelectedKeys() const
+{
+	const TArray<TSharedPtr<FModularRigTreeElement>> SelectedItems = GetSelectedItems();
 	
 	TArray<FString> SelectedKeys;
 	for (const TSharedPtr<FModularRigTreeElement>& SelectedItem : SelectedItems)
@@ -602,7 +609,7 @@ void SModularRigModel::HandleRenameModule()
 	{
 		FScopedTransaction Transaction(LOCTEXT("ModularRigModelRenameSelected", "Rename selected module"));
 
-		TArray<TSharedPtr<FModularRigTreeElement>> SelectedItems = TreeView->GetSelectedItems();
+		const TArray<TSharedPtr<FModularRigTreeElement>> SelectedItems = GetSelectedItems();
 		if (SelectedItems.Num() == 1)
 		{
 			SelectedItems[0]->RequestRename();
@@ -669,7 +676,7 @@ void SModularRigModel::HandleDeleteModules()
 	{
 		FScopedTransaction Transaction(LOCTEXT("ModularRigModelDeleteSelected", "Delete selected modules"));
 
-		TArray<TSharedPtr<FModularRigTreeElement>> SelectedItems = TreeView->GetSelectedItems();
+		const TArray<TSharedPtr<FModularRigTreeElement>> SelectedItems = GetSelectedItems();
 		TArray<FString> SelectedPaths;
 		Algo::Transform(SelectedItems, SelectedPaths, [](const TSharedPtr<FModularRigTreeElement>& Element)
 		{
@@ -730,7 +737,7 @@ void SModularRigModel::HandleMirrorModules()
 	UModularRig* Rig = GetDefaultModularRig();
 	if (Rig)
 	{
-		TArray<TSharedPtr<FModularRigTreeElement>> SelectedItems = TreeView->GetSelectedItems();
+		const TArray<TSharedPtr<FModularRigTreeElement>> SelectedItems = GetSelectedItems();
 		TArray<FString> SelectedPaths;
 		Algo::Transform(SelectedItems, SelectedPaths, [](const TSharedPtr<FModularRigTreeElement>& Element)
 		{
@@ -792,7 +799,7 @@ void SModularRigModel::HandleReresolveModules()
 	UModularRig* Rig = GetDefaultModularRig();
 	if (Rig)
 	{
-		TArray<TSharedPtr<FModularRigTreeElement>> SelectedItems = TreeView->GetSelectedItems();
+		const TArray<TSharedPtr<FModularRigTreeElement>> SelectedItems = GetSelectedItems();
 		TArray<FString> SelectedPaths;
 		Algo::Transform(SelectedItems, SelectedPaths, [](const TSharedPtr<FModularRigTreeElement>& Element)
 		{
@@ -892,9 +899,13 @@ bool SModularRigModel::CanSwapModules() const
 	if (Rig)
 	{
 		TSoftClassPtr<UControlRig> CommonClass = nullptr;
-		TArray<TSharedPtr<FModularRigTreeElement>> SelectedItems = TreeView->GetSelectedItems();
-		for (TSharedPtr<FModularRigTreeElement>& SelectedItem : SelectedItems)
+		const TArray<TSharedPtr<FModularRigTreeElement>> SelectedItems = GetSelectedItems();
+		for (const TSharedPtr<FModularRigTreeElement>& SelectedItem : SelectedItems)
 		{
+			if(!SelectedItem.IsValid())
+			{
+				continue;
+			}
 			TSoftClassPtr<UControlRig> ModuleClass;
 			if (const FRigModuleReference* Module = ControlRigBlueprint->ModularRigModel.FindModule(SelectedItem->ModulePath))
 			{
@@ -931,7 +942,7 @@ void SModularRigModel::HandleSwapClassForModules()
 	UModularRig* Rig = GetDefaultModularRig();
 	if (Rig)
 	{
-		TArray<TSharedPtr<FModularRigTreeElement>> SelectedItems = TreeView->GetSelectedItems();
+		const TArray<TSharedPtr<FModularRigTreeElement>> SelectedItems = GetSelectedItems();
 		TArray<FString> SelectedPaths;
 		Algo::Transform(SelectedItems, SelectedPaths, [](const TSharedPtr<FModularRigTreeElement>& Element)
 		{
diff --git a/Engine/Plugins/Animation/ControlRig/Source/ControlRigEditor/Private/Editor/SModularRigModel.h b/Engine/Plugins/Animation/ControlRig/Source/ControlRigEditor/Private/Editor/SModularRigModel.h
--- a/Engine/Plugins/Animation/ControlRig/Source/ControlRigEditor/Private/Editor/SModularRigModel.h
+++ b/Engine/Plugins/Animation/ControlRig/Source/ControlRigEditor/Private/Editor/SModularRigModel.h
@@ -54,6 +54,9 @@ private:
 	/** Rebuild the tree view */
 	void RefreshTreeView(bool bRebuildContent = true);
 
+	/** Returns all selected items */
+	TArray<TSharedPtr<FModularRigTreeElement>> GetSelectedItems() const;
+
 	/** Return all selected keys */
 	TArray<FString> GetSelectedKeys() const;
 
