[ControlRig & Sequencer] Added a flag on CRBP that allows users to add multiple instances of that rig in sequencer

-setting better display name for these tracks based on priority order will come next

diff --git a/Engine/Plugins/Animation/ControlRig/Source/ControlRigDeveloper/Public/ControlRigBlueprint.h b/Engine/Plugins/Animation/ControlRig/Source/ControlRigDeveloper/Public/ControlRigBlueprint.h
--- a/Engine/Plugins/Animation/ControlRig/Source/ControlRigDeveloper/Public/ControlRigBlueprint.h
+++ b/Engine/Plugins/Animation/ControlRig/Source/ControlRigDeveloper/Public/ControlRigBlueprint.h
@@ -320,6 +320,11 @@ private:
 	/** If set to true, this control rig has animatable controls */
 	UPROPERTY(AssetRegistrySearchable)
 	bool bExposesAnimatableControls;
+public:
+
+	/** If set to true, multiple control rig tracks can be created for the same rig in sequencer*/
+	UPROPERTY(EditAnywhere, Category="Sequencer", AssetRegistrySearchable)
+	bool bAllowMultipleInstances = false;
 
 private:
 
diff --git a/Engine/Plugins/Animation/ControlRig/Source/ControlRigEditor/Private/Sequencer/ControlRigParameterTrackEditor.cpp b/Engine/Plugins/Animation/ControlRig/Source/ControlRigEditor/Private/Sequencer/ControlRigParameterTrackEditor.cpp
--- a/Engine/Plugins/Animation/ControlRig/Source/ControlRigEditor/Private/Sequencer/ControlRigParameterTrackEditor.cpp
+++ b/Engine/Plugins/Animation/ControlRig/Source/ControlRigEditor/Private/Sequencer/ControlRigParameterTrackEditor.cpp
@@ -226,6 +226,21 @@ static USkeleton* AcquireSkeletonFromObjectGuid(const FGuid& Guid, UObject** Obj
 	return nullptr;
 }
 
+static bool DoesControlRigAllowMultipleInstances(const FTopLevelAssetPath& InGeneratedClassPath)
+{
+	const IAssetRegistry& AssetRegistry = FModuleManager::GetModuleChecked<FAssetRegistryModule>(TEXT("AssetRegistry")).Get();
+	
+	const FString BlueprintPath = InGeneratedClassPath.ToString().LeftChop(2); // Chop off _C
+	const FAssetData AssetData = AssetRegistry.GetAssetByObjectPath(FSoftObjectPath(BlueprintPath));
+
+	if (AssetData.GetTagValueRef<bool>(GET_MEMBER_NAME_CHECKED(UControlRigBlueprint, bAllowMultipleInstances)))
+	{
+		return true;
+	}
+
+	return false;
+}
+
 bool FControlRigParameterTrackEditor::bAutoGenerateControlRigTrack = true;
 FCriticalSection FControlRigParameterTrackEditor::ControlUndoTransactionMutex;
 
@@ -724,21 +739,31 @@ public:
 
 	virtual bool IsClassAllowed(const FClassViewerInitializationOptions& InInitOptions, const UClass* InClass, TSharedRef< class FClassViewerFilterFuncs > InFilterFuncs ) override
 	{
-		return !AlreadyAddedRigs.Contains(InClass);
+		return IsClassAllowed_Internal(InClass->GetClassPathName());
 	}
 	
 	virtual bool IsUnloadedClassAllowed(const FClassViewerInitializationOptions& InInitOptions,
 										const TSharedRef< const class IUnloadedBlueprintData > InUnloadedClassData,
 										TSharedRef< class FClassViewerFilterFuncs > InFilterFuncs) override
 	{
-		const FTopLevelAssetPath ClassPath = InUnloadedClassData.Get().GetClassPathName();
-		return !AlreadyAddedRigs.ContainsByPredicate([ClassPath](const UClass* Class)
-		{
-			return ClassPath == Class->GetClassPathName();
-		});
+		return IsClassAllowed_Internal(InUnloadedClassData.Get().GetClassPathName());
 	}
 	
 private:
+	bool IsClassAllowed_Internal(const FTopLevelAssetPath& InGeneratedClassPath)
+	{
+		if (DoesControlRigAllowMultipleInstances(InGeneratedClassPath))
+		{
+			return true;
+		}
+		
+		return !AlreadyAddedRigs.ContainsByPredicate([InGeneratedClassPath](const UClass* Class)
+			{
+				return InGeneratedClassPath == Class->GetClassPathName();
+			});
+	}
+
+	
 	TArray<UClass*> AlreadyAddedRigs;
 };
 
@@ -1739,13 +1764,16 @@ static UMovieSceneTrack* FindOrCreateControlRigTrack(TSharedPtr<ISequencer>& Seq
 	{
 		if (const FMovieSceneBinding* Binding = MovieScene->FindBinding(InBinding.BindingID))
 		{
-			TArray<UMovieSceneTrack*> Tracks = MovieScene->FindTracks(UMovieSceneControlRigParameterTrack::StaticClass(), Binding->GetObjectGuid(), NAME_None);
-			for (UMovieSceneTrack* AnyOleTrack : Tracks)
+			if (!DoesControlRigAllowMultipleInstances(ControlRigClass->GetClassPathName()))
 			{
-				UMovieSceneControlRigParameterTrack* Track = Cast<UMovieSceneControlRigParameterTrack>(AnyOleTrack);
-				if (Track && Track->GetControlRig() && Track->GetControlRig()->GetClass() == ControlRigClass)
+				TArray<UMovieSceneTrack*> Tracks = MovieScene->FindTracks(UMovieSceneControlRigParameterTrack::StaticClass(), Binding->GetObjectGuid(), NAME_None);
+				for (UMovieSceneTrack* AnyOleTrack : Tracks)
 				{
-					return Track;
+					UMovieSceneControlRigParameterTrack* Track = Cast<UMovieSceneControlRigParameterTrack>(AnyOleTrack);
+					if (Track && Track->GetControlRig() && Track->GetControlRig()->GetClass() == ControlRigClass)
+					{
+						return Track;
+					}
 				}
 			}
 
