Control Rig: Pre and Post construction events

#rnx

diff --git a/Engine/Plugins/Animation/ControlRig/Source/ControlRig/Private/ControlRig.cpp b/Engine/Plugins/Animation/ControlRig/Source/ControlRig/Private/ControlRig.cpp
--- a/Engine/Plugins/Animation/ControlRig/Source/ControlRig/Private/ControlRig.cpp
+++ b/Engine/Plugins/Animation/ControlRig/Source/ControlRig/Private/ControlRig.cpp
@@ -1604,6 +1604,19 @@ void UControlRig::AdaptEventQueueForEvaluate(TArray<FName>& InOutEventQueueToRun
 
 	for (int32 i=0; i<InOutEventQueueToRun.Num(); ++i)
 	{
+		if (InOutEventQueueToRun[i] == FRigUnit_PrepareForExecution::EventName)
+		{
+			if (SupportsEvent(FRigUnit_PrePrepareForExecution::EventName))
+			{
+				InOutEventQueueToRun.Insert(FRigUnit_PrePrepareForExecution::EventName, i);
+				i++; // skip preconstruction 
+			}
+			if (SupportsEvent(FRigUnit_PostPrepareForExecution::EventName))
+			{
+				i++; // skip construction
+				InOutEventQueueToRun.Insert(FRigUnit_PostPrepareForExecution::EventName, i);
+			}
+		}
 		if (InOutEventQueueToRun[i] == FRigUnit_BeginExecution::EventName)
 		{
 			if (SupportsEvent(FRigUnit_PreBeginExecution::EventName))
diff --git a/Engine/Plugins/Animation/ControlRig/Source/ControlRig/Private/Units/Execution/RigUnit_PrepareForExecution.cpp b/Engine/Plugins/Animation/ControlRig/Source/ControlRig/Private/Units/Execution/RigUnit_PrepareForExecution.cpp
--- a/Engine/Plugins/Animation/ControlRig/Source/ControlRig/Private/Units/Execution/RigUnit_PrepareForExecution.cpp
+++ b/Engine/Plugins/Animation/ControlRig/Source/ControlRig/Private/Units/Execution/RigUnit_PrepareForExecution.cpp
@@ -11,3 +11,14 @@ FRigUnit_PrepareForExecution_Execute()
 	ExecuteContext.SetEventName(FRigUnit_PrepareForExecution::EventName);
 }
 
+FRigUnit_PrePrepareForExecution_Execute()
+{
+	DECLARE_SCOPE_HIERARCHICAL_COUNTER_RIGUNIT()
+	ExecuteContext.SetEventName(FRigUnit_PrePrepareForExecution::EventName);
+}
+
+FRigUnit_PostPrepareForExecution_Execute()
+{
+	DECLARE_SCOPE_HIERARCHICAL_COUNTER_RIGUNIT()
+	ExecuteContext.SetEventName(FRigUnit_PostPrepareForExecution::EventName);
+}
diff --git a/Engine/Plugins/Animation/ControlRig/Source/ControlRig/Public/Units/Execution/RigUnit_PrepareForExecution.h b/Engine/Plugins/Animation/ControlRig/Source/ControlRig/Public/Units/Execution/RigUnit_PrepareForExecution.h
--- a/Engine/Plugins/Animation/ControlRig/Source/ControlRig/Public/Units/Execution/RigUnit_PrepareForExecution.h
+++ b/Engine/Plugins/Animation/ControlRig/Source/ControlRig/Public/Units/Execution/RigUnit_PrepareForExecution.h
@@ -26,3 +26,45 @@ struct CONTROLRIG_API FRigUnit_PrepareForExecution : public FRigUnit
 
 	static inline const FLazyName EventName = FLazyName(TEXT("Construction"));
 };
+
+/**
+ * Event to create / configure elements before any other event
+ */
+USTRUCT(meta=(DisplayName="Pre Construction", Category="Events", NodeColor="0.6, 0, 1", Keywords="Create,Build,Spawn,Setup,Init,Fit"))
+struct CONTROLRIG_API FRigUnit_PrePrepareForExecution : public FRigUnit
+{
+	GENERATED_BODY()
+
+	RIGVM_METHOD()
+	virtual void Execute() override;
+
+	virtual FName GetEventName() const override { return EventName; }
+	virtual bool CanOnlyExistOnce() const override { return true; }
+
+	// The execution result
+	UPROPERTY(EditAnywhere, Transient, DisplayName = "Execute", Category = "PrePrepareForExecution", meta = (Output))
+	FControlRigExecuteContext ExecuteContext;
+
+	static inline const FLazyName EventName = FLazyName(TEXT("PreConstruction"));
+};
+
+/**
+ * Event to create / configure elements before any other event
+ */
+USTRUCT(meta=(DisplayName="Post Construction", Category="Events", NodeColor="0.6, 0, 1", Keywords="Create,Build,Spawn,Setup,Init,Fit"))
+struct CONTROLRIG_API FRigUnit_PostPrepareForExecution : public FRigUnit
+{
+	GENERATED_BODY()
+
+	RIGVM_METHOD()
+	virtual void Execute() override;
+
+	virtual FName GetEventName() const override { return EventName; }
+	virtual bool CanOnlyExistOnce() const override { return true; }
+
+	// The execution result
+	UPROPERTY(EditAnywhere, Transient, DisplayName = "Execute", Category = "PostPrepareForExecution", meta = (Output))
+	FControlRigExecuteContext ExecuteContext;
+
+	static inline const FLazyName EventName = FLazyName(TEXT("PostConstruction"));
+};
diff --git a/Engine/Plugins/Animation/ControlRig/Source/ControlRigDeveloper/Private/Graph/ControlRigGraphSchema.cpp b/Engine/Plugins/Animation/ControlRig/Source/ControlRigDeveloper/Private/Graph/ControlRigGraphSchema.cpp
--- a/Engine/Plugins/Animation/ControlRig/Source/ControlRigDeveloper/Private/Graph/ControlRigGraphSchema.cpp
+++ b/Engine/Plugins/Animation/ControlRig/Source/ControlRigDeveloper/Private/Graph/ControlRigGraphSchema.cpp
@@ -52,7 +52,9 @@ bool UControlRigGraphSchema::IsRigVMDefaultEvent(const FName& InEventName) const
 		InEventName == FRigUnit_PreBeginExecution::EventName ||
 		InEventName == FRigUnit_PostBeginExecution::EventName ||
 		InEventName == FRigUnit_InverseExecution::EventName ||
+		InEventName == FRigUnit_PrePrepareForExecution::EventName ||
 		InEventName == FRigUnit_PrepareForExecution::EventName ||
+		InEventName == FRigUnit_PostPrepareForExecution::EventName ||
 		InEventName == FRigUnit_InteractionExecution::EventName ||
 		InEventName == FRigUnit_ConnectorExecution::EventName;
 }
