Control Rig: RigHierarchy not reacting to debugged rig

#rnx

diff --git a/Engine/Plugins/Animation/ControlRig/Source/ControlRigEditor/Private/Editor/SRigHierarchy.cpp b/Engine/Plugins/Animation/ControlRig/Source/ControlRigEditor/Private/Editor/SRigHierarchy.cpp
--- a/Engine/Plugins/Animation/ControlRig/Source/ControlRigEditor/Private/Editor/SRigHierarchy.cpp
+++ b/Engine/Plugins/Animation/ControlRig/Source/ControlRigEditor/Private/Editor/SRigHierarchy.cpp
@@ -249,6 +249,16 @@ void SRigHierarchy::Construct(const FArguments& InArgs, TSharedRef<FControlRigEd
 	
 	CreateContextMenu();
 	CreateDragDropMenu();
+
+	// after opening the editor the debugged rig won't exist yet. we'll have to wait for a tick so
+	// that we have a valid rig to listen to.
+	RegisterActiveTimer(0.f, FWidgetActiveTimerDelegate::CreateLambda([this](double, float) {
+		if(ControlRigBlueprint.IsValid())
+		{
+			(void)HandleSetObjectBeingDebugged(ControlRigBlueprint->GetDebuggedControlRig());
+		}
+		return EActiveTimerReturnType::Stop;
+	}));
 }
 
 void SRigHierarchy::OnEditorClose(const FRigVMEditor* InEditor, URigVMBlueprint* InBlueprint)
