Sequencer: Control Rig: Fix isssues with CR with Day Sequences, by making sure the track editor doesn't get set up with them, needed to move Sequencer setup out of constructor into OnInitialize()

diff --git a/Engine/Plugins/Animation/ControlRig/Source/ControlRigEditor/Private/Sequencer/ControlRigParameterTrackEditor.cpp b/Engine/Plugins/Animation/ControlRig/Source/ControlRigEditor/Private/Sequencer/ControlRigParameterTrackEditor.cpp
--- a/Engine/Plugins/Animation/ControlRig/Source/ControlRigEditor/Private/Sequencer/ControlRigParameterTrackEditor.cpp
+++ b/Engine/Plugins/Animation/ControlRig/Source/ControlRigEditor/Private/Sequencer/ControlRigParameterTrackEditor.cpp
@@ -261,15 +261,23 @@ FControlRigParameterTrackEditor::FControlRigParameterTrackEditor(TSharedRef<ISeq
 	{
 		GEditor->RegisterForUndo(this);
 	}
+}
 
-	UMovieScene* MovieScene = InSequencer->GetFocusedMovieSceneSequence()->GetMovieScene();
+void FControlRigParameterTrackEditor::OnInitialize()
+{
+	if (GetSequencer().IsValid() == false)
+	{
+		return;
+	}
+	ISequencer* SequencerPtr = GetSequencer().Get();
 
-	SelectionChangedHandle = InSequencer->GetSelectionChangedTracks().AddRaw(this, &FControlRigParameterTrackEditor::OnSelectionChanged);
-	SequencerChangedHandle = InSequencer->OnMovieSceneDataChanged().AddRaw(this, &FControlRigParameterTrackEditor::OnSequencerDataChanged);
-	OnActivateSequenceChangedHandle = InSequencer->OnActivateSequence().AddRaw(this, &FControlRigParameterTrackEditor::OnActivateSequenceChanged);
-	OnChannelChangedHandle = InSequencer->OnChannelChanged().AddRaw(this, &FControlRigParameterTrackEditor::OnChannelChanged);
+	UMovieScene* MovieScene = SequencerPtr->GetFocusedMovieSceneSequence()->GetMovieScene();
+	SelectionChangedHandle = SequencerPtr->GetSelectionChangedTracks().AddRaw(this, &FControlRigParameterTrackEditor::OnSelectionChanged);
+	SequencerChangedHandle = SequencerPtr->OnMovieSceneDataChanged().AddRaw(this, &FControlRigParameterTrackEditor::OnSequencerDataChanged);
+	OnActivateSequenceChangedHandle = SequencerPtr->OnActivateSequence().AddRaw(this, &FControlRigParameterTrackEditor::OnActivateSequenceChanged);
+	OnChannelChangedHandle = SequencerPtr->OnChannelChanged().AddRaw(this, &FControlRigParameterTrackEditor::OnChannelChanged);
 	OnMovieSceneChannelChangedHandle = MovieScene->OnChannelChanged().AddRaw(this, &FControlRigParameterTrackEditor::OnChannelChanged);
-	OnActorAddedToSequencerHandle = InSequencer->OnActorAddedToSequencer().AddRaw(this, &FControlRigParameterTrackEditor::HandleActorAdded);
+	OnActorAddedToSequencerHandle = SequencerPtr->OnActorAddedToSequencer().AddRaw(this, &FControlRigParameterTrackEditor::HandleActorAdded);
 
 	{
 		//we check for two things, one if the control rig has been replaced if so we need to switch.
@@ -466,12 +474,6 @@ FControlRigParameterTrackEditor::~FControlRigParameterTrackEditor()
 		GEditor->UnregisterForUndo(this);
 	}
 
-	UnbindAllControlRigs();
-	if (GetSequencer().IsValid())
-	{
-		//REMOVE ME IN UE5
-		GetSequencer()->GetObjectChangeListener().GetOnPropagateObjectChanges().RemoveAll(this);
-	}
 	FMovieSceneToolsModule::Get().UnregisterAnimationBakeHelper(this);
 }
 
@@ -659,6 +661,17 @@ TSharedRef<ISequencerTrackEditor> FControlRigParameterTrackEditor::CreateTrackEd
 	return MakeShareable(new FControlRigParameterTrackEditor(InSequencer));
 }
 
+bool FControlRigParameterTrackEditor::SupportsSequence(UMovieSceneSequence* InSequence) const
+{
+	ETrackSupport TrackSupported = InSequence ? InSequence->IsTrackSupported(UMovieSceneControlRigParameterTrack::StaticClass()) : ETrackSupport::Default;
+
+	if (TrackSupported == ETrackSupport::NotSupported)
+	{
+		return false;
+	}
+
+	return (InSequence && InSequence->IsA(ULevelSequence::StaticClass())) || TrackSupported == ETrackSupport::Supported;
+}
 
 bool FControlRigParameterTrackEditor::SupportsType(TSubclassOf<UMovieSceneTrack> Type) const
 {
diff --git a/Engine/Plugins/Animation/ControlRig/Source/ControlRigEditor/Private/Sequencer/ControlRigParameterTrackEditor.h b/Engine/Plugins/Animation/ControlRig/Source/ControlRigEditor/Private/Sequencer/ControlRigParameterTrackEditor.h
--- a/Engine/Plugins/Animation/ControlRig/Source/ControlRigEditor/Private/Sequencer/ControlRigParameterTrackEditor.h
+++ b/Engine/Plugins/Animation/ControlRig/Source/ControlRigEditor/Private/Sequencer/ControlRigParameterTrackEditor.h
@@ -63,11 +63,13 @@ public:
 public:
 
 	// ISequencerTrackEditor interface
+	virtual void OnInitialize() override;
 	virtual void OnRelease() override;
 	virtual void BuildObjectBindingContextMenu(FMenuBuilder& MenuBuilder, const TArray<FGuid>& ObjectBindings, const UClass* ObjectClass) override;
 	virtual void BuildObjectBindingTrackMenu(FMenuBuilder& MenuBuilder, const TArray<FGuid>& ObjectBindings, const UClass* ObjectClass) override;
 	virtual TSharedRef<ISequencerSection> MakeSectionInterface(UMovieSceneSection& SectionObject, UMovieSceneTrack& Track, FGuid ObjectBinding) override;
 	virtual bool SupportsType(TSubclassOf<UMovieSceneTrack> Type) const override;
+	virtual bool SupportsSequence(UMovieSceneSequence* InSequence) const override;
 	virtual bool HasTransformKeyBindings() const override { return true; }
 	virtual bool CanAddTransformKeysForSelectedObjects() const override;
 	virtual void OnAddTransformKeysForSelectedObjects(EMovieSceneTransformChannel Channel);
