Control Rig: Remove pre construction event and clean up event order

#rnx

diff --git a/Engine/Plugins/Animation/ControlRig/Source/ControlRig/Private/ControlRig.cpp b/Engine/Plugins/Animation/ControlRig/Source/ControlRig/Private/ControlRig.cpp
--- a/Engine/Plugins/Animation/ControlRig/Source/ControlRig/Private/ControlRig.cpp
+++ b/Engine/Plugins/Animation/ControlRig/Source/ControlRig/Private/ControlRig.cpp
@@ -749,6 +749,8 @@ bool UControlRig::Execute(const FName& InEventName)
 	const bool bIsEventFirstInQueue = !LocalEventQueueToRun.IsEmpty() && LocalEventQueueToRun[0] == InEventName; 
 	const bool bIsEventLastInQueue = !LocalEventQueueToRun.IsEmpty() && LocalEventQueueToRun.Last() == InEventName;
 	const bool bIsConstructionEvent = InEventName == FRigUnit_PrepareForExecution::EventName;
+	const bool bIsPostConstructionEvent = InEventName == FRigUnit_PostPrepareForExecution::EventName;
+	const bool bPostConstructionEventInQueue = LocalEventQueueToRun.Contains(FRigUnit_PostPrepareForExecution::EventName) && SupportsEvent(FRigUnit_PostPrepareForExecution::EventName);
 	const bool bPreForwardSolveInQueue = LocalEventQueueToRun.Contains(FRigUnit_PreBeginExecution::EventName);
 	const bool bPostForwardSolveInQueue = LocalEventQueueToRun.Contains(FRigUnit_PostBeginExecution::EventName);
 	const bool bIsPreForwardSolve = InEventName == FRigUnit_PreBeginExecution::EventName;
@@ -880,7 +882,7 @@ bool UControlRig::Execute(const FName& InEventName)
 	}
 
 	// disable any controller access outside of the construction event
-	FRigHierarchyEnableControllerBracket DisableHierarchyController(PublicContext.Hierarchy, bIsConstructionEvent);
+	FRigHierarchyEnableControllerBracket DisableHierarchyController(PublicContext.Hierarchy, bIsConstructionEvent || bIsPostConstructionEvent);
 
 	// given the outer scene component configure
 	// the transform lookups to map transforms from rig space to world space
@@ -958,7 +960,7 @@ bool UControlRig::Execute(const FName& InEventName)
 		UE_LOG(LogControlRig, Warning, TEXT("%s: Execute is being called recursively."), *GetPathName());
 		return false;
 	}
-	if(bIsConstructionEvent)
+	if(bIsConstructionEvent || bIsPostConstructionEvent)
 	{
 		if(IsRunningPreConstruction() || IsRunningPostConstruction())
 		{
@@ -1052,8 +1054,11 @@ bool UControlRig::Execute(const FName& InEventName)
 					bSuccess = Execute_Internal(FRigUnit_PrepareForExecution::EventName);
 					
 				} // destroy FTransientControlScope
-				
-				RunPostConstructionEvent();
+
+				if(!bPostConstructionEventInQueue)
+				{
+					RunPostConstructionEvent();
+				}
 
 				// Reset the input pose after construction
 				if (CurrentPose.Num() > 0)
@@ -1068,6 +1073,11 @@ bool UControlRig::Execute(const FName& InEventName)
 			
 		} // destroy DisableSelectionNotifications
 
+		if (bIsPostConstructionEvent || (bIsConstructionEvent && !bPostConstructionEventInQueue))
+		{
+			RunPostConstructionEvent();
+		}
+
 		if (bConstructionModeEnabled)
 		{
 #if WITH_EDITOR
@@ -1444,6 +1454,10 @@ bool UControlRig::Execute(const FName& InEventName)
 	{
 		RemoveRunOnceEvent(FRigUnit_PrepareForExecution::EventName);
 	}
+	if(bIsPostConstructionEvent)
+	{
+		RemoveRunOnceEvent(FRigUnit_PostPrepareForExecution::EventName);
+	}
 
 	return bSuccess;
 }
@@ -1586,6 +1600,7 @@ void UControlRig::RequestInit()
 void UControlRig::RequestConstruction()
 {
 	RequestRunOnceEvent(FRigUnit_PrepareForExecution::EventName, 0);
+	RequestRunOnceEvent(FRigUnit_PostPrepareForExecution::EventName, 1);
 }
 
 bool UControlRig::IsConstructionRequired() const
@@ -1606,11 +1621,6 @@ void UControlRig::AdaptEventQueueForEvaluate(TArray<FName>& InOutEventQueueToRun
 	{
 		if (InOutEventQueueToRun[i] == FRigUnit_PrepareForExecution::EventName)
 		{
-			if (SupportsEvent(FRigUnit_PrePrepareForExecution::EventName))
-			{
-				InOutEventQueueToRun.Insert(FRigUnit_PrePrepareForExecution::EventName, i);
-				i++; // skip preconstruction 
-			}
 			if (SupportsEvent(FRigUnit_PostPrepareForExecution::EventName))
 			{
 				i++; // skip construction
diff --git a/Engine/Plugins/Animation/ControlRig/Source/ControlRig/Private/Units/Execution/RigUnit_DynamicHierarchy.cpp b/Engine/Plugins/Animation/ControlRig/Source/ControlRig/Private/Units/Execution/RigUnit_DynamicHierarchy.cpp
--- a/Engine/Plugins/Animation/ControlRig/Source/ControlRig/Private/Units/Execution/RigUnit_DynamicHierarchy.cpp
+++ b/Engine/Plugins/Animation/ControlRig/Source/ControlRig/Private/Units/Execution/RigUnit_DynamicHierarchy.cpp
@@ -22,12 +22,12 @@ bool FRigUnit_DynamicHierarchyBase::IsValidToRunInContext(
 
 	if(bAllowOnlyConstructionEvent)
 	{
-		if(InExecuteContext.GetEventName() != FRigUnit_PrepareForExecution::EventName)
+		if(!InExecuteContext.IsRunningConstructionEvent())
 		{
 			if(OutErrorMessage)
 			{
-				static constexpr TCHAR ErrorMessageFormat[] = TEXT("Node can only run in %s Event");
-				*OutErrorMessage = FString::Printf(ErrorMessageFormat, *FRigUnit_PrepareForExecution::EventName.ToString());
+				static constexpr TCHAR ErrorMessageFormat[] = TEXT("Node can only run in %s or %s Event");
+				*OutErrorMessage = FString::Printf(ErrorMessageFormat, *FRigUnit_PrepareForExecution::EventName.ToString(), *FRigUnit_PostPrepareForExecution::EventName.ToString());
 			}
 			return false;
 		}
@@ -177,7 +177,7 @@ FRigUnit_SwitchParent_Execute()
 			}
 
 			// during construction event also change the initial weights
-			if(ExecuteContext.GetEventName() == FRigUnit_PrepareForExecution::EventName)
+			if(ExecuteContext.IsRunningConstructionEvent())
 			{
 				if(!ExecuteContext.Hierarchy->SwitchToParent(ChildElement, ParentElement, true, true, EmptyDependencyMap, &FailureReason))
 				{
@@ -251,7 +251,7 @@ FRigUnit_HierarchySetParentWeights_Execute()
 	ExecuteContext.Hierarchy->SetParentWeightArray(ChildElement, Weights, false, true);
 
 	// during construction event also change the initial weights
-	if(ExecuteContext.GetEventName() == FRigUnit_PrepareForExecution::EventName)
+	if(ExecuteContext.IsRunningConstructionEvent())
 	{
 		ExecuteContext.Hierarchy->SetParentWeightArray(ChildElement, Weights, true, true);
 	}
diff --git a/Engine/Plugins/Animation/ControlRig/Source/ControlRig/Private/Units/Execution/RigUnit_PrepareForExecution.cpp b/Engine/Plugins/Animation/ControlRig/Source/ControlRig/Private/Units/Execution/RigUnit_PrepareForExecution.cpp
--- a/Engine/Plugins/Animation/ControlRig/Source/ControlRig/Private/Units/Execution/RigUnit_PrepareForExecution.cpp
+++ b/Engine/Plugins/Animation/ControlRig/Source/ControlRig/Private/Units/Execution/RigUnit_PrepareForExecution.cpp
@@ -11,12 +11,6 @@ FRigUnit_PrepareForExecution_Execute()
 	ExecuteContext.SetEventName(FRigUnit_PrepareForExecution::EventName);
 }
 
-FRigUnit_PrePrepareForExecution_Execute()
-{
-	DECLARE_SCOPE_HIERARCHICAL_COUNTER_RIGUNIT()
-	ExecuteContext.SetEventName(FRigUnit_PrePrepareForExecution::EventName);
-}
-
 FRigUnit_PostPrepareForExecution_Execute()
 {
 	DECLARE_SCOPE_HIERARCHICAL_COUNTER_RIGUNIT()
diff --git a/Engine/Plugins/Animation/ControlRig/Source/ControlRig/Private/Units/RigUnitContext.cpp b/Engine/Plugins/Animation/ControlRig/Source/ControlRig/Private/Units/RigUnitContext.cpp
--- a/Engine/Plugins/Animation/ControlRig/Source/ControlRig/Private/Units/RigUnitContext.cpp
+++ b/Engine/Plugins/Animation/ControlRig/Source/ControlRig/Private/Units/RigUnitContext.cpp
@@ -3,9 +3,16 @@
 #include "Units/RigUnitContext.h"
 #include "ControlRig.h"
 #include "ModularRig.h"
+#include "Units/Execution/RigUnit_PrepareForExecution.h"
 
 #include UE_INLINE_GENERATED_CPP_BY_NAME(RigUnitContext)
 
+bool FControlRigExecuteContext::IsRunningConstructionEvent() const
+{
+	return GetEventName() == FRigUnit_PrepareForExecution::EventName ||
+		GetEventName() == FRigUnit_PostPrepareForExecution::EventName;
+}
+
 FName FControlRigExecuteContext::AddRigModuleNameSpace(const FName& InName) const
 {
 	if(IsRigModule())
diff --git a/Engine/Plugins/Animation/ControlRig/Source/ControlRig/Public/Units/Execution/RigUnit_PrepareForExecution.h b/Engine/Plugins/Animation/ControlRig/Source/ControlRig/Public/Units/Execution/RigUnit_PrepareForExecution.h
--- a/Engine/Plugins/Animation/ControlRig/Source/ControlRig/Public/Units/Execution/RigUnit_PrepareForExecution.h
+++ b/Engine/Plugins/Animation/ControlRig/Source/ControlRig/Public/Units/Execution/RigUnit_PrepareForExecution.h
@@ -27,27 +27,6 @@ struct CONTROLRIG_API FRigUnit_PrepareForExecution : public FRigUnit
 	static inline const FLazyName EventName = FLazyName(TEXT("Construction"));
 };
 
-/**
- * Event to create / configure elements before any other event
- */
-USTRUCT(meta=(DisplayName="Pre Construction", Category="Events", NodeColor="0.6, 0, 1", Keywords="Create,Build,Spawn,Setup,Init,Fit"))
-struct CONTROLRIG_API FRigUnit_PrePrepareForExecution : public FRigUnit
-{
-	GENERATED_BODY()
-
-	RIGVM_METHOD()
-	virtual void Execute() override;
-
-	virtual FName GetEventName() const override { return EventName; }
-	virtual bool CanOnlyExistOnce() const override { return true; }
-
-	// The execution result
-	UPROPERTY(EditAnywhere, Transient, DisplayName = "Execute", Category = "PrePrepareForExecution", meta = (Output))
-	FControlRigExecuteContext ExecuteContext;
-
-	static inline const FLazyName EventName = FLazyName(TEXT("PreConstruction"));
-};
-
 /**
  * Event to create / configure elements before any other event
  */
diff --git a/Engine/Plugins/Animation/ControlRig/Source/ControlRig/Public/Units/RigUnitContext.h b/Engine/Plugins/Animation/ControlRig/Source/ControlRig/Public/Units/RigUnitContext.h
--- a/Engine/Plugins/Animation/ControlRig/Source/ControlRig/Public/Units/RigUnitContext.h
+++ b/Engine/Plugins/Animation/ControlRig/Source/ControlRig/Public/Units/RigUnitContext.h
@@ -161,6 +161,11 @@ public:
 		return nullptr;
 	}
 
+	/**
+	 * Returns true if the event currently running is considered a construction event
+	 */
+	bool IsRunningConstructionEvent() const;
+
 	/**
 	 * Add the namespace from a given name
 	 */
diff --git a/Engine/Plugins/Animation/ControlRig/Source/ControlRigDeveloper/Private/Graph/ControlRigGraphSchema.cpp b/Engine/Plugins/Animation/ControlRig/Source/ControlRigDeveloper/Private/Graph/ControlRigGraphSchema.cpp
--- a/Engine/Plugins/Animation/ControlRig/Source/ControlRigDeveloper/Private/Graph/ControlRigGraphSchema.cpp
+++ b/Engine/Plugins/Animation/ControlRig/Source/ControlRigDeveloper/Private/Graph/ControlRigGraphSchema.cpp
@@ -52,7 +52,6 @@ bool UControlRigGraphSchema::IsRigVMDefaultEvent(const FName& InEventName) const
 		InEventName == FRigUnit_PreBeginExecution::EventName ||
 		InEventName == FRigUnit_PostBeginExecution::EventName ||
 		InEventName == FRigUnit_InverseExecution::EventName ||
-		InEventName == FRigUnit_PrePrepareForExecution::EventName ||
 		InEventName == FRigUnit_PrepareForExecution::EventName ||
 		InEventName == FRigUnit_PostPrepareForExecution::EventName ||
 		InEventName == FRigUnit_InteractionExecution::EventName ||
