Control Rig: Offer get / set animation channel nodes from item

#rnx

diff --git a/Engine/Plugins/Animation/ControlRig/Source/ControlRig/Private/Units/Hierarchy/RigUnit_ControlChannel.cpp b/Engine/Plugins/Animation/ControlRig/Source/ControlRig/Private/Units/Hierarchy/RigUnit_ControlChannel.cpp
--- a/Engine/Plugins/Animation/ControlRig/Source/ControlRig/Private/Units/Hierarchy/RigUnit_ControlChannel.cpp
+++ b/Engine/Plugins/Animation/ControlRig/Source/ControlRig/Private/Units/Hierarchy/RigUnit_ControlChannel.cpp
@@ -1,6 +1,7 @@
 // Copyright Epic Games, Inc. All Rights Reserved.
 
 #include "Units/Hierarchy/RigUnit_ControlChannel.h"
+#include "Units/Hierarchy/RigUnit_ControlChannelFromItem.h"
 #include "Units/RigUnitContext.h"
 
 #include UE_INLINE_GENERATED_CPP_BY_NAME(RigUnit_ControlChannel)
@@ -56,19 +57,7 @@ FRigUnit_GetBoolAnimationChannel_Execute()
 	{
 		return;
 	}
-
-	if(const FRigControlElement* ChannelElement = ExecuteContext.Hierarchy->Find<FRigControlElement>(CachedChannelKey))
-	{
-		if(ChannelElement->Settings.ControlType == ERigControlType::Bool)
-		{
-			const FRigControlValue StoredValue = ExecuteContext.Hierarchy->GetControlValueByIndex(ChannelElement->GetIndex(), bInitial ? ERigControlValueType::Initial : ERigControlValueType::Current);
-			Value = StoredValue.Get<bool>();
-		}
-		else
-		{
-			UE_CONTROLRIG_RIGUNIT_REPORT_ERROR(TEXT("Animation Channel %s is not a Bool."), *CachedChannelKey.ToString());
-		}
-	}
+	FRigUnit_GetBoolAnimationChannelFromItem::StaticExecute(ExecuteContext, Value, CachedChannelKey, bInitial);
 }
 
 FRigUnit_GetFloatAnimationChannel_Execute()
@@ -79,19 +68,7 @@ FRigUnit_GetFloatAnimationChannel_Execute()
 	{
 		return;
 	}
-
-	if(const FRigControlElement* ChannelElement = ExecuteContext.Hierarchy->Find<FRigControlElement>(CachedChannelKey))
-	{
-		if(ChannelElement->Settings.ControlType == ERigControlType::Float || ChannelElement->Settings.ControlType == ERigControlType::ScaleFloat)
-		{
-			const FRigControlValue StoredValue = ExecuteContext.Hierarchy->GetControlValueByIndex(ChannelElement->GetIndex(), bInitial ? ERigControlValueType::Initial : ERigControlValueType::Current);
-			Value = StoredValue.Get<float>();
-		}
-		else
-		{
-			UE_CONTROLRIG_RIGUNIT_REPORT_ERROR(TEXT("Animation Channel %s is not a Float."), *CachedChannelKey.ToString());
-		}
-	}
+	FRigUnit_GetFloatAnimationChannelFromItem::StaticExecute(ExecuteContext, Value, CachedChannelKey, bInitial);
 }
 
 FRigUnit_GetIntAnimationChannel_Execute()
@@ -102,19 +79,7 @@ FRigUnit_GetIntAnimationChannel_Execute()
 	{
 		return;
 	}
-
-	if(const FRigControlElement* ChannelElement = ExecuteContext.Hierarchy->Find<FRigControlElement>(CachedChannelKey))
-	{
-		if(ChannelElement->Settings.ControlType == ERigControlType::Integer)
-		{
-			const FRigControlValue StoredValue = ExecuteContext.Hierarchy->GetControlValueByIndex(ChannelElement->GetIndex(), bInitial ? ERigControlValueType::Initial : ERigControlValueType::Current);
-			Value = StoredValue.Get<int32>();
-		}
-		else
-		{
-			UE_CONTROLRIG_RIGUNIT_REPORT_ERROR(TEXT("Animation Channel %s is not an Integer (or Enum)."), *CachedChannelKey.ToString());
-		}
-	}
+	FRigUnit_GetIntAnimationChannelFromItem::StaticExecute(ExecuteContext, Value, CachedChannelKey, bInitial);
 }
 
 FRigUnit_GetVector2DAnimationChannel_Execute()
@@ -125,19 +90,7 @@ FRigUnit_GetVector2DAnimationChannel_Execute()
 	{
 		return;
 	}
-
-	if(const FRigControlElement* ChannelElement = ExecuteContext.Hierarchy->Find<FRigControlElement>(CachedChannelKey))
-	{
-		if(ChannelElement->Settings.ControlType == ERigControlType::Vector2D)
-		{
-			const FRigControlValue StoredValue = ExecuteContext.Hierarchy->GetControlValueByIndex(ChannelElement->GetIndex(), bInitial ? ERigControlValueType::Initial : ERigControlValueType::Current);
-			Value = FVector2D(StoredValue.Get<FVector2f>());
-		}
-		else
-		{
-			UE_CONTROLRIG_RIGUNIT_REPORT_ERROR(TEXT("Animation Channel %s is not a Vector2D."), *CachedChannelKey.ToString());
-		}
-	}
+	FRigUnit_GetVector2DAnimationChannelFromItem::StaticExecute(ExecuteContext, Value, CachedChannelKey, bInitial);
 }
 
 FRigUnit_GetVectorAnimationChannel_Execute()
@@ -148,19 +101,7 @@ FRigUnit_GetVectorAnimationChannel_Execute()
 	{
 		return;
 	}
-
-	if(const FRigControlElement* ChannelElement = ExecuteContext.Hierarchy->Find<FRigControlElement>(CachedChannelKey))
-	{
-		if(ChannelElement->Settings.ControlType == ERigControlType::Position || ChannelElement->Settings.ControlType == ERigControlType::Scale)
-		{
-			const FRigControlValue StoredValue = ExecuteContext.Hierarchy->GetControlValueByIndex(ChannelElement->GetIndex(), bInitial ? ERigControlValueType::Initial : ERigControlValueType::Current);
-			Value = FVector(StoredValue.Get<FVector3f>());
-		}
-		else
-		{
-			UE_CONTROLRIG_RIGUNIT_REPORT_ERROR(TEXT("Animation Channel %s is not a Vector (Position or Scale)."), *CachedChannelKey.ToString());
-		}
-	}
+	FRigUnit_GetVectorAnimationChannelFromItem::StaticExecute(ExecuteContext, Value, CachedChannelKey, bInitial);
 }
 
 FRigUnit_GetRotatorAnimationChannel_Execute()
@@ -171,19 +112,7 @@ FRigUnit_GetRotatorAnimationChannel_Execute()
 	{
 		return;
 	}
-
-	if(const FRigControlElement* ChannelElement = ExecuteContext.Hierarchy->Find<FRigControlElement>(CachedChannelKey))
-	{
-		if(ChannelElement->Settings.ControlType == ERigControlType::Rotator)
-		{
-			const FRigControlValue StoredValue = ExecuteContext.Hierarchy->GetControlValueByIndex(ChannelElement->GetIndex(), bInitial ? ERigControlValueType::Initial : ERigControlValueType::Current);
-			Value = FRotator::MakeFromEuler(FVector(StoredValue.Get<FVector3f>()));
-		}
-		else
-		{
-			UE_CONTROLRIG_RIGUNIT_REPORT_ERROR(TEXT("Animation Channel %s is not a Rotator."), *CachedChannelKey.ToString());
-		}
-	}
+	FRigUnit_GetRotatorAnimationChannelFromItem::StaticExecute(ExecuteContext, Value, CachedChannelKey, bInitial);
 }
 
 FRigUnit_GetTransformAnimationChannel_Execute()
@@ -194,24 +123,7 @@ FRigUnit_GetTransformAnimationChannel_Execute()
 	{
 		return;
 	}
-
-	if(const FRigControlElement* ChannelElement = ExecuteContext.Hierarchy->Find<FRigControlElement>(CachedChannelKey))
-	{
-		if(ChannelElement->Settings.ControlType == ERigControlType::Transform)
-		{
-			const FRigControlValue StoredValue = ExecuteContext.Hierarchy->GetControlValueByIndex(ChannelElement->GetIndex(), bInitial ? ERigControlValueType::Initial : ERigControlValueType::Current);
-			Value = StoredValue.Get<FRigControlValue::FTransform_Float>().ToTransform();
-		}
-		else if(ChannelElement->Settings.ControlType == ERigControlType::EulerTransform)
-		{
-			const FRigControlValue StoredValue = ExecuteContext.Hierarchy->GetControlValueByIndex(ChannelElement->GetIndex(), bInitial ? ERigControlValueType::Initial : ERigControlValueType::Current);
-			Value = StoredValue.Get<FRigControlValue::FEulerTransform_Float>().ToTransform().ToFTransform();
-		}
-		else
-		{
-			UE_CONTROLRIG_RIGUNIT_REPORT_ERROR(TEXT("Animation Channel %s is not a EulerTransform / Transform."), *CachedChannelKey.ToString());
-		}
-	}
+	FRigUnit_GetTransformAnimationChannelFromItem::StaticExecute(ExecuteContext, Value, CachedChannelKey, bInitial);
 }
 
 FRigUnit_SetBoolAnimationChannel_Execute()
@@ -220,19 +132,7 @@ FRigUnit_SetBoolAnimationChannel_Execute()
 	{
 		return;
 	}
-
-	if(const FRigControlElement* ChannelElement = ExecuteContext.Hierarchy->Find<FRigControlElement>(CachedChannelKey))
-	{
-		if(ChannelElement->Settings.ControlType == ERigControlType::Bool)
-		{
-			const FRigControlValue ValueToStore = FRigControlValue::Make<bool>(Value);
-			ExecuteContext.Hierarchy->SetControlValueByIndex(ChannelElement->GetIndex(), ValueToStore, bInitial ? ERigControlValueType::Initial : ERigControlValueType::Current);
-		}
-		else
-		{
-			UE_CONTROLRIG_RIGUNIT_REPORT_ERROR(TEXT("Animation Channel %s is not a Bool."), *CachedChannelKey.ToString());
-		}
-	}
+	FRigUnit_SetBoolAnimationChannelFromItem::StaticExecute(ExecuteContext, Value, CachedChannelKey, bInitial);
 }
 
 FRigUnit_SetFloatAnimationChannel_Execute()
@@ -241,19 +141,7 @@ FRigUnit_SetFloatAnimationChannel_Execute()
 	{
 		return;
 	}
-
-	if(const FRigControlElement* ChannelElement = ExecuteContext.Hierarchy->Find<FRigControlElement>(CachedChannelKey))
-	{
-		if(ChannelElement->Settings.ControlType == ERigControlType::Float || ChannelElement->Settings.ControlType == ERigControlType::ScaleFloat)
-		{
-			const FRigControlValue ValueToStore = FRigControlValue::Make<float>(Value);
-			ExecuteContext.Hierarchy->SetControlValueByIndex(ChannelElement->GetIndex(), ValueToStore, bInitial ? ERigControlValueType::Initial : ERigControlValueType::Current);
-		}
-		else
-		{
-			UE_CONTROLRIG_RIGUNIT_REPORT_ERROR(TEXT("Animation Channel %s is not a Float."), *CachedChannelKey.ToString());
-		}
-	}
+	FRigUnit_SetFloatAnimationChannelFromItem::StaticExecute(ExecuteContext, Value, CachedChannelKey, bInitial);
 }
 
 FRigUnit_SetIntAnimationChannel_Execute()
@@ -262,19 +150,7 @@ FRigUnit_SetIntAnimationChannel_Execute()
 	{
 		return;
 	}
-
-	if(const FRigControlElement* ChannelElement = ExecuteContext.Hierarchy->Find<FRigControlElement>(CachedChannelKey))
-	{
-		if(ChannelElement->Settings.ControlType == ERigControlType::Integer)
-		{
-			const FRigControlValue ValueToStore = FRigControlValue::Make<int32>(Value);
-			ExecuteContext.Hierarchy->SetControlValueByIndex(ChannelElement->GetIndex(), ValueToStore, bInitial ? ERigControlValueType::Initial : ERigControlValueType::Current);
-		}
-		else
-		{
-			UE_CONTROLRIG_RIGUNIT_REPORT_ERROR(TEXT("Animation Channel %s is not an Integer (or Enum)."), *CachedChannelKey.ToString());
-		}
-	}
+	FRigUnit_SetIntAnimationChannelFromItem::StaticExecute(ExecuteContext, Value, CachedChannelKey, bInitial);
 }
 
 FRigUnit_SetVector2DAnimationChannel_Execute()
@@ -283,19 +159,7 @@ FRigUnit_SetVector2DAnimationChannel_Execute()
 	{
 		return;
 	}
-
-	if(const FRigControlElement* ChannelElement = ExecuteContext.Hierarchy->Find<FRigControlElement>(CachedChannelKey))
-	{
-		if(ChannelElement->Settings.ControlType == ERigControlType::Vector2D)
-		{
-			const FRigControlValue ValueToStore = FRigControlValue::Make<FVector2f>(FVector2f(Value));
-			ExecuteContext.Hierarchy->SetControlValueByIndex(ChannelElement->GetIndex(), ValueToStore, bInitial ? ERigControlValueType::Initial : ERigControlValueType::Current);
-		}
-		else
-		{
-			UE_CONTROLRIG_RIGUNIT_REPORT_ERROR(TEXT("Animation Channel %s is not a Vector2D."), *CachedChannelKey.ToString());
-		}
-	}
+	FRigUnit_SetVector2DAnimationChannelFromItem::StaticExecute(ExecuteContext, Value, CachedChannelKey, bInitial);
 }
 
 FRigUnit_SetVectorAnimationChannel_Execute()
@@ -304,19 +168,7 @@ FRigUnit_SetVectorAnimationChannel_Execute()
 	{
 		return;
 	}
-
-	if(const FRigControlElement* ChannelElement = ExecuteContext.Hierarchy->Find<FRigControlElement>(CachedChannelKey))
-	{
-		if(ChannelElement->Settings.ControlType == ERigControlType::Position || ChannelElement->Settings.ControlType == ERigControlType::Scale)
-		{
-			const FRigControlValue ValueToStore = FRigControlValue::Make<FVector3f>(FVector3f(Value));
-			ExecuteContext.Hierarchy->SetControlValueByIndex(ChannelElement->GetIndex(), ValueToStore, bInitial ? ERigControlValueType::Initial : ERigControlValueType::Current);
-		}
-		else
-		{
-			UE_CONTROLRIG_RIGUNIT_REPORT_ERROR(TEXT("Animation Channel %s is not a Vector (Position or Scale)."), *CachedChannelKey.ToString());
-		}
-	}
+	FRigUnit_SetVectorAnimationChannelFromItem::StaticExecute(ExecuteContext, Value, CachedChannelKey, bInitial);
 }
 
 FRigUnit_SetRotatorAnimationChannel_Execute()
@@ -325,19 +177,7 @@ FRigUnit_SetRotatorAnimationChannel_Execute()
 	{
 		return;
 	}
-
-	if(const FRigControlElement* ChannelElement = ExecuteContext.Hierarchy->Find<FRigControlElement>(CachedChannelKey))
-	{
-		if(ChannelElement->Settings.ControlType == ERigControlType::Rotator)
-		{
-			const FRigControlValue ValueToStore = FRigControlValue::Make<FVector3f>(FVector3f(Value.Euler()));
-			ExecuteContext.Hierarchy->SetControlValueByIndex(ChannelElement->GetIndex(), ValueToStore, bInitial ? ERigControlValueType::Initial : ERigControlValueType::Current);
-		}
-		else
-		{
-			UE_CONTROLRIG_RIGUNIT_REPORT_ERROR(TEXT("Animation Channel %s is not a Rotator."), *CachedChannelKey.ToString());
-		}
-	}
+	FRigUnit_SetRotatorAnimationChannelFromItem::StaticExecute(ExecuteContext, Value, CachedChannelKey, bInitial);
 }
 
 FRigUnit_SetTransformAnimationChannel_Execute()
@@ -346,22 +186,5 @@ FRigUnit_SetTransformAnimationChannel_Execute()
 	{
 		return;
 	}
-
-	if(const FRigControlElement* ChannelElement = ExecuteContext.Hierarchy->Find<FRigControlElement>(CachedChannelKey))
-	{
-		if(ChannelElement->Settings.ControlType == ERigControlType::Transform)
-		{
-			const FRigControlValue ValueToStore = FRigControlValue::Make<FRigControlValue::FTransform_Float>(Value);
-			ExecuteContext.Hierarchy->SetControlValueByIndex(ChannelElement->GetIndex(), ValueToStore, bInitial ? ERigControlValueType::Initial : ERigControlValueType::Current);
-		}
-		else if(ChannelElement->Settings.ControlType == ERigControlType::EulerTransform)
-		{
-			const FRigControlValue ValueToStore = FRigControlValue::Make<FRigControlValue::FEulerTransform_Float>(FEulerTransform(Value));
-			ExecuteContext.Hierarchy->SetControlValueByIndex(ChannelElement->GetIndex(), ValueToStore, bInitial ? ERigControlValueType::Initial : ERigControlValueType::Current);
-		}
-		else
-		{
-			UE_CONTROLRIG_RIGUNIT_REPORT_ERROR(TEXT("Animation Channel %s is not a EulerTransform / Transform."), *CachedChannelKey.ToString());
-		}
-	}
+	FRigUnit_SetTransformAnimationChannelFromItem::StaticExecute(ExecuteContext, Value, CachedChannelKey, bInitial);
 }
diff --git a/Engine/Plugins/Animation/ControlRig/Source/ControlRig/Private/Units/Hierarchy/RigUnit_ControlChannelFromItem.cpp b/Engine/Plugins/Animation/ControlRig/Source/ControlRig/Private/Units/Hierarchy/RigUnit_ControlChannelFromItem.cpp
new file mode 100644
--- /dev/null
+++ b/Engine/Plugins/Animation/ControlRig/Source/ControlRig/Private/Units/Hierarchy/RigUnit_ControlChannelFromItem.cpp
@@ -0,0 +1,254 @@
+// Copyright Epic Games, Inc. All Rights Reserved.
+
+#include "Units/Hierarchy/RigUnit_ControlChannelFromItem.h"
+#include "Units/RigUnitContext.h"
+
+#include UE_INLINE_GENERATED_CPP_BY_NAME(RigUnit_ControlChannelFromItem)
+
+FRigUnit_GetBoolAnimationChannelFromItem_Execute()
+{
+	Value = false;
+	
+	if(const FRigControlElement* ChannelElement = ExecuteContext.Hierarchy->Find<FRigControlElement>(Item))
+	{
+		if(ChannelElement->Settings.ControlType == ERigControlType::Bool)
+		{
+			const FRigControlValue StoredValue = ExecuteContext.Hierarchy->GetControlValueByIndex(ChannelElement->GetIndex(), bInitial ? ERigControlValueType::Initial : ERigControlValueType::Current);
+			Value = StoredValue.Get<bool>();
+		}
+		else
+		{
+			UE_CONTROLRIG_RIGUNIT_REPORT_ERROR(TEXT("Animation Channel %s is not a Bool."), *Item.ToString());
+		}
+	}
+}
+
+FRigUnit_GetFloatAnimationChannelFromItem_Execute()
+{
+	Value = 0.f;
+	
+	if(const FRigControlElement* ChannelElement = ExecuteContext.Hierarchy->Find<FRigControlElement>(Item))
+	{
+		if(ChannelElement->Settings.ControlType == ERigControlType::Float || ChannelElement->Settings.ControlType == ERigControlType::ScaleFloat)
+		{
+			const FRigControlValue StoredValue = ExecuteContext.Hierarchy->GetControlValueByIndex(ChannelElement->GetIndex(), bInitial ? ERigControlValueType::Initial : ERigControlValueType::Current);
+			Value = StoredValue.Get<float>();
+		}
+		else
+		{
+			UE_CONTROLRIG_RIGUNIT_REPORT_ERROR(TEXT("Animation Channel %s is not a Float."), *Item.ToString());
+		}
+	}
+}
+
+FRigUnit_GetIntAnimationChannelFromItem_Execute()
+{
+	Value = 0;
+	
+	if(const FRigControlElement* ChannelElement = ExecuteContext.Hierarchy->Find<FRigControlElement>(Item))
+	{
+		if(ChannelElement->Settings.ControlType == ERigControlType::Integer)
+		{
+			const FRigControlValue StoredValue = ExecuteContext.Hierarchy->GetControlValueByIndex(ChannelElement->GetIndex(), bInitial ? ERigControlValueType::Initial : ERigControlValueType::Current);
+			Value = StoredValue.Get<int32>();
+		}
+		else
+		{
+			UE_CONTROLRIG_RIGUNIT_REPORT_ERROR(TEXT("Animation Channel %s is not an Integer (or Enum)."), *Item.ToString());
+		}
+	}
+}
+
+FRigUnit_GetVector2DAnimationChannelFromItem_Execute()
+{
+	Value = FVector2D::ZeroVector;
+	
+	if(const FRigControlElement* ChannelElement = ExecuteContext.Hierarchy->Find<FRigControlElement>(Item))
+	{
+		if(ChannelElement->Settings.ControlType == ERigControlType::Vector2D)
+		{
+			const FRigControlValue StoredValue = ExecuteContext.Hierarchy->GetControlValueByIndex(ChannelElement->GetIndex(), bInitial ? ERigControlValueType::Initial : ERigControlValueType::Current);
+			Value = FVector2D(StoredValue.Get<FVector2f>());
+		}
+		else
+		{
+			UE_CONTROLRIG_RIGUNIT_REPORT_ERROR(TEXT("Animation Channel %s is not a Vector2D."), *Item.ToString());
+		}
+	}
+}
+
+FRigUnit_GetVectorAnimationChannelFromItem_Execute()
+{
+	Value = FVector::ZeroVector;
+	
+	if(const FRigControlElement* ChannelElement = ExecuteContext.Hierarchy->Find<FRigControlElement>(Item))
+	{
+		if(ChannelElement->Settings.ControlType == ERigControlType::Position || ChannelElement->Settings.ControlType == ERigControlType::Scale)
+		{
+			const FRigControlValue StoredValue = ExecuteContext.Hierarchy->GetControlValueByIndex(ChannelElement->GetIndex(), bInitial ? ERigControlValueType::Initial : ERigControlValueType::Current);
+			Value = FVector(StoredValue.Get<FVector3f>());
+		}
+		else
+		{
+			UE_CONTROLRIG_RIGUNIT_REPORT_ERROR(TEXT("Animation Channel %s is not a Vector (Position or Scale)."), *Item.ToString());
+		}
+	}
+}
+
+FRigUnit_GetRotatorAnimationChannelFromItem_Execute()
+{
+	Value = FRotator::ZeroRotator;
+	
+	if(const FRigControlElement* ChannelElement = ExecuteContext.Hierarchy->Find<FRigControlElement>(Item))
+	{
+		if(ChannelElement->Settings.ControlType == ERigControlType::Rotator)
+		{
+			const FRigControlValue StoredValue = ExecuteContext.Hierarchy->GetControlValueByIndex(ChannelElement->GetIndex(), bInitial ? ERigControlValueType::Initial : ERigControlValueType::Current);
+			Value = FRotator::MakeFromEuler(FVector(StoredValue.Get<FVector3f>()));
+		}
+		else
+		{
+			UE_CONTROLRIG_RIGUNIT_REPORT_ERROR(TEXT("Animation Channel %s is not a Rotator."), *Item.ToString());
+		}
+	}
+}
+
+FRigUnit_GetTransformAnimationChannelFromItem_Execute()
+{
+	Value = FTransform::Identity;
+	
+	if(const FRigControlElement* ChannelElement = ExecuteContext.Hierarchy->Find<FRigControlElement>(Item))
+	{
+		if(ChannelElement->Settings.ControlType == ERigControlType::Transform)
+		{
+			const FRigControlValue StoredValue = ExecuteContext.Hierarchy->GetControlValueByIndex(ChannelElement->GetIndex(), bInitial ? ERigControlValueType::Initial : ERigControlValueType::Current);
+			Value = StoredValue.Get<FRigControlValue::FTransform_Float>().ToTransform();
+		}
+		else if(ChannelElement->Settings.ControlType == ERigControlType::EulerTransform)
+		{
+			const FRigControlValue StoredValue = ExecuteContext.Hierarchy->GetControlValueByIndex(ChannelElement->GetIndex(), bInitial ? ERigControlValueType::Initial : ERigControlValueType::Current);
+			Value = StoredValue.Get<FRigControlValue::FEulerTransform_Float>().ToTransform().ToFTransform();
+		}
+		else
+		{
+			UE_CONTROLRIG_RIGUNIT_REPORT_ERROR(TEXT("Animation Channel %s is not a EulerTransform / Transform."), *Item.ToString());
+		}
+	}
+}
+
+FRigUnit_SetBoolAnimationChannelFromItem_Execute()
+{
+	if(const FRigControlElement* ChannelElement = ExecuteContext.Hierarchy->Find<FRigControlElement>(Item))
+	{
+		if(ChannelElement->Settings.ControlType == ERigControlType::Bool)
+		{
+			const FRigControlValue ValueToStore = FRigControlValue::Make<bool>(Value);
+			ExecuteContext.Hierarchy->SetControlValueByIndex(ChannelElement->GetIndex(), ValueToStore, bInitial ? ERigControlValueType::Initial : ERigControlValueType::Current);
+		}
+		else
+		{
+			UE_CONTROLRIG_RIGUNIT_REPORT_ERROR(TEXT("Animation Channel %s is not a Bool."), *Item.ToString());
+		}
+	}
+}
+
+FRigUnit_SetFloatAnimationChannelFromItem_Execute()
+{
+	if(const FRigControlElement* ChannelElement = ExecuteContext.Hierarchy->Find<FRigControlElement>(Item))
+	{
+		if(ChannelElement->Settings.ControlType == ERigControlType::Float || ChannelElement->Settings.ControlType == ERigControlType::ScaleFloat)
+		{
+			const FRigControlValue ValueToStore = FRigControlValue::Make<float>(Value);
+			ExecuteContext.Hierarchy->SetControlValueByIndex(ChannelElement->GetIndex(), ValueToStore, bInitial ? ERigControlValueType::Initial : ERigControlValueType::Current);
+		}
+		else
+		{
+			UE_CONTROLRIG_RIGUNIT_REPORT_ERROR(TEXT("Animation Channel %s is not a Float."), *Item.ToString());
+		}
+	}
+}
+
+FRigUnit_SetIntAnimationChannelFromItem_Execute()
+{
+	if(const FRigControlElement* ChannelElement = ExecuteContext.Hierarchy->Find<FRigControlElement>(Item))
+	{
+		if(ChannelElement->Settings.ControlType == ERigControlType::Integer)
+		{
+			const FRigControlValue ValueToStore = FRigControlValue::Make<int32>(Value);
+			ExecuteContext.Hierarchy->SetControlValueByIndex(ChannelElement->GetIndex(), ValueToStore, bInitial ? ERigControlValueType::Initial : ERigControlValueType::Current);
+		}
+		else
+		{
+			UE_CONTROLRIG_RIGUNIT_REPORT_ERROR(TEXT("Animation Channel %s is not an Integer (or Enum)."), *Item.ToString());
+		}
+	}
+}
+
+FRigUnit_SetVector2DAnimationChannelFromItem_Execute()
+{
+	if(const FRigControlElement* ChannelElement = ExecuteContext.Hierarchy->Find<FRigControlElement>(Item))
+	{
+		if(ChannelElement->Settings.ControlType == ERigControlType::Vector2D)
+		{
+			const FRigControlValue ValueToStore = FRigControlValue::Make<FVector2f>(FVector2f(Value));
+			ExecuteContext.Hierarchy->SetControlValueByIndex(ChannelElement->GetIndex(), ValueToStore, bInitial ? ERigControlValueType::Initial : ERigControlValueType::Current);
+		}
+		else
+		{
+			UE_CONTROLRIG_RIGUNIT_REPORT_ERROR(TEXT("Animation Channel %s is not a Vector2D."), *Item.ToString());
+		}
+	}
+}
+
+FRigUnit_SetVectorAnimationChannelFromItem_Execute()
+{
+	if(const FRigControlElement* ChannelElement = ExecuteContext.Hierarchy->Find<FRigControlElement>(Item))
+	{
+		if(ChannelElement->Settings.ControlType == ERigControlType::Position || ChannelElement->Settings.ControlType == ERigControlType::Scale)
+		{
+			const FRigControlValue ValueToStore = FRigControlValue::Make<FVector3f>(FVector3f(Value));
+			ExecuteContext.Hierarchy->SetControlValueByIndex(ChannelElement->GetIndex(), ValueToStore, bInitial ? ERigControlValueType::Initial : ERigControlValueType::Current);
+		}
+		else
+		{
+			UE_CONTROLRIG_RIGUNIT_REPORT_ERROR(TEXT("Animation Channel %s is not a Vector (Position or Scale)."), *Item.ToString());
+		}
+	}
+}
+
+FRigUnit_SetRotatorAnimationChannelFromItem_Execute()
+{
+	if(const FRigControlElement* ChannelElement = ExecuteContext.Hierarchy->Find<FRigControlElement>(Item))
+	{
+		if(ChannelElement->Settings.ControlType == ERigControlType::Rotator)
+		{
+			const FRigControlValue ValueToStore = FRigControlValue::Make<FVector3f>(FVector3f(Value.Euler()));
+			ExecuteContext.Hierarchy->SetControlValueByIndex(ChannelElement->GetIndex(), ValueToStore, bInitial ? ERigControlValueType::Initial : ERigControlValueType::Current);
+		}
+		else
+		{
+			UE_CONTROLRIG_RIGUNIT_REPORT_ERROR(TEXT("Animation Channel %s is not a Rotator."), *Item.ToString());
+		}
+	}
+}
+
+FRigUnit_SetTransformAnimationChannelFromItem_Execute()
+{
+	if(const FRigControlElement* ChannelElement = ExecuteContext.Hierarchy->Find<FRigControlElement>(Item))
+	{
+		if(ChannelElement->Settings.ControlType == ERigControlType::Transform)
+		{
+			const FRigControlValue ValueToStore = FRigControlValue::Make<FRigControlValue::FTransform_Float>(Value);
+			ExecuteContext.Hierarchy->SetControlValueByIndex(ChannelElement->GetIndex(), ValueToStore, bInitial ? ERigControlValueType::Initial : ERigControlValueType::Current);
+		}
+		else if(ChannelElement->Settings.ControlType == ERigControlType::EulerTransform)
+		{
+			const FRigControlValue ValueToStore = FRigControlValue::Make<FRigControlValue::FEulerTransform_Float>(FEulerTransform(Value));
+			ExecuteContext.Hierarchy->SetControlValueByIndex(ChannelElement->GetIndex(), ValueToStore, bInitial ? ERigControlValueType::Initial : ERigControlValueType::Current);
+		}
+		else
+		{
+			UE_CONTROLRIG_RIGUNIT_REPORT_ERROR(TEXT("Animation Channel %s is not a EulerTransform / Transform."), *Item.ToString());
+		}
+	}
+}
diff --git a/Engine/Plugins/Animation/ControlRig/Source/ControlRig/Public/Units/Hierarchy/RigUnit_ControlChannelFromItem.h b/Engine/Plugins/Animation/ControlRig/Source/ControlRig/Public/Units/Hierarchy/RigUnit_ControlChannelFromItem.h
new file mode 100644
--- /dev/null
+++ b/Engine/Plugins/Animation/ControlRig/Source/ControlRig/Public/Units/Hierarchy/RigUnit_ControlChannelFromItem.h
@@ -0,0 +1,344 @@
+// Copyright Epic Games, Inc. All Rights Reserved.
+
+#pragma once
+
+#include "Units/RigUnit.h"
+#include "RigUnit_ControlChannelFromItem.generated.h"
+
+/**
+ * Get Animation Channel is used to retrieve a control's animation channel value
+ */
+USTRUCT(meta=(Abstract, Category="Controls", DocumentationPolicy = "Strict", NodeColor="0.462745, 1,0, 0.329412",Varying))
+struct CONTROLRIG_API FRigUnit_GetAnimationChannelFromItemBase : public FRigUnit
+{
+	GENERATED_BODY()
+
+	FRigUnit_GetAnimationChannelFromItemBase()
+	{
+		Item = FRigElementKey(NAME_None, ERigElementType::Control);
+		bInitial = false;
+	}
+
+	/**
+	 * The item representing the channel
+	 */
+	UPROPERTY(meta = (Input))
+	FRigElementKey Item;
+
+	/**
+	 * If set to true the initial value will be returned
+	 */
+	UPROPERTY(meta = (Input))
+	bool bInitial;
+};
+
+/**
+ * Get Bool Channel is used to retrieve a control's animation channel value
+ */
+USTRUCT(meta=(DisplayName="Get Bool Channel", Keywords="Animation,Channel", TemplateName="GetAnimationChannelFromItem"))
+struct CONTROLRIG_API FRigUnit_GetBoolAnimationChannelFromItem : public FRigUnit_GetAnimationChannelFromItemBase
+{
+	GENERATED_BODY()
+
+	FRigUnit_GetBoolAnimationChannelFromItem()
+		: FRigUnit_GetAnimationChannelFromItemBase()
+		, Value(false)
+	{}
+
+	RIGVM_METHOD()
+	virtual void Execute() override;
+
+	// The current value of the animation channel
+	UPROPERTY(meta=(Output))
+	bool Value;
+};
+
+/**
+ * Get Float Channel is used to retrieve a control's animation channel value
+ */
+USTRUCT(meta=(DisplayName="Get Float Channel", Keywords="Animation,Channel", TemplateName="GetAnimationChannelFromItem"))
+struct CONTROLRIG_API FRigUnit_GetFloatAnimationChannelFromItem : public FRigUnit_GetAnimationChannelFromItemBase
+{
+	GENERATED_BODY()
+
+	FRigUnit_GetFloatAnimationChannelFromItem()
+		: FRigUnit_GetAnimationChannelFromItemBase()
+		, Value(0.f)
+	{}
+
+	RIGVM_METHOD()
+	virtual void Execute() override;
+
+	// The current value of the animation channel
+	UPROPERTY(meta=(Output, UIMin=0, UIMax=1))
+	float Value;
+};
+
+/**
+ * Get Int Channel is used to retrieve a control's animation channel value
+ */
+USTRUCT(meta=(DisplayName="Get Int Channel", Keywords="Animation,Channel", TemplateName="GetAnimationChannelFromItem"))
+struct CONTROLRIG_API FRigUnit_GetIntAnimationChannelFromItem : public FRigUnit_GetAnimationChannelFromItemBase
+{
+	GENERATED_BODY()
+
+	FRigUnit_GetIntAnimationChannelFromItem()
+		: FRigUnit_GetAnimationChannelFromItemBase()
+		, Value(0)
+	{}
+
+	RIGVM_METHOD()
+	virtual void Execute() override;
+
+	// The current value of the animation channel
+	UPROPERTY(meta=(Output))
+	int32 Value;
+};
+
+/**
+ * Get Vector2D Channel is used to retrieve a control's animation channel value
+ */
+USTRUCT(meta=(DisplayName="Get Vector2D Channel", Keywords="Animation,Channel", TemplateName="GetAnimationChannelFromItem"))
+struct CONTROLRIG_API FRigUnit_GetVector2DAnimationChannelFromItem : public FRigUnit_GetAnimationChannelFromItemBase
+{
+	GENERATED_BODY()
+
+	FRigUnit_GetVector2DAnimationChannelFromItem()
+		: FRigUnit_GetAnimationChannelFromItemBase()
+		, Value(FVector2D::ZeroVector)
+	{}
+
+	RIGVM_METHOD()
+	virtual void Execute() override;
+
+	// The current value of the animation channel
+	UPROPERTY(meta=(Output))
+	FVector2D Value;
+};
+
+/**
+ * Get Vector Channel is used to retrieve a control's animation channel value
+ */
+USTRUCT(meta=(DisplayName="Get Vector Channel", Keywords="Animation,Channel", TemplateName="GetAnimationChannelFromItem"))
+struct CONTROLRIG_API FRigUnit_GetVectorAnimationChannelFromItem : public FRigUnit_GetAnimationChannelFromItemBase
+{
+	GENERATED_BODY()
+
+	FRigUnit_GetVectorAnimationChannelFromItem()
+		: FRigUnit_GetAnimationChannelFromItemBase()
+		, Value(FVector::ZeroVector)
+	{}
+
+	RIGVM_METHOD()
+	virtual void Execute() override;
+
+	// The current value of the animation channel
+	UPROPERTY(meta=(Output))
+	FVector Value;
+};
+
+/**
+ * Get Rotator Channel is used to retrieve a control's animation channel value
+ */
+USTRUCT(meta=(DisplayName="Get Rotator Channel", Keywords="Animation,Channel", TemplateName="GetAnimationChannelFromItem"))
+struct CONTROLRIG_API FRigUnit_GetRotatorAnimationChannelFromItem : public FRigUnit_GetAnimationChannelFromItemBase
+{
+	GENERATED_BODY()
+
+	FRigUnit_GetRotatorAnimationChannelFromItem()
+		: FRigUnit_GetAnimationChannelFromItemBase()
+		, Value(FRotator::ZeroRotator)
+	{}
+
+	RIGVM_METHOD()
+	virtual void Execute() override;
+
+	// The current value of the animation channel
+	UPROPERTY(meta=(Output))
+	FRotator Value;
+};
+
+/**
+ * Get Transform Channel is used to retrieve a control's animation channel value
+ */
+USTRUCT(meta=(DisplayName="Get Transform Channel", Keywords="Animation,Channel", TemplateName="GetAnimationChannelFromItem"))
+struct CONTROLRIG_API FRigUnit_GetTransformAnimationChannelFromItem : public FRigUnit_GetAnimationChannelFromItemBase
+{
+	GENERATED_BODY()
+
+	FRigUnit_GetTransformAnimationChannelFromItem()
+		: FRigUnit_GetAnimationChannelFromItemBase()
+		, Value(FTransform::Identity)
+	{}
+
+	RIGVM_METHOD()
+	virtual void Execute() override;
+
+	// The current value of the animation channel
+	UPROPERTY(meta=(Output))
+	FTransform Value;
+};
+
+/**
+ * Set Animation Channel is used to change a control's animation channel value
+ */
+USTRUCT(meta = (Abstract))
+struct CONTROLRIG_API FRigUnit_SetAnimationChannelBaseFromItem : public FRigUnit_GetAnimationChannelFromItemBase
+{
+	GENERATED_BODY()
+
+	FRigUnit_SetAnimationChannelBaseFromItem()
+		:FRigUnit_GetAnimationChannelFromItemBase()
+	{
+	}
+
+	UPROPERTY(DisplayName = "Execute", Transient, meta = (Input, Output))
+	FControlRigExecuteContext ExecuteContext;
+};
+
+/**
+ * Set Bool Channel is used to set a control's animation channel value
+ */
+USTRUCT(meta=(DisplayName="Set Bool Channel", Keywords="Animation,Channel", TemplateName="SetAnimationChannelFromItem"))
+struct CONTROLRIG_API FRigUnit_SetBoolAnimationChannelFromItem : public FRigUnit_SetAnimationChannelBaseFromItem
+{
+	GENERATED_BODY()
+
+	FRigUnit_SetBoolAnimationChannelFromItem()
+		: FRigUnit_SetAnimationChannelBaseFromItem()
+		, Value(false)
+	{}
+
+	RIGVM_METHOD()
+	virtual void Execute() override;
+
+	// The new value of the animation channel
+	UPROPERTY(meta=(Input))
+	bool Value;
+};
+
+/**
+ * Set Float Channel is used to set a control's animation channel value
+ */
+USTRUCT(meta=(DisplayName="Set Float Channel", Keywords="Animation,Channel", TemplateName="SetAnimationChannelFromItem"))
+struct CONTROLRIG_API FRigUnit_SetFloatAnimationChannelFromItem : public FRigUnit_SetAnimationChannelBaseFromItem
+{
+	GENERATED_BODY()
+
+	FRigUnit_SetFloatAnimationChannelFromItem()
+		: FRigUnit_SetAnimationChannelBaseFromItem()
+		, Value(0.f)
+	{}
+
+	RIGVM_METHOD()
+	virtual void Execute() override;
+
+	// The new value of the animation channel
+	UPROPERTY(meta=(Input, UIMin=0, UIMax=1))
+	float Value;
+};
+
+/**
+ * Set Int Channel is used to set a control's animation channel value
+ */
+USTRUCT(meta=(DisplayName="Set Int Channel", Keywords="Animation,Channel", TemplateName="SetAnimationChannelFromItem"))
+struct CONTROLRIG_API FRigUnit_SetIntAnimationChannelFromItem : public FRigUnit_SetAnimationChannelBaseFromItem
+{
+	GENERATED_BODY()
+
+	FRigUnit_SetIntAnimationChannelFromItem()
+		: FRigUnit_SetAnimationChannelBaseFromItem()
+		, Value(0)
+	{}
+
+	RIGVM_METHOD()
+	virtual void Execute() override;
+
+	// The new value of the animation channel
+	UPROPERTY(meta=(Input))
+	int32 Value;
+};
+
+/**
+ * Set Vector2D Channel is used to set a control's animation channel value
+ */
+USTRUCT(meta=(DisplayName="Set Vector2D Channel", Keywords="Animation,Channel", TemplateName="SetAnimationChannelFromItem"))
+struct CONTROLRIG_API FRigUnit_SetVector2DAnimationChannelFromItem : public FRigUnit_SetAnimationChannelBaseFromItem
+{
+	GENERATED_BODY()
+
+	FRigUnit_SetVector2DAnimationChannelFromItem()
+		: FRigUnit_SetAnimationChannelBaseFromItem()
+		, Value(FVector2D::ZeroVector)
+	{}
+
+	RIGVM_METHOD()
+	virtual void Execute() override;
+
+	// The new value of the animation channel
+	UPROPERTY(meta=(Input))
+	FVector2D Value;
+};
+
+/**
+ * Set Vector Channel is used to set a control's animation channel value
+ */
+USTRUCT(meta=(DisplayName="Set Vector Channel", Keywords="Animation,Channel", TemplateName="SetAnimationChannelFromItem"))
+struct CONTROLRIG_API FRigUnit_SetVectorAnimationChannelFromItem : public FRigUnit_SetAnimationChannelBaseFromItem
+{
+	GENERATED_BODY()
+
+	FRigUnit_SetVectorAnimationChannelFromItem()
+		: FRigUnit_SetAnimationChannelBaseFromItem()
+		, Value(FVector::ZeroVector)
+	{}
+
+	RIGVM_METHOD()
+	virtual void Execute() override;
+
+	// The new value of the animation channel
+	UPROPERTY(meta=(Input))
+	FVector Value;
+};
+
+/**
+ * Set Rotator Channel is used to set a control's animation channel value
+ */
+USTRUCT(meta=(DisplayName="Set Rotator Channel", Keywords="Animation,Channel", TemplateName="SetAnimationChannelFromItem"))
+struct CONTROLRIG_API FRigUnit_SetRotatorAnimationChannelFromItem : public FRigUnit_SetAnimationChannelBaseFromItem
+{
+	GENERATED_BODY()
+
+	FRigUnit_SetRotatorAnimationChannelFromItem()
+		: FRigUnit_SetAnimationChannelBaseFromItem()
+		, Value(FRotator::ZeroRotator)
+	{}
+
+	RIGVM_METHOD()
+	virtual void Execute() override;
+
+	// The new value of the animation channel
+	UPROPERTY(meta=(Input))
+	FRotator Value;
+};
+
+/**
+ * Set Transform Channel is used to set a control's animation channel value
+ */
+USTRUCT(meta=(DisplayName="Set Transform Channel", Keywords="Animation,Channel", TemplateName="SetAnimationChannelFromItem"))
+struct CONTROLRIG_API FRigUnit_SetTransformAnimationChannelFromItem : public FRigUnit_SetAnimationChannelBaseFromItem
+{
+	GENERATED_BODY()
+
+	FRigUnit_SetTransformAnimationChannelFromItem()
+		: FRigUnit_SetAnimationChannelBaseFromItem()
+		, Value(FTransform::Identity)
+	{}
+
+	RIGVM_METHOD()
+	virtual void Execute() override;
+
+	// The new value of the animation channel
+	UPROPERTY(meta=(Input))
+	FTransform Value;
+};
