[Control Rig & Sequencer] Changed default behavior when adding new control rig track: rigs added later now are evaluated after existing rigs by default, instead of before existing rigs

diff --git a/Engine/Plugins/Animation/ControlRig/Source/ControlRig/Private/Sequencer/ControlRigLayerInstanceProxy.cpp b/Engine/Plugins/Animation/ControlRig/Source/ControlRig/Private/Sequencer/ControlRigLayerInstanceProxy.cpp
--- a/Engine/Plugins/Animation/ControlRig/Source/ControlRig/Private/Sequencer/ControlRigLayerInstanceProxy.cpp
+++ b/Engine/Plugins/Animation/ControlRig/Source/ControlRig/Private/Sequencer/ControlRigLayerInstanceProxy.cpp
@@ -181,129 +181,42 @@ void FControlRigLayerInstanceProxy::AddControlRigTrack(int32 ControlRigID, UCont
 	{
 		UMovieSceneControlRigParameterTrack* Track = InControlRig->GetTypedOuter<UMovieSceneControlRigParameterTrack>();
 
-		if (ControlRigNodes.Num() > 0)
+		int32 PriorityOrder = Track ? Track->GetPriorityOrder() : INDEX_NONE;
+		if (PriorityOrder == INDEX_NONE) //track has no order so need to find a number such that it runs at the end, will happen on creation
 		{
-			int32 PriorityOrder = Track ? Track->GetPriorityOrder() : INDEX_NONE;
-			if (PriorityOrder == INDEX_NONE) //track has no order so just add to end, will happen on creation
-			{
-				if (InControlRig->IsAdditive()) //additive added to end of all
-				{
-					Node = ControlRigNodes[ControlRigNodes.Num() - 1].Get();
-					if (Track)
-					{
-						if (UMovieSceneControlRigParameterTrack* OtherTrack = Node->GetControlRig()->GetTypedOuter<UMovieSceneControlRigParameterTrack>())
-						{
-							Track->SetPriorityOrder(OtherTrack->GetPriorityOrder() + 1);
-						}
-						else
-						{
-							Track->SetPriorityOrder(DefaultPriorityOrder + ControlRigNodes.Num() -1);
-						}
-					}
-					Node = ControlRigNodes.Add_GetRef(MakeShared<FAnimNode_ControlRig_ExternalSource>()).Get();
-				}
-				else //if non-additive find first additive then insert there..
-				{
-					for (int32 Index = 0; Index < ControlRigNodes.Num(); ++Index)
-					{
-						Node = ControlRigNodes[Index].Get();
-						if (!Node || !Node->GetControlRig() || !Node->GetControlRig()->IsAdditive())
-						{
-							Node = nullptr;
-							continue; // not additive
-						}
-						//okay first additive so insert, but set priority first
-						if (Track)
-						{
-							if (UMovieSceneControlRigParameterTrack* OtherTrack = Node->GetControlRig()->GetTypedOuter<UMovieSceneControlRigParameterTrack>())
-							{
-								Track->SetPriorityOrder(OtherTrack->GetPriorityOrder() + 1);
-							}
-							else
-							{
-								Track->SetPriorityOrder(DefaultPriorityOrder + Index);
-							}
-						}
-						Node = ControlRigNodes.Insert_GetRef(MakeShared<FAnimNode_ControlRig_ExternalSource>(), Index).Get();
-						break;
-					}
-
-					if (Node == nullptr) //add to end
-					{
-						//set priority
-						Node = ControlRigNodes[ControlRigNodes.Num() - 1].Get();
-						if (Track)
-						{
-							if (UMovieSceneControlRigParameterTrack* OtherTrack = Node->GetControlRig()->GetTypedOuter<UMovieSceneControlRigParameterTrack>())
-							{
-								Track->SetPriorityOrder(OtherTrack->GetPriorityOrder() + 1);
-							}
-							else
-							{
-								Track->SetPriorityOrder(ControlRigNodes.Num() - 1);
-							}
-						}
-						Node = ControlRigNodes.Add_GetRef(MakeShared<FAnimNode_ControlRig_ExternalSource>()).Get();
-					}
-				}
-			}
-			else
+			SortControlRigNodes();
+			
+			PriorityOrder = DefaultPriorityOrder;
+			
+			// find the first node with the same type and make sure we set a even lower priority
+			for (int32 Index = 0; Index < ControlRigNodes.Num(); ++Index)
 			{
-				int32 NewNodePosition = 0;
-				for (; NewNodePosition < ControlRigNodes.Num(); ++NewNodePosition)
+				FAnimNode_ControlRig_ExternalSource* NodeToCheck = ControlRigNodes[Index].Get();
+				
+				if (!NodeToCheck ||
+					!NodeToCheck->GetControlRig() ||
+					NodeToCheck->GetControlRig()->IsAdditive() != InControlRig->IsAdditive())
 				{
-					Node = ControlRigNodes[NewNodePosition].Get();
-	
-					if (!Node || !Node->GetControlRig())
-					{
-						continue; // shouldn't happen
-					}
-					//additive but current isn't so keep going
-					if (InControlRig->IsAdditive() && !Node->GetControlRig()->IsAdditive())
-					{
-						continue;
-					}
-					//not additive but this one is so we add it here
-					else if (!InControlRig->IsAdditive() && Node->GetControlRig()->IsAdditive())
-					{
-						break;
-					}
-					if (UMovieSceneControlRigParameterTrack* OtherTrack = Node->GetControlRig()->GetTypedOuter<UMovieSceneControlRigParameterTrack>())
-					{
-						if (PriorityOrder >= OtherTrack->GetPriorityOrder())
-						{
-							continue;
-						}
-						else
-						{
-							break;
-						}
-					}
+					continue;
 				}
-				if (NewNodePosition >= ControlRigNodes.Num() - 1)
+				
+				if (UMovieSceneControlRigParameterTrack* OtherTrack = ControlRigNodes[0]->GetControlRig()->GetTypedOuter<UMovieSceneControlRigParameterTrack>())
 				{
-					Node = ControlRigNodes.Add_GetRef(MakeShared<FAnimNode_ControlRig_ExternalSource>()).Get();
+					PriorityOrder = OtherTrack->GetPriorityOrder() - 1;
 				}
 				else
 				{
-					Node = ControlRigNodes.Insert_GetRef(MakeShared<FAnimNode_ControlRig_ExternalSource>(), NewNodePosition).Get();
-				}
-			}
-		}
-		else //no nodes add first one
-		{
-			Node = ControlRigNodes.Add_GetRef(MakeShared<FAnimNode_ControlRig_ExternalSource>()).Get();
-			if (Track)
-			{
-				int32 PriorityOrder = Track->GetPriorityOrder();
-				if (PriorityOrder == INDEX_NONE) //track has no order so just add to end, will happen on creation
-				{
-					Track->SetPriorityOrder(DefaultPriorityOrder);
+					// Best guess of the lowest priority
+					PriorityOrder = DefaultPriorityOrder - ControlRigNodes.Num();
 				}
 			}
+
+			Track->SetPriorityOrder(PriorityOrder);
 		}
-	
-		check(Node);
+
+		// Simply add to the end, we will sort again
+		Node = ControlRigNodes.Add_GetRef(MakeShared<FAnimNode_ControlRig_ExternalSource>()).Get();
+		SortControlRigNodes();
 		
 		//this will set up the link nodes
 		ConstructNodes();
diff --git a/Engine/Plugins/Animation/ControlRig/Source/ControlRigEditor/Private/Sequencer/ControlRigParameterTrackEditor.cpp b/Engine/Plugins/Animation/ControlRig/Source/ControlRigEditor/Private/Sequencer/ControlRigParameterTrackEditor.cpp
--- a/Engine/Plugins/Animation/ControlRig/Source/ControlRigEditor/Private/Sequencer/ControlRigParameterTrackEditor.cpp
+++ b/Engine/Plugins/Animation/ControlRig/Source/ControlRigEditor/Private/Sequencer/ControlRigParameterTrackEditor.cpp
@@ -4276,7 +4276,7 @@ void FControlRigParameterTrackEditor::BuildTrackContextMenu(FMenuBuilder& MenuBu
 				SNew(SSpinBox<int32>)
 				.MinValue(0)
 				.Font(FAppStyle::GetFontStyle(TEXT("MenuItem.Font")))
-				.ToolTipText(LOCTEXT("OrderTooltip", "Order for this Control Rig to evaluate compared to others on the same binding"))
+				.ToolTipText(LOCTEXT("OrderTooltip", "Order for this Control Rig to evaluate compared to others on the same binding, higher number means earlier evaluation"))
 				.Value_Lambda([Track]() { return Track->GetPriorityOrder(); })
 				.OnValueChanged_Lambda([Track](int32 InValue) { Track->SetPriorityOrder(InValue); })
 				,
