[ControlRig] Fixed crash on creating new control rig due to nullptr access

diff --git a/Engine/Plugins/Animation/ControlRig/Source/ControlRig/Private/Sequencer/ControlRigLayerInstanceProxy.cpp b/Engine/Plugins/Animation/ControlRig/Source/ControlRig/Private/Sequencer/ControlRigLayerInstanceProxy.cpp
--- a/Engine/Plugins/Animation/ControlRig/Source/ControlRig/Private/Sequencer/ControlRigLayerInstanceProxy.cpp
+++ b/Engine/Plugins/Animation/ControlRig/Source/ControlRig/Private/Sequencer/ControlRigLayerInstanceProxy.cpp
@@ -181,38 +181,42 @@ void FControlRigLayerInstanceProxy::AddControlRigTrack(int32 ControlRigID, UCont
 	{
 		UMovieSceneControlRigParameterTrack* Track = InControlRig->GetTypedOuter<UMovieSceneControlRigParameterTrack>();
 
-		int32 PriorityOrder = Track ? Track->GetPriorityOrder() : INDEX_NONE;
-		if (PriorityOrder == INDEX_NONE) //track has no order so need to find a number such that it runs at the end, will happen on creation
+		if (Track)
 		{
-			SortControlRigNodes();
-			
-			PriorityOrder = DefaultPriorityOrder;
-			
-			// find the first node with the same type and make sure we set a even lower priority
-			for (int32 Index = 0; Index < ControlRigNodes.Num(); ++Index)
+			int32 PriorityOrder = Track->GetPriorityOrder();
+			if (PriorityOrder == INDEX_NONE) //track has no order so need to find a number such that it runs at the end, will happen on creation
 			{
-				FAnimNode_ControlRig_ExternalSource* NodeToCheck = ControlRigNodes[Index].Get();
+				PriorityOrder = DefaultPriorityOrder;
 				
-				if (!NodeToCheck ||
-					!NodeToCheck->GetControlRig() ||
-					NodeToCheck->GetControlRig()->IsAdditive() != InControlRig->IsAdditive())
+				SortControlRigNodes();
+			
+				// find the first node with the same type and make sure we set a even lower priority
+				for (int32 Index = 0; Index < ControlRigNodes.Num(); ++Index)
 				{
-					continue;
-				}
+					FAnimNode_ControlRig_ExternalSource* NodeToCheck = ControlRigNodes[Index].Get();
 				
-				if (UMovieSceneControlRigParameterTrack* OtherTrack = ControlRigNodes[0]->GetControlRig()->GetTypedOuter<UMovieSceneControlRigParameterTrack>())
-				{
-					PriorityOrder = OtherTrack->GetPriorityOrder() - 1;
-				}
-				else
-				{
-					// Best guess of the lowest priority
-					PriorityOrder = DefaultPriorityOrder - ControlRigNodes.Num();
+					if (!NodeToCheck ||
+						!NodeToCheck->GetControlRig() ||
+						NodeToCheck->GetControlRig()->IsAdditive() != InControlRig->IsAdditive())
+					{
+						continue;
+					}
+				
+					if (UMovieSceneControlRigParameterTrack* OtherTrack = ControlRigNodes[0]->GetControlRig()->GetTypedOuter<UMovieSceneControlRigParameterTrack>())
+					{
+						PriorityOrder = OtherTrack->GetPriorityOrder() - 1;
+					}
+					else
+					{
+						// Best guess of the lowest priority
+						PriorityOrder = DefaultPriorityOrder - ControlRigNodes.Num();
+					}
 				}
-			}
 
-			Track->SetPriorityOrder(PriorityOrder);
+				Track->SetPriorityOrder(PriorityOrder);
+			}	
 		}
+		
 
 		// Simply add to the end, we will sort again
 		Node = ControlRigNodes.Add_GetRef(MakeShared<FAnimNode_ControlRig_ExternalSource>()).Get();
