Control Rig: Animation Channels to support additional parents

#rnx

diff --git a/Engine/Plugins/Animation/ControlRig/Source/ControlRig/Private/Rigs/RigHierarchy.cpp b/Engine/Plugins/Animation/ControlRig/Source/ControlRig/Private/Rigs/RigHierarchy.cpp
--- a/Engine/Plugins/Animation/ControlRig/Source/ControlRig/Private/Rigs/RigHierarchy.cpp
+++ b/Engine/Plugins/Animation/ControlRig/Source/ControlRig/Private/Rigs/RigHierarchy.cpp
@@ -2304,6 +2304,16 @@ bool URigHierarchy::SetParentWeight(FRigBaseElement* InChild, int32 InParentInde
 	{
 		if(MultiParentElement->ParentConstraints.IsValidIndex(InParentIndex))
 		{
+			if(const FRigControlElement* ControlElement = Cast<FRigControlElement>(MultiParentElement))
+			{
+				// animation channels cannot change their parent weights, 
+				// they are not 3d things - so multi parenting doesn't matter for transforms.
+				if(ControlElement->IsAnimationChannel())
+				{
+					return false;
+				}
+			}
+			
 			InWeight.Location = FMath::Max(InWeight.Location, 0.f);
 			InWeight.Rotation = FMath::Max(InWeight.Rotation, 0.f);
 			InWeight.Scale = FMath::Max(InWeight.Scale, 0.f);
@@ -2410,6 +2420,16 @@ bool URigHierarchy::SetParentWeightArray(FRigBaseElement* InChild,  const TArray
 
 	if(FRigMultiParentElement* MultiParentElement = Cast<FRigMultiParentElement>(InChild))
 	{
+		if(const FRigControlElement* ControlElement = Cast<FRigControlElement>(MultiParentElement))
+		{
+			// animation channels cannot change their parent weights, 
+			// they are not 3d things - so multi parenting doesn't matter for transforms.
+			if(ControlElement->IsAnimationChannel())
+			{
+				return false;
+			}
+		}
+
 		if(MultiParentElement->ParentConstraints.Num() == InWeights.Num())
 		{
 			TArray<FRigElementWeight> InputWeights;
diff --git a/Engine/Plugins/Animation/ControlRig/Source/ControlRig/Private/Rigs/RigHierarchyController.cpp b/Engine/Plugins/Animation/ControlRig/Source/ControlRig/Private/Rigs/RigHierarchyController.cpp
--- a/Engine/Plugins/Animation/ControlRig/Source/ControlRig/Private/Rigs/RigHierarchyController.cpp
+++ b/Engine/Plugins/Animation/ControlRig/Source/ControlRig/Private/Rigs/RigHierarchyController.cpp
@@ -2735,16 +2735,8 @@ bool URigHierarchyController::AddParent(FRigBaseElement* InChild, FRigBaseElemen
 	{
 		if(ChildControlElement->IsAnimationChannel())
 		{
-			if(!bRemoveAllParents)
-			{
-				if (ChildControlElement->ParentConstraints.Num() > 0)
-				{
-					ReportErrorf(TEXT("Cannot add multiple parents to animation channel '%s'."), *InChild->Key.ToString());
-					return false;
-				}
-			}
-
 			bMaintainGlobalTransform = false;
+			InWeight = 0.f;
 		}
 
 		if(ChildControlElement->Settings.bRestrictSpaceSwitching)
@@ -3222,6 +3214,399 @@ bool URigHierarchyController::SetParent(FRigElementKey InChild, FRigElementKey I
 	return bParentSet;
 }
 
+bool URigHierarchyController::AddAvailableSpace(FRigElementKey InControl, FRigElementKey InSpace, bool bSetupUndo, bool bPrintPythonCommand)
+{
+	if(!IsValid())
+	{
+		return false;
+	}
+
+	URigHierarchy* Hierarchy = GetHierarchy();
+
+	FRigBaseElement* ControlBase = Hierarchy->Find(InControl);
+	if(ControlBase == nullptr)
+	{
+		ReportWarningf(TEXT("Cannot Add Available Space, Control '%s' not found."), *InControl.ToString());
+		return false;
+	}
+
+	FRigControlElement* Control = Cast<FRigControlElement>(ControlBase);
+	if(Control == nullptr)
+	{
+		ReportWarningf(TEXT("Cannot Add Available Space, '%s' is not a Control."), *InControl.ToString());
+		return false;
+	}
+
+	const FRigBaseElement* SpaceBase = Hierarchy->Find(InSpace);
+	if(SpaceBase == nullptr)
+	{
+		ReportWarningf(TEXT("Cannot Add Available Space, Space '%s' not found."), *InSpace.ToString());
+		return false;
+	}
+
+	const FRigTransformElement* Space = Cast<FRigTransformElement>(SpaceBase);
+	if(Space == nullptr)
+	{
+		ReportWarningf(TEXT("Cannot Add Available Space, '%s' is not a Transform."), *InSpace.ToString());
+		return false;
+	}
+
+#if WITH_EDITOR
+	TSharedPtr<FScopedTransaction> TransactionPtr;
+	if(bSetupUndo)
+	{
+		TransactionPtr = MakeShared<FScopedTransaction>(NSLOCTEXT("RigHierarchyController", "AddAvailableSpace", "Add Available Space"));
+		Hierarchy->Modify();
+	}
+#endif
+
+	const bool bSuccess = AddAvailableSpace(Control, Space);
+
+#if WITH_EDITOR
+	if(!bSuccess && TransactionPtr.IsValid())
+	{
+		TransactionPtr->Cancel();
+	}
+	TransactionPtr.Reset();
+
+	if (bSuccess && bPrintPythonCommand && !bSuspendPythonPrinting)
+	{
+		if (const UBlueprint* Blueprint = GetTypedOuter<UBlueprint>())
+		{
+			RigVMPythonUtils::Print(Blueprint->GetFName().ToString(),
+				FString::Printf(TEXT("hierarchy_controller.add_available_space(%s, %s)"),
+				*InControl.ToPythonString(),
+				*InSpace.ToPythonString()));
+		}
+	}
+#endif
+
+	return bSuccess;
+}
+
+bool URigHierarchyController::RemoveAvailableSpace(FRigElementKey InControl, FRigElementKey InSpace, bool bSetupUndo, bool bPrintPythonCommand)
+{
+	if(!IsValid())
+	{
+		return false;
+	}
+
+	URigHierarchy* Hierarchy = GetHierarchy();
+
+	FRigBaseElement* ControlBase = Hierarchy->Find(InControl);
+	if(ControlBase == nullptr)
+	{
+		ReportWarningf(TEXT("Cannot Remove Available Space, Control '%s' not found."), *InControl.ToString());
+		return false;
+	}
+
+	FRigControlElement* Control = Cast<FRigControlElement>(ControlBase);
+	if(Control == nullptr)
+	{
+		ReportWarningf(TEXT("Cannot Remove Available Space, '%s' is not a Control."), *InControl.ToString());
+		return false;
+	}
+
+	const FRigBaseElement* SpaceBase = Hierarchy->Find(InSpace);
+	if(SpaceBase == nullptr)
+	{
+		ReportWarningf(TEXT("Cannot Remove Available Space, Space '%s' not found."), *InSpace.ToString());
+		return false;
+	}
+
+	const FRigTransformElement* Space = Cast<FRigTransformElement>(SpaceBase);
+	if(Space == nullptr)
+	{
+		ReportWarningf(TEXT("Cannot Remove Available Space, '%s' is not a Transform."), *InSpace.ToString());
+		return false;
+	}
+
+#if WITH_EDITOR
+	TSharedPtr<FScopedTransaction> TransactionPtr;
+	if(bSetupUndo)
+	{
+		TransactionPtr = MakeShared<FScopedTransaction>(NSLOCTEXT("RigHierarchyController", "RemoveAvailableSpace", "Remove Available Space"));
+		Hierarchy->Modify();
+	}
+#endif
+
+	const bool bSuccess = RemoveAvailableSpace(Control, Space);
+
+#if WITH_EDITOR
+	if(!bSuccess && TransactionPtr.IsValid())
+	{
+		TransactionPtr->Cancel();
+	}
+	TransactionPtr.Reset();
+
+	if (bSuccess && bPrintPythonCommand && !bSuspendPythonPrinting)
+	{
+		if (const UBlueprint* Blueprint = GetTypedOuter<UBlueprint>())
+		{
+			RigVMPythonUtils::Print(Blueprint->GetFName().ToString(),
+				FString::Printf(TEXT("hierarchy_controller.remove_available_space(%s, %s)"),
+				*InControl.ToPythonString(),
+				*InSpace.ToPythonString()));
+		}
+	}
+#endif
+
+	return bSuccess;
+}
+
+bool URigHierarchyController::SetAvailableSpaceIndex(FRigElementKey InControl, FRigElementKey InSpace, int32 InIndex, bool bSetupUndo, bool bPrintPythonCommand)
+{
+	if(!IsValid())
+	{
+		return false;
+	}
+
+	URigHierarchy* Hierarchy = GetHierarchy();
+
+	FRigBaseElement* ControlBase = Hierarchy->Find(InControl);
+	if(ControlBase == nullptr)
+	{
+		ReportWarningf(TEXT("Cannot Set Available Space Index, Control '%s' not found."), *InControl.ToString());
+		return false;
+	}
+
+	FRigControlElement* Control = Cast<FRigControlElement>(ControlBase);
+	if(Control == nullptr)
+	{
+		ReportWarningf(TEXT("Cannot Set Available Space Index, '%s' is not a Control."), *InControl.ToString());
+		return false;
+	}
+
+	const FRigBaseElement* SpaceBase = Hierarchy->Find(InSpace);
+	if(SpaceBase == nullptr)
+	{
+		ReportWarningf(TEXT("Cannot Set Available Space Index, Space '%s' not found."), *InSpace.ToString());
+		return false;
+	}
+
+	const FRigTransformElement* Space = Cast<FRigTransformElement>(SpaceBase);
+	if(Space == nullptr)
+	{
+		ReportWarningf(TEXT("Cannot Set Available Space Index, '%s' is not a Transform."), *InSpace.ToString());
+		return false;
+	}
+
+#if WITH_EDITOR
+	TSharedPtr<FScopedTransaction> TransactionPtr;
+	if(bSetupUndo)
+	{
+		TransactionPtr = MakeShared<FScopedTransaction>(NSLOCTEXT("RigHierarchyController", "RemoveAvailableSpace", "Remove Available Space"));
+		Hierarchy->Modify();
+	}
+#endif
+
+	const bool bSuccess = SetAvailableSpaceIndex(Control, Space, InIndex);
+
+#if WITH_EDITOR
+	if(!bSuccess && TransactionPtr.IsValid())
+	{
+		TransactionPtr->Cancel();
+	}
+	TransactionPtr.Reset();
+
+	if (bSuccess && bPrintPythonCommand && !bSuspendPythonPrinting)
+	{
+		if (const UBlueprint* Blueprint = GetTypedOuter<UBlueprint>())
+		{
+			RigVMPythonUtils::Print(Blueprint->GetFName().ToString(),
+				FString::Printf(TEXT("hierarchy_controller.set_available_space_index(%s, %s)"),
+				*InControl.ToPythonString(),
+				*InSpace.ToPythonString()));
+		}
+	}
+#endif
+
+	return bSuccess;
+}
+
+bool URigHierarchyController::AddChannelHost(FRigElementKey InChannel, FRigElementKey InHost, bool bSetupUndo, bool bPrintPythonCommand)
+{
+	if(!IsValid())
+	{
+		return false;
+	}
+
+	URigHierarchy* Hierarchy = GetHierarchy();
+
+	FRigBaseElement* ChannelBase = Hierarchy->Find(InChannel);
+	if(ChannelBase == nullptr)
+	{
+		ReportWarningf(TEXT("Cannot Add Channel Host, Channel '%s' not found."), *InChannel.ToString());
+		return false;
+	}
+
+	FRigControlElement* Channel = Cast<FRigControlElement>(ChannelBase);
+	if(Channel == nullptr)
+	{
+		ReportWarningf(TEXT("Cannot Add Channel Host, '%s' is not a Control."), *InChannel.ToString());
+		return false;
+	}
+
+	if(!Channel->IsAnimationChannel())
+	{
+		ReportWarningf(TEXT("Cannot Add Channel Host, '%s' is not an animation channel."), *InChannel.ToString());
+		return false;
+	}
+
+	const FRigBaseElement* HostBase = Hierarchy->Find(InHost);
+	if(HostBase == nullptr)
+	{
+		ReportWarningf(TEXT("Cannot Add Channel Host, Host '%s' not found."), *InHost.ToString());
+		return false;
+	}
+
+	const FRigControlElement* Host = Cast<FRigControlElement>(HostBase);
+	if(Host == nullptr)
+	{
+		ReportWarningf(TEXT("Cannot Add Channel Host, '%s' is not a Control."), *InHost.ToString());
+		return false;
+	}
+
+	if(Host->IsAnimationChannel())
+	{
+		ReportWarningf(TEXT("Cannot Add Channel Host, '%s' is also an animation channel."), *InHost.ToString());
+		return false;
+	}
+
+	// the default parent cannot be added to the channel host list
+	if(Hierarchy->GetParents(InChannel).Contains(InHost))
+	{
+		ReportWarningf(TEXT("Cannot Add Channel Host, '%s' is the parent of channel '%s'."), *InHost.ToString(), *InChannel.ToString());
+		return false;
+	}
+
+	if(Channel->Settings.Customization.AvailableSpaces.Contains(Host->GetKey()))
+	{
+		ReportWarningf(TEXT("Cannot Add Channel Host, '%s' is already a host for channel '%s'."), *InHost.ToString(), *InChannel.ToString());
+		return false;
+	}
+
+#if WITH_EDITOR
+	TSharedPtr<FScopedTransaction> TransactionPtr;
+	if(bSetupUndo)
+	{
+		TransactionPtr = MakeShared<FScopedTransaction>(NSLOCTEXT("RigHierarchyController", "AddChannelHost", "Add Channel Host"));
+		Hierarchy->Modify();
+	}
+#endif
+
+	const bool bSuccess = AddAvailableSpace(Channel, Host);
+
+#if WITH_EDITOR
+	if(!bSuccess && TransactionPtr.IsValid())
+	{
+		TransactionPtr->Cancel();
+	}
+	TransactionPtr.Reset();
+
+	if (bSuccess && bPrintPythonCommand && !bSuspendPythonPrinting)
+	{
+		if (const UBlueprint* Blueprint = GetTypedOuter<UBlueprint>())
+		{
+			RigVMPythonUtils::Print(Blueprint->GetFName().ToString(),
+				FString::Printf(TEXT("hierarchy_controller.add_channel_host(%s, %s)"),
+				*InChannel.ToPythonString(),
+				*InHost.ToPythonString()));
+		}
+	}
+#endif
+
+	return bSuccess;
+}
+
+bool URigHierarchyController::RemoveChannelHost(FRigElementKey InChannel, FRigElementKey InHost, bool bSetupUndo, bool bPrintPythonCommand)
+{
+	if(!IsValid())
+	{
+		return false;
+	}
+
+	URigHierarchy* Hierarchy = GetHierarchy();
+
+	FRigBaseElement* ChannelBase = Hierarchy->Find(InChannel);
+	if(ChannelBase == nullptr)
+	{
+		ReportWarningf(TEXT("Cannot Remove Channel Host, Channel '%s' not found."), *InChannel.ToString());
+		return false;
+	}
+
+	FRigControlElement* Channel = Cast<FRigControlElement>(ChannelBase);
+	if(Channel == nullptr)
+	{
+		ReportWarningf(TEXT("Cannot Remove Channel Host, '%s' is not a Control."), *InChannel.ToString());
+		return false;
+	}
+
+	if(!Channel->IsAnimationChannel())
+	{
+		ReportWarningf(TEXT("Cannot Remove Channel Host, '%s' is not an animation channel."), *InChannel.ToString());
+		return false;
+	}
+
+	const FRigBaseElement* HostBase = Hierarchy->Find(InHost);
+	if(HostBase == nullptr)
+	{
+		ReportWarningf(TEXT("Cannot Remove Channel Host, Host '%s' not found."), *InHost.ToString());
+		return false;
+	}
+
+	const FRigControlElement* Host = Cast<FRigControlElement>(HostBase);
+	if(Host == nullptr)
+	{
+		ReportWarningf(TEXT("Cannot Remove Channel Host, '%s' is not a Control."), *InHost.ToString());
+		return false;
+	}
+
+	if(Host->IsAnimationChannel())
+	{
+		ReportWarningf(TEXT("Cannot Remove Channel Host, '%s' is also an animation channel."), *InHost.ToString());
+		return false;
+	}
+
+	if(!Channel->Settings.Customization.AvailableSpaces.Contains(Host->GetKey()))
+	{
+		ReportWarningf(TEXT("Cannot Remove Channel Host, '%s' is not a host for channel '%s'."), *InHost.ToString(), *InChannel.ToString());
+		return false;
+	}
+
+#if WITH_EDITOR
+	TSharedPtr<FScopedTransaction> TransactionPtr;
+	if(bSetupUndo)
+	{
+		TransactionPtr = MakeShared<FScopedTransaction>(NSLOCTEXT("RigHierarchyController", "RemoveChannelHost", "Remove Channel Host"));
+		Hierarchy->Modify();
+	}
+#endif
+
+	const bool bSuccess = RemoveAvailableSpace(Channel, Host);
+
+#if WITH_EDITOR
+	if(!bSuccess && TransactionPtr.IsValid())
+	{
+		TransactionPtr->Cancel();
+	}
+	TransactionPtr.Reset();
+
+	if (bSuccess && bPrintPythonCommand && !bSuspendPythonPrinting)
+	{
+		if (const UBlueprint* Blueprint = GetTypedOuter<UBlueprint>())
+		{
+			RigVMPythonUtils::Print(Blueprint->GetFName().ToString(),
+				FString::Printf(TEXT("hierarchy_controller.remove_channel_host(%s, %s)"),
+				*InChannel.ToPythonString(),
+				*InHost.ToPythonString()));
+		}
+	}
+#endif
+
+	return bSuccess;
+}
+
 TArray<FRigElementKey> URigHierarchyController::DuplicateElements(TArray<FRigElementKey> InKeys, bool bSelectNewElements, bool bSetupUndo, bool bPrintPythonCommands)
 {
 	const FString Content = ExportToText(InKeys);
@@ -3393,6 +3778,110 @@ bool URigHierarchyController::SetParent(FRigBaseElement* InChild, FRigBaseElemen
 	return AddParent(InChild, InParent, 1.f, bMaintainGlobalTransform, true);
 }
 
+bool URigHierarchyController::AddAvailableSpace(FRigControlElement* InControlElement, const FRigTransformElement* InSpaceElement)
+{
+	if(InControlElement == nullptr || InSpaceElement == nullptr)
+	{
+		return false;
+	}
+
+	// we cannot use animation channels as spaces / channel hosts
+	if(const FRigControlElement* SpaceControlElement = Cast<FRigControlElement>(InSpaceElement))
+	{
+		if(SpaceControlElement->IsAnimationChannel())
+		{
+			return false;
+		}
+	}
+
+	// in case of animation channels - we can only relate them to controls
+	if(InControlElement->IsAnimationChannel())
+	{
+		if(!InSpaceElement->IsA<FRigControlElement>())
+		{
+			return false;
+		}
+	}
+
+	// the default parent cannot be added to the available spaces list
+	if(GetHierarchy()->GetParents(InControlElement).Contains(InSpaceElement))
+	{
+		return false;
+	}
+
+	FRigControlSettings Settings = InControlElement->Settings;
+	TArray<FRigElementKey>& AvailableSpaces = Settings.Customization.AvailableSpaces;
+	if(AvailableSpaces.Contains(InSpaceElement->GetKey()))
+	{
+		return false;
+	}
+
+	AvailableSpaces.Add(InSpaceElement->GetKey());
+
+	GetHierarchy()->SetControlSettings(InControlElement, Settings, false, false, false);
+	return true;
+}
+
+bool URigHierarchyController::RemoveAvailableSpace(FRigControlElement* InControlElement, const FRigTransformElement* InSpaceElement)
+{
+	if(InControlElement == nullptr || InSpaceElement == nullptr)
+	{
+		return false;
+	}
+
+	FRigControlSettings Settings = InControlElement->Settings;
+	TArray<FRigElementKey>& AvailableSpaces =Settings.Customization.AvailableSpaces;
+	const int32 NumRemoved = AvailableSpaces.Remove(InSpaceElement->GetKey());
+	if(NumRemoved == 0)
+	{
+		return false;
+	}
+	
+	GetHierarchy()->SetControlSettings(InControlElement, Settings, false, false, false);
+	return true;
+}
+
+bool URigHierarchyController::SetAvailableSpaceIndex(FRigControlElement* InControlElement, const FRigTransformElement* InSpaceElement, int32 InIndex)
+{
+	if(InControlElement == nullptr || InSpaceElement == nullptr)
+	{
+		return false;
+	}
+
+	FRigControlSettings Settings = InControlElement->Settings;
+	TArray<FRigElementKey>& AvailableSpaces = Settings.Customization.AvailableSpaces;
+
+	bool bAddedAvailableSpace = false;
+	if(!AvailableSpaces.Contains(InSpaceElement->GetKey()))
+	{
+		bAddedAvailableSpace = AddAvailableSpace(InControlElement, InSpaceElement);
+		if(!bAddedAvailableSpace)
+		{
+			return false;
+		}
+	}
+
+	if(AvailableSpaces.Find(InSpaceElement->GetKey()) == InIndex)
+	{
+		return bAddedAvailableSpace;
+	}
+
+	(void)AvailableSpaces.Remove(InSpaceElement->GetKey());
+	InIndex = FMath::Max(InIndex, 0);
+
+	if(InIndex >= AvailableSpaces.Num())
+	{
+		AvailableSpaces.Add(InSpaceElement->GetKey());
+	}
+	else
+	{
+		AvailableSpaces.Insert(InSpaceElement->GetKey(), InIndex);
+	}
+
+	GetHierarchy()->SetControlSettings(InControlElement, Settings, false, false, false);
+	return true;
+}
+
 void URigHierarchyController::AddElementToDirty(FRigBaseElement* InParent, FRigBaseElement* InElementToAdd, int32 InHierarchyDistance) const
 {
 	if(InParent == nullptr)
diff --git a/Engine/Plugins/Animation/ControlRig/Source/ControlRig/Private/Units/Execution/RigUnit_DynamicHierarchy.cpp b/Engine/Plugins/Animation/ControlRig/Source/ControlRig/Private/Units/Execution/RigUnit_DynamicHierarchy.cpp
--- a/Engine/Plugins/Animation/ControlRig/Source/ControlRig/Private/Units/Execution/RigUnit_DynamicHierarchy.cpp
+++ b/Engine/Plugins/Animation/ControlRig/Source/ControlRig/Private/Units/Execution/RigUnit_DynamicHierarchy.cpp
@@ -102,6 +102,23 @@ FRigUnit_SetDefaultParent_Execute()
 	}
 }
 
+FRigUnit_SetChannelHosts_Execute()
+{
+	if(!FRigUnit_DynamicHierarchyBase::IsValidToRunInContext(ExecuteContext, true))
+	{
+		return;
+	}
+	
+	for(const FRigElementKey& Host : Hosts)
+	{
+		FRigHierarchyEnableControllerBracket EnableController(ExecuteContext.Hierarchy, true);
+		if(URigHierarchyController* Controller = ExecuteContext.Hierarchy->GetController(true))
+		{
+			(void)Controller->AddChannelHost(Channel, Host);
+		}
+	}
+}
+
 FRigUnit_SwitchParent_Execute()
 {
 	if(!FRigUnit_DynamicHierarchyBase::IsValidToRunInContext(ExecuteContext, false))
diff --git a/Engine/Plugins/Animation/ControlRig/Source/ControlRig/Public/Rigs/RigHierarchy.h b/Engine/Plugins/Animation/ControlRig/Source/ControlRig/Public/Rigs/RigHierarchy.h
--- a/Engine/Plugins/Animation/ControlRig/Source/ControlRig/Public/Rigs/RigHierarchy.h
+++ b/Engine/Plugins/Animation/ControlRig/Source/ControlRig/Public/Rigs/RigHierarchy.h
@@ -4889,7 +4889,9 @@ private:
 
 	// certain units are allowed to use this
 	friend struct FRigUnit_AddParent;
+	friend struct FRigUnit_AddParents;
 	friend struct FRigUnit_SetDefaultParent;
+	friend struct FRigUnit_SetChannelHosts;
 
 private:
 	TGuardValue<bool> GuardIsControllerAvailable;
diff --git a/Engine/Plugins/Animation/ControlRig/Source/ControlRig/Public/Rigs/RigHierarchyController.h b/Engine/Plugins/Animation/ControlRig/Source/ControlRig/Public/Rigs/RigHierarchyController.h
--- a/Engine/Plugins/Animation/ControlRig/Source/ControlRig/Public/Rigs/RigHierarchyController.h
+++ b/Engine/Plugins/Animation/ControlRig/Source/ControlRig/Public/Rigs/RigHierarchyController.h
@@ -525,6 +525,64 @@ public:
 	UFUNCTION(BlueprintCallable, Category = URigHierarchyController)
 	bool SetParent(FRigElementKey InChild, FRigElementKey InParent, bool bMaintainGlobalTransform = true, bool bSetupUndo = false, bool bPrintPythonCommand = false);
 
+	/**
+	 * Adds a new available space to the given control
+	 * @param InControl The control to add the available space for
+	 * @param InSpace The space to add to the available spaces list
+	 * @param bSetupUndo If set to true the stack will record the change for undo / redo
+	 * @param bPrintPythonCommand If set to true a python command equivalent to this call will be printed out
+	 * @return Returns true if successful.
+	 */
+	UFUNCTION(BlueprintCallable, Category = URigHierarchyController)
+	bool AddAvailableSpace(FRigElementKey InControl, FRigElementKey InSpace, bool bSetupUndo = false, bool bPrintPythonCommand = false);
+
+	/**
+	 * Removes an available space from the given control
+	 * @param InControl The control to remove the available space from
+	 * @param InSpace The space to remove from the available spaces list
+	 * @param bSetupUndo If set to true the stack will record the change for undo / redo
+	 * @param bPrintPythonCommand If set to true a python command equivalent to this call will be printed out
+	 * @return Returns true if successful.
+	 */
+	UFUNCTION(BlueprintCallable, Category = URigHierarchyController)
+	bool RemoveAvailableSpace(FRigElementKey InControl, FRigElementKey InSpace, bool bSetupUndo = false, bool bPrintPythonCommand = false);
+
+	/**
+	 * Reorders an available space for the given control
+	 * @param InControl The control to reorder the host for
+	 * @param InSpace The space to set the new index for
+	 * @param InIndex The new index of the available space
+	 * @param bSetupUndo If set to true the stack will record the change for undo / redo
+	 * @param bPrintPythonCommand If set to true a python command equivalent to this call will be printed out
+	 * @return Returns true if successful.
+	 */
+	UFUNCTION(BlueprintCallable, Category = URigHierarchyController)
+	bool SetAvailableSpaceIndex(FRigElementKey InControl, FRigElementKey InSpace, int32 InIndex, bool bSetupUndo = false, bool bPrintPythonCommand = false);
+
+	/**
+	 * Adds a new channel host to the animation channel
+	 * @note This is just an overload of AddAvailableSpace for readability
+	 * @param InChannel The animation channel to add the channel host for
+	 * @param InHost The host to add to the channel to
+	 * @param bSetupUndo If set to true the stack will record the change for undo / redo
+	 * @param bPrintPythonCommand If set to true a python command equivalent to this call will be printed out
+	 * @return Returns true if successful.
+	 */
+	UFUNCTION(BlueprintCallable, Category = URigHierarchyController)
+	bool AddChannelHost(FRigElementKey InChannel, FRigElementKey InHost, bool bSetupUndo = false, bool bPrintPythonCommand = false);
+
+	/**
+	 * Removes an channel host from the animation channel
+	 * @note This is just an overload of RemoveAvailableSpace for readability
+	 * @param InChannel The animation channel to remove the channel host from
+	 * @param InHost The host to remove from the channel from
+	 * @param bSetupUndo If set to true the stack will record the change for undo / redo
+	 * @param bPrintPythonCommand If set to true a python command equivalent to this call will be printed out
+	 * @return Returns true if successful.
+	 */
+	UFUNCTION(BlueprintCallable, Category = URigHierarchyController)
+	bool RemoveChannelHost(FRigElementKey InChannel, FRigElementKey InHost, bool bSetupUndo = false, bool bPrintPythonCommand = false);
+
 	/**
 	 * Duplicate the given elements
 	 * @param InKeys The keys of the elements to duplicate
@@ -720,6 +778,31 @@ private:
  	 */
 	bool SetParent(FRigBaseElement* InChild, FRigBaseElement* InParent, bool bMaintainGlobalTransform = true);
 
+	/**
+	 * Adds a new available space to the given control
+	 * @param InControlElement The control element to add the available space for
+	 * @param InSpaceElement The space element to add to the available spaces list
+	 * @return Returns true if successful.
+	 */
+	bool AddAvailableSpace(FRigControlElement* InControlElement, const FRigTransformElement* InSpaceElement);
+
+	/**
+	 * Removes an available space from the given control
+	 * @param InControlElement The control element to remove the available space from
+	 * @param InSpaceElement The space element to remove from the available spaces list
+	 * @return Returns true if successful.
+	 */
+	bool RemoveAvailableSpace(FRigControlElement* InControlElement, const FRigTransformElement* InSpaceElement);
+
+	/**
+	 * Reorders an available space for the given control
+	 * @param InControlElement The control element to remove the available space from
+	 * @param InSpaceElement The space element to remove from the available spaces list
+	 * @param InIndex The new index of the available space
+	 * @return Returns true if successful.
+	 */
+	bool SetAvailableSpaceIndex(FRigControlElement* InControlElement, const FRigTransformElement* InSpaceElement, int32 InIndex);
+
 	/**
 	 * Adds a new element to the dirty list of the given parent.
 	 * This function is recursive and will affect all parents in the tree.
diff --git a/Engine/Plugins/Animation/ControlRig/Source/ControlRig/Public/Units/Execution/RigUnit_DynamicHierarchy.h b/Engine/Plugins/Animation/ControlRig/Source/ControlRig/Public/Units/Execution/RigUnit_DynamicHierarchy.h
--- a/Engine/Plugins/Animation/ControlRig/Source/ControlRig/Public/Units/Execution/RigUnit_DynamicHierarchy.h
+++ b/Engine/Plugins/Animation/ControlRig/Source/ControlRig/Public/Units/Execution/RigUnit_DynamicHierarchy.h
@@ -82,6 +82,35 @@ struct CONTROLRIG_API FRigUnit_SetDefaultParent : public FRigUnit_DynamicHierarc
 	FRigElementKey Parent;
 };
 
+/**
+ * Allows an animation channel to be hosted by multiple controls
+ */
+USTRUCT(meta=(DisplayName="Set Channel Hosts", Keywords="Children,Parent,Constraint,Space,SetParent,AddParent,Channel", Varying))
+struct CONTROLRIG_API FRigUnit_SetChannelHosts : public FRigUnit_DynamicHierarchyBaseMutable
+{
+	GENERATED_BODY()
+
+	FRigUnit_SetChannelHosts()
+	{
+		Channel = FRigElementKey(NAME_None, ERigElementType::Control);
+	}
+
+	RIGVM_METHOD()
+	virtual void Execute() override;
+
+	/*
+	 * The channel to receive more hosts
+	 */
+	UPROPERTY(meta = (Input, ExpandByDefault))
+	FRigElementKey Channel;
+
+	/*
+	 * The hosts to add to the channel
+	 */
+	UPROPERTY(meta = (Input, ExpandByDefault))
+	TArray<FRigElementKey> Hosts;
+};
+
 UENUM(meta = (RigVMTypeAllowed))
 enum class ERigSwitchParentMode : uint8
 {
diff --git a/Engine/Plugins/Animation/ControlRig/Source/ControlRigEditor/Private/ControlRigElementDetails.cpp b/Engine/Plugins/Animation/ControlRig/Source/ControlRigEditor/Private/ControlRigElementDetails.cpp
--- a/Engine/Plugins/Animation/ControlRig/Source/ControlRigEditor/Private/ControlRigElementDetails.cpp
+++ b/Engine/Plugins/Animation/ControlRig/Source/ControlRigEditor/Private/ControlRigElementDetails.cpp
@@ -9,6 +9,7 @@
 #include "Widgets/Input/SVectorInputBox.h"
 #include "Widgets/Input/SCheckBox.h"
 #include "Widgets/Input/SButton.h"
+#include "Widgets/Layout/SExpandableArea.h"
 #include "Widgets/Colors/SColorBlock.h"
 #include "Widgets/Colors/SColorPicker.h"
 #include "ControlRigBlueprint.h"
@@ -22,7 +23,6 @@
 #include "Widgets/SRigVMGraphPinVariableBinding.h"
 #include "HAL/PlatformApplicationMisc.h"
 #include "Styling/AppStyle.h"
-#include "Editor/SRigHierarchyTreeView.h"
 #include "StructViewerFilter.h"
 #include "StructViewerModule.h"
 #include "Widgets/SRigVMGraphPinEnumPicker.h"
@@ -2812,6 +2812,7 @@ void FRigControlElementDetails::CustomizeDetails(IDetailLayoutBuilder& DetailBui
 	CustomizeValue(DetailBuilder);
 	CustomizeTransform(DetailBuilder);
 	CustomizeShape(DetailBuilder);
+	CustomizeAvailableSpaces(DetailBuilder);
 	CustomizeAnimationChannels(DetailBuilder);
 	CustomizeMetadata(DetailBuilder);
 }
@@ -3259,8 +3260,6 @@ void FRigControlElementDetails::CustomizeControl(IDetailLayoutBuilder& DetailBui
 		];
 	}
 
-	const TSharedPtr<IPropertyHandle> CustomizationHandle = SettingsHandle->GetChildHandle(TEXT("Customization"));
-
 	if(bSupportsShape)
 	{
 		const TSharedPtr<IPropertyHandle> RestrictSpaceSwitchingHandle = SettingsHandle->GetChildHandle(TEXT("bRestrictSpaceSwitching"));
@@ -3268,10 +3267,12 @@ void FRigControlElementDetails::CustomizeControl(IDetailLayoutBuilder& DetailBui
 		.AddProperty(RestrictSpaceSwitchingHandle.ToSharedRef())
 		.DisplayName(FText::FromString(TEXT("Restrict Switching")))
 		.IsEnabled(bIsEnabled);
-		
-		const TSharedPtr<IPropertyHandle> AvailableSpacesHandle = CustomizationHandle->GetChildHandle(TEXT("AvailableSpaces"));
-		ControlCategory.AddProperty(AvailableSpacesHandle.ToSharedRef())
-		.IsEnabled(bIsEnabled);
+
+		// Available Spaces is now handled by its own category (CustomizeAvailableSpaces)
+		//const TSharedPtr<IPropertyHandle> CustomizationHandle = SettingsHandle->GetChildHandle(TEXT("Customization"));
+		//const TSharedPtr<IPropertyHandle> AvailableSpacesHandle = CustomizationHandle->GetChildHandle(TEXT("AvailableSpaces"));
+		//ControlCategory.AddProperty(AvailableSpacesHandle.ToSharedRef())
+		//.IsEnabled(bIsEnabled);
 	}
 
 	TArray<FRigElementKey> Keys = GetElementKeys();
@@ -3447,6 +3448,7 @@ void FRigControlElementDetails::CustomizeAnimationChannels(IDetailLayoutBuilder&
 				}
 
 				TSharedPtr<SButton> SelectAnimationChannelButton;
+				TSharedPtr<SImage> SelectAnimationChannelImage;
 				TSharedPtr<SWidget> NameContent;
 
 				SAssignNew(NameContent, SHorizontalBox)
@@ -3506,20 +3508,20 @@ void FRigControlElementDetails::CustomizeAnimationChannels(IDetailLayoutBuilder&
 				[
 					SAssignNew(SelectAnimationChannelButton, SButton)
 					.ButtonStyle( FAppStyle::Get(), "NoBorder" )
-					.ButtonColorAndOpacity_Lambda([SelectAnimationChannelButton]() { return FRigElementKeyDetails::OnGetWidgetBackground(SelectAnimationChannelButton); })
 					.OnClicked_Lambda([this, ChildElementKey]() -> FReply
 					{
 						return OnSelectElementClicked(ChildElementKey);
 					})
 					.ContentPadding(0)
-					.ToolTipText(NSLOCTEXT("ControlRigElementDetails", "SelectParentInHierarchyToolTip", "Select Parent in hierarchy"))
+					.ToolTipText(NSLOCTEXT("ControlRigElementDetails", "SelectAnimationChannelInHierarchyToolTip", "Select Animation Channel"))
 					[
-						SNew(SImage)
-						.ColorAndOpacity_Lambda( [SelectAnimationChannelButton]() { return FRigElementKeyDetails::OnGetWidgetForeground(SelectAnimationChannelButton); })
+						SAssignNew(SelectAnimationChannelImage, SImage)
 						.Image(FAppStyle::GetBrush("Icons.Search"))
 					]
 				];
 
+				SelectAnimationChannelImage->SetColorAndOpacity(TAttribute<FSlateColor>::CreateLambda([this, SelectAnimationChannelButton]() { return FRigElementKeyDetails::OnGetWidgetForeground(SelectAnimationChannelButton); }));
+
 				const FText Label = FText::FromString(FString::Printf(TEXT("Channel%s"), *ChildControlElement->GetDisplayName().ToString()));
 				const TArray<FRigElementKey> ChildElementKeys = {ChildElementKey};
 				TAttribute<EVisibility> Visibility = EVisibility::Visible;
@@ -3697,6 +3699,323 @@ void FRigControlElementDetails::CustomizeAnimationChannels(IDetailLayoutBuilder&
 	}
 }
 
+void FRigControlElementDetails::CustomizeAvailableSpaces(IDetailLayoutBuilder& DetailBuilder)
+{
+	// only show this if only one control / animation channel is selected
+	if(PerElementInfos.Num() != 1)
+	{
+		return;
+	}
+
+	const FRigControlElement* ControlElement = PerElementInfos[0].GetElement<FRigControlElement>();
+	if(ControlElement == nullptr)
+	{
+		return;
+	}
+
+	const bool bIsAnimationChannel = IsAnyControlOfAnimationType(ERigControlAnimationType::AnimationChannel);
+	const bool bIsProcedural = IsAnyElementProcedural();
+	const bool bIsEnabled = !bIsProcedural;
+
+	URigHierarchy* Hierarchy = PerElementInfos[0].GetHierarchy();
+	URigHierarchy* HierarchyToChange = PerElementInfos[0].GetDefaultHierarchy();
+
+	static const FText ControlSpaces = LOCTEXT("AvailableSpaces", "Available Spaces"); 
+	static const FText ChannelHosts = LOCTEXT("ChannelHosts", "Channel Hosts"); 
+	static const FText ControlSpacesToolTip = LOCTEXT("AvailableSpacesToolTip", "Spaces available for this Control"); 
+	static const FText ChannelHostsToolTip = LOCTEXT("ChannelHostsToolTip", "A list of controls this channel is listed under"); 
+	IDetailCategoryBuilder& Category = DetailBuilder.EditCategory(TEXT("MultiParents"), bIsAnimationChannel ? ChannelHosts : ControlSpaces);
+	Category.SetToolTip(bIsAnimationChannel ? ChannelHostsToolTip : ControlSpaces);
+	
+	const TSharedRef<IPropertyUtilities> PropertyUtilities = DetailBuilder.GetPropertyUtilities();
+
+	DisplaySettings.bShowBones = true;
+	DisplaySettings.bShowControls = true;
+	DisplaySettings.bShowNulls = true;
+	DisplaySettings.bShowReferences = false;
+	DisplaySettings.bShowSockets = false;
+	DisplaySettings.bShowPhysics = false;
+	DisplaySettings.bHideParentsOnFilter = true;
+	DisplaySettings.bFlattenHierarchyOnFilter = true;
+	DisplaySettings.bShowIconColors = true;
+
+	const TSharedRef<SHorizontalBox> HeaderContentWidget = SNew(SHorizontalBox);
+	HeaderContentWidget->AddSlot()
+	.HAlign(HAlign_Right)
+	[
+		SAssignNew(AddSpaceMenuAnchor, SMenuAnchor)
+		.Placement( MenuPlacement_BelowAnchor )
+		.OnGetMenuContent( this, &FRigControlElementDetails::GetAddSpaceContent, PropertyUtilities)
+		.Content()
+		[
+			SNew(SImage)
+			.OnMouseButtonDown(this, &FRigControlElementDetails::OnAddSpaceMouseDown, PropertyUtilities)
+			.Image(FAppStyle::Get().GetBrush("Icons.PlusCircle"))
+			.ColorAndOpacity(FSlateColor::UseForeground())
+		]	
+	];
+	Category.HeaderContent(HeaderContentWidget);
+
+	TArray<FRigElementKey> AvailableSpaces = { Hierarchy->GetDefaultParent(ControlElement->GetKey()) };
+	for(const FRigElementKey& AvailableSpace : ControlElement->Settings.Customization.AvailableSpaces)
+	{
+		AvailableSpaces.AddUnique(AvailableSpace);
+	}
+
+	static const FText RemoveSpaceText = LOCTEXT("RemoveSpace", "Remove Space");
+	static const FText RemoveChannelHostText = LOCTEXT("RemoveChannelHost", "Remove Channel Host");
+	static const FText RemoveSpaceToolTipText = LOCTEXT("RemoveSpaceToolTip", "Removes this space from the list of available spaces");
+	static const FText RemoveChannelHostToolTipText = LOCTEXT("RemoveChannelHostToolTip", "Remove the channel from this hosting control");
+
+	for(int32 Index = 0; Index < AvailableSpaces.Num(); Index++)
+	{
+		const FRigElementKey ControlKey = ControlElement->GetKey();
+		const FRigElementKey AvailableSpace = AvailableSpaces[Index];
+		const bool bIsParentSpace = Index == 0;
+		const TPair<const FSlateBrush*, FSlateColor> BrushAndColor = SRigHierarchyItem::GetBrushForElementType(Hierarchy, AvailableSpace);
+
+		TSharedPtr<SButton> SelectSpaceButton, RemoveSpaceButton, MoveSpaceUpButton, MoveSpaceDownButton;
+		TSharedPtr<SImage> SelectSpaceImage, RemoveSpaceImage, MoveSpaceUpImage, MoveSpaceDownImage;
+
+		FDetailWidgetRow& WidgetRow = Category.AddCustomRow(FText::FromString(AvailableSpace.ToString()))
+		.NameContent()
+		.MinDesiredWidth(200.f)
+		.MaxDesiredWidth(800.f)
+		[
+			SNew(SHorizontalBox)
+			.IsEnabled(bIsEnabled)
+
+			+ SHorizontalBox::Slot()
+			.MaxWidth(32)
+			.HAlign(HAlign_Left)
+			.VAlign(VAlign_Center)
+			.Padding(FMargin(0.f, 0.f, 3.f, 0.f))
+			[
+				SNew(SImage)
+				.Image(BrushAndColor.Key)
+				.ColorAndOpacity(BrushAndColor.Value)
+			]
+
+			+ SHorizontalBox::Slot()
+			.AutoWidth()
+			.HAlign(HAlign_Left)
+			.VAlign(VAlign_Center)
+			.Padding(FMargin(0.f, 0.f, 8.f, 0.f))
+			[
+				SNew(STextBlock)
+				.Font(IDetailLayoutBuilder::GetDetailFont())
+				.Text_Lambda([this, AvailableSpace]() -> FText
+				{
+					return GetDisplayNameForElement(AvailableSpace);
+				})
+			]
+
+			+SHorizontalBox::Slot()
+			.AutoWidth()
+			.HAlign(HAlign_Right)
+			.VAlign(VAlign_Center)
+			.Padding(FMargin(0.f, 0.f, 0.f, 0.f))
+			[
+				SAssignNew(SelectSpaceButton, SButton)
+				.ButtonStyle( FAppStyle::Get(), "NoBorder" )
+				.OnClicked_Lambda([this, AvailableSpace]() -> FReply
+				{
+					return OnSelectElementClicked(AvailableSpace);
+				})
+				.ContentPadding(0)
+				.ToolTipText(LOCTEXT("SelectElementInHierarchy", "Select Element in hierarchy"))
+				[
+					SAssignNew(SelectSpaceImage, SImage)
+					.Image(FAppStyle::GetBrush("Icons.Search"))
+				]
+			]
+		]
+		.ValueContent()
+		[
+			SNew(SHorizontalBox)
+			.IsEnabled(bIsEnabled)
+			
+			+ SHorizontalBox::Slot()
+			.AutoWidth()
+			[
+				SAssignNew(MoveSpaceUpButton, SButton)
+				.Visibility((Index > 0 && !bIsAnimationChannel) ? EVisibility::Visible : EVisibility::Collapsed)
+				.ButtonStyle(FAppStyle::Get(), TEXT("SimpleButton"))
+				.ContentPadding(0)
+				.IsEnabled(Index > 1)
+				.OnClicked_Lambda([this, ControlKey, AvailableSpace, HierarchyToChange, Index, PropertyUtilities]
+				{
+					if(URigHierarchyController* Controller = HierarchyToChange->GetController(true))
+					{
+						FScopedTransaction Transaction(LOCTEXT("MoveAvailableSpaceUp", "Move Available Space Up"));
+						HierarchyToChange->Modify();
+						Controller->SetAvailableSpaceIndex(ControlKey, AvailableSpace, Index - 2);
+						PropertyUtilities->ForceRefresh();
+					}
+					return FReply::Handled();
+				})
+				.ToolTipText(LOCTEXT("MoveUp", "Move up"))
+				[
+					SAssignNew(MoveSpaceUpImage, SImage)
+					.Image(FAppStyle::GetBrush("Icons.ChevronUp"))
+				]
+			]
+
+			+ SHorizontalBox::Slot()
+			.AutoWidth()
+			[
+				SAssignNew(MoveSpaceDownButton, SButton)
+				.Visibility((Index > 0 && !bIsAnimationChannel) ? EVisibility::Visible : EVisibility::Collapsed)
+				.ButtonStyle(FAppStyle::Get(), TEXT("SimpleButton"))
+				.ContentPadding(0)
+				.IsEnabled(Index > 0 && Index < AvailableSpaces.Num() - 1)
+				.OnClicked_Lambda([this, ControlKey, AvailableSpace, HierarchyToChange, Index, PropertyUtilities]
+				{
+					if(URigHierarchyController* Controller = HierarchyToChange->GetController(true))
+					{
+						FScopedTransaction Transaction(LOCTEXT("MoveAvailableSpaceDown", "Move Available Space Down"));
+						HierarchyToChange->Modify();
+						Controller->SetAvailableSpaceIndex(ControlKey, AvailableSpace, Index);
+						PropertyUtilities->ForceRefresh();
+					}
+					return FReply::Handled();
+				})
+				.ToolTipText(LOCTEXT("MoveDown", "Move down"))
+				[
+					SAssignNew(MoveSpaceDownImage, SImage)
+					.Image(FAppStyle::GetBrush("Icons.ChevronDown"))
+				]
+			]
+
+			+ SHorizontalBox::Slot()
+			.AutoWidth()
+			[
+				SAssignNew(RemoveSpaceButton, SButton)
+				.Visibility(Index > 0 ? EVisibility::Visible : EVisibility::Collapsed)
+				.ButtonStyle(FAppStyle::Get(), TEXT("SimpleButton"))
+				.ContentPadding(0)
+				.IsEnabled(Index > 0)
+				.OnClicked_Lambda([this, ControlKey, AvailableSpace, HierarchyToChange, Index, PropertyUtilities]
+				{
+					if(URigHierarchyController* Controller = HierarchyToChange->GetController(true))
+					{
+						const bool bIsAnimationChannel = IsAnyControlOfAnimationType(ERigControlAnimationType::AnimationChannel);
+						FScopedTransaction Transaction(bIsAnimationChannel ? RemoveChannelHostText : RemoveSpaceText);
+						HierarchyToChange->Modify();
+						if(bIsAnimationChannel)
+						{
+							Controller->RemoveChannelHost(ControlKey, AvailableSpace);
+						}
+						else
+						{
+							Controller->RemoveAvailableSpace(ControlKey, AvailableSpace);
+						}
+						PropertyUtilities->ForceRefresh();
+					}
+					return FReply::Handled();
+				})
+				.ToolTipText(LOCTEXT("Remove", "Remove"))
+				[
+					SAssignNew(RemoveSpaceImage, SImage)
+					.Image(FAppStyle::GetBrush("Icons.Delete"))
+				]
+			]
+		];
+
+		SelectSpaceImage->SetColorAndOpacity(TAttribute<FSlateColor>::CreateLambda([this, SelectSpaceButton]() { return FRigElementKeyDetails::OnGetWidgetForeground(SelectSpaceButton); }));
+		MoveSpaceUpImage->SetColorAndOpacity(TAttribute<FSlateColor>::CreateLambda([this, MoveSpaceUpButton]() { return FRigElementKeyDetails::OnGetWidgetForeground(MoveSpaceUpButton); }));
+		MoveSpaceDownImage->SetColorAndOpacity(TAttribute<FSlateColor>::CreateLambda([this, MoveSpaceDownButton]() { return FRigElementKeyDetails::OnGetWidgetForeground(MoveSpaceDownButton); }));
+		RemoveSpaceImage->SetColorAndOpacity(TAttribute<FSlateColor>::CreateLambda([this, RemoveSpaceButton]() { return FRigElementKeyDetails::OnGetWidgetForeground(RemoveSpaceButton); }));
+
+		if(!bIsProcedural)
+		{
+			if(!bIsAnimationChannel)
+			{
+				WidgetRow.AddCustomContextMenuAction(FUIAction(
+					FExecuteAction::CreateLambda([this, ControlKey, AvailableSpace, HierarchyToChange, Index, PropertyUtilities]()
+					{
+						if(URigHierarchyController* Controller = HierarchyToChange->GetController(true))
+						{
+							FScopedTransaction Transaction(LOCTEXT("MoveAvailableSpaceUp", "Move Available Space Up"));
+							HierarchyToChange->Modify();
+							Controller->SetAvailableSpaceIndex(ControlKey, AvailableSpace, Index - 2);
+							PropertyUtilities->ForceRefresh();
+						}
+					}),
+					FCanExecuteAction::CreateLambda([Index](){ return Index > 1; })),
+					LOCTEXT("MoveUp", "Move Up"),
+					LOCTEXT("MoveAvailableSpaceUpToolTip", "Moves this available space up in the list of spaces"),
+					FSlateIcon());
+
+				const int32 NumSpaces = AvailableSpaces.Num();
+				WidgetRow.AddCustomContextMenuAction(FUIAction(
+					FExecuteAction::CreateLambda([this, ControlKey, AvailableSpace, HierarchyToChange, Index, PropertyUtilities]()
+					{
+						if(URigHierarchyController* Controller = HierarchyToChange->GetController(true))
+						{
+							FScopedTransaction Transaction(LOCTEXT("MoveAvailableSpaceDown", "Move Available Space Down"));
+							HierarchyToChange->Modify();
+							Controller->SetAvailableSpaceIndex(ControlKey, AvailableSpace, Index);
+							PropertyUtilities->ForceRefresh();
+						}
+					}),
+					FCanExecuteAction::CreateLambda([Index, NumSpaces](){ return Index > 0 && Index < NumSpaces - 1; })),
+					LOCTEXT("MoveDown", "Move Down"),
+					LOCTEXT("MoveAvailableSpaceDownToolTip", "Moves this available space down in the list of spaces"),
+					FSlateIcon());
+			}
+
+			WidgetRow.AddCustomContextMenuAction(FUIAction(
+				FExecuteAction::CreateLambda([this, ControlKey, AvailableSpace, HierarchyToChange, PropertyUtilities]()
+				{
+					if(URigHierarchyController* Controller = HierarchyToChange->GetController(true))
+					{
+						const bool bIsAnimationChannel = IsAnyControlOfAnimationType(ERigControlAnimationType::AnimationChannel);
+						FScopedTransaction Transaction(bIsAnimationChannel ? RemoveChannelHostText : RemoveSpaceText);
+						HierarchyToChange->Modify();
+						if(bIsAnimationChannel)
+						{
+							Controller->RemoveChannelHost(ControlKey, AvailableSpace);
+						}
+						else
+						{
+							Controller->RemoveAvailableSpace(ControlKey, AvailableSpace);
+						}
+						PropertyUtilities->ForceRefresh();
+					}
+				}),
+				FCanExecuteAction::CreateLambda([bIsParentSpace](){ return !bIsParentSpace; })),
+				bIsAnimationChannel ? RemoveChannelHostText : RemoveSpaceText,
+				bIsAnimationChannel ? RemoveChannelHostToolTipText : RemoveSpaceToolTipText,
+				FSlateIcon());
+		}
+	}
+
+	Category.InitiallyCollapsed(AvailableSpaces.Num() < 2);
+	if(AvailableSpaces.IsEmpty())
+	{
+		static const FText NoSpacesText = LOCTEXT("NoSpacesSet", "No Available Spaces set");
+		static const FText NoChannelHostsText = LOCTEXT("NoChannelHostsSet", "No Channel Hosts set");
+
+		Category.AddCustomRow(FText()).WholeRowContent()
+		[
+			SNew(SHorizontalBox)
+	
+			+ SHorizontalBox::Slot()
+			.AutoWidth()
+			.HAlign(HAlign_Left)
+			.VAlign(VAlign_Center)
+			.Padding(FMargin(0.f, 0.f, 0.f, 0.f))
+			[
+				SNew(STextBlock)
+				.IsEnabled(bIsEnabled)
+				.Font(IDetailLayoutBuilder::GetDetailFont())
+				.Text(bIsAnimationChannel ? NoChannelHostsText : NoSpacesText)
+			]
+		];
+	}
+}
+
 FReply FRigControlElementDetails::OnAddAnimationChannelClicked()
 {
 	if(IsAnyElementNotOfType(ERigElementType::Control) || IsAnyElementProcedural())
@@ -3756,6 +4075,65 @@ TSharedRef<ITableRow> FRigControlElementDetails::HandleGenerateAnimationChannelT
 	];
 }
 
+TSharedRef<SWidget> FRigControlElementDetails::GetAddSpaceContent(const TSharedRef<IPropertyUtilities> PropertyUtilities)
+{
+	if(PerElementInfos.IsEmpty())
+	{
+		return SNullWidget::NullWidget;
+	}
+
+	FRigTreeDelegates RigTreeDelegates;
+	RigTreeDelegates.OnGetHierarchy = FOnGetRigTreeHierarchy::CreateLambda([this]()
+	{
+		return PerElementInfos[0].GetHierarchy();
+	});
+	RigTreeDelegates.OnGetDisplaySettings = FOnGetRigTreeDisplaySettings::CreateSP(this, &FRigControlElementDetails::GetDisplaySettings);
+	RigTreeDelegates.OnGetSelection = FOnRigTreeGetSelection::CreateLambda([]() -> TArray<FRigElementKey> { return {}; });
+	RigTreeDelegates.OnSelectionChanged = FOnRigTreeSelectionChanged::CreateSP(this, &FRigControlElementDetails::OnAddSpaceSelection, PropertyUtilities);
+
+	return SNew(SBox)
+		.Padding(2.5)
+		.MinDesiredWidth(200)
+		.MinDesiredHeight(300)
+		[
+			SNew(SRigHierarchyTreeView)
+			.RigTreeDelegates(RigTreeDelegates)
+			.PopulateOnConstruct(true)
+		];
+}
+
+FReply FRigControlElementDetails::OnAddSpaceMouseDown(const FGeometry& InGeometry, const FPointerEvent& InPointerEvent, const TSharedRef<IPropertyUtilities> PropertyUtilities)
+{
+	AddSpaceMenuAnchor->SetIsOpen(true);
+	return FReply::Handled();
+}
+
+void FRigControlElementDetails::OnAddSpaceSelection(TSharedPtr<FRigTreeElement> Selection, ESelectInfo::Type SelectInfo, const TSharedRef<IPropertyUtilities> PropertyUtilities)
+{
+	if(Selection)
+	{
+		const FRigElementKey ChildKey = PerElementInfos[0].GetElement()->GetKey();
+		const FRigElementKey NewParentKey = Selection->Key;
+		URigHierarchy* HierarchyToChange = PerElementInfos[0].GetDefaultHierarchy();
+		AddSpaceMenuAnchor->SetIsOpen(false);
+
+		const bool bIsAnimationChannel = IsAnyControlOfAnimationType(ERigControlAnimationType::AnimationChannel);
+		static const FText AddSpaceText = LOCTEXT("AddSpace", "Add Space");
+		static const FText AddChannelHostText = LOCTEXT("AddChannelHost", "Add Channel Host");
+		FScopedTransaction Transaction(bIsAnimationChannel ? AddChannelHostText : AddSpaceText);
+		HierarchyToChange->Modify();
+		if(bIsAnimationChannel)
+		{
+			HierarchyToChange->GetController(true)->AddChannelHost(ChildKey, NewParentKey);
+		}
+		else
+		{
+			HierarchyToChange->GetController(true)->AddAvailableSpace(ChildKey, NewParentKey);
+		}
+		PropertyUtilities->ForceRefresh();
+	}
+}
+
 void FRigControlElementDetails::HandleControlTypeChanged(TSharedPtr<ERigControlType> ControlType, ESelectInfo::Type SelectInfo, FRigElementKey ControlKey, const TSharedRef<IPropertyUtilities> PropertyUtilities)
 {
 	HandleControlTypeChanged(*ControlType.Get(), {ControlKey}, PropertyUtilities);
@@ -4528,7 +4906,7 @@ FText FRigControlElementDetails::GetDisplayNameForElement(const FRigElementKey&
 	const FRigControlElement* ControlElement = Hierarchy->Find<FRigControlElement>(InKey);
 	if(ControlElement == nullptr)
 	{
-		return FText();
+		return FText::FromName(InKey.Name);
 	}
 
 	return FText::FromName(ControlElement->GetDisplayName());
diff --git a/Engine/Plugins/Animation/ControlRig/Source/ControlRigEditor/Private/ControlRigElementDetails.h b/Engine/Plugins/Animation/ControlRig/Source/ControlRigEditor/Private/ControlRigElementDetails.h
--- a/Engine/Plugins/Animation/ControlRig/Source/ControlRigEditor/Private/ControlRigElementDetails.h
+++ b/Engine/Plugins/Animation/ControlRig/Source/ControlRigEditor/Private/ControlRigElementDetails.h
@@ -16,6 +16,7 @@
 #include "IPropertyUtilities.h"
 #include "SSearchableComboBox.h"
 #include "Widgets/Input/SSegmentedControl.h"
+#include "Widgets/Input/SMenuAnchor.h"
 #include "DetailCategoryBuilder.h"
 #include "DetailLayoutBuilder.h"
 #include "DetailWidgetRow.h"
@@ -25,6 +26,7 @@
 #include "ScopedTransaction.h"
 #include "Styling/AppStyle.h"
 #include "Algo/Transform.h"
+#include "Editor/SRigHierarchyTreeView.h"
 
 #include "ControlRigElementDetails.generated.h"
 
@@ -409,6 +411,7 @@ public:
 	void CustomizeValue(IDetailLayoutBuilder& DetailBuilder);
 	void CustomizeControl(IDetailLayoutBuilder& DetailBuilder);
 	void CustomizeAnimationChannels(IDetailLayoutBuilder& DetailBuilder);
+	void CustomizeAvailableSpaces(IDetailLayoutBuilder& DetailBuilder);
 	void CustomizeShape(IDetailLayoutBuilder& DetailBuilder);
 	virtual void BeginDestroy() override;
 
@@ -741,6 +744,13 @@ private:
 	// animation channel related callbacks
 	FReply OnAddAnimationChannelClicked();
 	TSharedRef<ITableRow> HandleGenerateAnimationChannelTypeRow(TSharedPtr<ERigControlType> ControlType, const TSharedRef<STableViewBase>& OwnerTable, FRigElementKey ControlKey);
+
+	// multi parent related callbacks
+	const FRigTreeDisplaySettings& GetDisplaySettings() const { return DisplaySettings; }
+	TSharedRef<SWidget> GetAddSpaceContent(const TSharedRef<IPropertyUtilities> PropertyUtilities);
+	FReply OnAddSpaceMouseDown(const FGeometry& InGeometry, const FPointerEvent& InPointerEvent, const TSharedRef<IPropertyUtilities> PropertyUtilities);
+	void OnAddSpaceSelection(TSharedPtr<FRigTreeElement> Selection, ESelectInfo::Type SelectInfo, const TSharedRef<IPropertyUtilities> PropertyUtilities);
+
 	void HandleControlTypeChanged(TSharedPtr<ERigControlType> ControlType, ESelectInfo::Type SelectInfo, FRigElementKey ControlKey, const TSharedRef<IPropertyUtilities> PropertyUtilities);
 	void HandleControlTypeChanged(ERigControlType ControlType, TArray<FRigElementKey> ControlKeys, const TSharedRef<IPropertyUtilities> PropertyUtilities);
 	void HandleControlEnumChanged(TSharedPtr<FString> InItem, ESelectInfo::Type InSelectionInfo, const TSharedRef<IPropertyUtilities> PropertyUtilities);
@@ -755,6 +765,9 @@ private:
 
 	TSharedPtr<SControlRigShapeNameList> ShapeNameListWidget; 
 	static TSharedPtr<TArray<ERigControlValueType>> PickedValueTypes;
+
+	FRigTreeDisplaySettings DisplaySettings;
+	TSharedPtr<SMenuAnchor> AddSpaceMenuAnchor;
 };
 
 template<>
diff --git a/Engine/Plugins/Animation/ControlRig/Source/ControlRigEditor/Private/EditMode/ControlRigEditMode.cpp b/Engine/Plugins/Animation/ControlRig/Source/ControlRigEditor/Private/EditMode/ControlRigEditMode.cpp
--- a/Engine/Plugins/Animation/ControlRig/Source/ControlRigEditor/Private/EditMode/ControlRigEditMode.cpp
+++ b/Engine/Plugins/Animation/ControlRig/Source/ControlRigEditor/Private/EditMode/ControlRigEditMode.cpp
@@ -4509,8 +4509,14 @@ void FControlRigEditMode::OnHierarchyModified(ERigHierarchyNotification InNotif,
 														ParentControlElement->GetKey() == ControlElement->GetKey())
 													{
 														ControlProxy->SelectProxy(ControlRig, AnimationChannelControl, bSelected);
+														return true;
 													}
 												}
+												if(AnimationChannelControl->Settings.Customization.AvailableSpaces.Contains(ControlElement->GetKey()))
+												{
+													ControlProxy->SelectProxy(ControlRig, AnimationChannelControl, bSelected);
+													return true;
+												}
 											}
 											return true;
 										}
diff --git a/Engine/Plugins/Animation/ControlRig/Source/ControlRigEditor/Private/EditMode/ControlsProxyDetailCustomization.cpp b/Engine/Plugins/Animation/ControlRig/Source/ControlRigEditor/Private/EditMode/ControlsProxyDetailCustomization.cpp
--- a/Engine/Plugins/Animation/ControlRig/Source/ControlRigEditor/Private/EditMode/ControlsProxyDetailCustomization.cpp
+++ b/Engine/Plugins/Animation/ControlRig/Source/ControlRigEditor/Private/EditMode/ControlsProxyDetailCustomization.cpp
@@ -950,7 +950,7 @@ void FAnimDetailProxyDetails::CustomizeDetails(IDetailLayoutBuilder& DetailBuild
 	}
 	else
 	{
-		AttributesCategory.SetCategoryVisibility(true);
+		AttributesCategory.SetCategoryVisibility(false);
 	}
 }
 
diff --git a/Engine/Plugins/Animation/ControlRig/Source/ControlRigEditor/Private/EditMode/SControlRigDetails.cpp b/Engine/Plugins/Animation/ControlRig/Source/ControlRigEditor/Private/EditMode/SControlRigDetails.cpp
--- a/Engine/Plugins/Animation/ControlRig/Source/ControlRigEditor/Private/EditMode/SControlRigDetails.cpp
+++ b/Engine/Plugins/Animation/ControlRig/Source/ControlRigEditor/Private/EditMode/SControlRigDetails.cpp
@@ -143,36 +143,47 @@ void SControlRigDetails::HandleControlSelected(UControlRig* Subject, FRigControl
 	UpdateProxies();
 }
 
-static UControlRigControlsProxy* GetParentProxy(UControlRigControlsProxy* ChildProxy, const TArray<UControlRigControlsProxy*>& Proxies)
+static TArray<UControlRigControlsProxy*> GetParentProxies(UControlRigControlsProxy* ChildProxy, const TArray<UControlRigControlsProxy*>& Proxies)
 {
 	if (!ChildProxy || !ChildProxy->OwnerControlRig.IsValid())
 	{
-		return nullptr;
+		return {};
 	}
 	if (!ChildProxy->OwnerControlElement.UpdateCache(ChildProxy->OwnerControlRig->GetHierarchy()))
 	{
-		return nullptr;
+		return {};
 	}
-	FRigBaseElement* ChildParent = (ChildProxy && ChildProxy->OwnerControlRig.IsValid()) ?
-		ChildProxy->OwnerControlRig->GetHierarchy()->GetFirstParent(ChildProxy->OwnerControlElement.GetElement()) : nullptr;
-	if (ChildParent == nullptr)
+	TArray<FRigBaseElement*> Parents;
+	if(ChildProxy && ChildProxy->OwnerControlRig.IsValid())
 	{
-		return nullptr;
+		Parents = ChildProxy->OwnerControlRig->GetHierarchy()->GetParents(ChildProxy->OwnerControlElement.GetElement()); 
 	}
+
+	TArray<UControlRigControlsProxy*> ParentProxies;
 	for (UControlRigControlsProxy* Proxy : Proxies)
 	{
 		if (Proxy && Proxy->OwnerControlRig.IsValid())
 		{
 			if (Proxy->OwnerControlElement.UpdateCache(Proxy->OwnerControlRig->GetHierarchy()))
 			{
-				if (ChildParent == Proxy->OwnerControlElement.GetElement())
+				if(const FRigControlElement* OwnerControlElement = Cast<FRigControlElement>(Proxy->OwnerControlElement.GetElement()))
 				{
-					return Proxy;
+					if (Parents.Contains(OwnerControlElement))
+					{
+						ParentProxies.AddUnique(Proxy);
+					}
+					if(const FRigControlElement* ChildControlElement = Cast<FRigControlElement>(ChildProxy->OwnerControlElement.GetElement()))
+					{
+						if(ChildControlElement->Settings.Customization.AvailableSpaces.Contains(OwnerControlElement->GetKey()))
+						{
+							ParentProxies.AddUnique(Proxy);
+						}
+					}
 				}
 			}
 		}
 	}
-	return nullptr;
+	return ParentProxies;
 }
 static UControlRigControlsProxy* GetProxyWithSameType(TArray<TWeakObjectPtr<>>& AllProxies, ERigControlType ControlType, bool bIsEnum)
 {
@@ -381,7 +392,8 @@ void SControlRigDetails::UpdateProxies()
 					//now add child proxies to parents if parents also selected...
 					for (UControlRigControlsProxy* Proxy : ChildProxies)
 					{
-						if (UControlRigControlsProxy* ParentProxy = GetParentProxy(Proxy, Proxies))
+						TArray<UControlRigControlsProxy*> ParentProxies = GetParentProxies(Proxy, Proxies);
+						for (UControlRigControlsProxy* ParentProxy : ParentProxies)
 						{
 							TObjectPtr<UEnum> EnumPtr = nullptr;
 							if (ParentProxy->OwnerControlRig.IsValid())
@@ -399,7 +411,8 @@ void SControlRigDetails::UpdateProxies()
 								ExistingProxy->AddChildProxy(Proxy);
 							}
 						}
-						else
+
+						if(ParentProxies.IsEmpty())
 						{
 							AllProxies.Add(Proxy);
 						}
diff --git a/Engine/Plugins/Animation/ControlRig/Source/ControlRigEditor/Private/Editor/ControlRigEditor.cpp b/Engine/Plugins/Animation/ControlRig/Source/ControlRigEditor/Private/Editor/ControlRigEditor.cpp
--- a/Engine/Plugins/Animation/ControlRig/Source/ControlRigEditor/Private/Editor/ControlRigEditor.cpp
+++ b/Engine/Plugins/Animation/ControlRig/Source/ControlRigEditor/Private/Editor/ControlRigEditor.cpp
@@ -1176,13 +1176,35 @@ void FControlRigEditor::SetDetailViewForRigElements(const TArray<FRigElementKey>
 		return;
 	}
 
+	TArray<FRigElementKey> Keys = InKeys;
+	if(Keys.IsEmpty())
+	{
+		TArray< TWeakObjectPtr<UObject> > SelectedObjects = GetSelectedObjects();
+		for (TWeakObjectPtr<UObject> SelectedObject : SelectedObjects)
+		{
+			if (SelectedObject.IsValid())
+			{
+				if(const URigVMDetailsViewWrapperObject* WrapperObject = Cast<URigVMDetailsViewWrapperObject>(SelectedObject.Get()))
+				{
+					if(const UScriptStruct* WrappedStruct = WrapperObject->GetWrappedStruct())
+					{
+						if (WrappedStruct->IsChildOf(FRigBaseElement::StaticStruct()))
+						{
+							Keys.Add(WrapperObject->GetContent<FRigBaseElement>().Key);
+						}
+					}
+				}
+			}
+		}
+	}
+
 	ClearDetailObject();
 
 	UControlRigBlueprint* RigBlueprint = Cast<UControlRigBlueprint>(GetBlueprintObj());
 	URigHierarchy* HierarchyBeingDebugged = GetHierarchyBeingDebugged();
 	TArray<UObject*> Objects;
 
-	for(const FRigElementKey& Key : InKeys)
+	for(const FRigElementKey& Key : Keys)
 	{
 		FRigBaseElement* Element = HierarchyBeingDebugged->Find(Key);
 		if (Element == nullptr)
diff --git a/Engine/Plugins/Animation/ControlRig/Source/ControlRigEditor/Private/Editor/SRigHierarchyTreeView.cpp b/Engine/Plugins/Animation/ControlRig/Source/ControlRigEditor/Private/Editor/SRigHierarchyTreeView.cpp
--- a/Engine/Plugins/Animation/ControlRig/Source/ControlRigEditor/Private/Editor/SRigHierarchyTreeView.cpp
+++ b/Engine/Plugins/Animation/ControlRig/Source/ControlRigEditor/Private/Editor/SRigHierarchyTreeView.cpp
@@ -319,6 +319,11 @@ void SRigHierarchyTreeView::Construct(const FArguments& InArgs)
 
 	LastMousePosition = FVector2D::ZeroVector;
 	TimeAtMousePosition = 0.0;
+
+	if(InArgs._PopulateOnConstruct)
+	{
+		RefreshTreeView(true);
+	}
 }
 
 void SRigHierarchyTreeView::Tick(const FGeometry& AllottedGeometry, const double InCurrentTime, const float InDeltaTime)
diff --git a/Engine/Plugins/Animation/ControlRig/Source/ControlRigEditor/Private/Editor/SRigHierarchyTreeView.h b/Engine/Plugins/Animation/ControlRig/Source/ControlRigEditor/Private/Editor/SRigHierarchyTreeView.h
--- a/Engine/Plugins/Animation/ControlRig/Source/ControlRigEditor/Private/Editor/SRigHierarchyTreeView.h
+++ b/Engine/Plugins/Animation/ControlRig/Source/ControlRigEditor/Private/Editor/SRigHierarchyTreeView.h
@@ -288,9 +288,11 @@ public:
 
 	SLATE_BEGIN_ARGS(SRigHierarchyTreeView)
 		: _AutoScrollEnabled(false)
+		, _PopulateOnConstruct(false)
 	{}
 		SLATE_ARGUMENT(FRigTreeDelegates, RigTreeDelegates)
 		SLATE_ARGUMENT(bool, AutoScrollEnabled)
+		SLATE_ARGUMENT(bool, PopulateOnConstruct)
 	SLATE_END_ARGS()
 
 	void Construct(const FArguments& InArgs);
diff --git a/Engine/Plugins/Animation/ControlRig/Source/ControlRigEditor/Private/Sequencer/ControlRigParameterTrackEditor.cpp b/Engine/Plugins/Animation/ControlRig/Source/ControlRigEditor/Private/Sequencer/ControlRigParameterTrackEditor.cpp
--- a/Engine/Plugins/Animation/ControlRig/Source/ControlRigEditor/Private/Sequencer/ControlRigParameterTrackEditor.cpp
+++ b/Engine/Plugins/Animation/ControlRig/Source/ControlRigEditor/Private/Sequencer/ControlRigParameterTrackEditor.cpp
@@ -3393,11 +3393,17 @@ void FControlRigParameterTrackEditor::HandleControlSelected(UControlRig* Subject
 							if (ParentConstraint.ParentElement == ControlElement)
 							{
 								Controller->SelectElement(OtherControlElement->GetKey(), bSelected);
-								break;
+								return true;
 							}
 						}
 					}
 
+					if(OtherControlElement->IsAnimationChannel() && OtherControlElement->Settings.Customization.AvailableSpaces.Contains(ControlElement->GetKey()))
+					{
+						Controller->SelectElement(OtherControlElement->GetKey(), bSelected);
+						return true;
+					}
+
 					return true;
 				});
 		}
