ControlRig: Make sure to update the hierarchy topology version so that it reflects the structure of the hierarchy copied from

#rnx

diff --git a/Engine/Plugins/Animation/ControlRig/Source/ControlRig/Private/Rigs/RigHierarchy.cpp b/Engine/Plugins/Animation/ControlRig/Source/ControlRig/Private/Rigs/RigHierarchy.cpp
--- a/Engine/Plugins/Animation/ControlRig/Source/ControlRig/Private/Rigs/RigHierarchy.cpp
+++ b/Engine/Plugins/Animation/ControlRig/Source/ControlRig/Private/Rigs/RigHierarchy.cpp
@@ -647,6 +647,11 @@ void URigHierarchy::CopyHierarchy(URigHierarchy* InHierarchy)
 
 		PreviousNameMap.Append(InHierarchy->PreviousNameMap);
 
+		// copy the topology version from the hierarchy.
+		// for this we'll include the hash of the hierarchy to make sure we
+		// are deterministic for different hierarchies.
+		TopologyVersion = HashCombine(InHierarchy->GetTopologyVersion(), InHierarchy->GetTopologyHash());
+
 		// Increment the topology version to invalidate our cached children.
 		IncrementTopologyVersion();
 
