Control Rig: Integrate element search with FindInBlueprint

#rnx

diff --git a/Engine/Plugins/Animation/ControlRig/Source/ControlRigDeveloper/Private/Graph/ControlRigGraphNode.cpp b/Engine/Plugins/Animation/ControlRig/Source/ControlRigDeveloper/Private/Graph/ControlRigGraphNode.cpp
--- a/Engine/Plugins/Animation/ControlRig/Source/ControlRigDeveloper/Private/Graph/ControlRigGraphNode.cpp
+++ b/Engine/Plugins/Animation/ControlRig/Source/ControlRigDeveloper/Private/Graph/ControlRigGraphNode.cpp
@@ -14,5 +14,40 @@ UControlRigGraphNode::UControlRigGraphNode()
 {
 }
 
+#if WITH_EDITOR
+
+void UControlRigGraphNode::AddPinSearchMetaDataInfo(const UEdGraphPin* Pin, TArray<FSearchTagDataPair>& OutTaggedMetaData) const
+{
+	Super::AddPinSearchMetaDataInfo(Pin, OutTaggedMetaData);
+
+	if(const URigVMPin* ModelPin = FindModelPinFromGraphPin(Pin))
+	{
+		if(ModelPin->GetCPPTypeObject() == FRigElementKey::StaticStruct())
+		{
+			const FString DefaultValue = ModelPin->GetDefaultValue();
+			if(!DefaultValue.IsEmpty())
+			{
+				FString RigElementKeys;
+				if(ModelPin->IsArray())
+				{
+					RigElementKeys = DefaultValue;
+				}
+				else
+				{
+					RigElementKeys = FString::Printf(TEXT("(%s)"), *DefaultValue);
+				}
+				if(!RigElementKeys.IsEmpty())
+				{
+					RigElementKeys.ReplaceInline(TEXT("="), TEXT(","));
+					RigElementKeys.ReplaceInline(TEXT("\""), TEXT(""));
+					OutTaggedMetaData.Emplace(FText::FromString(TEXT("Rig Items")), FText::FromString(RigElementKeys));
+				}
+			}
+		}
+	}
+}
+
+#endif
+
 #undef LOCTEXT_NAMESPACE
 
diff --git a/Engine/Plugins/Animation/ControlRig/Source/ControlRigDeveloper/Public/Graph/ControlRigGraphNode.h b/Engine/Plugins/Animation/ControlRig/Source/ControlRigDeveloper/Public/Graph/ControlRigGraphNode.h
--- a/Engine/Plugins/Animation/ControlRig/Source/ControlRigDeveloper/Public/Graph/ControlRigGraphNode.h
+++ b/Engine/Plugins/Animation/ControlRig/Source/ControlRigDeveloper/Public/Graph/ControlRigGraphNode.h
@@ -39,6 +39,10 @@ public:
 
 	UControlRigGraphNode();
 
+#if WITH_EDITOR
+	virtual void AddPinSearchMetaDataInfo(const UEdGraphPin* Pin, TArray<FSearchTagDataPair>& OutTaggedMetaData) const override;
+#endif
+
 private:
 
 	friend class SRigVMGraphNode;
diff --git a/Engine/Plugins/Animation/ControlRig/Source/ControlRigEditor/Private/ControlRigHierarchyCommands.cpp b/Engine/Plugins/Animation/ControlRig/Source/ControlRigEditor/Private/ControlRigHierarchyCommands.cpp
--- a/Engine/Plugins/Animation/ControlRig/Source/ControlRigEditor/Private/ControlRigHierarchyCommands.cpp
+++ b/Engine/Plugins/Animation/ControlRig/Source/ControlRigEditor/Private/ControlRigHierarchyCommands.cpp
@@ -12,6 +12,7 @@ void FControlRigHierarchyCommands::RegisterCommands()
 	UI_COMMAND(AddNullItem, "New Null", "Add new null at the origin (0, 0, 0) to the hierarchy.", EUserInterfaceActionType::Button, FInputChord());
 	UI_COMMAND(AddConnectorItem, "New Connector", "Add new connector at the origin (0, 0, 0) to the hierarchy.\nNote: This is only available for Rig Modules.", EUserInterfaceActionType::Button, FInputChord());
 	UI_COMMAND(AddSocketItem, "New Socket", "Add new socket at the origin (0, 0, 0) to the hierarchy.\nNote: This is only available for Rig Modules and Modular Rigs.", EUserInterfaceActionType::Button, FInputChord());
+	UI_COMMAND(FindReferencesOfItem, "Find References", "Find all references of the selected item.", EUserInterfaceActionType::Button, FInputChord(EKeys::F, EModifierKey::Control));
 	UI_COMMAND(DuplicateItem, "Duplicate", "Duplicate the selected items in the hierarchy.", EUserInterfaceActionType::Button, FInputChord(EKeys::D, EModifierKey::Control));
 	UI_COMMAND(MirrorItem, "Mirror", "Mirror the selected items in the hierarchy.", EUserInterfaceActionType::Button, FInputChord());
 	UI_COMMAND(DeleteItem, "Delete", "Deletes the selected items from the hierarchy.", EUserInterfaceActionType::Button, FInputChord(EKeys::Delete));
diff --git a/Engine/Plugins/Animation/ControlRig/Source/ControlRigEditor/Private/ControlRigHierarchyCommands.h b/Engine/Plugins/Animation/ControlRig/Source/ControlRigEditor/Private/ControlRigHierarchyCommands.h
--- a/Engine/Plugins/Animation/ControlRig/Source/ControlRigEditor/Private/ControlRigHierarchyCommands.h
+++ b/Engine/Plugins/Animation/ControlRig/Source/ControlRigEditor/Private/ControlRigHierarchyCommands.h
@@ -36,6 +36,9 @@ public:
 	/** Add Item at origin */
 	TSharedPtr< FUICommandInfo > AddSocketItem;
 
+	/** Find all references of an item in the graphs */
+	TSharedPtr< FUICommandInfo > FindReferencesOfItem;
+	
 	/** Duplicate currently selected items */
 	TSharedPtr< FUICommandInfo > DuplicateItem;
 
diff --git a/Engine/Plugins/Animation/ControlRig/Source/ControlRigEditor/Private/Editor/ControlRigEditor.cpp b/Engine/Plugins/Animation/ControlRig/Source/ControlRigEditor/Private/Editor/ControlRigEditor.cpp
--- a/Engine/Plugins/Animation/ControlRig/Source/ControlRigEditor/Private/Editor/ControlRigEditor.cpp
+++ b/Engine/Plugins/Animation/ControlRig/Source/ControlRigEditor/Private/Editor/ControlRigEditor.cpp
@@ -2856,6 +2856,15 @@ FVector2D FControlRigEditor::ComputePersonaProjectedScreenPos(const FVector& InW
 	return FVector2D::ZeroVector;
 }
 
+void FControlRigEditor::FindReferencesOfItem(const FRigElementKey& InKey)
+{
+	static constexpr TCHAR Format[] = TEXT("Type,%s,Name,%s");
+	static const UEnum* TypeEnum = StaticEnum<ERigElementType>();
+	const FText TypeText = TypeEnum->GetDisplayNameTextByValue((int64)InKey.Type);
+	const FString Query = FString::Printf(Format, *TypeText.ToString(), *InKey.Name.ToString()); 
+	SummonSearchUI(true, Query, true);
+}
+
 void FControlRigEditor::HandlePreviewMeshChanged(USkeletalMesh* InOldSkeletalMesh, USkeletalMesh* InNewSkeletalMesh)
 {
 	RebindToSkeletalMeshComponent();
diff --git a/Engine/Plugins/Animation/ControlRig/Source/ControlRigEditor/Private/Editor/ControlRigEditor.h b/Engine/Plugins/Animation/ControlRig/Source/ControlRigEditor/Private/Editor/ControlRigEditor.h
--- a/Engine/Plugins/Animation/ControlRig/Source/ControlRigEditor/Private/Editor/ControlRigEditor.h
+++ b/Engine/Plugins/Animation/ControlRig/Source/ControlRigEditor/Private/Editor/ControlRigEditor.h
@@ -199,6 +199,8 @@ public:
 	FSimpleMulticastDelegate& OnRequestNavigateToConnectorWarning() { return RequestNavigateToConnectorWarningDelegate; }
 
 	FVector2D ComputePersonaProjectedScreenPos(const FVector& InWorldPos, bool bClampToScreenRectangle = false);
+
+	void FindReferencesOfItem(const FRigElementKey& InKey);
 	
 protected:
 
diff --git a/Engine/Plugins/Animation/ControlRig/Source/ControlRigEditor/Private/Editor/SRigHierarchy.cpp b/Engine/Plugins/Animation/ControlRig/Source/ControlRigEditor/Private/Editor/SRigHierarchy.cpp
--- a/Engine/Plugins/Animation/ControlRig/Source/ControlRigEditor/Private/Editor/SRigHierarchy.cpp
+++ b/Engine/Plugins/Animation/ControlRig/Source/ControlRigEditor/Private/Editor/SRigHierarchy.cpp
@@ -311,6 +311,10 @@ void SRigHierarchy::BindCommands()
 		FExecuteAction::CreateSP(this, &SRigHierarchy::HandleNewItem, ERigElementType::Socket, false),
 		FCanExecuteAction::CreateSP(this, &SRigHierarchy::CanAddElement, ERigElementType::Socket));
 
+	CommandList->MapAction(Commands.FindReferencesOfItem,
+		FExecuteAction::CreateSP(this, &SRigHierarchy::HandleFindReferencesOfItem),
+		FCanExecuteAction::CreateSP(this, &SRigHierarchy::CanFindReferencesOfItem));
+
 	CommandList->MapAction(Commands.DuplicateItem,
 		FExecuteAction::CreateSP(this, &SRigHierarchy::HandleDuplicateItem),
 		FCanExecuteAction::CreateSP(this, &SRigHierarchy::CanDuplicateItem));
@@ -1287,6 +1291,7 @@ void SRigHierarchy::CreateContextMenu()
 					
 					ElementsSection.AddMenuEntry(Commands.DeleteItem);
 					ElementsSection.AddMenuEntry(Commands.DuplicateItem);
+					ElementsSection.AddMenuEntry(Commands.FindReferencesOfItem);
 					ElementsSection.AddMenuEntry(Commands.RenameItem);
 					ElementsSection.AddMenuEntry(Commands.MirrorItem);
 
@@ -2185,6 +2190,21 @@ void SRigHierarchy::HandleNewItem(ERigElementType InElementType, bool bIsAnimati
 	RefreshTreeView();
 }
 
+bool SRigHierarchy::CanFindReferencesOfItem() const
+{
+	return !GetSelectedKeys().IsEmpty();
+}
+
+void SRigHierarchy::HandleFindReferencesOfItem()
+{
+	if(!ControlRigEditor.IsValid() || GetSelectedKeys().IsEmpty())
+	{
+		return;
+	}
+
+	ControlRigEditor.Pin()->FindReferencesOfItem(GetSelectedKeys()[0]);
+}
+
 /** Check whether we can deleting the selected item(s) */
 bool SRigHierarchy::CanDuplicateItem() const
 {
diff --git a/Engine/Plugins/Animation/ControlRig/Source/ControlRigEditor/Private/Editor/SRigHierarchy.h b/Engine/Plugins/Animation/ControlRig/Source/ControlRigEditor/Private/Editor/SRigHierarchy.h
--- a/Engine/Plugins/Animation/ControlRig/Source/ControlRigEditor/Private/Editor/SRigHierarchy.h
+++ b/Engine/Plugins/Animation/ControlRig/Source/ControlRigEditor/Private/Editor/SRigHierarchy.h
@@ -82,6 +82,12 @@ private:
 	/** Create a new item */
 	void HandleNewItem(ERigElementType InElementType, bool bIsAnimationChannel);
 
+	/** Check we can find the references of an item */
+	bool CanFindReferencesOfItem() const;
+
+	/** Find all references of an item */
+	void HandleFindReferencesOfItem();
+
 	/** Check whether we can deleting the selected item(s) */
 	bool CanDuplicateItem() const;
 
diff --git a/Engine/Plugins/Runtime/RigVM/Source/RigVMDeveloper/Private/EdGraph/RigVMEdGraphNode.cpp b/Engine/Plugins/Runtime/RigVM/Source/RigVMDeveloper/Private/EdGraph/RigVMEdGraphNode.cpp
--- a/Engine/Plugins/Runtime/RigVM/Source/RigVMDeveloper/Private/EdGraph/RigVMEdGraphNode.cpp
+++ b/Engine/Plugins/Runtime/RigVM/Source/RigVMDeveloper/Private/EdGraph/RigVMEdGraphNode.cpp
@@ -2137,17 +2137,13 @@ void URigVMEdGraphNode::BeginDestroy()
 void URigVMEdGraphNode::AddPinSearchMetaDataInfo(const UEdGraphPin* Pin, TArray<FSearchTagDataPair>& OutTaggedMetaData) const
 {
 	Super::AddPinSearchMetaDataInfo(Pin, OutTaggedMetaData);
-	if (const URigVMNode* Node = GetModelNode())
+	if(const URigVMPin* ModelPin = FindModelPinFromGraphPin(Pin))
 	{
-		const FString PinName = Pin->GetName();
-		FString Left, Right = PinName;
-		URigVMPin::SplitPinPathAtStart(PinName, Left, Right);
-		if (URigVMPin* ModelPin = Node->FindPin(Right))
+		OutTaggedMetaData.Emplace(FText::FromString(TEXT("Data Type")), FText::FromString(ModelPin->GetCPPType()));
+		
+		if (ModelPin->IsBoundToVariable())
 		{
-			if (ModelPin->IsBoundToVariable())
-			{
-				OutTaggedMetaData.Add(FSearchTagDataPair(FText::FromString(TEXT("Binding")), FText::FromString(ModelPin->GetBoundVariablePath())));
-			}
+			OutTaggedMetaData.Add(FSearchTagDataPair(FText::FromString(TEXT("Binding")), FText::FromString(ModelPin->GetBoundVariablePath())));
 		}
 	}
 }
