Control Rig: Hide connectors from add space tree view

#rnx

diff --git a/Engine/Plugins/Animation/ControlRig/Source/ControlRigEditor/Private/Editor/SRigSpacePickerWidget.cpp b/Engine/Plugins/Animation/ControlRig/Source/ControlRigEditor/Private/Editor/SRigSpacePickerWidget.cpp
--- a/Engine/Plugins/Animation/ControlRig/Source/ControlRigEditor/Private/Editor/SRigSpacePickerWidget.cpp
+++ b/Engine/Plugins/Animation/ControlRig/Source/ControlRigEditor/Private/Editor/SRigSpacePickerWidget.cpp
@@ -698,7 +698,12 @@ void SRigSpacePickerWidget::HandleSpaceDelete(FRigElementKey InKey)
 
 FReply SRigSpacePickerWidget::HandleAddElementClicked()
 {
+	HierarchyDisplaySettings.bShowConnectors = false;
+	HierarchyDisplaySettings.bShowSockets = false;
+	HierarchyDisplaySettings.bShowPhysics = false;
+	
 	FRigTreeDelegates TreeDelegates;
+	TreeDelegates.OnGetDisplaySettings = FOnGetRigTreeDisplaySettings::CreateSP(this, &SRigSpacePickerWidget::GetHierarchyDisplaySettings); 
 	TreeDelegates.OnGetHierarchy = FOnGetRigTreeHierarchy::CreateSP(this, &SRigSpacePickerWidget::GetHierarchyConst);
 	TreeDelegates.OnMouseButtonClick = FOnRigTreeMouseButtonClick::CreateLambda([this](TSharedPtr<FRigTreeElement> InItem)
 	{
diff --git a/Engine/Plugins/Animation/ControlRig/Source/ControlRigEditor/Private/Editor/SRigSpacePickerWidget.h b/Engine/Plugins/Animation/ControlRig/Source/ControlRigEditor/Private/Editor/SRigSpacePickerWidget.h
--- a/Engine/Plugins/Animation/ControlRig/Source/ControlRigEditor/Private/Editor/SRigSpacePickerWidget.h
+++ b/Engine/Plugins/Animation/ControlRig/Source/ControlRigEditor/Private/Editor/SRigSpacePickerWidget.h
@@ -81,6 +81,8 @@ public:
 		}
 		return nullptr;
 	}
+
+	const FRigTreeDisplaySettings& GetHierarchyDisplaySettings() const { return HierarchyDisplaySettings; }
 	const URigHierarchy* GetHierarchyConst() const { return GetHierarchy(); }
 	
 	const TArray<FRigElementKey>& GetControls() const { return ControlKeys; }
@@ -173,6 +175,7 @@ private:
 	TWeakPtr<IMenu> ContextMenu;
 	FDelegateHandle HierarchyModifiedHandle;
 	FDelegateHandle ActiveSpaceChangedWindowHandle;
+	FRigTreeDisplaySettings HierarchyDisplaySettings;
 
 	static FRigElementKey InValidKey;
 };
