Control Rig: Fix localization conflict

#rnx

diff --git a/Engine/Plugins/Animation/ControlRig/Source/ControlRigEditor/Private/ControlRigElementDetails.cpp b/Engine/Plugins/Animation/ControlRig/Source/ControlRigEditor/Private/ControlRigElementDetails.cpp
--- a/Engine/Plugins/Animation/ControlRig/Source/ControlRigEditor/Private/ControlRigElementDetails.cpp
+++ b/Engine/Plugins/Animation/ControlRig/Source/ControlRigEditor/Private/ControlRigElementDetails.cpp
@@ -3854,7 +3854,7 @@ void FRigControlElementDetails::CustomizeAvailableSpaces(IDetailLayoutBuilder& D
 					}
 					return FReply::Handled();
 				})
-				.ToolTipText(LOCTEXT("MoveUp", "Move up"))
+				.ToolTipText(LOCTEXT("MoveUp", "Move Up"))
 				[
 					SAssignNew(MoveSpaceUpImage, SImage)
 					.Image(FAppStyle::GetBrush("Icons.ChevronUp"))
@@ -3880,7 +3880,7 @@ void FRigControlElementDetails::CustomizeAvailableSpaces(IDetailLayoutBuilder& D
 					}
 					return FReply::Handled();
 				})
-				.ToolTipText(LOCTEXT("MoveDown", "Move down"))
+				.ToolTipText(LOCTEXT("MoveDown", "Move Down"))
 				[
 					SAssignNew(MoveSpaceDownImage, SImage)
 					.Image(FAppStyle::GetBrush("Icons.ChevronDown"))
