Control Rig: Improving the Modular Rig Tree

#rnx

diff --git a/Engine/Plugins/Animation/ControlRig/Source/ControlRig/Private/ModularRigController.cpp b/Engine/Plugins/Animation/ControlRig/Source/ControlRig/Private/ModularRigController.cpp
--- a/Engine/Plugins/Animation/ControlRig/Source/ControlRig/Private/ModularRigController.cpp
+++ b/Engine/Plugins/Animation/ControlRig/Source/ControlRig/Private/ModularRigController.cpp
@@ -1016,6 +1016,8 @@ bool UModularRigController::DeleteModule(const FString& InModulePath, bool bSetu
 	}
 #endif
 
+	(void)DeselectModule(Module->GetPath());
+
 	// Delete children
 	TArray<FString> ChildrenPaths;
 	Algo::Transform(Module->CachedChildren, ChildrenPaths, [](const FRigModuleReference* Child){ return Child->GetPath(); });
@@ -1118,6 +1120,13 @@ FString UModularRigController::RenameModule(const FString& InModulePath, const F
 	
 	const FString OldPath = (Module->ParentPath.IsEmpty()) ? OldName : URigHierarchy::JoinNameSpace(Module->ParentPath, OldName);
 	const FString NewPath = (Module->ParentPath.IsEmpty()) ? *NewName :  URigHierarchy::JoinNameSpace(Module->ParentPath, NewName);
+
+	const int32 SelectionIndex = Model->SelectedModulePaths.Find(OldPath);
+	if(SelectionIndex != INDEX_NONE)
+	{
+		Notify(EModularRigNotification::ModuleDeselected, Module);
+	}
+
 	Module->PreviousName = Module->Name;
 	Module->Name = InNewName;
 	TArray<FRigModuleReference*> Children;
@@ -1165,6 +1174,12 @@ FString UModularRigController::RenameModule(const FString& InModulePath, const F
 	UpdateShortNames();
 	Notify(EModularRigNotification::ModuleRenamed, Module);
 
+	if(SelectionIndex != INDEX_NONE)
+	{
+		Model->SelectedModulePaths[SelectionIndex] = NewPath;
+		Notify(EModularRigNotification::ModuleSelected, Module);
+	}
+
 #if WITH_EDITOR
 	TransactionPtr.Reset();
 #endif
@@ -1233,6 +1248,13 @@ FString UModularRigController::ReparentModule(const FString& InModulePath, const
 
 	// Reparent or unparent children
 	const FString OldPath = Module->GetPath();
+
+	const int32 SelectionIndex = Model->SelectedModulePaths.Find(OldPath);
+	if(SelectionIndex != INDEX_NONE)
+	{
+		Notify(EModularRigNotification::ModuleDeselected, Module);
+	}
+	
 	Module->PreviousParentPath = Module->ParentPath;
 	Module->PreviousName = Module->Name;
 	Module->ParentPath = (NewParentModule) ? NewParentModule->GetPath() : FString();
@@ -1247,7 +1269,6 @@ FString UModularRigController::ReparentModule(const FString& InModulePath, const
 		SubTree.Append(SubTree[Index]->CachedChildren);
 	}
 
-
 	Model->UpdateCachedChildren();
 	UpdateShortNames();
 
@@ -1302,7 +1323,13 @@ FString UModularRigController::ReparentModule(const FString& InModulePath, const
 	(void)DisconnectCyclicConnectors(bSetupUndo);
 
 	Notify(EModularRigNotification::ModuleReparented, Module);
-	
+
+	if(SelectionIndex != INDEX_NONE)
+	{
+		Model->SelectedModulePaths[SelectionIndex] = NewPath;
+		Notify(EModularRigNotification::ModuleSelected, Module);
+	}
+
 #if WITH_EDITOR
  	TransactionPtr.Reset();
 #endif
@@ -1631,6 +1658,43 @@ bool UModularRigController::SwapModulesOfClass(TSubclassOf<UControlRig> InOldCla
 	return true;
 }
 
+bool UModularRigController::SelectModule(const FString& InModulePath, const bool InSelected)
+{
+	const bool bCurrentlySelected = Model->SelectedModulePaths.Contains(InModulePath);
+	if(bCurrentlySelected == InSelected)
+	{
+		return false;
+	}
+
+	const FRigModuleReference* Module = FindModule(InModulePath);
+	if(Module == nullptr)
+	{
+		return false;
+	}
+
+	if(InSelected)
+	{
+		Model->SelectedModulePaths.Add(InModulePath);
+	}
+	else
+	{
+		Model->SelectedModulePaths.Remove(InModulePath);
+	}
+
+	Notify(InSelected ? EModularRigNotification::ModuleSelected : EModularRigNotification::ModuleDeselected, Module);
+	return true;
+}
+
+bool UModularRigController::DeselectModule(const FString& InModulePath)
+{
+	return SelectModule(InModulePath, false);
+}
+
+TArray<FString> UModularRigController::GetSelectedModules() const
+{
+	return Model->SelectedModulePaths;
+}
+
 void UModularRigController::RefreshModuleVariables(bool bSetupUndo)
 {
 	Model->ForEachModule([this, bSetupUndo](const FRigModuleReference* Element) -> bool
diff --git a/Engine/Plugins/Animation/ControlRig/Source/ControlRig/Public/ModularRigController.h b/Engine/Plugins/Animation/ControlRig/Source/ControlRig/Public/ModularRigController.h
--- a/Engine/Plugins/Animation/ControlRig/Source/ControlRig/Public/ModularRigController.h
+++ b/Engine/Plugins/Animation/ControlRig/Source/ControlRig/Public/ModularRigController.h
@@ -79,6 +79,15 @@ class CONTROLRIG_API UModularRigController : public UObject
 	UFUNCTION(BlueprintCallable, Category = "ControlRig | Modules")
 	bool SwapModulesOfClass(TSubclassOf<UControlRig> InOldClass, TSubclassOf<UControlRig> InNewClass, bool bSetupUndo = true);
 
+	UFUNCTION(BlueprintCallable, Category = "ControlRig | Modules")
+	bool SelectModule(const FString& InModulePath, const bool InSelected = true);
+
+	UFUNCTION(BlueprintCallable, Category = "ControlRig | Modules")
+	bool DeselectModule(const FString& InModulePath);
+
+	UFUNCTION(BlueprintPure, Category = "ControlRig | Modules")
+	TArray<FString> GetSelectedModules() const;
+
 	void RefreshModuleVariables(bool bSetupUndo = true);
 	void RefreshModuleVariables(const FRigModuleReference* InModule, bool bSetupUndo = true);
 
diff --git a/Engine/Plugins/Animation/ControlRig/Source/ControlRig/Public/ModularRigModel.h b/Engine/Plugins/Animation/ControlRig/Source/ControlRig/Public/ModularRigModel.h
--- a/Engine/Plugins/Animation/ControlRig/Source/ControlRig/Public/ModularRigModel.h
+++ b/Engine/Plugins/Animation/ControlRig/Source/ControlRig/Public/ModularRigModel.h
@@ -30,7 +30,11 @@ enum class EModularRigNotification : uint8
 	
 	InteractionBracketCanceled, // A bracket has been canceled
 
-	ModuleClassChanged, 
+	ModuleClassChanged,
+
+	ModuleSelected,
+
+	ModuleDeselected,
 
 	/** MAX - invalid */
 	Max UMETA(Hidden),
@@ -324,6 +328,7 @@ public:
 
 private:
 	TWeakObjectPtr<UObject> OuterClientHost;
+	TArray<FString> SelectedModulePaths;
 
 	friend class UModularRigController;
 	friend struct FRigModuleReference;
diff --git a/Engine/Plugins/Animation/ControlRig/Source/ControlRigEditor/Private/Editor/ControlRigEditor.cpp b/Engine/Plugins/Animation/ControlRig/Source/ControlRigEditor/Private/Editor/ControlRigEditor.cpp
--- a/Engine/Plugins/Animation/ControlRig/Source/ControlRigEditor/Private/Editor/ControlRigEditor.cpp
+++ b/Engine/Plugins/Animation/ControlRig/Source/ControlRigEditor/Private/Editor/ControlRigEditor.cpp
@@ -4218,11 +4218,23 @@ void FControlRigEditor::HandleRigTypeChanged(UControlRigBlueprint* InBlueprint)
 
 void FControlRigEditor::HandleModularRigModified(EModularRigNotification InNotification, const FRigModuleReference* InModule)
 {
+	UControlRigBlueprint* RigBlueprint = GetControlRigBlueprint();
+	if(RigBlueprint == nullptr)
+	{
+		return;
+	}
+	
+	UModularRigController* ModularRigController = RigBlueprint->GetModularRigController();
+	if(ModularRigController == nullptr)
+	{
+		return;
+	}
+
 	switch(InNotification)
 	{
 		case EModularRigNotification::ModuleAdded:
 		{
-			ModulesSelected = {InModule->GetPath()};
+			ModularRigController->SelectModule(InModule->GetPath());
 			break;
 		}
 		case EModularRigNotification::ModuleRemoved:
@@ -4247,11 +4259,6 @@ void FControlRigEditor::HandleModularRigModified(EModularRigNotification InNotif
 			{
 				OldPath = URigHierarchy::JoinNameSpace(InModule->PreviousParentPath, InModule->Name.ToString());
 			}
-			ModulesSelected.Remove(OldPath);
-			ModulesSelected.Add(InModule->GetPath());
-			RefreshDetailView();
-
-			// todo: update SchematicGraph
 			break;
 		}
 		case EModularRigNotification::ConnectionChanged:
@@ -4259,6 +4266,13 @@ void FControlRigEditor::HandleModularRigModified(EModularRigNotification InNotif
 			// todo: update SchematicGraph
 			break;
 		}
+		case EModularRigNotification::ModuleSelected:
+		case EModularRigNotification::ModuleDeselected:
+		{
+			ModulesSelected = ModularRigController->GetSelectedModules();
+			SetDetailViewForRigModules(ModulesSelected);
+			break;
+		}
 	}
 }
 
diff --git a/Engine/Plugins/Animation/ControlRig/Source/ControlRigEditor/Private/Editor/SModularRigModel.cpp b/Engine/Plugins/Animation/ControlRig/Source/ControlRigEditor/Private/Editor/SModularRigModel.cpp
--- a/Engine/Plugins/Animation/ControlRig/Source/ControlRigEditor/Private/Editor/SModularRigModel.cpp
+++ b/Engine/Plugins/Animation/ControlRig/Source/ControlRigEditor/Private/Editor/SModularRigModel.cpp
@@ -81,6 +81,11 @@ void SModularRigModel::Construct(const FArguments& InArgs, TSharedRef<FControlRi
 	ControlRigBlueprint->OnModularRigPreCompiled().AddRaw(this, &SModularRigModel::HandlePreCompileModularRigs);
 	ControlRigBlueprint->OnModularRigCompiled().AddRaw(this, &SModularRigModel::HandlePostCompileModularRigs);
 
+	if(UModularRigController* ModularRigController = ControlRigBlueprint->GetModularRigController())
+	{
+		ModularRigController->OnModified().AddSP(this, &SModularRigModel::OnModularRigModified);
+	}
+
 	// for deleting, renaming, dragging
 	CommandList = MakeShared<FUICommandList>();
 
@@ -92,6 +97,11 @@ void SModularRigModel::Construct(const FArguments& InArgs, TSharedRef<FControlRi
 
 	BindCommands();
 
+	bShowSecondaryConnectors = false;
+	bShowOptionalConnectors = false;
+	bShowUnresolvedConnectors = true;
+	bIsPerformingSelection = false;
+
 	// setup all delegates for the modular rig model widget
 	FModularRigTreeDelegates Delegates;
 	Delegates.OnGetModularRig = FOnGetModularRigTreeRig::CreateSP(this, &SModularRigModel::GetModularRigForTreeView);
@@ -106,6 +116,7 @@ void SModularRigModel::Construct(const FArguments& InArgs, TSharedRef<FControlRi
 	Delegates.OnVerifyModuleNameChanged = FOnModularRigTreeVerifyElementNameChanged::CreateSP(this, &SModularRigModel::HandleVerifyNameChanged);
 	Delegates.OnResolveConnector = FOnModularRigTreeResolveConnector::CreateSP(this, &SModularRigModel::HandleConnectorResolved);
 	Delegates.OnDisconnectConnector = FOnModularRigTreeDisconnectConnector::CreateSP(this, &SModularRigModel::HandleConnectorDisconnect);
+	Delegates.OnSelectionChanged = FOnModularRigTreeSelectionChanged::CreateSP(this, &SModularRigModel::HandleSelectionChanged);
 
 	HeaderRowWidget = SNew(SHeaderRow)
 		.Visibility(EVisibility::Visible);
@@ -144,6 +155,41 @@ void SModularRigModel::Construct(const FArguments& InArgs, TSharedRef<FControlRi
 	ChildSlot
 	[
 		SNew(SVerticalBox)
+		
+		+SVerticalBox::Slot()
+		.Padding(0.0f, 0.0f)
+		.AutoHeight()
+		[
+			SNew(SHorizontalBox)
+			
+			+SHorizontalBox::Slot()
+			.Padding(0.0f, 0.0f)
+			.HAlign(HAlign_Left)
+			.AutoWidth()
+			[
+				SNew(SComboButton)
+			   .ComboButtonStyle(&FAppStyle::Get().GetWidgetStyle<FComboButtonStyle>("SimpleComboButtonWithIcon"))
+			   .ForegroundColor(FSlateColor::UseStyle())
+			   .ToolTipText(LOCTEXT("OptionsToolTip", "Open the Options Menu ."))
+			   .OnGetMenuContent(this, &SModularRigModel::OnGetOptionsMenu)
+			   .ContentPadding(FMargin(1, 0))
+			   .ButtonContent()
+			   [
+					SNew(SImage)
+					.Image(FAppStyle::Get().GetBrush("Icons.Filter"))
+					.ColorAndOpacity(FSlateColor::UseForeground())
+			   ]
+			]
+
+			+SHorizontalBox::Slot()
+			.Padding(4.0f, 0.0f, 0.0f, 0.0f)
+			.HAlign(HAlign_Fill)
+			[
+				SAssignNew(FilterBox, SSearchBox)
+				.OnTextChanged(this, &SModularRigModel::OnFilterTextChanged)
+			]
+		]
+		
 		+SVerticalBox::Slot()
 		.Padding(0.0f, 0.0f)
 		[
@@ -159,6 +205,10 @@ void SModularRigModel::Construct(const FArguments& InArgs, TSharedRef<FControlRi
 					.HeaderRow(HeaderRowWidget)
 					.RigTreeDelegates(Delegates)
 					.AutoScrollEnabled(true)
+					.FilterText_Lambda([this]() { return FilterText; })
+					.ShowSecondaryConnectors_Lambda([this]() { return bShowSecondaryConnectors; })
+					.ShowOptionalConnectors_Lambda([this]() { return bShowOptionalConnectors; })
+					.ShowUnresolvedConnectors_Lambda([this]() { return bShowUnresolvedConnectors; })
 				]
 			]
 		]
@@ -202,6 +252,10 @@ void SModularRigModel::OnEditorClose(const FRigVMEditor* InEditor, URigVMBluepri
 		InBlueprint->OnSetObjectBeingDebugged().RemoveAll(this);
 		BP->OnModularRigPreCompiled().RemoveAll(this);
 		BP->OnModularRigCompiled().RemoveAll(this);
+		if(UModularRigController* ModularRigController = BP->GetModularRigController())
+		{
+			ModularRigController->OnModified().RemoveAll(this);
+		}
 	}
 
 	if(const UModularRig* Rig = GetModularRigForTreeView())
@@ -349,6 +403,82 @@ void SModularRigModel::HandleSetObjectBeingDebugged(UObject* InObject)
 	RefreshTreeView();
 }
 
+TSharedRef<SWidget> SModularRigModel::OnGetOptionsMenu()
+{
+	FMenuBuilder MenuBuilder(true, nullptr);
+
+	const FCanExecuteAction CanExecuteAction = FCanExecuteAction::CreateLambda([]() { return true; });
+
+	MenuBuilder.BeginSection("FilterOptions", LOCTEXT("FilterOptions", "Filter Options"));
+	{
+		MenuBuilder.AddMenuEntry(
+			LOCTEXT("SecondaryConnectors", "Secondary Connectors"),
+			LOCTEXT("SecondaryConnectorsToolTip", "Toggle the display of secondary connectors"),
+			FSlateIcon(FControlRigEditorStyle::Get().GetStyleSetName(), TEXT("ControlRig.ConnectorSecondary")),
+			FUIAction(
+				FExecuteAction::CreateLambda([this](){
+					bShowSecondaryConnectors = !bShowSecondaryConnectors;
+					RefreshTreeView(true);
+				}),
+				CanExecuteAction,
+				FIsActionChecked::CreateLambda([this]()
+				{
+					return bShowSecondaryConnectors;
+				})
+			),
+			NAME_None,
+			EUserInterfaceActionType::Check
+		);
+
+		MenuBuilder.AddMenuEntry(
+			LOCTEXT("OptionalConnectors", "Optional Connectors"),
+			LOCTEXT("OptionalConnectorsToolTip", "Toggle the display of secondary connectors"),
+			FSlateIcon(FControlRigEditorStyle::Get().GetStyleSetName(), TEXT("ControlRig.ConnectorOptional")),
+			FUIAction(
+				FExecuteAction::CreateLambda([this](){
+					bShowOptionalConnectors = !bShowOptionalConnectors;
+					RefreshTreeView(true);
+				}),
+				CanExecuteAction,
+				FIsActionChecked::CreateLambda([this]()
+				{
+					return bShowOptionalConnectors;
+				})
+			),
+			NAME_None,
+			EUserInterfaceActionType::Check
+		);
+
+		MenuBuilder.AddMenuEntry(
+			LOCTEXT("UnresolvedConnectors", "Unresolved Connectors"),
+			LOCTEXT("UnresolvedConnectorsToolTip", "Toggle the display of unresolved connectors"),
+			FSlateIcon(FControlRigEditorStyle::Get().GetStyleSetName(), TEXT("ControlRig.ConnectorWarning")),
+			FUIAction(
+				FExecuteAction::CreateLambda([this](){
+					bShowUnresolvedConnectors = !bShowUnresolvedConnectors;
+					RefreshTreeView(true);
+				}),
+				CanExecuteAction,
+				FIsActionChecked::CreateLambda([this]()
+				{
+					return bShowUnresolvedConnectors;
+				})
+			),
+			NAME_None,
+			EUserInterfaceActionType::Check
+		);
+	}
+	MenuBuilder.EndSection();
+
+	return MenuBuilder.MakeWidget();
+}
+
+void SModularRigModel::OnFilterTextChanged(const FText& SearchText)
+{
+	FilterText = SearchText;
+	RefreshTreeView(true);
+}
+
 TSharedPtr< SWidget > SModularRigModel::CreateContextMenuWidget()
 {
 	UToolMenus* ToolMenus = UToolMenus::Get();
@@ -363,13 +493,6 @@ TSharedPtr< SWidget > SModularRigModel::CreateContextMenuWidget()
 
 void SModularRigModel::OnItemClicked(TSharedPtr<FModularRigTreeElement> InItem)
 {
-	UModularRig* Rig = GetModularRig();
-	check(Rig);
-
-	if (ControlRigEditor.IsValid() && InItem.IsValid())
-	{
-		ControlRigEditor.Pin()->SetDetailViewForRigModules(TreeView->GetSelectedKeys());
-	}
 }
 
 void SModularRigModel::OnItemDoubleClicked(TSharedPtr<FModularRigTreeElement> InItem)
@@ -1049,8 +1172,88 @@ void SModularRigModel::HandleConnectorDisconnect(const FRigElementKey& InConnect
 	}
 }
 
-void SModularRigModel::OnSelectionChanged(TSharedPtr<FModularRigTreeElement> Selection, ESelectInfo::Type SelectInfo)
+void SModularRigModel::HandleSelectionChanged(TSharedPtr<FModularRigTreeElement> Selection, ESelectInfo::Type SelectInfo)
 {
+	if(bIsPerformingSelection)
+	{
+		return;
+	}
+	
+	if (ControlRigBlueprint.IsValid())
+	{
+		UModularRigController* Controller = ControlRigBlueprint->GetModularRigController();
+		check(Controller);
+
+		const TArray<FString> NewSelection = TreeView->GetSelectedKeys();
+		const TArray<FString> OldSelection = Controller->GetSelectedModules();
+
+		const TGuardValue<bool> GuardSelection(bIsPerformingSelection, true);
+
+		for(const FString& PreviouslySelectedModule : OldSelection)
+		{
+			if(!NewSelection.Contains(PreviouslySelectedModule))
+			{
+				Controller->DeselectModule(PreviouslySelectedModule);
+			}
+		}
+		for(const FString& NewModuleToSelect : NewSelection)
+		{
+			if(!OldSelection.Contains(NewModuleToSelect))
+			{
+				Controller->SelectModule(NewModuleToSelect);
+			}
+		}
+	}
+}
+
+void SModularRigModel::OnModularRigModified(EModularRigNotification InNotif, const FRigModuleReference* InModule)
+{
+	if (!ControlRigBlueprint.IsValid())
+	{
+		return;
+	}
+
+	switch(InNotif)
+	{
+		case EModularRigNotification::ModuleSelected:
+		case EModularRigNotification::ModuleDeselected:
+		{
+			if(!bIsPerformingSelection)
+			{
+				const TGuardValue<bool> GuardSelection(bIsPerformingSelection, true);
+				if(const UModularRigController* ModularRigController = ControlRigBlueprint->GetModularRigController())
+				{
+					const TArray<FString> SelectedModulePaths = ModularRigController->GetSelectedModules();
+					TArray<TSharedPtr<FModularRigTreeElement>> NewSelection;
+					for(const FString& SelectedModulePath : SelectedModulePaths)
+					{
+						if(const TSharedPtr<FModularRigTreeElement> Module = TreeView->FindElement(SelectedModulePath))
+						{
+							NewSelection.Add(Module);
+						}
+					}
+					TreeView->SetSelection(NewSelection);
+				}
+			}
+			break;
+		}
+		case EModularRigNotification::ModuleAdded:
+		case EModularRigNotification::ModuleRenamed:
+		case EModularRigNotification::ModuleRemoved:
+		case EModularRigNotification::ModuleReparented:
+		case EModularRigNotification::ConnectionChanged:
+		case EModularRigNotification::ModuleConfigValueChanged:
+		case EModularRigNotification::ModuleShortNameChanged:
+		case EModularRigNotification::ModuleClassChanged:
+		{
+			TreeView->RefreshTreeView();
+			break;
+		}
+		default:
+		{
+			break;
+		}
+	}
 }
 
 void SModularRigModel::OnHierarchyModified(ERigHierarchyNotification InNotif, URigHierarchy* InHierarchy, const FRigBaseElement* InElement)
diff --git a/Engine/Plugins/Animation/ControlRig/Source/ControlRigEditor/Private/Editor/SModularRigModel.h b/Engine/Plugins/Animation/ControlRig/Source/ControlRigEditor/Private/Editor/SModularRigModel.h
--- a/Engine/Plugins/Animation/ControlRig/Source/ControlRigEditor/Private/Editor/SModularRigModel.h
+++ b/Engine/Plugins/Animation/ControlRig/Source/ControlRigEditor/Private/Editor/SModularRigModel.h
@@ -97,7 +97,7 @@ private:
 	void HandleConnectorDisconnect(const FRigElementKey& InConnector);
 
 	/** Set Selection Changed */
-	void OnSelectionChanged(TSharedPtr<FModularRigTreeElement> Selection, ESelectInfo::Type SelectInfo);
+	void HandleSelectionChanged(TSharedPtr<FModularRigTreeElement> Selection, ESelectInfo::Type SelectInfo);
 
 	TSharedPtr< SWidget > CreateContextMenuWidget();
 	void OnItemClicked(TSharedPtr<FModularRigTreeElement> InItem);
@@ -143,11 +143,23 @@ private:
 	void OnRequestDetailsInspection(const FString& InKey);
 	void HandlePreCompileModularRigs(URigVMBlueprint* InBlueprint);
 	void HandlePostCompileModularRigs(URigVMBlueprint* InBlueprint);
+	void OnModularRigModified(EModularRigNotification InNotif, const FRigModuleReference* InModule);
 	void OnHierarchyModified(ERigHierarchyNotification InNotif, URigHierarchy* InHierarchy, const FRigBaseElement* InElement);
 
 	void HandleRefreshEditorFromBlueprint(URigVMBlueprint* InBlueprint);
 	void HandleSetObjectBeingDebugged(UObject* InObject);
 
+	TSharedRef<SWidget> OnGetOptionsMenu();
+	void OnFilterTextChanged(const FText& SearchText);
+
+	bool bShowSecondaryConnectors;
+	bool bShowOptionalConnectors;
+	bool bShowUnresolvedConnectors;
+	FText FilterText;
+
+	TSharedPtr<SSearchBox> FilterBox;
+	bool bIsPerformingSelection;
+
 public:
 
 	friend class FModularRigTreeElement;
diff --git a/Engine/Plugins/Animation/ControlRig/Source/ControlRigEditor/Private/Editor/SModularRigTreeView.cpp b/Engine/Plugins/Animation/ControlRig/Source/ControlRigEditor/Private/Editor/SModularRigTreeView.cpp
--- a/Engine/Plugins/Animation/ControlRig/Source/ControlRigEditor/Private/Editor/SModularRigTreeView.cpp
+++ b/Engine/Plugins/Animation/ControlRig/Source/ControlRigEditor/Private/Editor/SModularRigTreeView.cpp
@@ -9,6 +9,7 @@
 #include "Widgets/Text/STextBlock.h"
 #include "PropertyCustomizationHelpers.h"
 #include "Framework/Application/SlateApplication.h"
+#include "Framework/Application/SlateUser.h"
 #include "Editor/EditorEngine.h"
 #include "HelperUtil.h"
 #include "Widgets/Text/SInlineEditableTextBlock.h"
@@ -143,6 +144,16 @@ void SModularRigModelItem::PopulateConnectorTargetList(const FRigElementKey InCo
 			if (Cast<FRigConnectorElement>(Hierarchy->Find(InConnectorKey)))
 			{
 				ConnectorComboBox->GetTreeView()->RefreshTreeView(true);
+
+				// set the focus to the search box so you can start typing right away
+				RegisterActiveTimer(0.f, FWidgetActiveTimerDelegate::CreateLambda([this](double,float)
+				{
+					FSlateApplication::Get().ForEachUser([this](FSlateUser& User)
+					{
+						User.SetFocus(ConnectorComboBox->GetSearchBox());
+					});
+					return EActiveTimerReturnType::Stop;
+				}));
 			}
 		}
 	}
@@ -659,12 +670,18 @@ void SModularRigTreeView::Construct(const FArguments& InArgs)
 	Delegates = InArgs._RigTreeDelegates;
 	bAutoScrollEnabled = InArgs._AutoScrollEnabled;
 
+	FilterText = InArgs._FilterText;
+	ShowSecondaryConnectors = InArgs._ShowSecondaryConnectors;
+	ShowOptionalConnectors = InArgs._ShowOptionalConnectors;
+	ShowUnresolvedConnectors = InArgs._ShowUnresolvedConnectors;
+
 	STreeView<TSharedPtr<FModularRigTreeElement>>::FArguments SuperArgs;
 	SuperArgs.HeaderRow(InArgs._HeaderRow);
 	SuperArgs.TreeItemsSource(&RootElements);
 	SuperArgs.SelectionMode(ESelectionMode::Multi);
 	SuperArgs.OnGenerateRow(this, &SModularRigTreeView::MakeTableRowWidget, false);
 	SuperArgs.OnGetChildren(this, &SModularRigTreeView::HandleGetChildrenForTree);
+	SuperArgs.OnSelectionChanged(FOnModularRigTreeSelectionChanged::CreateRaw(&Delegates, &FModularRigTreeDelegates::HandleSelectionChanged));
 	SuperArgs.OnContextMenuOpening(Delegates.OnContextMenuOpening);
 	SuperArgs.HighlightParentNodesForSelection(true);
 	SuperArgs.AllowInvisibleItemSelection(true);  //without this we deselect everything when we filter or we collapse
@@ -785,7 +802,7 @@ TSharedPtr<FModularRigTreeElement> SModularRigTreeView::FindElement(const FStrin
 	return TSharedPtr<FModularRigTreeElement>();
 }
 
-bool SModularRigTreeView::AddElement(FString InKey, FString InParentKey)
+bool SModularRigTreeView::AddElement(FString InKey, FString InParentKey, bool bApplyFilterText)
 {
 	if(ElementMap.Contains(InKey))
 	{
@@ -794,32 +811,33 @@ bool SModularRigTreeView::AddElement(FString InKey, FString InParentKey)
 
 	if (!InKey.IsEmpty())
 	{
-		TSharedPtr<FModularRigTreeElement> NewItem = MakeShared<FModularRigTreeElement>(InKey, SharedThis(this), true);
+		const FString ModulePath = InKey;
 
-		ElementMap.Add(InKey, NewItem);
-		if (!InParentKey.IsEmpty())
+		bool bFilteredOutElement = false;
+		const FString FilterTextString = FilterText.Get().ToString();
+		if(!FilterTextString.IsEmpty())
 		{
-			ParentMap.Add(InKey, InParentKey);
-
-			TSharedPtr<FModularRigTreeElement>* FoundItem = ElementMap.Find(InParentKey);
-			check(FoundItem);
-			FoundItem->Get()->Children.Add(NewItem);
-		}
-		else
-		{
-			RootElements.Add(NewItem);
+			FString StringToSearch = InKey;
+			(void)URigHierarchy::SplitNameSpace(StringToSearch, nullptr, &StringToSearch);
+			
+			if(!StringToSearch.Contains(FilterTextString, ESearchCase::IgnoreCase))
+			{
+				bFilteredOutElement = true;
+			}
 		}
 
-		SetItemExpansion(NewItem, true);
-
+		TArray<FString> FilteredConnectors;
 		if (const UModularRig* ModularRig = Delegates.GetModularRig())
 		{
-			if (const FRigModuleInstance* Module = ModularRig->FindModule(NewItem->ModulePath))
+			const FModularRigModel& Model = ModularRig->GetModularRigModel();
+			
+			if (const FRigModuleInstance* Module = ModularRig->FindModule(ModulePath))
 			{
 				if (const UControlRig* ModuleRig = Module->GetRig())
 				{
 					const UControlRig* CDO = ModuleRig->GetClass()->GetDefaultObject<UControlRig>();
 					const TArray<FRigModuleConnector>& Connectors = CDO->GetRigModuleSettings().ExposedConnectors;
+
 					for (const FRigModuleConnector& Connector : Connectors)
 					{
 						if (Connector.IsPrimary())
@@ -827,21 +845,86 @@ bool SModularRigTreeView::AddElement(FString InKey, FString InParentKey)
 							continue;
 						}
 
-						const FString Key = FString::Printf(TEXT("%s:%s"), *NewItem->ModulePath, *Connector.Name);
-						TSharedPtr<FModularRigTreeElement> ConnectorItem = MakeShared<FModularRigTreeElement>(Key, SharedThis(this), false);
-						NewItem.Get()->Children.Add(ConnectorItem);
-						ElementMap.Add(Key, ConnectorItem);
-						ParentMap.Add(Key, InKey);
+						const FString Key = FString::Printf(TEXT("%s:%s"), *ModulePath, *Connector.Name);
+						bool bShouldFilterByConnectorType = true;
+
+						if(!FilterTextString.IsEmpty())
+						{
+							const bool bMatchesFilter = Connector.Name.Contains(FilterTextString, ESearchCase::IgnoreCase);
+							if(bFilteredOutElement)
+							{
+								if(!bMatchesFilter)
+								{
+									continue;
+								}
+							}
+							bShouldFilterByConnectorType = !bMatchesFilter;
+						}
+
+						if(bShouldFilterByConnectorType)
+						{
+							const bool bIsConnected = Model.Connections.HasConnection(FRigElementKey(*Key, ERigElementType::Connector));
+							if(bIsConnected || !ShowUnresolvedConnectors.Get())
+							{
+								if(Connector.IsOptional())
+								{
+									if(ShowOptionalConnectors.Get() == false)
+									{
+										continue;
+									}
+								}
+								else if(Connector.IsSecondary())
+								{
+									if(ShowSecondaryConnectors.Get() == false)
+									{
+										continue;
+									}
+								}
+							}
+						}
+
+						FilteredConnectors.Add(Key);
 					}
 				}
 			}
 		}
+		
+		if(bFilteredOutElement && bApplyFilterText && FilteredConnectors.IsEmpty())
+		{
+			return false;
+		}
+		
+		TSharedPtr<FModularRigTreeElement> NewItem = MakeShared<FModularRigTreeElement>(InKey, SharedThis(this), true);
+
+		ElementMap.Add(InKey, NewItem);
+		if (!InParentKey.IsEmpty())
+		{
+			ParentMap.Add(InKey, InParentKey);
+
+			TSharedPtr<FModularRigTreeElement>* FoundItem = ElementMap.Find(InParentKey);
+			check(FoundItem);
+			FoundItem->Get()->Children.Add(NewItem);
+		}
+		else
+		{
+			RootElements.Add(NewItem);
+		}
+
+		SetItemExpansion(NewItem, true);
+
+		for (const FString& Key : FilteredConnectors)
+		{
+			TSharedPtr<FModularRigTreeElement> ConnectorItem = MakeShared<FModularRigTreeElement>(Key, SharedThis(this), false);
+			NewItem.Get()->Children.Add(ConnectorItem);
+			ElementMap.Add(Key, ConnectorItem);
+			ParentMap.Add(Key, InKey);
+		}
 	}
 
 	return true;
 }
 
-bool SModularRigTreeView::AddElement(const FRigModuleInstance* InElement)
+bool SModularRigTreeView::AddElement(const FRigModuleInstance* InElement, bool bApplyFilterText)
 {
 	check(InElement);
 	
@@ -853,7 +936,7 @@ bool SModularRigTreeView::AddElement(const FRigModuleInstance* InElement)
 	const UModularRig* ModularRig = Delegates.GetModularRig();
 
 	FString ElementPath = InElement->GetPath();
-	if(!AddElement(ElementPath, FString()))
+	if(!AddElement(ElementPath, FString(), bApplyFilterText))
 	{
 		return false;
 	}
@@ -867,7 +950,7 @@ bool SModularRigTreeView::AddElement(const FRigModuleInstance* InElement)
 			{
 				if(const FRigModuleInstance* ParentElement = ModularRig->FindModule(ParentPath))
 				{
-					AddElement(ParentElement);
+					AddElement(ParentElement, false);
 
 					if(ElementMap.Contains(ParentPath))
 					{
@@ -967,7 +1050,7 @@ void SModularRigTreeView::RefreshTreeView(bool bRebuildContent)
 		{
 			ModularRig->ForEachModule([&](const FRigModuleInstance* Element)
 			{
-				AddElement(Element);
+				AddElement(Element, true);
 				return true;
 			});
 
diff --git a/Engine/Plugins/Animation/ControlRig/Source/ControlRigEditor/Private/Editor/SModularRigTreeView.h b/Engine/Plugins/Animation/ControlRig/Source/ControlRigEditor/Private/Editor/SModularRigTreeView.h
--- a/Engine/Plugins/Animation/ControlRig/Source/ControlRigEditor/Private/Editor/SModularRigTreeView.h
+++ b/Engine/Plugins/Animation/ControlRig/Source/ControlRigEditor/Private/Editor/SModularRigTreeView.h
@@ -20,14 +20,16 @@ DECLARE_DELEGATE_TwoParams(FOnModularRigTreeResolveConnector, const FRigElementK
 DECLARE_DELEGATE_OneParam(FOnModularRigTreeDisconnectConnector, const FRigElementKey& /*Connector*/);
 DECLARE_DELEGATE_RetVal_ThreeParams(bool, FOnModularRigTreeVerifyElementNameChanged, const FString& /*OldPath*/, const FName& /*NewName*/, FText& /*OutErrorMessage*/);
 
-
 typedef STreeView<TSharedPtr<FModularRigTreeElement>>::FOnMouseButtonClick FOnModularRigTreeMouseButtonClick;
 typedef STreeView<TSharedPtr<FModularRigTreeElement>>::FOnMouseButtonDoubleClick FOnModularRigTreeMouseButtonDoubleClick;
 typedef STableRow<TSharedPtr<FModularRigTreeElement>>::FOnCanAcceptDrop FOnModularRigTreeCanAcceptDrop;
 typedef STableRow<TSharedPtr<FModularRigTreeElement>>::FOnAcceptDrop FOnModularRigTreeAcceptDrop;
+typedef STreeView<TSharedPtr<FModularRigTreeElement>>::FOnSelectionChanged FOnModularRigTreeSelectionChanged;
 
 struct CONTROLRIGEDITOR_API FModularRigTreeDelegates
 {
+public:
+	
 	FOnGetModularRigTreeRig OnGetModularRig;
 	FOnModularRigTreeMouseButtonClick OnMouseButtonClick;
 	FOnModularRigTreeMouseButtonDoubleClick OnMouseButtonDoubleClick;
@@ -40,6 +42,7 @@ struct CONTROLRIGEDITOR_API FModularRigTreeDelegates
 	FOnModularRigTreeVerifyElementNameChanged OnVerifyModuleNameChanged;
 	FOnModularRigTreeResolveConnector OnResolveConnector;
 	FOnModularRigTreeDisconnectConnector OnDisconnectConnector;
+	FOnModularRigTreeSelectionChanged OnSelectionChanged;
 	
 	FModularRigTreeDelegates()
 	{
@@ -91,6 +94,20 @@ struct CONTROLRIGEDITOR_API FModularRigTreeDelegates
 		}
 		return false;
 	}
+
+	void HandleSelectionChanged(TSharedPtr<FModularRigTreeElement> Selection, ESelectInfo::Type SelectInfo)
+	{
+		if(bSelectionIsChanged)
+		{
+			return;
+		}
+		TGuardValue<bool> Guard(bSelectionIsChanged, true);
+		(void)OnSelectionChanged.ExecuteIfBound(Selection, SelectInfo);
+	}
+
+private:
+
+	bool bSelectionIsChanged = false;
 };
 
 
@@ -180,10 +197,18 @@ public:
 
 	SLATE_BEGIN_ARGS(SModularRigTreeView)
 		: _AutoScrollEnabled(false)
+		, _FilterText()
+		, _ShowSecondaryConnectors(false)
+		, _ShowOptionalConnectors(false)
+		, _ShowUnresolvedConnectors(true)
 	{}
 		SLATE_ARGUMENT( TSharedPtr<SHeaderRow>, HeaderRow )
 		SLATE_ARGUMENT(FModularRigTreeDelegates, RigTreeDelegates)
 		SLATE_ARGUMENT(bool, AutoScrollEnabled)
+		SLATE_ATTRIBUTE(FText, FilterText)
+		SLATE_ATTRIBUTE(bool, ShowSecondaryConnectors)
+		SLATE_ATTRIBUTE(bool, ShowOptionalConnectors)
+		SLATE_ATTRIBUTE(bool, ShowUnresolvedConnectors)
 	SLATE_END_ARGS()
 
 	void Construct(const FArguments& InArgs);
@@ -222,8 +247,8 @@ public:
 
 	TSharedPtr<FModularRigTreeElement> FindElement(const FString& InElementKey);
 	static TSharedPtr<FModularRigTreeElement> FindElement(const FString& InElementKey, TSharedPtr<FModularRigTreeElement> CurrentItem);
-	bool AddElement(FString InKey, FString InParentKey = FString());
-	bool AddElement(const FRigModuleInstance* InElement);
+	bool AddElement(FString InKey, FString InParentKey = FString(), bool bApplyFilterText = true);
+	bool AddElement(const FRigModuleInstance* InElement, bool bApplyFilterText);
 	void AddSpacerElement();
 	bool ReparentElement(const FString InKey, const FString InParentKey);
 	void RefreshTreeView(bool bRebuildContent = true);
@@ -258,6 +283,11 @@ private:
 	FVector2D LastMousePosition;
 	double TimeAtMousePosition;
 
+	TAttribute<FText> FilterText;
+	TAttribute<bool> ShowSecondaryConnectors;
+	TAttribute<bool> ShowOptionalConnectors;
+	TAttribute<bool> ShowUnresolvedConnectors;
+
 	friend class SModularRigModel;
 };
 
diff --git a/Engine/Plugins/Animation/ControlRig/Source/ControlRigEditor/Private/Editor/SRigHierarchy.cpp b/Engine/Plugins/Animation/ControlRig/Source/ControlRigEditor/Private/Editor/SRigHierarchy.cpp
--- a/Engine/Plugins/Animation/ControlRig/Source/ControlRigEditor/Private/Editor/SRigHierarchy.cpp
+++ b/Engine/Plugins/Animation/ControlRig/Source/ControlRigEditor/Private/Editor/SRigHierarchy.cpp
@@ -160,32 +160,16 @@ void SRigHierarchy::Construct(const FArguments& InArgs, TSharedRef<FControlRigEd
 					[
 						SNew(SComboButton)
 						.Visibility(EVisibility::Visible)
-						.ComboButtonStyle(FAppStyle::Get(), "GenericFilters.ComboButtonStyle")
-						.ForegroundColor(FLinearColor::White)
+						.ComboButtonStyle(&FAppStyle::Get().GetWidgetStyle<FComboButtonStyle>("SimpleComboButtonWithIcon"))
+						.ForegroundColor(FSlateColor::UseStyle())
 						.ContentPadding(0.0f)
 						.OnGetMenuContent(this, &SRigHierarchy::CreateFilterMenu)
 						.ButtonContent()
 						[
-							SNew(SHorizontalBox)
-							+ SHorizontalBox::Slot()
-							.AutoWidth()
-							.VAlign(VAlign_Center)
-							[
-								SNew(STextBlock)
-								.TextStyle(FAppStyle::Get(), "GenericFilters.TextStyle")
-								.Font(FAppStyle::Get().GetFontStyle("FontAwesome.9"))
-								.Text(FText::FromString(FString(TEXT("\xf0b0"))) /*fa-filter*/)
-							]
-							+ SHorizontalBox::Slot()
-							.AutoWidth()
-							.Padding(2, 0, 0, 0)
-							.VAlign(VAlign_Center)
-							[
-								SNew(STextBlock)
-								.TextStyle(FAppStyle::Get(), "GenericFilters.TextStyle")
-								.Text(LOCTEXT("FilterMenuLabel", "Options"))
-							]
-						]
+							SNew(SImage)
+							.Image(FAppStyle::Get().GetBrush("Icons.Filter"))
+							.ColorAndOpacity(FSlateColor::UseForeground())
+						 ]
 					]
 					+SHorizontalBox::Slot()
 					.VAlign(VAlign_Center)
