Control Rig: Sync selection between modules + elements

#rnx

diff --git a/Engine/Plugins/Animation/ControlRig/Source/ControlRigDeveloper/Private/ControlRigBlueprint.cpp b/Engine/Plugins/Animation/ControlRig/Source/ControlRigDeveloper/Private/ControlRigBlueprint.cpp
--- a/Engine/Plugins/Animation/ControlRig/Source/ControlRigDeveloper/Private/ControlRigBlueprint.cpp
+++ b/Engine/Plugins/Animation/ControlRig/Source/ControlRigDeveloper/Private/ControlRigBlueprint.cpp
@@ -2844,6 +2844,12 @@ void UControlRigBlueprint::HandleRigModulesModified(EModularRigNotification InNo
 			ModulesRecompilationBracket--;
 			break;
 		}
+		case EModularRigNotification::ModuleSelected:
+		case EModularRigNotification::ModuleDeselected:
+		{
+			ModulesRecompilationBracket--;
+			break;
+		}
 		default:
 		{
 			break;
diff --git a/Engine/Plugins/Animation/ControlRig/Source/ControlRigEditor/Private/Editor/SModularRigModel.cpp b/Engine/Plugins/Animation/ControlRig/Source/ControlRigEditor/Private/Editor/SModularRigModel.cpp
--- a/Engine/Plugins/Animation/ControlRig/Source/ControlRigEditor/Private/Editor/SModularRigModel.cpp
+++ b/Engine/Plugins/Animation/ControlRig/Source/ControlRigEditor/Private/Editor/SModularRigModel.cpp
@@ -1178,6 +1178,8 @@ void SModularRigModel::HandleSelectionChanged(TSharedPtr<FModularRigTreeElement>
 	{
 		return;
 	}
+
+	TreeView->ClearHighlightedItems();
 	
 	if (ControlRigBlueprint.IsValid())
 	{
@@ -1268,38 +1270,24 @@ void SModularRigModel::OnHierarchyModified(ERigHierarchyNotification InNotif, UR
 		case ERigHierarchyNotification::ElementSelected:
 		case ERigHierarchyNotification::ElementDeselected:
 		{
-			const FRigConnectorElement* Connector = Cast<FRigConnectorElement>(InElement);
-			if(Connector == nullptr)
-			{
-				for(const FModularRigSingleConnection& Connection : ControlRigBlueprint->ModularRigModel.Connections)
-				{
-					check(Connection.Connector.Type == ERigElementType::Connector);
-					if(Connection.Target == InElement->GetKey())
-					{
-						if(const FRigConnectorElement* TargetConnector = InHierarchy->Find<FRigConnectorElement>(Connection.Connector))
-						{
-							OnHierarchyModified(InNotif, InHierarchy, TargetConnector);
-						}
-					}
-				}
-				return;
-			}
+			FString ModulePathOrConnectorName = InHierarchy->GetModulePath(InElement->GetKey());
 
-			FString ModulePathOrConnectorName;
-			if(Connector->IsPrimary())
+			if(const FRigConnectorElement* Connector = Cast<FRigConnectorElement>(InElement))
 			{
-				ModulePathOrConnectorName = InHierarchy->GetModulePath(Connector->GetKey());
-			}
-			else
-			{
-				ModulePathOrConnectorName = Connector->GetName();
+				if(Connector->IsPrimary())
+            	{
+					ModulePathOrConnectorName = Connector->GetName();
+            	}
 			}
 
-			TSharedPtr<FModularRigTreeElement> Item = TreeView->FindElement(ModulePathOrConnectorName);
-			if(Item.IsValid())
+			if(!ModulePathOrConnectorName.IsEmpty())
 			{
-				const bool bSelected = InNotif == ERigHierarchyNotification::ElementSelected;
-				TreeView->SetItemSelection(Item, bSelected);
+				if(TSharedPtr<FModularRigTreeElement> Item = TreeView->FindElement(ModulePathOrConnectorName))
+				{
+					const bool bSelected = InNotif == ERigHierarchyNotification::ElementSelected;
+					TreeView->SetItemHighlighted(Item, bSelected);
+					TreeView->RequestScrollIntoView(Item);
+				}
 			}
 		}
 		default:
diff --git a/Engine/Plugins/Animation/ControlRig/Source/ControlRigEditor/Private/Editor/SRigHierarchy.cpp b/Engine/Plugins/Animation/ControlRig/Source/ControlRigEditor/Private/Editor/SRigHierarchy.cpp
--- a/Engine/Plugins/Animation/ControlRig/Source/ControlRigEditor/Private/Editor/SRigHierarchy.cpp
+++ b/Engine/Plugins/Animation/ControlRig/Source/ControlRigEditor/Private/Editor/SRigHierarchy.cpp
@@ -80,6 +80,11 @@ void SRigHierarchy::Construct(const FArguments& InArgs, TSharedRef<FControlRigEd
 	ControlRigBlueprint->OnRefreshEditor().AddRaw(this, &SRigHierarchy::HandleRefreshEditorFromBlueprint);
 	ControlRigBlueprint->OnSetObjectBeingDebugged().AddRaw(this, &SRigHierarchy::HandleSetObjectBeingDebugged);
 
+	if(UModularRigController* ModularRigController = ControlRigBlueprint->GetModularRigController())
+	{
+		ModularRigController->OnModified().AddSP(this, &SRigHierarchy::OnModularRigModified);
+	}
+
 	// for deleting, renaming, dragging
 	CommandList = MakeShared<FUICommandList>();
 
@@ -260,6 +265,11 @@ void SRigHierarchy::OnEditorClose(const FRigVMEditor* InEditor, URigVMBlueprint*
 		BP->Hierarchy->OnModified().RemoveAll(this);
 		InBlueprint->OnRefreshEditor().RemoveAll(this);
 		InBlueprint->OnSetObjectBeingDebugged().RemoveAll(this);
+
+		if(UModularRigController* ModularRigController = BP->GetModularRigController())
+		{
+			ModularRigController->OnModified().RemoveAll(this);
+		}
 	}
 	
 	ControlRigEditor.Reset();
@@ -590,6 +600,8 @@ void SRigHierarchy::OnSelectionChanged(TSharedPtr<FRigTreeElement> Selection, ES
 		return;
 	}
 
+	TreeView->ClearHighlightedItems();
+
 	// an element to use for the control rig editor's detail panel
 	FRigElementKey LastSelectedElement;
 
@@ -855,6 +867,57 @@ void SRigHierarchy::OnHierarchyModified_AnyThread(ERigHierarchyNotification InNo
 	}
 }
 
+void SRigHierarchy::OnModularRigModified(EModularRigNotification InNotif, const FRigModuleReference* InModule)
+{
+	if(!ControlRigBlueprint.IsValid())
+	{
+		return;
+	}
+
+	switch(InNotif)
+	{
+		case EModularRigNotification::ModuleSelected:
+		case EModularRigNotification::ModuleDeselected:
+		{
+			const bool bSelected = InNotif == EModularRigNotification::ModuleSelected;
+			if(ControlRigEditor.IsValid())
+			{
+				if(UControlRig* ControlRig = ControlRigEditor.Pin()->GetControlRig())
+				{
+					if(URigHierarchy* Hierarchy = ControlRig->GetHierarchy())
+					{
+						const FString ModulePath = InModule->GetPath();
+						TArray<FRigElementKey> Keys = Hierarchy->GetAllKeys();
+						Keys = Keys.FilterByPredicate([Hierarchy, ModulePath](const FRigElementKey& InKey)
+						{
+							return ModulePath == Hierarchy->GetModulePath(InKey);
+						});
+
+						bool bScrollIntoView = true;
+						for(const FRigElementKey& Key : Keys)
+						{
+							if(const TSharedPtr<FRigTreeElement> TreeElement = TreeView->FindElement(Key))
+							{
+								TreeView->SetItemHighlighted(TreeElement, bSelected);
+								if(bScrollIntoView)
+								{
+									TreeView->RequestScrollIntoView(TreeElement);
+									bScrollIntoView = false;
+								}
+							}
+						}
+					}
+				}
+			}
+			break;
+		}
+		default:
+		{
+			break;
+		}
+	}
+}
+
 void SRigHierarchy::HandleRefreshEditorFromBlueprint(URigVMBlueprint* InBlueprint)
 {
 	if (bIsChangingRigHierarchy)
diff --git a/Engine/Plugins/Animation/ControlRig/Source/ControlRigEditor/Private/Editor/SRigHierarchy.h b/Engine/Plugins/Animation/ControlRig/Source/ControlRigEditor/Private/Editor/SRigHierarchy.h
--- a/Engine/Plugins/Animation/ControlRig/Source/ControlRigEditor/Private/Editor/SRigHierarchy.h
+++ b/Engine/Plugins/Animation/ControlRig/Source/ControlRigEditor/Private/Editor/SRigHierarchy.h
@@ -214,6 +214,7 @@ private:
 	bool bIsChangingRigHierarchy;
 	void OnHierarchyModified(ERigHierarchyNotification InNotif, URigHierarchy* InHierarchy, const FRigBaseElement* InElement);
 	void OnHierarchyModified_AnyThread(ERigHierarchyNotification InNotif, URigHierarchy* InHierarchy, const FRigBaseElement* InElement);
+	void OnModularRigModified(EModularRigNotification InNotif, const FRigModuleReference* InModule);
 	void HandleRefreshEditorFromBlueprint(URigVMBlueprint* InBlueprint);
 	void HandleSetObjectBeingDebugged(UObject* InObject);
 	void OnPreConstruction_AnyThread(UControlRig* InRig, const FName& InEventName);
diff --git a/Engine/Plugins/Animation/ControlRig/Source/ControlRigEditor/Private/Editor/SRigHierarchyTreeView.cpp b/Engine/Plugins/Animation/ControlRig/Source/ControlRigEditor/Private/Editor/SRigHierarchyTreeView.cpp
--- a/Engine/Plugins/Animation/ControlRig/Source/ControlRigEditor/Private/Editor/SRigHierarchyTreeView.cpp
+++ b/Engine/Plugins/Animation/ControlRig/Source/ControlRigEditor/Private/Editor/SRigHierarchyTreeView.cpp
@@ -380,6 +380,19 @@ void SRigHierarchyTreeView::Tick(const FGeometry& AllottedGeometry, const double
 	}
 }
 
+TSharedPtr<FRigTreeElement> SRigHierarchyTreeView::FindElement(const FRigElementKey& InElementKey) const
+{
+	for (int32 RootIndex = 0; RootIndex < RootElements.Num(); ++RootIndex)
+	{
+		TSharedPtr<FRigTreeElement> Found = FindElement(InElementKey, RootElements[RootIndex]);
+		if (Found.IsValid())
+		{
+			return Found;
+		}
+	}
+	return TSharedPtr<FRigTreeElement>();
+}
+
 TSharedPtr<FRigTreeElement> SRigHierarchyTreeView::FindElement(const FRigElementKey& InElementKey, TSharedPtr<FRigTreeElement> CurrentItem)
 {
 	if (CurrentItem->Key == InElementKey)
diff --git a/Engine/Plugins/Animation/ControlRig/Source/ControlRigEditor/Private/Editor/SRigHierarchyTreeView.h b/Engine/Plugins/Animation/ControlRig/Source/ControlRigEditor/Private/Editor/SRigHierarchyTreeView.h
--- a/Engine/Plugins/Animation/ControlRig/Source/ControlRigEditor/Private/Editor/SRigHierarchyTreeView.h
+++ b/Engine/Plugins/Animation/ControlRig/Source/ControlRigEditor/Private/Editor/SRigHierarchyTreeView.h
@@ -337,6 +337,7 @@ public:
 		}
 	}
 
+	TSharedPtr<FRigTreeElement> FindElement(const FRigElementKey& InElementKey) const;
 	static TSharedPtr<FRigTreeElement> FindElement(const FRigElementKey& InElementKey, TSharedPtr<FRigTreeElement> CurrentItem);
 	bool AddElement(FRigElementKey InKey, FRigElementKey InParentKey = FRigElementKey());
 	bool AddElement(const FRigBaseElement* InElement);
