SchematicView: Disable auto grouping for Control Rig

#rnx

diff --git a/Engine/Plugins/Animation/ControlRig/Source/ControlRigEditor/Private/Editor/ControlRigSchematicModel.h b/Engine/Plugins/Animation/ControlRig/Source/ControlRigEditor/Private/Editor/ControlRigSchematicModel.h
--- a/Engine/Plugins/Animation/ControlRig/Source/ControlRigEditor/Private/Editor/ControlRigSchematicModel.h
+++ b/Engine/Plugins/Animation/ControlRig/Source/ControlRigEditor/Private/Editor/ControlRigSchematicModel.h
@@ -95,6 +95,7 @@ public:
 	void OnSetObjectBeingDebugged(UObject* InObject);
 	void HandleModularRigModified(EModularRigNotification InNotification, const FRigModuleReference* InModule);
 
+	virtual bool IsAutoGroupingEnabled() const { return false; }
 	virtual FSchematicGraphGroupNode* AddAutoGroupNode() override;
 	virtual FVector2d GetPositionForNode(const FSchematicGraphNode* InNode) const override;
 	virtual bool GetPositionAnimationEnabledForNode(const FSchematicGraphNode* InNode) const override;
