Control Rig: Support middle click for tree selection

#rnx

diff --git a/Engine/Plugins/Animation/ControlRig/Source/ControlRig/Private/ModularRigController.cpp b/Engine/Plugins/Animation/ControlRig/Source/ControlRig/Private/ModularRigController.cpp
--- a/Engine/Plugins/Animation/ControlRig/Source/ControlRig/Private/ModularRigController.cpp
+++ b/Engine/Plugins/Animation/ControlRig/Source/ControlRig/Private/ModularRigController.cpp
@@ -1690,6 +1690,35 @@ bool UModularRigController::DeselectModule(const FString& InModulePath)
 	return SelectModule(InModulePath, false);
 }
 
+bool UModularRigController::SetModuleSelection(const TArray<FString>& InModulePaths)
+{
+	bool bResult = false;
+	const TArray<FString> OldSelection = GetSelectedModules();
+
+	for(const FString& PreviouslySelectedModule : OldSelection)
+	{
+		if(!InModulePaths.Contains(PreviouslySelectedModule))
+		{
+			if(DeselectModule(PreviouslySelectedModule))
+			{
+				bResult = true;
+			}
+		}
+	}
+	for(const FString& NewModuleToSelect : InModulePaths)
+	{
+		if(!OldSelection.Contains(NewModuleToSelect))
+		{
+			if(SelectModule(NewModuleToSelect))
+			{
+				bResult = true;
+			}
+		}
+	}
+
+	return bResult;
+}
+
 TArray<FString> UModularRigController::GetSelectedModules() const
 {
 	return Model->SelectedModulePaths;
diff --git a/Engine/Plugins/Animation/ControlRig/Source/ControlRig/Public/ModularRigController.h b/Engine/Plugins/Animation/ControlRig/Source/ControlRig/Public/ModularRigController.h
--- a/Engine/Plugins/Animation/ControlRig/Source/ControlRig/Public/ModularRigController.h
+++ b/Engine/Plugins/Animation/ControlRig/Source/ControlRig/Public/ModularRigController.h
@@ -85,6 +85,9 @@ class CONTROLRIG_API UModularRigController : public UObject
 	UFUNCTION(BlueprintCallable, Category = "ControlRig | Modules")
 	bool DeselectModule(const FString& InModulePath);
 
+	UFUNCTION(BlueprintCallable, Category = "ControlRig | Modules")
+	bool SetModuleSelection(const TArray<FString>& InModulePaths);
+
 	UFUNCTION(BlueprintPure, Category = "ControlRig | Modules")
 	TArray<FString> GetSelectedModules() const;
 
diff --git a/Engine/Plugins/Animation/ControlRig/Source/ControlRigEditor/Private/Editor/SModularRigModel.cpp b/Engine/Plugins/Animation/ControlRig/Source/ControlRigEditor/Private/Editor/SModularRigModel.cpp
--- a/Engine/Plugins/Animation/ControlRig/Source/ControlRigEditor/Private/Editor/SModularRigModel.cpp
+++ b/Engine/Plugins/Animation/ControlRig/Source/ControlRigEditor/Private/Editor/SModularRigModel.cpp
@@ -309,6 +309,48 @@ FReply SModularRigModel::OnKeyDown(const FGeometry& MyGeometry, const FKeyEvent&
 	return FReply::Unhandled();
 }
 
+FReply SModularRigModel::OnMouseButtonDown(const FGeometry& MyGeometry, const FPointerEvent& MouseEvent)
+{
+	const FReply Reply = SCompoundWidget::OnMouseButtonDown(MyGeometry, MouseEvent);
+	if(Reply.IsEventHandled())
+	{
+		return Reply;
+	}
+
+	if(MouseEvent.GetEffectingButton() == EKeys::MiddleMouseButton)
+	{
+		if(const TSharedPtr<FModularRigTreeElement>* ItemPtr = TreeView->FindItemAtPosition(MouseEvent.GetScreenSpacePosition()))
+		{
+			if(const TSharedPtr<FModularRigTreeElement>& Item = *ItemPtr)
+			{
+				if (ControlRigBlueprint.IsValid())
+				{
+					UModularRigController* Controller = ControlRigBlueprint->GetModularRigController();
+					check(Controller);
+
+					if(const FRigModuleReference* Module = Controller->FindModule(Item->ModulePath))
+					{
+						TArray<const FRigModuleReference*> ModulesToSelect = {Module};
+						TArray<FString> ModulePaths;
+						for(int32 Index = 0; Index < ModulesToSelect.Num(); Index++)
+						{
+							ModulePaths.Add(ModulesToSelect[Index]->GetPath());
+							for(const FRigModuleReference* ChildModule : ModulesToSelect[Index]->CachedChildren)
+							{
+								ModulesToSelect.AddUnique(ChildModule);
+							}
+						}
+						
+						Controller->SetModuleSelection(ModulePaths);
+					}
+				}
+			}
+		}
+	}
+
+	return FReply::Unhandled();
+}
+
 void SModularRigModel::RefreshTreeView(bool bRebuildContent)
 {
 	bool bDummySuspensionFlag = false;
@@ -1186,25 +1228,9 @@ void SModularRigModel::HandleSelectionChanged(TSharedPtr<FModularRigTreeElement>
 		UModularRigController* Controller = ControlRigBlueprint->GetModularRigController();
 		check(Controller);
 
-		const TArray<FString> NewSelection = TreeView->GetSelectedKeys();
-		const TArray<FString> OldSelection = Controller->GetSelectedModules();
-
 		const TGuardValue<bool> GuardSelection(bIsPerformingSelection, true);
-
-		for(const FString& PreviouslySelectedModule : OldSelection)
-		{
-			if(!NewSelection.Contains(PreviouslySelectedModule))
-			{
-				Controller->DeselectModule(PreviouslySelectedModule);
-			}
-		}
-		for(const FString& NewModuleToSelect : NewSelection)
-		{
-			if(!OldSelection.Contains(NewModuleToSelect))
-			{
-				Controller->SelectModule(NewModuleToSelect);
-			}
-		}
+		const TArray<FString> NewSelection = TreeView->GetSelectedKeys();
+		Controller->SetModuleSelection(NewSelection);
 	}
 }
 
diff --git a/Engine/Plugins/Animation/ControlRig/Source/ControlRigEditor/Private/Editor/SModularRigModel.h b/Engine/Plugins/Animation/ControlRig/Source/ControlRigEditor/Private/Editor/SModularRigModel.h
--- a/Engine/Plugins/Animation/ControlRig/Source/ControlRigEditor/Private/Editor/SModularRigModel.h
+++ b/Engine/Plugins/Animation/ControlRig/Source/ControlRigEditor/Private/Editor/SModularRigModel.h
@@ -50,6 +50,7 @@ private:
 
 	/** SWidget interface */
 	virtual FReply OnKeyDown(const FGeometry& MyGeometry, const FKeyEvent& InKeyEvent);
+	virtual FReply OnMouseButtonDown(const FGeometry& MyGeometry, const FPointerEvent& MouseEvent) override;
 
 	/** Rebuild the tree view */
 	void RefreshTreeView(bool bRebuildContent = true);
diff --git a/Engine/Plugins/Animation/ControlRig/Source/ControlRigEditor/Private/Editor/SRigHierarchy.cpp b/Engine/Plugins/Animation/ControlRig/Source/ControlRigEditor/Private/Editor/SRigHierarchy.cpp
--- a/Engine/Plugins/Animation/ControlRig/Source/ControlRigEditor/Private/Editor/SRigHierarchy.cpp
+++ b/Engine/Plugins/Animation/ControlRig/Source/ControlRigEditor/Private/Editor/SRigHierarchy.cpp
@@ -471,6 +471,37 @@ FReply SRigHierarchy::OnKeyDown(const FGeometry& MyGeometry, const FKeyEvent& In
 	return FReply::Unhandled();
 }
 
+FReply SRigHierarchy::OnMouseButtonUp(const FGeometry& MyGeometry, const FPointerEvent& MouseEvent)
+{
+	FReply Reply = SCompoundWidget::OnMouseButtonUp(MyGeometry, MouseEvent);
+	if(Reply.IsEventHandled())
+	{
+		return Reply;
+	}
+
+	if(MouseEvent.GetEffectingButton() == EKeys::MiddleMouseButton)
+	{
+		if(const TSharedPtr<FRigTreeElement>* ItemPtr = TreeView->FindItemAtPosition(MouseEvent.GetScreenSpacePosition()))
+		{
+			if(const TSharedPtr<FRigTreeElement>& Item = *ItemPtr)
+			{
+				if (URigHierarchy* Hierarchy = GetHierarchy())
+				{
+					TArray<FRigElementKey> KeysToSelect = {Item->Key};
+					KeysToSelect.Append(Hierarchy->GetChildren(Item->Key, true));
+
+					URigHierarchyController* Controller = Hierarchy->GetController(true);
+					check(Controller);
+		
+					Controller->SetSelection(KeysToSelect);
+				}
+			}
+		}
+	}
+
+	return FReply::Unhandled();
+}
+
 EVisibility SRigHierarchy::IsToolbarVisible() const
 {
 	if (URigHierarchy* Hierarchy = GetHierarchy())
diff --git a/Engine/Plugins/Animation/ControlRig/Source/ControlRigEditor/Private/Editor/SRigHierarchy.h b/Engine/Plugins/Animation/ControlRig/Source/ControlRigEditor/Private/Editor/SRigHierarchy.h
--- a/Engine/Plugins/Animation/ControlRig/Source/ControlRigEditor/Private/Editor/SRigHierarchy.h
+++ b/Engine/Plugins/Animation/ControlRig/Source/ControlRigEditor/Private/Editor/SRigHierarchy.h
@@ -66,6 +66,7 @@ private:
 
 	/** SWidget interface */
 	virtual FReply OnKeyDown(const FGeometry& MyGeometry, const FKeyEvent& InKeyEvent);
+	virtual FReply OnMouseButtonUp(const FGeometry& MyGeometry, const FPointerEvent& MouseEvent) override;
 
 	/** Rebuild the tree view */
 	void RefreshTreeView(bool bRebuildContent = true);
