Control Rig: Show connector icon in details panel

#rnx

diff --git a/Engine/Plugins/Animation/ControlRig/Source/ControlRigEditor/Private/ControlRigModuleDetails.cpp b/Engine/Plugins/Animation/ControlRig/Source/ControlRigEditor/Private/ControlRigModuleDetails.cpp
--- a/Engine/Plugins/Animation/ControlRig/Source/ControlRigEditor/Private/ControlRigModuleDetails.cpp
+++ b/Engine/Plugins/Animation/ControlRig/Source/ControlRigEditor/Private/ControlRigModuleDetails.cpp
@@ -28,6 +28,7 @@
 #include "ScopedTransaction.h"
 #include "Editor/SRigHierarchyTreeView.h"
 #include "Widgets/SRigVMVariantTagWidget.h"
+#include "Algo/Sort.h"
 
 #define LOCTEXT_NAMESPACE "ControlRigModuleDetails"
 
@@ -314,6 +315,13 @@ void FRigModuleInstanceDetails::CustomizeDetails(IDetailLayoutBuilder& DetailBui
 		if (bDisplayConnectors)
 		{
 			TArray<FRigModuleConnector> Connectors = GetConnectors();
+
+			// sort connectors primary first, then secondary, then optional
+			Algo::SortBy(Connectors, [](const FRigModuleConnector& Connector) -> int32
+			{
+				return Connector.IsPrimary() ? 0 : (Connector.IsOptional() ? 2 : 1);
+			});
+			
 			for(const FRigModuleConnector& Connector : Connectors)
 			{
 				const FText Label = FText::FromString(Connector.Name);
@@ -394,14 +402,40 @@ void FRigModuleInstanceDetails::CustomizeDetails(IDetailLayoutBuilder& DetailBui
 				});
 				TreeDelegates.OnSelectionChanged.BindSP(this, &FRigModuleInstanceDetails::OnConnectorTargetChanged, Connector);
 			
+				static const FSlateBrush* PrimaryBrush = FControlRigEditorStyle::Get().GetBrush("ControlRig.ConnectorPrimary");
+				static const FSlateBrush* SecondaryBrush = FControlRigEditorStyle::Get().GetBrush("ControlRig.ConnectorSecondary");
+				static const FSlateBrush* OptionalBrush = FControlRigEditorStyle::Get().GetBrush("ControlRig.ConnectorOptional");
+
+				const FSlateBrush* IconBrush = Connector.IsPrimary() ? PrimaryBrush : (Connector.IsOptional() ? OptionalBrush : SecondaryBrush);
 
 				ConnectionsCategory.AddCustomRow(Label)
 					.NameContent()
 					[
-						SNew(STextBlock)
-						.Text(Label)
-						.Font(IDetailLayoutBuilder::GetDetailFont())
-						.IsEnabled(true)
+						SNew(SHorizontalBox)
+
+						+SHorizontalBox::Slot()
+						.AutoWidth()
+						.Padding(0.f, 0.f, 4.f, 0.f)
+						.HAlign(HAlign_Left)
+						.VAlign(VAlign_Center)
+						[
+							SNew(SImage)
+							.Image(IconBrush)
+							.ColorAndOpacity(FSlateColor::UseForeground())
+							.DesiredSizeOverride(FVector2D(16, 16))
+						]
+						
+						+SHorizontalBox::Slot()
+						.AutoWidth()
+						.Padding(0.f, 0.f, 0.f, 0.f)
+						.HAlign(HAlign_Left)
+						.VAlign(VAlign_Center)
+						[
+							SNew(STextBlock)
+							.Text(Label)
+							.Font(IDetailLayoutBuilder::GetDetailFont())
+							.IsEnabled(true)
+						]
 					]
 					.ValueContent()
 					[
