ControlRig / RigVM: Filter functions and assets based on invalid tag

#rnx

diff --git a/Engine/Plugins/Animation/ControlRig/Source/ControlRigEditor/Private/Editor/SRigModuleAssetBrowser.cpp b/Engine/Plugins/Animation/ControlRig/Source/ControlRigEditor/Private/Editor/SRigModuleAssetBrowser.cpp
--- a/Engine/Plugins/Animation/ControlRig/Source/ControlRigEditor/Private/Editor/SRigModuleAssetBrowser.cpp
+++ b/Engine/Plugins/Animation/ControlRig/Source/ControlRigEditor/Private/Editor/SRigModuleAssetBrowser.cpp
@@ -10,6 +10,7 @@
 
 #include "ControlRigEditor.h"
 #include "Widgets/Layout/SUniformGridPanel.h"
+#include "Widgets/SRigVMVariantWidget.h"
 
 #define LOCTEXT_NAMESPACE "RigModuleAssetBrowser"
 
@@ -130,7 +131,29 @@ bool SRigModuleAssetBrowser::OnShouldFilterAsset(const struct FAssetData& AssetD
 	}
 
 	const EControlRigType ControlRigType = (EControlRigType)(ControlTypeEnum->GetValueByName(*ControlRigTypeStr));
-	return ControlRigType != EControlRigType::RigModule;
+	if(ControlRigType != EControlRigType::RigModule)
+	{
+		return true;
+	}
+
+	static const FName AssetVariantPropertyName = GET_MEMBER_NAME_CHECKED(URigVMBlueprint, AssetVariant);
+	const FProperty* AssetVariantProperty = CastField<FProperty>(URigVMBlueprint::StaticClass()->FindPropertyByName(AssetVariantPropertyName));
+	const FString VariantStr = AssetData.GetTagValueRef<FString>(AssetVariantPropertyName);
+	if(!VariantStr.IsEmpty())
+	{
+		FRigVMVariant AssetVariant;
+		AssetVariantProperty->ImportText_Direct(*VariantStr, &AssetVariant, nullptr, EPropertyPortFlags::PPF_None);
+
+		for(const FRigVMTag& VariantTag : AssetVariant.Tags)
+		{
+			if(VariantTag.bMarksSubjectAsInvalid)
+			{
+				return true;
+			}
+		}
+	}
+
+	return false;
 }
 
 void SRigModuleAssetBrowser::OnAssetDoubleClicked(const FAssetData& AssetData)
@@ -167,6 +190,50 @@ TSharedRef<SToolTip> SRigModuleAssetBrowser::CreateCustomAssetToolTip(FAssetData
 
 	// Add asset registry tags to a text list; except skeleton as that is implied in Persona
 	TSharedRef<SVerticalBox> DescriptionBox = SNew(SVerticalBox);
+
+	static const FName AssetVariantPropertyName = GET_MEMBER_NAME_CHECKED(URigVMBlueprint, AssetVariant);
+	const FProperty* AssetVariantProperty = CastField<FProperty>(URigVMBlueprint::StaticClass()->FindPropertyByName(AssetVariantPropertyName));
+	const FString VariantStr = AssetData.GetTagValueRef<FString>(AssetVariantPropertyName);
+	if(!VariantStr.IsEmpty())
+	{
+		FRigVMVariant AssetVariant;
+		AssetVariantProperty->ImportText_Direct(*VariantStr, &AssetVariant, nullptr, EPropertyPortFlags::PPF_None);
+
+		if(!AssetVariant.Tags.IsEmpty())
+		{
+			DescriptionBox->AddSlot()
+			.AutoHeight()
+			.Padding(0,0,5,0)
+			[
+				SNew(SHorizontalBox)
+				+ SHorizontalBox::Slot()
+				.HAlign(HAlign_Left)
+				.VAlign(VAlign_Center)
+				.AutoWidth()
+				[
+					SNew(STextBlock)
+					.Text(LOCTEXT("AssetBrowser_RigVMTagsLabel", "Tags :"))
+					.ColorAndOpacity(FSlateColor::UseSubduedForeground())
+				]
+
+				+ SHorizontalBox::Slot()
+				.AutoWidth()
+				.HAlign(HAlign_Left)
+				.VAlign(VAlign_Center)
+				.Padding(4, 0, 0, 0)
+				[
+					SNew(SRigVMVariantTagWidget)
+					.Visibility(EVisibility::Visible)
+					.CanAddTags(false)
+					.EnableContextMenu(false)
+					.EnableTick(false)
+					.Orientation(EOrientation::Orient_Horizontal)
+					.OnGetTags_Lambda([AssetVariant]() { return AssetVariant.Tags; })
+				]
+			];
+		}
+	}
+	
 	for(TPair<FName, FAssetTagValueRef> TagPair : AssetData.TagsAndValues)
 	{
 		if(TagsToShow.Contains(TagPair.Key))
diff --git a/Engine/Plugins/Runtime/RigVM/Source/RigVMDeveloper/Public/RigVMBlueprint.h b/Engine/Plugins/Runtime/RigVM/Source/RigVMDeveloper/Public/RigVMBlueprint.h
--- a/Engine/Plugins/Runtime/RigVM/Source/RigVMDeveloper/Public/RigVMBlueprint.h
+++ b/Engine/Plugins/Runtime/RigVM/Source/RigVMDeveloper/Public/RigVMBlueprint.h
@@ -854,6 +854,7 @@ private:
 	friend class FRigVMTreePackageNode;
 	friend class SRigVMGraphNode;
 	friend struct FGuardSkipDirtyBlueprintStatus;
+	friend class SRigModuleAssetBrowser;
 };
 
 class RIGVMDEVELOPER_API FRigVMBlueprintCompileScope
diff --git a/Engine/Plugins/Runtime/RigVM/Source/RigVMEditor/Private/EdGraph/NodeSpawners/RigVMEdGraphFunctionRefNodeSpawner.cpp b/Engine/Plugins/Runtime/RigVM/Source/RigVMEditor/Private/EdGraph/NodeSpawners/RigVMEdGraphFunctionRefNodeSpawner.cpp
--- a/Engine/Plugins/Runtime/RigVM/Source/RigVMEditor/Private/EdGraph/NodeSpawners/RigVMEdGraphFunctionRefNodeSpawner.cpp
+++ b/Engine/Plugins/Runtime/RigVM/Source/RigVMEditor/Private/EdGraph/NodeSpawners/RigVMEdGraphFunctionRefNodeSpawner.cpp
@@ -368,6 +368,18 @@ bool URigVMEdGraphFunctionRefNodeSpawner::IsTemplateNodeFilteredOut(FBlueprintAc
 	{
 		return true;
 	}
+
+	// filter outdated functions
+	if(ReferencedPublicFunctionHeader.IsValid())
+	{
+		for(const FRigVMTag& Tag : ReferencedPublicFunctionHeader.Variant.Tags)
+		{
+			if(Tag.bMarksSubjectAsInvalid)
+			{
+				return true;
+			}
+		}
+	}
 	
 	if(bIsLocalFunction)
 	{
