ControlRig: Fix fallout from selection changes

#rnx

diff --git a/Engine/Plugins/Animation/ControlRig/Source/ControlRigDeveloper/Private/ControlRigBlueprint.cpp b/Engine/Plugins/Animation/ControlRig/Source/ControlRigDeveloper/Private/ControlRigBlueprint.cpp
--- a/Engine/Plugins/Animation/ControlRig/Source/ControlRigDeveloper/Private/ControlRigBlueprint.cpp
+++ b/Engine/Plugins/Animation/ControlRig/Source/ControlRigDeveloper/Private/ControlRigBlueprint.cpp
@@ -2847,8 +2847,8 @@ void UControlRigBlueprint::HandleRigModulesModified(EModularRigNotification InNo
 		case EModularRigNotification::ModuleSelected:
 		case EModularRigNotification::ModuleDeselected:
 		{
-			ModulesRecompilationBracket--;
-			break;
+			// don't do anything during selection
+			return;
 		}
 		default:
 		{
