Control Rig: Display animation channels in details panel

#rnx

diff --git a/Engine/Plugins/Animation/ControlRig/Source/ControlRig/Private/Rigs/RigHierarchy.cpp b/Engine/Plugins/Animation/ControlRig/Source/ControlRig/Private/Rigs/RigHierarchy.cpp
--- a/Engine/Plugins/Animation/ControlRig/Source/ControlRig/Private/Rigs/RigHierarchy.cpp
+++ b/Engine/Plugins/Animation/ControlRig/Source/ControlRig/Private/Rigs/RigHierarchy.cpp
@@ -1861,17 +1861,6 @@ TArray<FRigElementKey> URigHierarchy::GetChildren(FRigElementKey InKey, bool bRe
 {
 	LLM_SCOPE_BYNAME(TEXT("Animation/ControlRig"));
 
-	auto ConvertElementsToKeys = [this](TConstArrayView<FRigBaseElement*> InElements) -> TArray<FRigElementKey>
-	{
-		TArray<FRigElementKey> ElementKeys;
-		ElementKeys.Reserve(InElements.Num());
-		for (const FRigBaseElement* Element: InElements)
-		{
-			ElementKeys.Add(Element->Key);
-		}
-		return ElementKeys;
-	};
-
 	if (bRecursive)
 	{
 		return ConvertElementsToKeys(GetChildren(Find(InKey), true));
@@ -1939,30 +1928,19 @@ TArray<int32> URigHierarchy::GetChildren(int32 InIndex, bool bRecursive) const
 		return {};
 	}
 	
-	TArray<int32> ChildIndexes;
-
-	auto AddIndexesFromElements = [this, &ChildIndexes](TConstArrayView<FRigBaseElement*> InElements) -> void
-	{
-		ChildIndexes.Reserve(ChildIndexes.Num() + InElements.Num());
-		for (const FRigBaseElement* Element: InElements)
-		{
-			ChildIndexes.Add(Element->Index);
-		}
-	};
-	
-	AddIndexesFromElements(GetChildren(Elements[InIndex]));
+	TArray<int32> ChildIndices;
+	ConvertElementsToIndices(GetChildren(Elements[InIndex]), ChildIndices);
 
 	if (bRecursive)
 	{
 		// Go along the children array and add all children. Once we stop adding children, the traversal index
 		// will reach the end and we're done.
-		for(int32 TraversalIndex = 0; TraversalIndex != ChildIndexes.Num(); TraversalIndex++)
+		for(int32 TraversalIndex = 0; TraversalIndex != ChildIndices.Num(); TraversalIndex++)
 		{
-			AddIndexesFromElements(GetChildren(Elements[ChildIndexes[TraversalIndex]]));
+			ConvertElementsToIndices(GetChildren(Elements[ChildIndices[TraversalIndex]]), ChildIndices);
 		}
 	}
-
-	return ChildIndexes;
+	return ChildIndices;
 }
 
 TConstArrayView<FRigBaseElement*> URigHierarchy::GetChildren(const FRigBaseElement* InElement) const
@@ -2746,6 +2724,56 @@ FRigElementKey URigHierarchy::GetWorldSpaceReferenceKey()
 	return FRigElementKey(WorldSpaceReferenceName, ERigElementType::Reference); 
 }
 
+TArray<FRigElementKey> URigHierarchy::GetAnimationChannels(FRigElementKey InKey, bool bOnlyDirectChildren) const
+{
+	LLM_SCOPE_BYNAME(TEXT("Animation/ControlRig"));
+	return ConvertElementsToKeys(GetAnimationChannels(Find<FRigControlElement>(InKey), bOnlyDirectChildren));
+}
+
+TArray<int32> URigHierarchy::GetAnimationChannels(int32 InIndex, bool bOnlyDirectChildren) const
+{
+	return ConvertElementsToIndices(GetAnimationChannels(Get<FRigControlElement>(InIndex), bOnlyDirectChildren));
+}
+
+TArray<FRigControlElement*> URigHierarchy::GetAnimationChannels(const FRigControlElement* InElement, bool bOnlyDirectChildren) const
+{
+	if(InElement == nullptr)
+	{
+		return {};
+	}
+
+	const TConstArrayView<FRigBaseElement*> AllChildren = GetChildren(InElement);
+	const TArray<FRigBaseElement*> FilteredChildren = AllChildren.FilterByPredicate([](const FRigBaseElement* Element) -> bool
+	{
+		if(const FRigControlElement* ControlElement = Cast<FRigControlElement>(Element))
+		{
+			return ControlElement->IsAnimationChannel();
+		}
+		return false;
+	});
+
+	TArray<FRigControlElement*> AnimationChannels = ConvertElements<FRigControlElement>(FilteredChildren);
+	if(!bOnlyDirectChildren)
+	{
+		AnimationChannels.Append(GetFilteredElements<FRigControlElement>(
+			[InElement](const FRigControlElement* PotentialAnimationChannel) -> bool
+			{
+				if(PotentialAnimationChannel->IsAnimationChannel())
+				{
+					if(PotentialAnimationChannel->Settings.Customization.AvailableSpaces.Contains(InElement->Key))
+					{
+						return true;
+					}
+				}
+				return false;
+			},
+		true /* traverse */
+		));
+	}
+	
+	return AnimationChannels;
+}
+
 TArray<FRigElementKey> URigHierarchy::GetAllKeys(bool bTraverse, ERigElementType InElementType) const
 {
 	LLM_SCOPE_BYNAME(TEXT("Animation/ControlRig"))
diff --git a/Engine/Plugins/Animation/ControlRig/Source/ControlRig/Public/Rigs/RigHierarchy.h b/Engine/Plugins/Animation/ControlRig/Source/ControlRig/Public/Rigs/RigHierarchy.h
--- a/Engine/Plugins/Animation/ControlRig/Source/ControlRig/Public/Rigs/RigHierarchy.h
+++ b/Engine/Plugins/Animation/ControlRig/Source/ControlRig/Public/Rigs/RigHierarchy.h
@@ -3396,6 +3396,30 @@ public:
 		return false;
 	}
 
+	/**
+	 * Returns the animation channels of a given element key
+	 * @param InKey The key of the element to retrieve the animation channels for
+	 * @param bOnlyDirectChildren If set to false also animation channels with secondary parenting relationships will be retrieved
+	 * @return Returns the animation channels' indices
+	 */
+	TArray<FRigElementKey> GetAnimationChannels(FRigElementKey InKey, bool bOnlyDirectChildren = true) const;
+
+	/**
+	 * Returns the animation channels of a given element index
+	 * @param InIndex The index of the element to retrieve the animation channels for
+	 * @param bOnlyDirectChildren If set to false also animation channels with secondary parenting relationships will be retrieved
+	 * @return Returns the animation channels' indices
+	 */
+	TArray<int32> GetAnimationChannels(int32 InIndex, bool bOnlyDirectChildren = true) const;
+
+	/**
+	 * Returns the animation channels of a given element
+	 * @param InElement The element to retrieve the animation channels for
+	 * @param bOnlyDirectChildren If set to false also animation channels with secondary parenting relationships will be retrieved
+	 * @return Returns the animation channels
+	 */
+	TArray<FRigControlElement*> GetAnimationChannels(const FRigControlElement* InElement, bool bOnlyDirectChildren = true) const;
+
 	/**
 	 * Returns all element keys of this hierarchy
 	 * @param bTraverse If set to true the keys will be returned by depth first traversal
@@ -4834,6 +4858,68 @@ private:
 	void ForEachListeningHierarchy(TFunctionRef<void(const FRigHierarchyListener&)> PerListeningHierarchyFunction);
 #endif
 
+public:
+	
+	template<typename RangeType>
+	static void ConvertElementsToKeys(RangeType InElements, TArray<FRigElementKey>& OutKeys)
+	{
+		OutKeys.Reserve(OutKeys.Num() + InElements.Num());
+		for (const typename RangeType::ElementType Element: InElements)
+		{
+			OutKeys.Add(Element->Key);
+		}
+	}
+
+	template<typename RangeType>
+	static void ConvertElementsToIndices(RangeType InElements, TArray<int32>& OutIndices)
+	{
+		OutIndices.Reserve(OutIndices.Num() + InElements.Num());
+		for (const typename RangeType::ElementType Element: InElements)
+		{
+			OutIndices.Add(Element->Index);
+		}
+	}
+
+	template<typename ElementType = FRigBaseElement, typename RangeType = TArray<FRigBaseElement*>>
+	static void ConvertElements(RangeType InElements, TArray<ElementType*>& OutElements, bool bFilterNull = true)
+	{
+		OutElements.Reserve(OutElements.Num() + InElements.Num());
+		for (const typename RangeType::ElementType Element: InElements)
+		{
+			ElementType* CastElement = Cast<ElementType>(Element);
+			if(CastElement || bFilterNull)
+			{
+				OutElements.Add(CastElement);
+			}
+		}
+	}
+
+	template<typename RangeType>
+	static TArray<FRigElementKey> ConvertElementsToKeys(RangeType InElements)
+	{
+		TArray<FRigElementKey> ElementKeys;
+		ConvertElementsToKeys(InElements, ElementKeys);
+		return ElementKeys;
+	}
+
+	template<typename RangeType>
+	static TArray<int32> ConvertElementsToIndices(RangeType InElements)
+	{
+		TArray<int32> ElementIndices;
+		ConvertElementsToIndices(InElements, ElementIndices);
+		return ElementIndices;
+	}
+
+	template<typename ElementType = FRigBaseElement, typename RangeType = TArray<FRigBaseElement*>>
+	static TArray<ElementType*> ConvertElements(RangeType InElements, bool bFilterNull = true)
+	{
+		TArray<ElementType*> OutElements;
+		ConvertElements<ElementType, RangeType>(InElements, OutElements, bFilterNull);
+		return OutElements;
+	}
+
+private:
+	
 	// the currently destroyed element - used to avoid notification storms
 	const FRigBaseElement* ElementBeingDestroyed; 
 	
diff --git a/Engine/Plugins/Animation/ControlRig/Source/ControlRigEditor/Private/ControlRigElementDetails.cpp b/Engine/Plugins/Animation/ControlRig/Source/ControlRigEditor/Private/ControlRigElementDetails.cpp
--- a/Engine/Plugins/Animation/ControlRig/Source/ControlRigEditor/Private/ControlRigElementDetails.cpp
+++ b/Engine/Plugins/Animation/ControlRig/Source/ControlRigEditor/Private/ControlRigElementDetails.cpp
@@ -3421,259 +3421,278 @@ void FRigControlElementDetails::CustomizeAnimationChannels(IDetailLayoutBuilder&
 	];
 	Category.HeaderContent(HeaderContentWidget);
 
-	bool bHasAnimationChannels = false;
-	for(const FRigBaseElement* ChildElement : Hierarchy->GetChildren(ControlElement))
+	const TArray<FRigControlElement*> AnimationChannels = Hierarchy->GetAnimationChannels(ControlElement, false);
+	const bool bHasAnimationChannels = !AnimationChannels.IsEmpty();
+	const FRigElementKey ControlElementKey = ControlElement->GetKey();
+
+	for(const FRigControlElement* AssignedAnimationChannel : AnimationChannels)
 	{
-		if(const FRigControlElement* ChildControlElement = Cast<FRigControlElement>(ChildElement))
+		const FRigElementKey ChildElementKey = AssignedAnimationChannel->GetKey();
+		const bool bIsDirectlyParentedAnimationChannel = HierarchyToChange->GetDefaultParent(ChildElementKey) == ControlElementKey;
+		
+		const TPair<const FSlateBrush*, FSlateColor> BrushAndColor = SRigHierarchyItem::GetBrushForElementType(Hierarchy, ChildElementKey);
+
+		static TArray<TSharedPtr<ERigControlType>> ControlValueTypes;
+		if(ControlValueTypes.IsEmpty())
 		{
-			if(ChildControlElement->IsAnimationChannel())
+			const UEnum* ValueTypeEnum = StaticEnum<ERigControlType>();
+			for(int64 EnumValue = 0; EnumValue < ValueTypeEnum->GetMaxEnumValue(); EnumValue++)
 			{
-				bHasAnimationChannels = true;
-				const FRigElementKey ChildElementKey = ChildElement->GetKey();
-				
-				const TPair<const FSlateBrush*, FSlateColor> BrushAndColor = SRigHierarchyItem::GetBrushForElementType(Hierarchy, ChildElementKey);
-
-				static TArray<TSharedPtr<ERigControlType>> ControlValueTypes;
-				if(ControlValueTypes.IsEmpty())
+				if(ValueTypeEnum->HasMetaData(TEXT("Hidden"), (int32)EnumValue))
 				{
-					const UEnum* ValueTypeEnum = StaticEnum<ERigControlType>();
-					for(int64 EnumValue = 0; EnumValue < ValueTypeEnum->GetMaxEnumValue(); EnumValue++)
-					{
-						if(ValueTypeEnum->HasMetaData(TEXT("Hidden"), (int32)EnumValue))
-						{
-							continue;
-						}
-						ControlValueTypes.Add(MakeShareable(new ERigControlType((ERigControlType)EnumValue)));
-					}
+					continue;
 				}
+				ControlValueTypes.Add(MakeShareable(new ERigControlType((ERigControlType)EnumValue)));
+			}
+		}
 
-				TSharedPtr<SButton> SelectAnimationChannelButton;
-				TSharedPtr<SImage> SelectAnimationChannelImage;
-				TSharedPtr<SWidget> NameContent;
+		TSharedPtr<SButton> SelectAnimationChannelButton;
+		TSharedPtr<SImage> SelectAnimationChannelImage;
+		TSharedPtr<SWidget> NameContent;
 
-				SAssignNew(NameContent, SHorizontalBox)
-				.IsEnabled(bIsEnabled)
+		SAssignNew(NameContent, SHorizontalBox)
+		.IsEnabled(bIsEnabled)
 
-				+ SHorizontalBox::Slot()
-				.MaxWidth(32)
-				.HAlign(HAlign_Left)
-				.VAlign(VAlign_Center)
-				.Padding(FMargin(0.f, 0.f, 3.f, 0.f))
-				[
-					SNew(SComboButton)
-					.ContentPadding(0)
-					.HasDownArrow(false)
-					.ButtonContent()
-					[
-						SNew(SImage)
-						.Image(BrushAndColor.Key)
-						.ColorAndOpacity(BrushAndColor.Value)
-					]
-					.MenuContent()
-					[
-						SNew(SListView<TSharedPtr<ERigControlType>>)
-						.ListItemsSource( &ControlValueTypes )
-						.OnGenerateRow(this, &FRigControlElementDetails::HandleGenerateAnimationChannelTypeRow, ChildElementKey)
-						.OnSelectionChanged(this, &FRigControlElementDetails::HandleControlTypeChanged, ChildElementKey, PropertyUtilities)
-					]
-				]
+		+ SHorizontalBox::Slot()
+		.MaxWidth(32)
+		.HAlign(HAlign_Left)
+		.VAlign(VAlign_Center)
+		.Padding(FMargin(0.f, 0.f, 3.f, 0.f))
+		[
+			SNew(SComboButton)
+			.ContentPadding(0)
+			.HasDownArrow(false)
+			.ButtonContent()
+			[
+				SNew(SImage)
+				.Image(BrushAndColor.Key)
+				.ColorAndOpacity(BrushAndColor.Value)
+			]
+			.MenuContent()
+			[
+				SNew(SListView<TSharedPtr<ERigControlType>>)
+				.ListItemsSource( &ControlValueTypes )
+				.OnGenerateRow(this, &FRigControlElementDetails::HandleGenerateAnimationChannelTypeRow, ChildElementKey)
+				.OnSelectionChanged(this, &FRigControlElementDetails::HandleControlTypeChanged, ChildElementKey, PropertyUtilities)
+			]
+		]
 
-				+ SHorizontalBox::Slot()
-				.AutoWidth()
-				.HAlign(HAlign_Left)
-				.VAlign(VAlign_Center)
-				.Padding(FMargin(0.f, 0.f, 8.f, 0.f))
-				[
-					SNew(SInlineEditableTextBlock)
-					.Font(IDetailLayoutBuilder::GetDetailFont())
-					.Text_Lambda([this, ChildElementKey]() -> FText
-					{
-						return GetDisplayNameForElement(ChildElementKey);
-					})
-					.OnTextCommitted_Lambda([this, ChildElementKey](const FText& InNewText, ETextCommit::Type InCommitType)
-					{
-						SetDisplayNameForElement(InNewText, InCommitType, ChildElementKey);
-					})
-					.OnVerifyTextChanged_Lambda([this, ChildElementKey](const FText& InText, FText& OutErrorMessage) -> bool
-					{
-						return OnVerifyDisplayNameChanged(InText, OutErrorMessage, ChildElementKey);
-					})
-				]
+		+ SHorizontalBox::Slot()
+		.AutoWidth()
+		.HAlign(HAlign_Left)
+		.VAlign(VAlign_Center)
+		.Padding(FMargin(0.f, 0.f, 8.f, 0.f))
+		[
+			SNew(SInlineEditableTextBlock)
+			.Font(bIsDirectlyParentedAnimationChannel ? IDetailLayoutBuilder::GetDetailFont() : IDetailLayoutBuilder::GetDetailFontItalic())
+			.Text_Lambda([this, ChildElementKey]() -> FText
+			{
+				return GetDisplayNameForElement(ChildElementKey);
+			})
+			.OnTextCommitted_Lambda([this, ChildElementKey](const FText& InNewText, ETextCommit::Type InCommitType)
+			{
+				SetDisplayNameForElement(InNewText, InCommitType, ChildElementKey);
+			})
+			.OnVerifyTextChanged_Lambda([this, ChildElementKey](const FText& InText, FText& OutErrorMessage) -> bool
+			{
+				return OnVerifyDisplayNameChanged(InText, OutErrorMessage, ChildElementKey);
+			})
+		]
 
-				+SHorizontalBox::Slot()
-				.AutoWidth()
-				.HAlign(HAlign_Right)
-				.VAlign(VAlign_Center)
-				.Padding(FMargin(0.f, 0.f, 0.f, 0.f))
-				[
-					SAssignNew(SelectAnimationChannelButton, SButton)
-					.ButtonStyle( FAppStyle::Get(), "NoBorder" )
-					.OnClicked_Lambda([this, ChildElementKey]() -> FReply
-					{
-						return OnSelectElementClicked(ChildElementKey);
-					})
-					.ContentPadding(0)
-					.ToolTipText(NSLOCTEXT("ControlRigElementDetails", "SelectAnimationChannelInHierarchyToolTip", "Select Animation Channel"))
-					[
-						SAssignNew(SelectAnimationChannelImage, SImage)
-						.Image(FAppStyle::GetBrush("Icons.Search"))
-					]
-				];
+		+SHorizontalBox::Slot()
+		.AutoWidth()
+		.HAlign(HAlign_Right)
+		.VAlign(VAlign_Center)
+		.Padding(FMargin(0.f, 0.f, 0.f, 0.f))
+		[
+			SAssignNew(SelectAnimationChannelButton, SButton)
+			.ButtonStyle( FAppStyle::Get(), "NoBorder" )
+			.OnClicked_Lambda([this, ChildElementKey]() -> FReply
+			{
+				return OnSelectElementClicked(ChildElementKey);
+			})
+			.ContentPadding(0)
+			.ToolTipText(NSLOCTEXT("ControlRigElementDetails", "SelectAnimationChannelInHierarchyToolTip", "Select Animation Channel"))
+			[
+				SAssignNew(SelectAnimationChannelImage, SImage)
+				.Image(FAppStyle::GetBrush("Icons.Search"))
+			]
+		];
 
-				SelectAnimationChannelImage->SetColorAndOpacity(TAttribute<FSlateColor>::CreateLambda([this, SelectAnimationChannelButton]() { return FRigElementKeyDetails::OnGetWidgetForeground(SelectAnimationChannelButton); }));
+		SelectAnimationChannelImage->SetColorAndOpacity(TAttribute<FSlateColor>::CreateLambda([this, SelectAnimationChannelButton]() { return FRigElementKeyDetails::OnGetWidgetForeground(SelectAnimationChannelButton); }));
 
-				const FText Label = FText::FromString(FString::Printf(TEXT("Channel%s"), *ChildControlElement->GetDisplayName().ToString()));
-				const TArray<FRigElementKey> ChildElementKeys = {ChildElementKey};
-				TAttribute<EVisibility> Visibility = EVisibility::Visible;
+		const FText Label = FText::FromString(FString::Printf(TEXT("Channel%s"), *AssignedAnimationChannel->GetDisplayName().ToString()));
+		const TArray<FRigElementKey> ChildElementKeys = {ChildElementKey};
+		TAttribute<EVisibility> Visibility = EVisibility::Visible;
 
-				FDetailWidgetRow* WidgetRow = nullptr; 
-				switch(ChildControlElement->Settings.ControlType)
+		FDetailWidgetRow* WidgetRow = nullptr; 
+		switch(AssignedAnimationChannel->Settings.ControlType)
+		{
+			case ERigControlType::Bool:
+			{
+				WidgetRow = &CreateBoolValueWidgetRow(ChildElementKeys, Category, Label, FText(), ERigControlValueType::Current, Visibility, NameContent);
+				break;
+			}
+			case ERigControlType::Float:
+			case ERigControlType::ScaleFloat:
+			{
+				WidgetRow = &CreateFloatValueWidgetRow(ChildElementKeys, Category, Label, FText(), ERigControlValueType::Current, Visibility, NameContent);
+				break;
+			}
+			case ERigControlType::Integer:
+			{
+				if(AssignedAnimationChannel->Settings.ControlEnum)
 				{
-					case ERigControlType::Bool:
-					{
-						WidgetRow = &CreateBoolValueWidgetRow(ChildElementKeys, Category, Label, FText(), ERigControlValueType::Current, Visibility, NameContent);
-						break;
-					}
-					case ERigControlType::Float:
-					case ERigControlType::ScaleFloat:
-					{
-						WidgetRow = &CreateFloatValueWidgetRow(ChildElementKeys, Category, Label, FText(), ERigControlValueType::Current, Visibility, NameContent);
-						break;
-					}
-					case ERigControlType::Integer:
-					{
-						if(ChildControlElement->Settings.ControlEnum)
-						{
-							WidgetRow = &CreateEnumValueWidgetRow(ChildElementKeys, Category, Label, FText(), ERigControlValueType::Current, Visibility, NameContent);
-						}
-						else
-						{
-							WidgetRow = &CreateIntegerValueWidgetRow(ChildElementKeys, Category, Label, FText(), ERigControlValueType::Current, Visibility, NameContent);
-						}
-						break;
-					}
-					case ERigControlType::Vector2D:
-					{
-						WidgetRow = &CreateVector2DValueWidgetRow(ChildElementKeys, Category, Label, FText(), ERigControlValueType::Current, Visibility, NameContent);
-						break;
-					}
-					case ERigControlType::Position:
-					case ERigControlType::Rotator:
-					case ERigControlType::Scale:
-					{
-						SAdvancedTransformInputBox<FEulerTransform>::FArguments TransformWidgetArgs =
-							SAdvancedTransformInputBox<FEulerTransform>::FArguments()
-						.DisplayToggle(false)
-						.DisplayRelativeWorld(false)
-						.Visibility(EVisibility::Visible);
-
-						WidgetRow = &CreateTransformComponentValueWidgetRow(
-							ChildControlElement->Settings.ControlType,
-							ChildElementKeys,
-							TransformWidgetArgs,
-							Category,
-							Label,
-							FText(),
-							GetTransformTypeFromValueType(ERigControlValueType::Current),
-							ERigControlValueType::Current,
-							NameContent);
-						break;
-					}
-					case ERigControlType::Transform:
-					case ERigControlType::EulerTransform:
-					{
-						SAdvancedTransformInputBox<FEulerTransform>::FArguments TransformWidgetArgs =
-							SAdvancedTransformInputBox<FEulerTransform>::FArguments()
-						.DisplayToggle(false)
-						.DisplayRelativeWorld(false)
-						.Font(IDetailLayoutBuilder::GetDetailFont())
-						.Visibility(EVisibility::Visible);
-
-						WidgetRow = &CreateEulerTransformValueWidgetRow(
-							ChildElementKeys,
-							TransformWidgetArgs,
-							Category,
-							Label,
-							FText(),
-							ERigTransformElementDetailsTransform::Current,
-							ERigControlValueType::Current,
-							NameContent);
-						break;
-					}
-					default:
-					{
-						WidgetRow = &Category.AddCustomRow(Label)
-						.NameContent()
-						[
-							NameContent.ToSharedRef()
-						];
-						break;
-					}
+					WidgetRow = &CreateEnumValueWidgetRow(ChildElementKeys, Category, Label, FText(), ERigControlValueType::Current, Visibility, NameContent);
 				}
-
-				if(WidgetRow)
+				else
 				{
-					WidgetRow->AddCustomContextMenuAction(FUIAction(
-						FExecuteAction::CreateLambda([this, ChildElementKeys, HierarchyToChange]()
-						{
-							if(URigHierarchyController* Controller = HierarchyToChange->GetController(true))
-							{
-								FScopedTransaction Transaction(LOCTEXT("DeleteAnimationChannels", "Delete Animation Channels"));
-								HierarchyToChange->Modify();
-								
-								for(const FRigElementKey& KeyToRemove : ChildElementKeys)
-								{
-									Controller->RemoveElement(KeyToRemove, true, true);
-								}
-							}
-						})),
-						LOCTEXT("DeleteAnimationChannel", "Delete"),
-						LOCTEXT("DeleteAnimationChannelTooltip", "Deletes this animation channel"),
-						FSlateIcon());
+					WidgetRow = &CreateIntegerValueWidgetRow(ChildElementKeys, Category, Label, FText(), ERigControlValueType::Current, Visibility, NameContent);
+				}
+				break;
+			}
+			case ERigControlType::Vector2D:
+			{
+				WidgetRow = &CreateVector2DValueWidgetRow(ChildElementKeys, Category, Label, FText(), ERigControlValueType::Current, Visibility, NameContent);
+				break;
+			}
+			case ERigControlType::Position:
+			case ERigControlType::Rotator:
+			case ERigControlType::Scale:
+			{
+				SAdvancedTransformInputBox<FEulerTransform>::FArguments TransformWidgetArgs =
+					SAdvancedTransformInputBox<FEulerTransform>::FArguments()
+				.DisplayToggle(false)
+				.DisplayRelativeWorld(false)
+				.Visibility(EVisibility::Visible);
 
-					const FRigElementKey ControlElementKey = ControlElement->GetKey();
-					
-					// move up or down
-					WidgetRow->AddCustomContextMenuAction(FUIAction(
-						FExecuteAction::CreateLambda([this, ControlElementKey, ChildElementKeys, HierarchyToChange]()
+				WidgetRow = &CreateTransformComponentValueWidgetRow(
+					AssignedAnimationChannel->Settings.ControlType,
+					ChildElementKeys,
+					TransformWidgetArgs,
+					Category,
+					Label,
+					FText(),
+					GetTransformTypeFromValueType(ERigControlValueType::Current),
+					ERigControlValueType::Current,
+					NameContent);
+				break;
+			}
+			case ERigControlType::Transform:
+			case ERigControlType::EulerTransform:
+			{
+				SAdvancedTransformInputBox<FEulerTransform>::FArguments TransformWidgetArgs =
+					SAdvancedTransformInputBox<FEulerTransform>::FArguments()
+				.DisplayToggle(false)
+				.DisplayRelativeWorld(false)
+				.Font(IDetailLayoutBuilder::GetDetailFont())
+				.Visibility(EVisibility::Visible);
+
+				WidgetRow = &CreateEulerTransformValueWidgetRow(
+					ChildElementKeys,
+					TransformWidgetArgs,
+					Category,
+					Label,
+					FText(),
+					ERigTransformElementDetailsTransform::Current,
+					ERigControlValueType::Current,
+					NameContent);
+				break;
+			}
+			default:
+			{
+				WidgetRow = &Category.AddCustomRow(Label)
+				.NameContent()
+				[
+					NameContent.ToSharedRef()
+				];
+				break;
+			}
+		}
+
+		if(WidgetRow)
+		{
+			if(bIsDirectlyParentedAnimationChannel)
+			{
+				WidgetRow->AddCustomContextMenuAction(FUIAction(
+					FExecuteAction::CreateLambda([this, ChildElementKeys, HierarchyToChange]()
+					{
+						if(URigHierarchyController* Controller = HierarchyToChange->GetController(true))
 						{
-							if(URigHierarchyController* Controller = HierarchyToChange->GetController(true))
+							FScopedTransaction Transaction(LOCTEXT("DeleteAnimationChannels", "Delete Animation Channels"));
+							HierarchyToChange->Modify();
+							
+							for(const FRigElementKey& KeyToRemove : ChildElementKeys)
 							{
-								FScopedTransaction Transaction(LOCTEXT("MoveAnimationChannelUpTransaction", "Move Animation Channel Up"));
-								HierarchyToChange->Modify();
-								
-								for(const FRigElementKey& KeyToMove : ChildElementKeys)
-								{
-									const int32 LocalIndex = HierarchyToChange->GetLocalIndex(KeyToMove);
-									Controller->ReorderElement(KeyToMove, LocalIndex - 1, true);
-								}
-								Controller->SelectElement(ControlElementKey, true, true);
+								Controller->RemoveElement(KeyToRemove, true, true);
 							}
-						})),
-						LOCTEXT("MoveAnimationChannelUp", "Move Up"),
-						LOCTEXT("MoveAnimationChannelUpTooltip", "Reorders this animation channel to show up one higher"),
-						FSlateIcon());
-					WidgetRow->AddCustomContextMenuAction(FUIAction(
-						FExecuteAction::CreateLambda([this, ControlElementKey, ChildElementKeys, HierarchyToChange]()
+						}
+					})),
+					LOCTEXT("DeleteAnimationChannel", "Delete"),
+					LOCTEXT("DeleteAnimationChannelTooltip", "Deletes this animation channel"),
+					FSlateIcon());
+			}
+			else
+			{
+				WidgetRow->AddCustomContextMenuAction(FUIAction(
+					FExecuteAction::CreateLambda([this, ControlElementKey, ChildElementKeys, HierarchyToChange, PropertyUtilities]()
+					{
+						if(URigHierarchyController* Controller = HierarchyToChange->GetController(true))
 						{
-							if(URigHierarchyController* Controller = HierarchyToChange->GetController(true))
+							FScopedTransaction Transaction(LOCTEXT("RemoveAnimationChannelHosts", "Remove Animation Channel Hosts"));
+							HierarchyToChange->Modify();
+											
+							for(const FRigElementKey& KeyToRemove : ChildElementKeys)
 							{
-								FScopedTransaction Transaction(LOCTEXT("MoveAnimationChannelDownTransaction", "Move Animation Channel Down"));
-								HierarchyToChange->Modify();
-								
-								for(const FRigElementKey& KeyToMove : ChildElementKeys)
-								{
-									const int32 LocalIndex = HierarchyToChange->GetLocalIndex(KeyToMove);
-									Controller->ReorderElement(KeyToMove, LocalIndex + 1, true);
-								}
-								Controller->SelectElement(ControlElementKey, true, true);
+								Controller->RemoveChannelHost(KeyToRemove, ControlElementKey, true, true);
 							}
-						})),
-						LOCTEXT("MoveAnimationChannelDown", "Move Down"),
-						LOCTEXT("MoveAnimationChannelDownTooltip", "Reorders this animation channel to show up one lower"),
-						FSlateIcon());
-				}
+							PropertyUtilities->ForceRefresh();
+						}
+					})),
+					LOCTEXT("RemoveAnimationChannelHost", "Remove from this host"),
+					LOCTEXT("RemoveAnimationChannelHostTooltip", "Remove the animation channel from this host"),
+					FSlateIcon());
 			}
+			
+			// move up or down
+			WidgetRow->AddCustomContextMenuAction(FUIAction(
+				FExecuteAction::CreateLambda([this, ControlElementKey, ChildElementKeys, HierarchyToChange]()
+				{
+					if(URigHierarchyController* Controller = HierarchyToChange->GetController(true))
+					{
+						FScopedTransaction Transaction(LOCTEXT("MoveAnimationChannelUpTransaction", "Move Animation Channel Up"));
+						HierarchyToChange->Modify();
+						
+						for(const FRigElementKey& KeyToMove : ChildElementKeys)
+						{
+							const int32 LocalIndex = HierarchyToChange->GetLocalIndex(KeyToMove);
+							Controller->ReorderElement(KeyToMove, LocalIndex - 1, true);
+						}
+						Controller->SelectElement(ControlElementKey, true, true);
+					}
+				})),
+				LOCTEXT("MoveAnimationChannelUp", "Move Up"),
+				LOCTEXT("MoveAnimationChannelUpTooltip", "Reorders this animation channel to show up one higher"),
+				FSlateIcon());
+			WidgetRow->AddCustomContextMenuAction(FUIAction(
+				FExecuteAction::CreateLambda([this, ControlElementKey, ChildElementKeys, HierarchyToChange]()
+				{
+					if(URigHierarchyController* Controller = HierarchyToChange->GetController(true))
+					{
+						FScopedTransaction Transaction(LOCTEXT("MoveAnimationChannelDownTransaction", "Move Animation Channel Down"));
+						HierarchyToChange->Modify();
+						
+						for(const FRigElementKey& KeyToMove : ChildElementKeys)
+						{
+							const int32 LocalIndex = HierarchyToChange->GetLocalIndex(KeyToMove);
+							Controller->ReorderElement(KeyToMove, LocalIndex + 1, true);
+						}
+						Controller->SelectElement(ControlElementKey, true, true);
+					}
+				})),
+				LOCTEXT("MoveAnimationChannelDown", "Move Down"),
+				LOCTEXT("MoveAnimationChannelDownTooltip", "Reorders this animation channel to show up one lower"),
+				FSlateIcon());
 		}
 	}
 
