Control Rig: Ensure to update thumbnails for modules as needed

#rnx

diff --git a/Engine/Plugins/Animation/ControlRig/Source/ControlRigDeveloper/Private/ControlRigBlueprint.cpp b/Engine/Plugins/Animation/ControlRig/Source/ControlRigDeveloper/Private/ControlRigBlueprint.cpp
--- a/Engine/Plugins/Animation/ControlRig/Source/ControlRigDeveloper/Private/ControlRigBlueprint.cpp
+++ b/Engine/Plugins/Animation/ControlRig/Source/ControlRigDeveloper/Private/ControlRigBlueprint.cpp
@@ -944,6 +944,7 @@ void UControlRigBlueprint::PreSave(FObjectPreSaveContext ObjectSaveContext)
 	{
 		ControlRigType = EControlRigType::RigModule;
 		ItemTypeDisplayName = TEXT("Rig Module");
+		CustomThumbnail = RigModuleSettings.Icon.ToString();
 	}
 	else if (GetControlRigClass()->IsChildOf(UModularRig::StaticClass()))
 	{
diff --git a/Engine/Plugins/Animation/ControlRig/Source/ControlRigDeveloper/Public/ControlRigBlueprint.h b/Engine/Plugins/Animation/ControlRig/Source/ControlRigDeveloper/Public/ControlRigBlueprint.h
--- a/Engine/Plugins/Animation/ControlRig/Source/ControlRigDeveloper/Public/ControlRigBlueprint.h
+++ b/Engine/Plugins/Animation/ControlRig/Source/ControlRigDeveloper/Public/ControlRigBlueprint.h
@@ -190,6 +190,12 @@ public:
 	UPROPERTY(EditAnywhere, Category = "Hierarchy", AssetRegistrySearchable)
 	FRigModuleSettings RigModuleSettings;
 
+	// This relates to FAssetThumbnailPool::CustomThumbnailTagName and allows
+	// the thumbnail pool to show the thumbnail of the icon rather than the
+	// rig itself to avoid deploying the 3D renderer.
+	UPROPERTY(EditAnywhere, Category = "Hierarchy", AssetRegistrySearchable)
+	FString CustomThumbnail;
+
 	/** Asset searchable information module references in this rig */
 	UPROPERTY(AssetRegistrySearchable)
 	TArray<FModuleReferenceData> ModuleReferenceData;
