[ControlRig] added a null check for shape actor as they might not have an attached parent to avoid a hard crash.

diff --git a/Engine/Plugins/Animation/ControlRig/Source/ControlRigEditor/Private/EditMode/ControlRigEditMode.cpp b/Engine/Plugins/Animation/ControlRig/Source/ControlRigEditor/Private/EditMode/ControlRigEditMode.cpp
--- a/Engine/Plugins/Animation/ControlRig/Source/ControlRigEditor/Private/EditMode/ControlRigEditMode.cpp
+++ b/Engine/Plugins/Animation/ControlRig/Source/ControlRigEditor/Private/EditMode/ControlRigEditMode.cpp
@@ -4956,8 +4956,15 @@ void FControlRigEditMode::MoveControlShape(AControlRigShapeActor* ShapeActor, co
 		FTransform CurrentTransform;
 		if (ShapeActor->StaticMeshComponent && ShapeActor->StaticMeshComponent->GetStaticMesh())
 		{
-			const FTransform& ParentTransform = ShapeActor->GetAttachParentActor()->GetTransform();
-			CurrentTransform = ShapeActor->GetGlobalTransform() * ParentTransform;
+			if (ShapeActor->GetAttachParentActor())
+			{
+				const FTransform& ParentTransform = ShapeActor->GetAttachParentActor()->GetTransform();
+				CurrentTransform = ShapeActor->GetGlobalTransform() * ParentTransform;
+			}
+			else
+			{
+				CurrentTransform = ShapeActor->GetGlobalTransform();
+			}
 		}
 		else
 		{
