Sequencer:  Control Rig: Restore selection state when refreshing tree view and make sure to clear sequencer when selecting in viewport/outliner

diff --git a/Engine/Plugins/Animation/ControlRig/Source/ControlRigEditor/Private/EditMode/ControlRigEditMode.cpp b/Engine/Plugins/Animation/ControlRig/Source/ControlRigEditor/Private/EditMode/ControlRigEditMode.cpp
--- a/Engine/Plugins/Animation/ControlRig/Source/ControlRigEditor/Private/EditMode/ControlRigEditMode.cpp
+++ b/Engine/Plugins/Animation/ControlRig/Source/ControlRigEditor/Private/EditMode/ControlRigEditMode.cpp
@@ -14,6 +14,9 @@
 #include "EditMode/SControlRigDetails.h"
 #include "EditMode/SControlRigOutliner.h"
 #include "ISequencer.h"
+#include "ISequencer.h"
+#include "MVVM/ViewModels/SequencerEditorViewModel.h"
+#include "MVVM/Selection/Selection.h"
 #include "MovieScene.h"
 #include "Editor.h"
 #include "EditorViewportClient.h"
@@ -1463,6 +1466,14 @@ bool FControlRigEditMode::HandleClick(FEditorViewportClient* InViewportClient, H
 								GEditor->SelectNone(false, true);
 								GEditor->NoteSelectionChange();
 							}
+							//also need to clear explicitly in sequencer
+							if (WeakSequencer.IsValid())
+							{
+								if (ISequencer* SequencerPtr = WeakSequencer.Pin().Get())
+								{
+									SequencerPtr->GetViewModel()->GetSelection()->Empty();
+								}
+							}
 						}
 						ClearRigElementSelection(ValidControlTypeMask());
 						SetRigElementSelection(ShapeActor->ControlRig.Get(),ERigElementType::Control, ControlName, true);
diff --git a/Engine/Plugins/Animation/ControlRig/Source/ControlRigEditor/Private/EditMode/SControlRigOutliner.cpp b/Engine/Plugins/Animation/ControlRig/Source/ControlRigEditor/Private/EditMode/SControlRigOutliner.cpp
--- a/Engine/Plugins/Animation/ControlRig/Source/ControlRigEditor/Private/EditMode/SControlRigOutliner.cpp
+++ b/Engine/Plugins/Animation/ControlRig/Source/ControlRigEditor/Private/EditMode/SControlRigOutliner.cpp
@@ -17,6 +17,8 @@
 #include "EditMode/ControlRigEditMode.h"
 #include "EditorModeManager.h"
 #include "ISequencer.h"
+#include "MVVM/ViewModels/SequencerEditorViewModel.h"
+#include "MVVM/Selection/Selection.h"
 #include "LevelSequence.h"
 #include "Selection.h"
 #include "Editor.h"
@@ -896,14 +898,25 @@ TArray<URigHierarchy*> SMultiRigHierarchyTreeView::GetHierarchy() const
 void SMultiRigHierarchyTreeView::SetControlRigs(TArrayView < TWeakObjectPtr<UControlRig>>& InControlRigs)
 {
 	ControlRigs.SetNum(0);
+	TArray < TPair<UControlRig*, TArray<FName>>> SelectedControls;
 	for (TWeakObjectPtr<UControlRig>& ControlRig : InControlRigs)
 	{
 		if (ControlRig.IsValid())
 		{
 			ControlRigs.Add(ControlRig.Get());
+			SelectedControls.Add(TPair<UControlRig*, TArray<FName>>(ControlRig.Get(), ControlRig->CurrentControlSelection()));
 		}
 	}
+	
 	RefreshTreeView(true);
+	//reselect controls that will have gotten cleared by the refresh, this situation can happen on save
+	for (TPair<UControlRig*, TArray<FName>>& CRS : SelectedControls)
+	{
+		for (const FName& Name : CRS.Value)
+		{
+			CRS.Key->SelectControl(Name, true);
+		}
+	}
 }
 
 //////////////////////////////////////////////////////////////
@@ -1138,6 +1151,15 @@ void SControlRigOutliner::HandleSelectionChanged(TSharedPtr<FMultiRigTreeElement
 				GEditor->SelectNone(true, true);
 				GEditor->RedrawLevelEditingViewports();
 			}
+			const TWeakPtr<ISequencer>& WeakSequencer = EditMode->GetWeakSequencer();
+			//also need to clear explicitly in sequencer
+			if (WeakSequencer.IsValid())
+			{
+				if (ISequencer* SequencerPtr = WeakSequencer.Pin().Get())
+				{
+					SequencerPtr->GetViewModel()->GetSelection()->Empty();
+				}
+			}
 		}
 	}
 
