Control Rig: Separate storage for pose, dirty state and curve from hierarchy

#rnx

diff --git a/Engine/Plugins/Animation/AnimationData/Source/AnimationData/Private/AnimSequencerController.cpp b/Engine/Plugins/Animation/AnimationData/Source/AnimationData/Private/AnimSequencerController.cpp
--- a/Engine/Plugins/Animation/AnimationData/Source/AnimationData/Private/AnimSequencerController.cpp
+++ b/Engine/Plugins/Animation/AnimationData/Source/AnimationData/Private/AnimSequencerController.cpp
@@ -3054,7 +3054,7 @@ bool UAnimSequencerController::AddCurveControl(const FName& CurveName) const
 						Settings.ControlType = ERigControlType::Float;
 						Settings.DisplayName = FName(*(CurveKey.Name.ToString() + TEXT(" Curve")));
 								
-						HierarchyController->AddControl(CurveControlKey.Name, FRigElementKey(), Settings, FRigControlValue::Make(CurveElement->Value), FTransform::Identity, FTransform::Identity, false);
+						HierarchyController->AddControl(CurveControlKey.Name, FRigElementKey(), Settings, FRigControlValue::Make(CurveElement->Get()), FTransform::Identity, FTransform::Identity, false);
 
 						const FRigControlElement* ControlElement = Hierarchy->FindChecked<FRigControlElement>(CurveControlKey);
 						ensure(ControlElement);
@@ -3509,7 +3509,7 @@ bool UAnimSequencerController::DuplicateCurveControl(const FName& CurveName, con
 						Settings.ControlType = ERigControlType::Float;
 						Settings.DisplayName = FName(*(NewCurveKey.Name.ToString() + TEXT(" Curve")));
 			
-						HierarchyController->AddControl(NewCurveControlKey.Name, FRigElementKey(), Settings, FRigControlValue::Make(CurveElement->Value), FTransform::Identity, FTransform::Identity, false);
+						HierarchyController->AddControl(NewCurveControlKey.Name, FRigElementKey(), Settings, FRigControlValue::Make(CurveElement->Get()), FTransform::Identity, FTransform::Identity, false);
 			
 						// Rename the curve driving the control value
 						const FScalarParameterNameAndCurve* ParameterCurvePair = Section->GetScalarParameterNamesAndCurves().FindByPredicate([CurveName](const FScalarParameterNameAndCurve& Parameter)
diff --git a/Engine/Plugins/Animation/ControlRig/Source/ControlRig/Private/AnimNode_ControlRigBase.cpp b/Engine/Plugins/Animation/ControlRig/Source/ControlRig/Private/AnimNode_ControlRigBase.cpp
--- a/Engine/Plugins/Animation/ControlRig/Source/ControlRig/Private/AnimNode_ControlRigBase.cpp
+++ b/Engine/Plugins/Animation/ControlRig/Source/ControlRig/Private/AnimNode_ControlRigBase.cpp
@@ -364,9 +364,9 @@ void FAnimNode_ControlRigBase::UpdateOutput(UControlRig* ControlRig, FPoseContex
 		ControlRigCurves.Reserve(Hierarchy->Num(ERigElementType::Curve));
 		Hierarchy->ForEach<FRigCurveElement>([&ControlRigCurves](const FRigCurveElement* InElement)
 		{
-			if(InElement->bIsValueSet)
+			if(InElement->IsValueSet())
 			{
-				ControlRigCurves.Add(InElement->GetFName(), InElement->Value);
+				ControlRigCurves.Add(InElement->GetFName(), InElement->Get());
 			}
 			return true;
 		});
diff --git a/Engine/Plugins/Animation/ControlRig/Source/ControlRig/Private/ControlRig.cpp b/Engine/Plugins/Animation/ControlRig/Source/ControlRig/Private/ControlRig.cpp
--- a/Engine/Plugins/Animation/ControlRig/Source/ControlRig/Private/ControlRig.cpp
+++ b/Engine/Plugins/Animation/ControlRig/Source/ControlRig/Private/ControlRig.cpp
@@ -2010,7 +2010,7 @@ void UControlRig::CreateRigControlsForCurveContainer()
 				Settings.bDrawLimits = false;
 
 				FRigControlValue Value;
-				Value.Set<float>(CurveElement->Value);
+				Value.Set<float>(CurveElement->Get());
 
 				Controller->AddControl(CurveElement->GetFName(), FRigElementKey(), Settings, Value, FTransform::Identity, FTransform::Identity); 
 			}
diff --git a/Engine/Plugins/Animation/ControlRig/Source/ControlRig/Private/ControlRigGizmoActor.cpp b/Engine/Plugins/Animation/ControlRig/Source/ControlRig/Private/ControlRigGizmoActor.cpp
--- a/Engine/Plugins/Animation/ControlRig/Source/ControlRig/Private/ControlRigGizmoActor.cpp
+++ b/Engine/Plugins/Animation/ControlRig/Source/ControlRig/Private/ControlRigGizmoActor.cpp
@@ -176,7 +176,7 @@ bool AControlRigShapeActor::UpdateControlSettings(
 	// update the shape transform
 	if(bShapeTransformChanged)
 	{
-		const FTransform ShapeTransform = InControlElement->Shape.Get(ERigTransformType::CurrentLocal);
+		const FTransform ShapeTransform = InControlElement->GetShapeTransform().Get(ERigTransformType::CurrentLocal);
 		StaticMeshComponent->SetRelativeTransform(MeshTransform * ShapeTransform);
 	}
 	
diff --git a/Engine/Plugins/Animation/ControlRig/Source/ControlRig/Private/Rigs/FKControlRig.cpp b/Engine/Plugins/Animation/ControlRig/Source/ControlRig/Private/Rigs/FKControlRig.cpp
--- a/Engine/Plugins/Animation/ControlRig/Source/ControlRig/Private/Rigs/FKControlRig.cpp
+++ b/Engine/Plugins/Animation/ControlRig/Source/ControlRig/Private/Rigs/FKControlRig.cpp
@@ -417,7 +417,7 @@ void UFKControlRig::CreateRigElements(const FReferenceSkeleton& InReferenceSkele
 
 				const FName DisplayCurveControlName(*(CurveElement->GetFName().ToString() + TEXT(" Curve")));
 				Settings.DisplayName = DisplayCurveControlName;
-				Controller->AddControl(ControlName, FRigElementKey(), Settings, FRigControlValue::Make(CurveElement->Value), FTransform::Identity, FTransform::Identity, false);
+				Controller->AddControl(ControlName, FRigElementKey(), Settings, FRigControlValue::Make(CurveElement->Get()), FTransform::Identity, FTransform::Identity, false);
 				
 				return true;
 			});
diff --git a/Engine/Plugins/Animation/ControlRig/Source/ControlRig/Private/Rigs/RigHierarchy.cpp b/Engine/Plugins/Animation/ControlRig/Source/ControlRig/Private/Rigs/RigHierarchy.cpp
--- a/Engine/Plugins/Animation/ControlRig/Source/ControlRig/Private/Rigs/RigHierarchy.cpp
+++ b/Engine/Plugins/Animation/ControlRig/Source/ControlRig/Private/Rigs/RigHierarchy.cpp
@@ -214,8 +214,11 @@ void URigHierarchy::Save(FArchive& Ar)
 		// store the key
 		FRigElementKey Key = Element->GetKey();
 		Ar << Key;
+	}
 
-		// allow the element to store more information
+	for(int32 ElementIndex = 0; ElementIndex < ElementCount; ElementIndex++)
+	{
+		FRigBaseElement* Element = Elements[ElementIndex];
 		Element->Serialize(Ar, FRigBaseElement::StaticData);
 	}
 
@@ -284,6 +287,11 @@ void URigHierarchy::Load(FArchive& Ar)
 
 	PoseVersionPerElement.Reset();
 
+	int32 NumTransforms = 0;
+	int32 NumDirtyStates = 0;
+	int32 NumCurves = 0;
+
+	const bool bAllocateStoragePerElement = Ar.CustomVer(FControlRigObjectVersion::GUID) < FControlRigObjectVersion::RigHierarchyIndirectElementStorage;
 	for(int32 ElementIndex = 0; ElementIndex < ElementCount; ElementIndex++)
 	{
 		FRigElementKey Key;
@@ -296,10 +304,83 @@ void URigHierarchy::Load(FArchive& Ar)
 		Element->Index = Elements.Add(Element);
 		ElementsPerType[RigElementTypeToFlatIndex(Key.Type)].Add(Element);
 		IndexLookup.Add(Key, Element->Index);
-		
-		Element->Load(Ar, FRigBaseElement::StaticData);
+
+		if(bAllocateStoragePerElement)
+		{
+			AllocateDefaultElementStorage(Element, false);
+			Element->Load(Ar, FRigBaseElement::StaticData);
+		}
+		else
+		{
+			NumTransforms += Element->GetNumTransforms();
+			NumDirtyStates += Element->GetNumTransforms();
+			NumCurves += Element->GetNumCurves();
+		}
 	}
 
+	if(bAllocateStoragePerElement)
+	{
+		// update all storage pointers now that we've created
+		// all elements.
+		UpdateElementStorage();
+	}
+	else // if we can allocate the storage as one big buffer...
+	{
+		const TArray<int32, TInlineAllocator<4>> TransformIndices = ElementTransforms.Allocate(NumTransforms, FTransform::Identity);
+		const TArray<int32, TInlineAllocator<4>> DirtyStateIndices = ElementDirtyStates.Allocate(NumDirtyStates, false);
+		const TArray<int32, TInlineAllocator<4>> CurveIndices = ElementCurves.Allocate(NumCurves, 0.f);
+		int32 UsedTransformIndex = 0;
+		int32 UsedDirtyStateIndex = 0;
+		int32 UsedCurveIndex = 0;
+
+		ElementTransforms.Shrink();
+		ElementDirtyStates.Shrink();
+		ElementCurves.Shrink();
+
+		for(int32 ElementIndex = 0; ElementIndex < ElementCount; ElementIndex++)
+		{
+			FRigBaseElement* Element = Elements[ElementIndex];
+			
+			if(FRigTransformElement* TransformElement = Cast<FRigTransformElement>(Element))
+			{
+				TransformElement->PoseStorage.Initial.Local.Index = TransformIndices[UsedTransformIndex++];
+				TransformElement->PoseStorage.Initial.Global.Index = TransformIndices[UsedTransformIndex++];
+				TransformElement->PoseStorage.Current.Local.Index = TransformIndices[UsedTransformIndex++];
+				TransformElement->PoseStorage.Current.Global.Index = TransformIndices[UsedTransformIndex++];
+				TransformElement->PoseDirtyState.Initial.Local.Index = DirtyStateIndices[UsedDirtyStateIndex++];
+				TransformElement->PoseDirtyState.Initial.Global.Index = DirtyStateIndices[UsedDirtyStateIndex++];
+				TransformElement->PoseDirtyState.Current.Local.Index = DirtyStateIndices[UsedDirtyStateIndex++];
+				TransformElement->PoseDirtyState.Current.Global.Index = DirtyStateIndices[UsedDirtyStateIndex++];
+
+				if(FRigControlElement* ControlElement = Cast<FRigControlElement>(TransformElement))
+				{
+					ControlElement->OffsetStorage.Initial.Local.Index = TransformIndices[UsedTransformIndex++];
+					ControlElement->OffsetStorage.Initial.Global.Index = TransformIndices[UsedTransformIndex++];
+					ControlElement->OffsetStorage.Current.Local.Index = TransformIndices[UsedTransformIndex++];
+					ControlElement->OffsetStorage.Current.Global.Index = TransformIndices[UsedTransformIndex++];
+					ControlElement->OffsetDirtyState.Initial.Local.Index = DirtyStateIndices[UsedDirtyStateIndex++];
+					ControlElement->OffsetDirtyState.Initial.Global.Index = DirtyStateIndices[UsedDirtyStateIndex++];
+					ControlElement->OffsetDirtyState.Current.Local.Index = DirtyStateIndices[UsedDirtyStateIndex++];
+					ControlElement->OffsetDirtyState.Current.Global.Index = DirtyStateIndices[UsedDirtyStateIndex++];
+					ControlElement->ShapeStorage.Initial.Local.Index = TransformIndices[UsedTransformIndex++];
+					ControlElement->ShapeStorage.Initial.Global.Index = TransformIndices[UsedTransformIndex++];
+					ControlElement->ShapeStorage.Current.Local.Index = TransformIndices[UsedTransformIndex++];
+					ControlElement->ShapeStorage.Current.Global.Index = TransformIndices[UsedTransformIndex++];
+					ControlElement->ShapeDirtyState.Initial.Local.Index = DirtyStateIndices[UsedDirtyStateIndex++];
+					ControlElement->ShapeDirtyState.Initial.Global.Index = DirtyStateIndices[UsedDirtyStateIndex++];
+					ControlElement->ShapeDirtyState.Current.Local.Index = DirtyStateIndices[UsedDirtyStateIndex++];
+					ControlElement->ShapeDirtyState.Current.Global.Index = DirtyStateIndices[UsedDirtyStateIndex++];
+				}
+			}
+			else if(FRigCurveElement* CurveElement = Cast<FRigCurveElement>(Element))
+			{
+				CurveElement->StorageIndex = CurveIndices[UsedCurveIndex++];
+			}
+
+			Element->LinkStorage(ElementTransforms.Storage, ElementDirtyStates.Storage, ElementCurves.Storage);
+			Element->Load(Ar, FRigBaseElement::StaticData);
+		}
+	}
 	IncrementTopologyVersion();
 
 	for(int32 ElementIndex = 0; ElementIndex < ElementCount; ElementIndex++)
@@ -351,7 +432,6 @@ void URigHierarchy::Load(FArchive& Ar)
 	if (Ar.CustomVer(FControlRigObjectVersion::GUID) >= FControlRigObjectVersion::RigHierarchyStoresElementMetadata)
 	{
 		ElementMetadata.Reset();
-		ElementMetadataFreeList.Reset();
 		TMap<FRigElementKey, FMetadataStorage> LoadedElementMetadata;
 		
 		Ar << LoadedElementMetadata;
@@ -359,7 +439,7 @@ void URigHierarchy::Load(FArchive& Ar)
 		{
 			FRigBaseElement* Element = Find(Entry.Key);
 			Element->MetadataStorageIndex = ElementMetadata.Num();
-			ElementMetadata.Add(MoveTemp(Entry.Value));
+			ElementMetadata.Storage.Add(MoveTemp(Entry.Value));
 		}
 	}
 }
@@ -418,15 +498,17 @@ void URigHierarchy::Reset_Impl(bool bResetElements)
 		}
 		IndexLookup.Reset();
 
-		for (FMetadataStorage& MetadataStorage: ElementMetadata)
+		ElementTransforms.Reset();
+		ElementDirtyStates.Reset();
+		ElementCurves.Reset();
+
+		ElementMetadata.Reset([](int32 InIndex, FMetadataStorage& MetadataStorage)
 		{
 			for (TTuple<FName, FRigBaseMetadata*>& Item: MetadataStorage.MetadataMap)
 			{
 				FRigBaseMetadata::DestroyMetadata(&Item.Value);
 			} 
-		}
-		ElementMetadata.Reset();
-		ElementMetadataFreeList.Reset();
+		});
 	}
 
 	ResetPoseHash = INDEX_NONE;
@@ -544,7 +626,11 @@ void URigHierarchy::CopyHierarchy(URigHierarchy* InHierarchy)
 			sizeof(FRigReferenceElement),
 			sizeof(FRigConnectorElement),
 			sizeof(FRigSocketElement),
-		}; 
+		};
+
+		int32 NumTransforms = 0;
+		int32 NumDirtyStates = 0;
+		int32 NumCurves = 0;
 
 		if(bReallocateElements)
 		{
@@ -587,7 +673,11 @@ void URigHierarchy::CopyHierarchy(URigHierarchy* InHierarchy)
 				FRigBaseElement* Target = reinterpret_cast<FRigBaseElement*>(&NewElementsPerType[ElementTypeIndex][StructureSize * SubIndex]);
 
 				Target->InitializeFrom(Source);
-			
+
+				NumTransforms += Target->GetNumTransforms();
+				NumDirtyStates += Target->GetNumTransforms();
+				NumCurves += Target->GetNumCurves();
+
 				Target->SubIndex = SubIndex;
 				Target->Index = Elements.Add(Target);
 
@@ -628,18 +718,74 @@ void URigHierarchy::CopyHierarchy(URigHierarchy* InHierarchy)
 				check(Target->Key.Type == Source->Key.Type);
 				Target->InitializeFrom(Source);
 
+				NumTransforms += Target->GetNumTransforms();
+				NumDirtyStates += Target->GetNumTransforms();
+				NumCurves += Target->GetNumCurves();
+
 				IncrementPoseVersion(Index);
 			}
 
 			IndexLookup = InHierarchy->IndexLookup;
 		}
 
+		ElementTransforms.Reset();
+		ElementDirtyStates.Reset();
+		ElementCurves.Reset();
+		
+		const TArray<int32, TInlineAllocator<4>> TransformIndices = ElementTransforms.Allocate(NumTransforms, FTransform::Identity);
+		const TArray<int32, TInlineAllocator<4>> DirtyStateIndices = ElementDirtyStates.Allocate(NumDirtyStates, false);
+		const TArray<int32, TInlineAllocator<4>> CurveIndices = ElementCurves.Allocate(NumCurves, 0.f);
+		int32 UsedTransformIndex = 0;
+		int32 UsedDirtyStateIndex = 0;
+		int32 UsedCurveIndex = 0;
+
+		ElementTransforms.Shrink();
+		ElementDirtyStates.Shrink();
+		ElementCurves.Shrink();
+
 		// Copy all the element subclass data and all elements' metadata over.
 		for(int32 Index = 0; Index < InHierarchy->Num(); Index++)
 		{
 			const FRigBaseElement* Source = InHierarchy->Get(Index);
 			FRigBaseElement* Target = Elements[Index];
 
+			if(FRigTransformElement* TransformElement = Cast<FRigTransformElement>(Target))
+			{
+				TransformElement->PoseStorage.Initial.Local.Index = TransformIndices[UsedTransformIndex++];
+				TransformElement->PoseStorage.Initial.Global.Index = TransformIndices[UsedTransformIndex++];
+				TransformElement->PoseStorage.Current.Local.Index = TransformIndices[UsedTransformIndex++];
+				TransformElement->PoseStorage.Current.Global.Index = TransformIndices[UsedTransformIndex++];
+				TransformElement->PoseDirtyState.Initial.Local.Index = DirtyStateIndices[UsedDirtyStateIndex++];
+				TransformElement->PoseDirtyState.Initial.Global.Index = DirtyStateIndices[UsedDirtyStateIndex++];
+				TransformElement->PoseDirtyState.Current.Local.Index = DirtyStateIndices[UsedDirtyStateIndex++];
+				TransformElement->PoseDirtyState.Current.Global.Index = DirtyStateIndices[UsedDirtyStateIndex++];
+
+				if(FRigControlElement* ControlElement = Cast<FRigControlElement>(TransformElement))
+				{
+					ControlElement->OffsetStorage.Initial.Local.Index = TransformIndices[UsedTransformIndex++];
+					ControlElement->OffsetStorage.Initial.Global.Index = TransformIndices[UsedTransformIndex++];
+					ControlElement->OffsetStorage.Current.Local.Index = TransformIndices[UsedTransformIndex++];
+					ControlElement->OffsetStorage.Current.Global.Index = TransformIndices[UsedTransformIndex++];
+					ControlElement->OffsetDirtyState.Initial.Local.Index = DirtyStateIndices[UsedDirtyStateIndex++];
+					ControlElement->OffsetDirtyState.Initial.Global.Index = DirtyStateIndices[UsedDirtyStateIndex++];
+					ControlElement->OffsetDirtyState.Current.Local.Index = DirtyStateIndices[UsedDirtyStateIndex++];
+					ControlElement->OffsetDirtyState.Current.Global.Index = DirtyStateIndices[UsedDirtyStateIndex++];
+					ControlElement->ShapeStorage.Initial.Local.Index = TransformIndices[UsedTransformIndex++];
+					ControlElement->ShapeStorage.Initial.Global.Index = TransformIndices[UsedTransformIndex++];
+					ControlElement->ShapeStorage.Current.Local.Index = TransformIndices[UsedTransformIndex++];
+					ControlElement->ShapeStorage.Current.Global.Index = TransformIndices[UsedTransformIndex++];
+					ControlElement->ShapeDirtyState.Initial.Local.Index = DirtyStateIndices[UsedDirtyStateIndex++];
+					ControlElement->ShapeDirtyState.Initial.Global.Index = DirtyStateIndices[UsedDirtyStateIndex++];
+					ControlElement->ShapeDirtyState.Current.Local.Index = DirtyStateIndices[UsedDirtyStateIndex++];
+					ControlElement->ShapeDirtyState.Current.Global.Index = DirtyStateIndices[UsedDirtyStateIndex++];
+				}
+			}
+			else if(FRigCurveElement* CurveElement = Cast<FRigCurveElement>(Target))
+			{
+				CurveElement->StorageIndex = CurveIndices[UsedCurveIndex++];
+			}
+			Target->LinkStorage(ElementTransforms.Storage, ElementDirtyStates.Storage, ElementCurves.Storage);
+
 			Target->CopyFrom(Source);
 
 			CopyAllMetadataFromElement(Target, Source);
@@ -829,18 +975,20 @@ void URigHierarchy::CopyPose(URigHierarchy* InHierarchy, bool bCurrent, bool bIn
 						{
 							if(bCurrent)
 							{
-								ControlElementA->Offset.Set(ERigTransformType::CurrentGlobal, InHierarchy->GetControlOffsetTransform(ControlElementB, ERigTransformType::CurrentGlobal));
-								ControlElementA->Offset.MarkDirty(ERigTransformType::CurrentLocal);
-								ControlElementA->Pose.MarkDirty(ERigTransformType::CurrentGlobal);
-								ControlElementA->Shape.MarkDirty(ERigTransformType::CurrentGlobal);
+								ControlElementA->GetOffsetTransform().Set(ERigTransformType::CurrentGlobal, InHierarchy->GetControlOffsetTransform(ControlElementB, ERigTransformType::CurrentGlobal));
+								ControlElementA->GetOffsetDirtyState().MarkClean(ERigTransformType::CurrentGlobal);
+								ControlElementA->GetOffsetDirtyState().MarkDirty(ERigTransformType::CurrentLocal);
+								ControlElementA->GetDirtyState().MarkDirty(ERigTransformType::CurrentGlobal);
+								ControlElementA->GetShapeDirtyState().MarkDirty(ERigTransformType::CurrentGlobal);
 								IncrementPoseVersion(ControlElementA->Index);
 							}
 							if(bInitial)
 							{
-								ControlElementA->Offset.Set(ERigTransformType::InitialGlobal, InHierarchy->GetControlOffsetTransform(ControlElementB, ERigTransformType::InitialGlobal));
-								ControlElementA->Offset.MarkDirty(ERigTransformType::InitialLocal);
-								ControlElementA->Pose.MarkDirty(ERigTransformType::InitialGlobal);
-								ControlElementA->Shape.MarkDirty(ERigTransformType::InitialGlobal);
+								ControlElementA->GetOffsetTransform().Set(ERigTransformType::InitialGlobal, InHierarchy->GetControlOffsetTransform(ControlElementB, ERigTransformType::InitialGlobal));
+								ControlElementA->GetOffsetDirtyState().MarkClean(ERigTransformType::InitialGlobal);
+								ControlElementA->GetOffsetDirtyState().MarkDirty(ERigTransformType::InitialLocal);
+								ControlElementA->GetDirtyState().MarkDirty(ERigTransformType::InitialGlobal);
+								ControlElementA->GetShapeDirtyState().MarkDirty(ERigTransformType::InitialGlobal);
 								IncrementPoseVersion(ControlElementA->Index);
 							}
 						}
@@ -848,14 +996,16 @@ void URigHierarchy::CopyPose(URigHierarchy* InHierarchy, bool bCurrent, bool bIn
 						{
 							if(bCurrent)
 							{
-								MultiParentElementA->Pose.Set(ERigTransformType::CurrentGlobal, InHierarchy->GetTransform(MultiParentElementB, ERigTransformType::CurrentGlobal));
-								MultiParentElementA->Pose.MarkDirty(ERigTransformType::CurrentLocal);
+								MultiParentElementA->GetTransform().Set(ERigTransformType::CurrentGlobal, InHierarchy->GetTransform(MultiParentElementB, ERigTransformType::CurrentGlobal));
+								MultiParentElementA->GetDirtyState().MarkClean(ERigTransformType::CurrentGlobal);
+								MultiParentElementA->GetDirtyState().MarkDirty(ERigTransformType::CurrentLocal);
 								IncrementPoseVersion(MultiParentElementA->Index);
 							}
 							if(bInitial)
 							{
-								MultiParentElementA->Pose.Set(ERigTransformType::InitialGlobal, InHierarchy->GetTransform(MultiParentElementB, ERigTransformType::InitialGlobal));
-								MultiParentElementA->Pose.MarkDirty(ERigTransformType::InitialLocal);
+								MultiParentElementA->GetTransform().Set(ERigTransformType::InitialGlobal, InHierarchy->GetTransform(MultiParentElementB, ERigTransformType::InitialGlobal));
+								MultiParentElementA->GetDirtyState().MarkClean(ERigTransformType::InitialGlobal);
+								MultiParentElementA->GetDirtyState().MarkDirty(ERigTransformType::InitialLocal);
 								IncrementPoseVersion(MultiParentElementA->Index);
 							}
 						}
@@ -975,17 +1125,17 @@ void URigHierarchy::ResetPoseToInitial(ERigElementType InTypeFilter)
 		{
 			if(FRigTransformElement* TransformElement = Get<FRigTransformElement>(ElementIndex))
 			{
-				if(TransformElement->Pose.IsDirty(ERigTransformType::CurrentGlobal))
+				if(TransformElement->GetDirtyState().IsDirty(ERigTransformType::CurrentGlobal))
 				{
 					continue;
 				}
 				
-				TransformElement->Pose.MarkDirty(ERigTransformType::CurrentGlobal);
+				TransformElement->GetDirtyState().MarkDirty(ERigTransformType::CurrentGlobal);
 
 				if(FRigControlElement* ControlElement = Cast<FRigControlElement>(TransformElement))
 				{
-					ControlElement->Offset.MarkDirty(ERigTransformType::CurrentGlobal);
-					ControlElement->Shape.MarkDirty(ERigTransformType::CurrentGlobal);
+					ControlElement->GetOffsetDirtyState().MarkDirty(ERigTransformType::CurrentGlobal);
+					ControlElement->GetShapeDirtyState().MarkDirty(ERigTransformType::CurrentGlobal);
 				}
 
 				PropagateDirtyFlags(TransformElement, false, true, false, true);
@@ -1014,14 +1164,17 @@ void URigHierarchy::ResetPoseToInitial(ERigElementType InTypeFilter)
 
 		if(FRigControlElement* ControlElement = Cast<FRigControlElement>(Elements[ElementIndex]))
 		{
-			ControlElement->Offset.Current = ControlElement->Offset.Initial;
-			ControlElement->Shape.Current = ControlElement->Shape.Initial;
+			ControlElement->GetOffsetTransform().Current = ControlElement->GetOffsetTransform().Initial;
+			ControlElement->GetOffsetDirtyState().Current = ControlElement->GetOffsetDirtyState().Initial;
+			ControlElement->GetShapeTransform().Current = ControlElement->GetShapeTransform().Initial;
+			ControlElement->GetShapeDirtyState().Current = ControlElement->GetShapeDirtyState().Initial;
 			ControlElement->PreferredEulerAngles.Current = ControlElement->PreferredEulerAngles.Initial;
 		}
 
 		if(FRigTransformElement* TransformElement = Cast<FRigTransformElement>(Elements[ElementIndex]))
 		{
-			TransformElement->Pose.Current = TransformElement->Pose.Initial;
+			TransformElement->GetTransform().Current = TransformElement->GetTransform().Initial;
+			TransformElement->GetDirtyState().Current = TransformElement->GetDirtyState().Initial;
 		}
 	}
 	
@@ -2323,7 +2476,7 @@ bool URigHierarchy::SetParentWeight(FRigBaseElement* InChild, int32 InParentInde
 					GetControlOffsetTransform(ControlElement, LocalType);
 				}
 				GetTransform(MultiParentElement, LocalType);
-				MultiParentElement->Pose.MarkDirty(GlobalType);
+				MultiParentElement->GetDirtyState().MarkDirty(GlobalType);
 			}
 			else
 			{
@@ -2333,14 +2486,14 @@ bool URigHierarchy::SetParentWeight(FRigBaseElement* InChild, int32 InParentInde
 					GetControlOffsetTransform(ControlElement, GlobalType);
 				}
 				GetTransform(MultiParentElement, GlobalType);
-				MultiParentElement->Pose.MarkDirty(LocalType);
+				MultiParentElement->GetDirtyState().MarkDirty(LocalType);
 			}
 
 			TargetWeight = InWeight;
 
 			if(FRigControlElement* ControlElement = Cast<FRigControlElement>(MultiParentElement))
 			{
-				ControlElement->Offset.MarkDirty(GlobalType);
+				ControlElement->GetOffsetDirtyState().MarkDirty(GlobalType);
 			}
 
 			PropagateDirtyFlags(MultiParentElement, ERigTransformType::IsInitial(LocalType), bAffectChildren);
@@ -2450,12 +2603,12 @@ bool URigHierarchy::SetParentWeightArray(FRigBaseElement* InChild,  const TArray
 			if(bAffectChildren)
 			{
 				GetTransform(MultiParentElement, LocalType);
-				MultiParentElement->Pose.MarkDirty(GlobalType);
+				MultiParentElement->GetDirtyState().MarkDirty(GlobalType);
 			}
 			else
 			{
 				GetTransform(MultiParentElement, GlobalType);
-				MultiParentElement->Pose.MarkDirty(LocalType);
+				MultiParentElement->GetDirtyState().MarkDirty(LocalType);
 			}
 
 			for(int32 WeightIndex=0; WeightIndex < InWeights.Num(); WeightIndex++)
@@ -2472,8 +2625,8 @@ bool URigHierarchy::SetParentWeightArray(FRigBaseElement* InChild,  const TArray
 
 			if(FRigControlElement* ControlElement = Cast<FRigControlElement>(MultiParentElement))
 			{
-				ControlElement->Offset.MarkDirty(GlobalType);
-				ControlElement->Shape.MarkDirty(GlobalType);
+				ControlElement->GetOffsetDirtyState().MarkDirty(GlobalType);
+				ControlElement->GetShapeDirtyState().MarkDirty(GlobalType);
 			}
 
 			PropagateDirtyFlags(MultiParentElement, ERigTransformType::IsInitial(LocalType), bAffectChildren);
@@ -3524,11 +3677,11 @@ FTransform URigHierarchy::GetTransform(FRigTransformElement* InTransformElement,
 	
 #endif
 	
-	if(InTransformElement->Pose.IsDirty(InTransformType))
+	if(InTransformElement->GetDirtyState().IsDirty(InTransformType))
 	{
 		const ERigTransformType::Type OpposedType = SwapLocalAndGlobal(InTransformType);
 		const ERigTransformType::Type GlobalType = MakeGlobal(InTransformType);
-		ensure(!InTransformElement->Pose.IsDirty(OpposedType));
+		ensure(!InTransformElement->GetDirtyState().IsDirty(OpposedType));
 
 		FTransform ParentTransform;
 		if(IsLocal(InTransformType))
@@ -3547,17 +3700,19 @@ FTransform URigHierarchy::GetTransform(FRigTransformElement* InTransformElement,
 					const FVector ParentScale = ParentTransform.GetScale3D();
 					if(FMath::IsNearlyZero(ParentScale.X) || FMath::IsNearlyZero(ParentScale.Y) || FMath::IsNearlyZero(ParentScale.Z))
 					{
-						Transform.SetTranslation(InTransformElement->Pose.Get(InTransformType).GetTranslation());
-						Transform.SetScale3D(InTransformElement->Pose.Get(InTransformType).GetScale3D());
+						const FTransform& InputTransform = InTransformElement->GetTransform().Get(InTransformType); 
+						Transform.SetTranslation(InputTransform.GetTranslation());
+						Transform.SetScale3D(InputTransform.GetScale3D());
 					}
 				}
 			};
 			
 			if(FRigControlElement* ControlElement = Cast<FRigControlElement>(InTransformElement))
 			{
-				FTransform NewTransform = ComputeLocalControlValue(ControlElement, ControlElement->Pose.Get(OpposedType), GlobalType);
+				FTransform NewTransform = ComputeLocalControlValue(ControlElement, ControlElement->GetTransform().Get(OpposedType), GlobalType);
 				CompensateZeroScale(NewTransform);
-				ControlElement->Pose.Set(InTransformType, NewTransform);
+				ControlElement->GetTransform().Set(InTransformType, NewTransform);
+				ControlElement->GetDirtyState().MarkClean(InTransformType);
 				/** from mikez we do not want geting a pose to set these preferred angles
 				switch(ControlElement->Settings.ControlType)
 				{
@@ -3580,21 +3735,23 @@ FTransform URigHierarchy::GetTransform(FRigTransformElement* InTransformElement,
 			{
 				// this is done for nulls and any element that can have more than one parent which 
 				// is not a control
-				const FTransform& GlobalTransform = MultiParentElement->Pose.Get(GlobalType);
+				const FTransform& GlobalTransform = MultiParentElement->GetTransform().Get(GlobalType);
 				FTransform LocalTransform = InverseSolveParentConstraints(
 					GlobalTransform, 
 					MultiParentElement->ParentConstraints, GlobalType, FTransform::Identity);
 				CompensateZeroScale(LocalTransform);
-				MultiParentElement->Pose.Set(InTransformType, LocalTransform);
+				MultiParentElement->GetTransform().Set(InTransformType, LocalTransform);
+				MultiParentElement->GetDirtyState().MarkClean(InTransformType);
 			}
 			else
 			{
 				ParentTransform = GetParentTransform(InTransformElement, GlobalType);
 
-				FTransform NewTransform = InTransformElement->Pose.Get(OpposedType).GetRelativeTransform(ParentTransform);
+				FTransform NewTransform = InTransformElement->GetTransform().Get(OpposedType).GetRelativeTransform(ParentTransform);
 				NewTransform.NormalizeRotation();
 				CompensateZeroScale(NewTransform);
-				InTransformElement->Pose.Set(InTransformType, NewTransform);
+				InTransformElement->GetTransform().Set(InTransformType, NewTransform);
+				InTransformElement->GetDirtyState().MarkClean(InTransformType);
 			}
 		}
 		else
@@ -3607,8 +3764,9 @@ FTransform URigHierarchy::GetTransform(FRigTransformElement* InTransformElement,
 				const FTransform NewTransform = SolveParentConstraints(
 					ControlElement->ParentConstraints, InTransformType,
 					GetControlOffsetTransform(ControlElement, OpposedType), true,
-					ControlElement->Pose.Get(OpposedType), true);
-				ControlElement->Pose.Set(InTransformType, NewTransform);
+					ControlElement->GetTransform().Get(OpposedType), true);
+				ControlElement->GetTransform().Set(InTransformType, NewTransform);
+				ControlElement->GetDirtyState().MarkClean(InTransformType);
 			}
 			else if(FRigMultiParentElement* MultiParentElement = Cast<FRigMultiParentElement>(InTransformElement))
 			{
@@ -3617,22 +3775,24 @@ FTransform URigHierarchy::GetTransform(FRigTransformElement* InTransformElement,
 				const FTransform NewTransform = SolveParentConstraints(
 					MultiParentElement->ParentConstraints, InTransformType,
 					FTransform::Identity, false,
-					MultiParentElement->Pose.Get(OpposedType), true);
-				MultiParentElement->Pose.Set(InTransformType, NewTransform);
+					MultiParentElement->GetTransform().Get(OpposedType), true);
+				MultiParentElement->GetTransform().Set(InTransformType, NewTransform);
+				MultiParentElement->GetDirtyState().MarkClean(InTransformType);
 			}
 			else
 			{
 				ParentTransform = GetParentTransform(InTransformElement, GlobalType);
 
-				FTransform NewTransform = InTransformElement->Pose.Get(OpposedType) * ParentTransform;
+				FTransform NewTransform = InTransformElement->GetTransform().Get(OpposedType) * ParentTransform;
 				NewTransform.NormalizeRotation();
-				InTransformElement->Pose.Set(InTransformType, NewTransform);
+				InTransformElement->GetTransform().Set(InTransformType, NewTransform);
+				InTransformElement->GetDirtyState().MarkClean(InTransformType);
 			}
 		}
 
 		EnsureCacheValidity();
 	}
-	return InTransformElement->Pose.Get(InTransformType);
+	return InTransformElement->GetTransform().Get(InTransformType);
 }
 
 void URigHierarchy::SetTransform(FRigTransformElement* InTransformElement, const FTransform& InTransform, const ERigTransformType::Type InTransformType, bool bAffectChildren, bool bSetupUndo, bool bForce, bool bPrintPythonCommands)
@@ -3714,9 +3874,9 @@ void URigHierarchy::SetTransform(FRigTransformElement* InTransformElement, const
 	
 #endif
 
-	if(!InTransformElement->Pose.IsDirty(InTransformType))
+	if(!InTransformElement->GetDirtyState().IsDirty(InTransformType))
 	{
-		const FTransform PreviousTransform = InTransformElement->Pose.Get(InTransformType);
+		const FTransform PreviousTransform = InTransformElement->GetTransform().Get(InTransformType);
 		if(!bForce && FRigComputedTransform::Equals(PreviousTransform, InTransform))
 		{
 			return;
@@ -3727,13 +3887,14 @@ void URigHierarchy::SetTransform(FRigTransformElement* InTransformElement, const
 	PropagateDirtyFlags(InTransformElement, ERigTransformType::IsInitial(InTransformType), bAffectChildren);
 
 	const ERigTransformType::Type OpposedType = SwapLocalAndGlobal(InTransformType);
-	InTransformElement->Pose.Set(InTransformType, InTransform);
-	InTransformElement->Pose.MarkDirty(OpposedType);
+	InTransformElement->GetTransform().Set(InTransformType, InTransform);
+	InTransformElement->GetDirtyState().MarkClean(InTransformType);
+	InTransformElement->GetDirtyState().MarkDirty(OpposedType);
 	IncrementPoseVersion(InTransformElement->Index);
 
 	if(FRigControlElement* ControlElement = Cast<FRigControlElement>(InTransformElement))
 	{
-		ControlElement->Shape.MarkDirty(MakeGlobal(InTransformType));
+		ControlElement->GetShapeDirtyState().MarkDirty(MakeGlobal(InTransformType));
 
 		if(bUsePreferredEulerAngles && ERigTransformType::IsLocal(InTransformType))
 		{
@@ -3752,7 +3913,7 @@ void URigHierarchy::SetTransform(FRigTransformElement* InTransformElement, const
 			ERigTransformStackEntryType::TransformPose,
 			InTransformType,
 			PreviousTransform,
-			InTransformElement->Pose.Get(InTransformType),
+			InTransformElement->GetTransform().Get(InTransformType),
 			bAffectChildren,
 			bSetupUndo);
 	}
@@ -3853,19 +4014,20 @@ FTransform URigHierarchy::GetControlOffsetTransform(FRigControlElement* InContro
 	
 #endif
 
-	if(InControlElement->Offset.IsDirty(InTransformType))
+	if(InControlElement->GetOffsetDirtyState().IsDirty(InTransformType))
 	{
 		const ERigTransformType::Type OpposedType = SwapLocalAndGlobal(InTransformType);
 		const ERigTransformType::Type GlobalType = MakeGlobal(InTransformType);
-		ensure(!InControlElement->Offset.IsDirty(OpposedType));
+		ensure(!InControlElement->GetOffsetDirtyState().IsDirty(OpposedType));
 
 		if(IsLocal(InTransformType))
 		{
-			const FTransform& GlobalTransform = InControlElement->Offset.Get(GlobalType);
+			const FTransform& GlobalTransform = InControlElement->GetOffsetTransform().Get(GlobalType);
 			const FTransform LocalTransform = InverseSolveParentConstraints(
 				GlobalTransform, 
 				InControlElement->ParentConstraints, GlobalType, FTransform::Identity);
-			InControlElement->Offset.Set(InTransformType, LocalTransform);
+			InControlElement->GetOffsetTransform().Set(InTransformType, LocalTransform);
+			InControlElement->GetOffsetDirtyState().MarkClean(InTransformType);
 
 			if(bEnableCacheValidityCheck)
 			{
@@ -3885,12 +4047,13 @@ FTransform URigHierarchy::GetControlOffsetTransform(FRigControlElement* InContro
 		}
 		else
 		{
-			const FTransform& LocalTransform = InControlElement->Offset.Get(OpposedType); 
+			const FTransform& LocalTransform = InControlElement->GetOffsetTransform().Get(OpposedType); 
 			const FTransform GlobalTransform = SolveParentConstraints(
 				InControlElement->ParentConstraints, InTransformType,
 				LocalTransform, true,
 				FTransform::Identity, false);
-			InControlElement->Offset.Set(InTransformType, GlobalTransform);
+			InControlElement->GetOffsetTransform().Set(InTransformType, GlobalTransform);
+			InControlElement->GetOffsetDirtyState().MarkClean(InTransformType);
 
 			if(bEnableCacheValidityCheck)
 			{
@@ -3910,7 +4073,7 @@ FTransform URigHierarchy::GetControlOffsetTransform(FRigControlElement* InContro
 
 		EnsureCacheValidity();
 	}
-	return InControlElement->Offset.Get(InTransformType);
+	return InControlElement->GetOffsetTransform().Get(InTransformType);
 }
 
 void URigHierarchy::SetControlOffsetTransform(FRigControlElement* InControlElement, const FTransform& InTransform,
@@ -3941,9 +4104,9 @@ void URigHierarchy::SetControlOffsetTransform(FRigControlElement* InControlEleme
 	
 #endif
 
-	if(!InControlElement->Offset.IsDirty(InTransformType))
+	if(!InControlElement->GetOffsetDirtyState().IsDirty(InTransformType))
 	{
-		const FTransform PreviousTransform = InControlElement->Offset.Get(InTransformType);
+		const FTransform PreviousTransform = InControlElement->GetOffsetTransform().Get(InTransformType);
 		if(!bForce && FRigComputedTransform::Equals(PreviousTransform, InTransform))
 		{
 			return;
@@ -3954,12 +4117,13 @@ void URigHierarchy::SetControlOffsetTransform(FRigControlElement* InControlEleme
 	PropagateDirtyFlags(InControlElement, ERigTransformType::IsInitial(InTransformType), bAffectChildren);
 
 	GetTransform(InControlElement, MakeLocal(InTransformType));
-	InControlElement->Pose.MarkDirty(MakeGlobal(InTransformType));
+	InControlElement->GetDirtyState().MarkDirty(MakeGlobal(InTransformType));
 
 	const ERigTransformType::Type OpposedType = SwapLocalAndGlobal(InTransformType);
-	InControlElement->Offset.Set(InTransformType, InTransform);
-	InControlElement->Offset.MarkDirty(OpposedType);
-	InControlElement->Shape.MarkDirty(MakeGlobal(InTransformType));
+	InControlElement->GetOffsetTransform().Set(InTransformType, InTransform);
+	InControlElement->GetOffsetDirtyState().MarkClean(InTransformType);
+	InControlElement->GetOffsetDirtyState().MarkDirty(OpposedType);
+	InControlElement->GetShapeDirtyState().MarkDirty(MakeGlobal(InTransformType));
 
 	EnsureCacheValidity();
 
@@ -3980,7 +4144,7 @@ void URigHierarchy::SetControlOffsetTransform(FRigControlElement* InControlEleme
             ERigTransformStackEntryType::ControlOffset,
             InTransformType,
             PreviousTransform,
-            InControlElement->Offset.Get(InTransformType),
+            InControlElement->GetOffsetTransform().Get(InTransformType),
             bAffectChildren,
             bSetupUndo);
 	}
@@ -4038,29 +4202,31 @@ FTransform URigHierarchy::GetControlShapeTransform(FRigControlElement* InControl
 		return FTransform::Identity;
 	}
 	
-	if(InControlElement->Shape.IsDirty(InTransformType))
+	if(InControlElement->GetShapeDirtyState().IsDirty(InTransformType))
 	{
 		const ERigTransformType::Type OpposedType = SwapLocalAndGlobal(InTransformType);
 		const ERigTransformType::Type GlobalType = MakeGlobal(InTransformType);
-		ensure(!InControlElement->Shape.IsDirty(OpposedType));
+		ensure(!InControlElement->GetShapeDirtyState().IsDirty(OpposedType));
 
 		const FTransform ParentTransform = GetTransform(InControlElement, GlobalType);
 		if(IsLocal(InTransformType))
 		{
-			FTransform LocalTransform = InControlElement->Shape.Get(OpposedType).GetRelativeTransform(ParentTransform);
+			FTransform LocalTransform = InControlElement->GetShapeTransform().Get(OpposedType).GetRelativeTransform(ParentTransform);
 			LocalTransform.NormalizeRotation();
-			InControlElement->Shape.Set(InTransformType, LocalTransform);
+			InControlElement->GetShapeTransform().Set(InTransformType, LocalTransform);
+			InControlElement->GetShapeDirtyState().MarkClean(InTransformType);
 		}
 		else
 		{
-			FTransform GlobalTransform = InControlElement->Shape.Get(OpposedType) * ParentTransform;
+			FTransform GlobalTransform = InControlElement->GetShapeTransform().Get(OpposedType) * ParentTransform;
 			GlobalTransform.NormalizeRotation();
-			InControlElement->Shape.Set(InTransformType, GlobalTransform);
+			InControlElement->GetShapeTransform().Set(InTransformType, GlobalTransform);
+			InControlElement->GetShapeDirtyState().MarkClean(InTransformType);
 		}
 
 		EnsureCacheValidity();
 	}
-	return InControlElement->Shape.Get(InTransformType);
+	return InControlElement->GetShapeTransform().Get(InTransformType);
 }
 
 void URigHierarchy::SetControlShapeTransform(FRigControlElement* InControlElement, const FTransform& InTransform,
@@ -4072,9 +4238,9 @@ void URigHierarchy::SetControlShapeTransform(FRigControlElement* InControlElemen
 		return;
 	}
 
-	if(!InControlElement->Shape.IsDirty(InTransformType))
+	if(!InControlElement->GetShapeDirtyState().IsDirty(InTransformType))
 	{
-		const FTransform PreviousTransform = InControlElement->Shape.Get(InTransformType);
+		const FTransform PreviousTransform = InControlElement->GetShapeTransform().Get(InTransformType);
 		if(!bForce && FRigComputedTransform::Equals(PreviousTransform, InTransform))
 		{
 			return;
@@ -4083,8 +4249,9 @@ void URigHierarchy::SetControlShapeTransform(FRigControlElement* InControlElemen
 
 	const FTransform PreviousTransform = GetControlShapeTransform(InControlElement, InTransformType);
 	const ERigTransformType::Type OpposedType = SwapLocalAndGlobal(InTransformType);
-	InControlElement->Shape.Set(InTransformType, InTransform);
-	InControlElement->Shape.MarkDirty(OpposedType);
+	InControlElement->GetShapeTransform().Set(InTransformType, InTransform);
+	InControlElement->GetShapeDirtyState().MarkClean(InTransformType);
+	InControlElement->GetShapeDirtyState().MarkDirty(OpposedType);
 
 	if (IsInitial(InTransformType))
 	{
@@ -4104,7 +4271,7 @@ void URigHierarchy::SetControlShapeTransform(FRigControlElement* InControlElemen
             ERigTransformStackEntryType::ControlShape,
             InTransformType,
             PreviousTransform,
-            InControlElement->Shape.Get(InTransformType),
+            InControlElement->GetShapeTransform().Get(InTransformType),
             false,
             bSetupUndo);
 	}
@@ -4646,7 +4813,7 @@ float URigHierarchy::GetCurveValue(FRigCurveElement* InCurveElement) const
 	{
 		return 0.f;
 	}
-	return InCurveElement->bIsValueSet ? InCurveElement->Value : 0.f;
+	return InCurveElement->bIsValueSet ? InCurveElement->Get() : 0.f;
 }
 
 
@@ -4665,19 +4832,18 @@ void URigHierarchy::SetCurveValue(FRigCurveElement* InCurveElement, float InValu
 	}
 
 	const bool bPreviousIsValueSet = InCurveElement->bIsValueSet; 
-	const float PreviousValue = InCurveElement->Value;
+	const float PreviousValue = InCurveElement->Get();
 	if(!bForce && InCurveElement->bIsValueSet && FMath::IsNearlyZero(PreviousValue - InValue))
 	{
 		return;
 	}
 
-	InCurveElement->bIsValueSet = true;
-	InCurveElement->Value = InValue;
+	InCurveElement->Set(InValue);
 
 #if WITH_EDITOR
 	if(bSetupUndo || IsTracingChanges())
 	{
-		PushCurveToStack(InCurveElement->GetKey(), PreviousValue, InCurveElement->Value, bPreviousIsValueSet, true, bSetupUndo);
+		PushCurveToStack(InCurveElement->GetKey(), PreviousValue, InCurveElement->Get(), bPreviousIsValueSet, true, bSetupUndo);
 	}
 
 	if (!bPropagatingChange)
@@ -4724,7 +4890,7 @@ void URigHierarchy::UnsetCurveValue(FRigCurveElement* InCurveElement, bool bSetu
 #if WITH_EDITOR
 	if(bSetupUndo || IsTracingChanges())
 	{
-		PushCurveToStack(InCurveElement->GetKey(), InCurveElement->Value, InCurveElement->Value, bPreviousIsValueSet, false, bSetupUndo);
+		PushCurveToStack(InCurveElement->GetKey(), InCurveElement->Get(), InCurveElement->Get(), bPreviousIsValueSet, false, bSetupUndo);
 	}
 
 	if (!bPropagatingChange)
@@ -4796,10 +4962,10 @@ bool URigHierarchy::IsDependentOn(FRigBaseElement* InDependent, FRigBaseElement*
 	const int32 DependencyElementIndex = InDependency->GetIndex();
 	const TTuple<int32,int32> CacheKey(DependentElementIndex, DependencyElementIndex);
 
-		if(!ElementDependencyCache.IsValid(GetTopologyVersion()))
-		{
-			ElementDependencyCache.Set(TMap<TTuple<int32, int32>, bool>(), GetTopologyVersion());
-		}
+	if(!ElementDependencyCache.IsValid(GetTopologyVersion()))
+	{
+		ElementDependencyCache.Set(TMap<TTuple<int32, int32>, bool>(), GetTopologyVersion());
+	}
 
 	// we'll only update the caches if we are following edges on the actual topology
 	if(const bool* bCachedResult = ElementDependencyCache.Get().Find(CacheKey))
@@ -5308,6 +5474,7 @@ void URigHierarchy::DestroyElement(FRigBaseElement*& InElement)
 			for(int32 Index=0;Index<Count;Index++)
 			{
 				TGuardValue<const FRigBaseElement*> DestroyGuard(ElementBeingDestroyed, &ExistingElements[Index]);
+				DeallocateElementStorage(&ExistingElements[Index]);
 				ExistingElements[Index].~FRigBoneElement(); 
 			}
 			break;
@@ -5318,6 +5485,7 @@ void URigHierarchy::DestroyElement(FRigBaseElement*& InElement)
 			for(int32 Index=0;Index<Count;Index++)
 			{
 				TGuardValue<const FRigBaseElement*> DestroyGuard(ElementBeingDestroyed, &ExistingElements[Index]);
+				DeallocateElementStorage(&ExistingElements[Index]);
 				ExistingElements[Index].~FRigNullElement(); 
 			}
 			break;
@@ -5328,6 +5496,7 @@ void URigHierarchy::DestroyElement(FRigBaseElement*& InElement)
 			for(int32 Index=0;Index<Count;Index++)
 			{
 				TGuardValue<const FRigBaseElement*> DestroyGuard(ElementBeingDestroyed, &ExistingElements[Index]);
+				DeallocateElementStorage(&ExistingElements[Index]);
 				ExistingElements[Index].~FRigControlElement(); 
 			}
 			break;
@@ -5338,6 +5507,7 @@ void URigHierarchy::DestroyElement(FRigBaseElement*& InElement)
 			for(int32 Index=0;Index<Count;Index++)
 			{
 				TGuardValue<const FRigBaseElement*> DestroyGuard(ElementBeingDestroyed, &ExistingElements[Index]);
+				DeallocateElementStorage(&ExistingElements[Index]);
 				ExistingElements[Index].~FRigCurveElement(); 
 			}
 			break;
@@ -5348,6 +5518,7 @@ void URigHierarchy::DestroyElement(FRigBaseElement*& InElement)
 			for(int32 Index=0;Index<Count;Index++)
 			{
 				TGuardValue<const FRigBaseElement*> DestroyGuard(ElementBeingDestroyed, &ExistingElements[Index]);
+				DeallocateElementStorage(&ExistingElements[Index]);
 				ExistingElements[Index].~FRigPhysicsElement(); 
 			}
 			break;
@@ -5358,6 +5529,7 @@ void URigHierarchy::DestroyElement(FRigBaseElement*& InElement)
 			for(int32 Index=0;Index<Count;Index++)
 			{
 				TGuardValue<const FRigBaseElement*> DestroyGuard(ElementBeingDestroyed, &ExistingElements[Index]);
+				DeallocateElementStorage(&ExistingElements[Index]);
 				ExistingElements[Index].~FRigReferenceElement(); 
 			}
 			break;
@@ -5368,6 +5540,7 @@ void URigHierarchy::DestroyElement(FRigBaseElement*& InElement)
 			for(int32 Index=0;Index<Count;Index++)
 			{
 				TGuardValue<const FRigBaseElement*> DestroyGuard(ElementBeingDestroyed, &ExistingElements[Index]);
+				DeallocateElementStorage(&ExistingElements[Index]);
 				ExistingElements[Index].~FRigConnectorElement(); 
 			}
 			break;
@@ -5378,6 +5551,7 @@ void URigHierarchy::DestroyElement(FRigBaseElement*& InElement)
 			for(int32 Index=0;Index<Count;Index++)
 			{
 				TGuardValue<const FRigBaseElement*> DestroyGuard(ElementBeingDestroyed, &ExistingElements[Index]);
+				DeallocateElementStorage(&ExistingElements[Index]);
 				ExistingElements[Index].~FRigSocketElement(); 
 			}
 			break;
@@ -5425,9 +5599,9 @@ void URigHierarchy::PropagateDirtyFlags(FRigTransformElement* InTransformElement
 				
 				if(ERigTransformType::IsGlobal(TypeToDirty))
 				{
-					if(ControlElement->Offset.IsDirty(TypeToDirty) &&
-						ControlElement->Pose.IsDirty(TypeToDirty) &&
-						ControlElement->Shape.IsDirty(TypeToDirty))
+					if(ControlElement->GetOffsetDirtyState().IsDirty(TypeToDirty) &&
+						ControlElement->GetDirtyState().IsDirty(TypeToDirty) &&
+						ControlElement->GetShapeDirtyState().IsDirty(TypeToDirty))
 					{
 						continue;
 					}
@@ -5437,7 +5611,7 @@ void URigHierarchy::PropagateDirtyFlags(FRigTransformElement* InTransformElement
 			{
 				if(ERigTransformType::IsGlobal(TypeToDirty))
 				{
-					if(MultiParentElement->Pose.IsDirty(TypeToDirty))
+					if(MultiParentElement->GetDirtyState().IsDirty(TypeToDirty))
 					{
 						continue;
 					}
@@ -5445,7 +5619,7 @@ void URigHierarchy::PropagateDirtyFlags(FRigTransformElement* InTransformElement
 			}
 			else
 			{
-				if(ElementToDirty.Element->Pose.IsDirty(TypeToDirty))
+				if(ElementToDirty.Element->GetDirtyState().IsDirty(TypeToDirty))
 				{
 					continue;
 				}
@@ -5481,9 +5655,9 @@ void URigHierarchy::PropagateDirtyFlags(FRigTransformElement* InTransformElement
 
 				if(ERigTransformType::IsGlobal(TypeToDirty))
 				{
-					if(ControlElement->Offset.IsDirty(TypeToDirty) &&
-						ControlElement->Pose.IsDirty(TypeToDirty) &&
-					    ControlElement->Shape.IsDirty(TypeToDirty))
+					if(ControlElement->GetOffsetDirtyState().IsDirty(TypeToDirty) &&
+						ControlElement->GetDirtyState().IsDirty(TypeToDirty) &&
+					    ControlElement->GetShapeDirtyState().IsDirty(TypeToDirty))
 					{
 						continue;
 					}
@@ -5493,7 +5667,7 @@ void URigHierarchy::PropagateDirtyFlags(FRigTransformElement* InTransformElement
 			{
 				if(ERigTransformType::IsGlobal(TypeToDirty))
 				{
-					if(MultiParentElement->Pose.IsDirty(TypeToDirty))
+					if(MultiParentElement->GetDirtyState().IsDirty(TypeToDirty))
 					{
 						continue;
 					}
@@ -5501,18 +5675,18 @@ void URigHierarchy::PropagateDirtyFlags(FRigTransformElement* InTransformElement
 			}
 			else
 			{
-				if(ElementToDirty.Element->Pose.IsDirty(TypeToDirty))
+				if(ElementToDirty.Element->GetDirtyState().IsDirty(TypeToDirty))
 				{
 					continue;
 				}
 			}
 						
-			ElementToDirty.Element->Pose.MarkDirty(TypeToDirty);
+			ElementToDirty.Element->GetDirtyState().MarkDirty(TypeToDirty);
 		
 			if(FRigControlElement* ControlElement = Cast<FRigControlElement>(ElementToDirty.Element))
 			{
-				ControlElement->Offset.MarkDirty(GlobalType);
-				ControlElement->Shape.MarkDirty(GlobalType);
+				ControlElement->GetOffsetDirtyState().MarkDirty(GlobalType);
+				ControlElement->GetShapeDirtyState().MarkDirty(GlobalType);
 			}
 
 			if(bAffectChildren)
@@ -5542,9 +5716,9 @@ void URigHierarchy::CleanupInvalidCaches()
 
 	struct Local
 	{
-		static bool NeedsCheck(bool bDirty[FRigLocalAndGlobalTransform::EDirtyMax])
+		static bool NeedsCheck(const FRigLocalAndGlobalDirtyState& InDirtyState)
 		{
-			return !bDirty[FRigLocalAndGlobalTransform::ELocal] && !bDirty[FRigLocalAndGlobalTransform::EGlobal];
+			return !InDirtyState.Local.Get() && !InDirtyState.Global.Get();
 		}
 	};
 
@@ -5554,37 +5728,37 @@ void URigHierarchy::CleanupInvalidCaches()
 		FRigBaseElement* BaseElement = HierarchyForCacheValidation->Elements[ElementIndex];
 		if(FRigControlElement* ControlElement = Cast<FRigControlElement>(BaseElement))
 		{
-			if(Local::NeedsCheck(ControlElement->Offset.Initial.bDirty))
+			if(Local::NeedsCheck(ControlElement->GetOffsetDirtyState().Initial))
 			{
-				ControlElement->Offset.MarkDirty(ERigTransformType::InitialGlobal);
+				ControlElement->GetOffsetDirtyState().MarkDirty(ERigTransformType::InitialGlobal);
 			}
 
-			if(Local::NeedsCheck(ControlElement->Pose.Initial.bDirty))
+			if(Local::NeedsCheck(ControlElement->GetDirtyState().Initial))
 			{
-				ControlElement->Pose.MarkDirty(ERigTransformType::InitialGlobal);
+				ControlElement->GetDirtyState().MarkDirty(ERigTransformType::InitialGlobal);
 			}
 
-			if(Local::NeedsCheck(ControlElement->Shape.Initial.bDirty))
+			if(Local::NeedsCheck(ControlElement->GetShapeDirtyState().Initial))
 			{
-				ControlElement->Shape.MarkDirty(ERigTransformType::InitialGlobal);
+				ControlElement->GetShapeDirtyState().MarkDirty(ERigTransformType::InitialGlobal);
 			}
 			continue;
 		}
 
 		if(FRigMultiParentElement* MultiParentElement = Cast<FRigMultiParentElement>(BaseElement))
 		{
-			if(Local::NeedsCheck(MultiParentElement->Pose.Initial.bDirty))
+			if(Local::NeedsCheck(MultiParentElement->GetDirtyState().Initial))
 			{
-				MultiParentElement->Pose.MarkDirty(ERigTransformType::InitialLocal);
+				MultiParentElement->GetDirtyState().MarkDirty(ERigTransformType::InitialLocal);
 			}
 			continue;
 		}
 
 		if(FRigTransformElement* TransformElement = Cast<FRigTransformElement>(BaseElement))
 		{
-			if(Local::NeedsCheck(TransformElement->Pose.Initial.bDirty))
+			if(Local::NeedsCheck(TransformElement->GetDirtyState().Initial))
 			{
-				TransformElement->Pose.MarkDirty(ERigTransformType::InitialGlobal);
+				TransformElement->GetDirtyState().MarkDirty(ERigTransformType::InitialGlobal);
 			}
 		}
 	}
@@ -5601,36 +5775,36 @@ void URigHierarchy::CleanupInvalidCaches()
 		{
 			FRigControlElement* OtherControlElement = HierarchyForCacheValidation->FindChecked<FRigControlElement>(ControlElement->GetKey());
 			
-			if(Local::NeedsCheck(ControlElement->Offset.Initial.bDirty))
+			if(Local::NeedsCheck(ControlElement->GetOffsetDirtyState().Initial))
 			{
-				const FTransform CachedGlobalTransform = OtherControlElement->Offset.Get(ERigTransformType::InitialGlobal);
+				const FTransform CachedGlobalTransform = OtherControlElement->GetOffsetTransform().Get(ERigTransformType::InitialGlobal);
 				const FTransform ComputedGlobalTransform = HierarchyForCacheValidation->GetControlOffsetTransform(OtherControlElement, ERigTransformType::InitialGlobal);
 
 				if(!FRigComputedTransform::Equals(ComputedGlobalTransform, CachedGlobalTransform, 0.01))
 				{
-					ControlElement->Offset.MarkDirty(ERigTransformType::InitialGlobal);
+					ControlElement->GetOffsetDirtyState().MarkDirty(ERigTransformType::InitialGlobal);
 				}
 			}
 
-			if(Local::NeedsCheck(ControlElement->Pose.Initial.bDirty))
+			if(Local::NeedsCheck(ControlElement->GetDirtyState().Initial))
 			{
-				const FTransform CachedGlobalTransform = ControlElement->Pose.Get(ERigTransformType::InitialGlobal);
+				const FTransform CachedGlobalTransform = ControlElement->GetTransform().Get(ERigTransformType::InitialGlobal);
 				const FTransform ComputedGlobalTransform = HierarchyForCacheValidation->GetTransform(OtherControlElement, ERigTransformType::InitialGlobal);
 				
 				if(!FRigComputedTransform::Equals(ComputedGlobalTransform, CachedGlobalTransform, 0.01))
 				{
-					ControlElement->Pose.MarkDirty(ERigTransformType::InitialGlobal);
+					ControlElement->GetDirtyState().MarkDirty(ERigTransformType::InitialGlobal);
 				}
 			}
 
-			if(Local::NeedsCheck(ControlElement->Shape.Initial.bDirty))
+			if(Local::NeedsCheck(ControlElement->GetShapeDirtyState().Initial))
 			{
-				const FTransform CachedGlobalTransform = ControlElement->Shape.Get(ERigTransformType::InitialGlobal);
+				const FTransform CachedGlobalTransform = ControlElement->GetShapeTransform().Get(ERigTransformType::InitialGlobal);
 				const FTransform ComputedGlobalTransform = HierarchyForCacheValidation->GetControlShapeTransform(OtherControlElement, ERigTransformType::InitialGlobal);
 				
 				if(!FRigComputedTransform::Equals(ComputedGlobalTransform, CachedGlobalTransform, 0.01))
 				{
-					ControlElement->Shape.MarkDirty(ERigTransformType::InitialGlobal);
+					ControlElement->GetShapeDirtyState().MarkDirty(ERigTransformType::InitialGlobal);
 				}
 			}
 			continue;
@@ -5640,15 +5814,15 @@ void URigHierarchy::CleanupInvalidCaches()
 		{
 			FRigMultiParentElement* OtherMultiParentElement = HierarchyForCacheValidation->FindChecked<FRigMultiParentElement>(MultiParentElement->GetKey());
 
-			if(Local::NeedsCheck(MultiParentElement->Pose.Initial.bDirty))
+			if(Local::NeedsCheck(MultiParentElement->GetDirtyState().Initial))
 			{
-				const FTransform CachedGlobalTransform = MultiParentElement->Pose.Get(ERigTransformType::InitialGlobal);
+				const FTransform CachedGlobalTransform = MultiParentElement->GetTransform().Get(ERigTransformType::InitialGlobal);
 				const FTransform ComputedGlobalTransform = HierarchyForCacheValidation->GetTransform(OtherMultiParentElement, ERigTransformType::InitialGlobal);
 				
 				if(!FRigComputedTransform::Equals(ComputedGlobalTransform, CachedGlobalTransform, 0.01))
 				{
 					// for nulls we perceive the local transform as less relevant
-					MultiParentElement->Pose.MarkDirty(ERigTransformType::InitialLocal);
+					MultiParentElement->GetDirtyState().MarkDirty(ERigTransformType::InitialLocal);
 				}
 			}
 			continue;
@@ -5658,14 +5832,14 @@ void URigHierarchy::CleanupInvalidCaches()
 		{
 			FRigTransformElement* OtherTransformElement = HierarchyForCacheValidation->FindChecked<FRigTransformElement>(TransformElement->GetKey());
 
-			if(Local::NeedsCheck(TransformElement->Pose.Initial.bDirty))
+			if(Local::NeedsCheck(TransformElement->GetDirtyState().Initial))
 			{
-				const FTransform CachedGlobalTransform = TransformElement->Pose.Get(ERigTransformType::InitialGlobal);
+				const FTransform CachedGlobalTransform = TransformElement->GetTransform().Get(ERigTransformType::InitialGlobal);
 				const FTransform ComputedGlobalTransform = HierarchyForCacheValidation->GetTransform(OtherTransformElement, ERigTransformType::InitialGlobal);
 				
 				if(!FRigComputedTransform::Equals(ComputedGlobalTransform, CachedGlobalTransform, 0.01))
 				{
-					TransformElement->Pose.MarkDirty(ERigTransformType::InitialGlobal);
+					TransformElement->GetDirtyState().MarkDirty(ERigTransformType::InitialGlobal);
 				}
 			}
 		}
@@ -5733,6 +5907,113 @@ void URigHierarchy::FMetadataStorage::Serialize(FArchive& Ar)
 	}
 }
 
+void URigHierarchy::AllocateDefaultElementStorage(FRigBaseElement* InElement, bool bUpdateAllElements)
+{
+	if(FRigTransformElement* TransformElement = Cast<FRigTransformElement>(InElement))
+	{
+		// only for the default allocation we want to catch double allocations.
+		// so if we are already linked to the default storage, exit early.
+		if(ElementTransforms.Contains(TransformElement->PoseStorage.Initial.Local.Index, TransformElement->PoseStorage.Initial.Local.Storage))
+		{
+			check(ElementTransforms.Contains(TransformElement->PoseStorage.Initial.Global.Index, TransformElement->PoseStorage.Initial.Global.Storage));
+			check(ElementTransforms.Contains(TransformElement->PoseStorage.Current.Local.Index, TransformElement->PoseStorage.Current.Local.Storage));
+			check(ElementTransforms.Contains(TransformElement->PoseStorage.Current.Global.Index, TransformElement->PoseStorage.Current.Global.Storage));
+			check(ElementDirtyStates.Contains(TransformElement->PoseDirtyState.Initial.Local.Index, TransformElement->PoseDirtyState.Initial.Local.Storage));
+			check(ElementDirtyStates.Contains(TransformElement->PoseDirtyState.Initial.Global.Index, TransformElement->PoseDirtyState.Initial.Global.Storage));
+			check(ElementDirtyStates.Contains(TransformElement->PoseDirtyState.Current.Local.Index, TransformElement->PoseDirtyState.Current.Local.Storage));
+			check(ElementDirtyStates.Contains(TransformElement->PoseDirtyState.Current.Global.Index, TransformElement->PoseDirtyState.Current.Global.Storage));
+
+			if(const FRigControlElement* ControlElement = Cast<FRigControlElement>(TransformElement))
+			{
+				check(ElementTransforms.Contains(ControlElement->OffsetStorage.Initial.Local.Index, ControlElement->OffsetStorage.Initial.Local.Storage));
+				check(ElementTransforms.Contains(ControlElement->OffsetStorage.Initial.Global.Index, ControlElement->OffsetStorage.Initial.Global.Storage));
+				check(ElementTransforms.Contains(ControlElement->OffsetStorage.Current.Local.Index, ControlElement->OffsetStorage.Current.Local.Storage));
+				check(ElementTransforms.Contains(ControlElement->OffsetStorage.Current.Global.Index, ControlElement->OffsetStorage.Current.Global.Storage));
+				check(ElementDirtyStates.Contains(ControlElement->OffsetDirtyState.Initial.Local.Index, ControlElement->OffsetDirtyState.Initial.Local.Storage));
+				check(ElementDirtyStates.Contains(ControlElement->OffsetDirtyState.Initial.Global.Index, ControlElement->OffsetDirtyState.Initial.Global.Storage));
+				check(ElementDirtyStates.Contains(ControlElement->OffsetDirtyState.Current.Local.Index, ControlElement->OffsetDirtyState.Current.Local.Storage));
+				check(ElementDirtyStates.Contains(ControlElement->OffsetDirtyState.Current.Global.Index, ControlElement->OffsetDirtyState.Current.Global.Storage));
+				check(ElementTransforms.Contains(ControlElement->ShapeStorage.Initial.Local.Index, ControlElement->ShapeStorage.Initial.Local.Storage));
+				check(ElementTransforms.Contains(ControlElement->ShapeStorage.Initial.Global.Index, ControlElement->ShapeStorage.Initial.Global.Storage));
+				check(ElementTransforms.Contains(ControlElement->ShapeStorage.Current.Local.Index, ControlElement->ShapeStorage.Current.Local.Storage));
+				check(ElementTransforms.Contains(ControlElement->ShapeStorage.Current.Global.Index, ControlElement->ShapeStorage.Current.Global.Storage));
+				check(ElementDirtyStates.Contains(ControlElement->ShapeDirtyState.Initial.Local.Index, ControlElement->ShapeDirtyState.Initial.Local.Storage));
+				check(ElementDirtyStates.Contains(ControlElement->ShapeDirtyState.Initial.Global.Index, ControlElement->ShapeDirtyState.Initial.Global.Storage));
+				check(ElementDirtyStates.Contains(ControlElement->ShapeDirtyState.Current.Local.Index, ControlElement->ShapeDirtyState.Current.Local.Storage));
+				check(ElementDirtyStates.Contains(ControlElement->ShapeDirtyState.Current.Global.Index, ControlElement->ShapeDirtyState.Current.Global.Storage));
+			}
+			return;
+		}
+		const TArray<int32, TInlineAllocator<4>> TransformIndices = ElementTransforms.Allocate(TransformElement->GetNumTransforms(), FTransform::Identity);
+		check(TransformIndices.Num() >= 4);
+		const TArray<int32, TInlineAllocator<4>> DirtyStateIndices = ElementDirtyStates.Allocate(TransformElement->GetNumTransforms(), false);
+		check(DirtyStateIndices.Num() >= 4);
+		TransformElement->PoseStorage.Initial.Local.Index = TransformIndices[0];
+		TransformElement->PoseStorage.Initial.Global.Index = TransformIndices[1];
+		TransformElement->PoseStorage.Current.Local.Index = TransformIndices[2];
+		TransformElement->PoseStorage.Current.Global.Index = TransformIndices[3];
+		TransformElement->PoseDirtyState.Initial.Local.Index = DirtyStateIndices[0];
+		TransformElement->PoseDirtyState.Initial.Global.Index = DirtyStateIndices[1];
+		TransformElement->PoseDirtyState.Current.Local.Index = DirtyStateIndices[2];
+		TransformElement->PoseDirtyState.Current.Global.Index = DirtyStateIndices[3];
+
+		if(FRigControlElement* ControlElement = Cast<FRigControlElement>(TransformElement))
+		{
+			check(TransformIndices.Num() >= 12);
+			check(DirtyStateIndices.Num() >= 12);
+			ControlElement->OffsetStorage.Initial.Local.Index = TransformIndices[4];
+			ControlElement->OffsetStorage.Initial.Global.Index = TransformIndices[5];
+			ControlElement->OffsetStorage.Current.Local.Index = TransformIndices[6];
+			ControlElement->OffsetStorage.Current.Global.Index = TransformIndices[7];
+			ControlElement->OffsetDirtyState.Initial.Local.Index = DirtyStateIndices[4];
+			ControlElement->OffsetDirtyState.Initial.Global.Index = DirtyStateIndices[5];
+			ControlElement->OffsetDirtyState.Current.Local.Index = DirtyStateIndices[6];
+			ControlElement->OffsetDirtyState.Current.Global.Index = DirtyStateIndices[7];
+			ControlElement->ShapeStorage.Initial.Local.Index = TransformIndices[8];
+			ControlElement->ShapeStorage.Initial.Global.Index = TransformIndices[9];
+			ControlElement->ShapeStorage.Current.Local.Index = TransformIndices[10];
+			ControlElement->ShapeStorage.Current.Global.Index = TransformIndices[11];
+			ControlElement->ShapeDirtyState.Initial.Local.Index = DirtyStateIndices[8];
+			ControlElement->ShapeDirtyState.Initial.Global.Index = DirtyStateIndices[9];
+			ControlElement->ShapeDirtyState.Current.Local.Index = DirtyStateIndices[10];
+			ControlElement->ShapeDirtyState.Current.Global.Index = DirtyStateIndices[11];
+		}
+	}
+	else if(FRigCurveElement* CurveElement = Cast<FRigCurveElement>(InElement))
+	{
+		if(ElementCurves.Contains(CurveElement->StorageIndex, CurveElement->Storage))
+		{
+			return;
+		}
+		
+		const TArray<int32, TInlineAllocator<4>> CurveIndices = ElementCurves.Allocate(CurveElement->GetNumCurves(), 0.f);
+		check(CurveIndices.Num() >= 1);
+		CurveElement->StorageIndex = CurveIndices[0];
+	}
+
+	if(bUpdateAllElements)
+	{
+		UpdateElementStorage();
+	}
+	else
+	{
+		InElement->LinkStorage(ElementTransforms.Storage, ElementDirtyStates.Storage, ElementCurves.Storage);
+	}
+}
+
+void URigHierarchy::DeallocateElementStorage(FRigBaseElement* InElement)
+{
+	InElement->UnlinkStorage(ElementTransforms, ElementDirtyStates, ElementCurves);
+}
+
+void URigHierarchy::UpdateElementStorage()
+{
+	for(FRigBaseElement* Element : Elements)
+	{
+		Element->LinkStorage(ElementTransforms.Storage, ElementDirtyStates.Storage, ElementCurves.Storage);
+	}
+}
+
 void URigHierarchy::PropagateMetadata(const FRigElementKey& InKey, const FName& InName, bool bNotify)
 {
 	if(const FRigBaseElement* Element = Find(InKey))
@@ -5811,16 +6092,7 @@ FRigBaseMetadata* URigHierarchy::GetMetadataForElement(FRigBaseElement* InElemen
 {
 	if (!ElementMetadata.IsValidIndex(InElement->MetadataStorageIndex))
 	{
-		// Do we have entries in the freelist we can recycle?
-		if (!ElementMetadataFreeList.IsEmpty())
-		{
-			InElement->MetadataStorageIndex = ElementMetadataFreeList.Pop(EAllowShrinking::No);
-		}
-		else
-		{
-			InElement->MetadataStorageIndex = ElementMetadata.Num();
-			ElementMetadata.AddDefaulted();
-		}
+		InElement->MetadataStorageIndex = ElementMetadata.Allocate(1, FMetadataStorage())[0];
 	}
 	
 	FMetadataStorage& Storage = ElementMetadata[InElement->MetadataStorageIndex];
@@ -5930,7 +6202,7 @@ bool URigHierarchy::RemoveMetadataForElement(FRigBaseElement* InElement, const F
 	// metadata storage can just recycle that.
 	if (Storage.MetadataMap.IsEmpty())
 	{
-		ElementMetadataFreeList.Add(InElement->MetadataStorageIndex);
+		ElementMetadata.Deallocate(InElement->MetadataStorageIndex);
 		InElement->MetadataStorageIndex = INDEX_NONE;
 	}
 	else if (Storage.LastAccessName == InName)
@@ -5961,7 +6233,7 @@ bool URigHierarchy::RemoveAllMetadataForElement(FRigBaseElement* InElement)
 	// Clear the storage for the next user.
 	Storage.Reset();
 	
-	ElementMetadataFreeList.Push(InElement->MetadataStorageIndex);
+	ElementMetadata.Deallocate(InElement->MetadataStorageIndex);
 	InElement->MetadataStorageIndex = INDEX_NONE;
 
 	if(ElementBeingDestroyed != InElement)
@@ -6026,7 +6298,7 @@ void URigHierarchy::EnsureCacheValidityImpl()
 				continue;
 			}
 
-			if(!TransformElement->Pose.IsDirty(GlobalType))
+			if(!TransformElement->GetDirtyState().IsDirty(GlobalType))
 			{
 				continue;
 			}
@@ -6037,19 +6309,19 @@ void URigHierarchy::EnsureCacheValidityImpl()
 				{
                     if(FRigControlElement* ControlElementToDirty = Cast<FRigControlElement>(ElementToDirty.Element))
                     {
-                        if(ControlElementToDirty->Offset.IsDirty(GlobalType))
+                        if(ControlElementToDirty->GetOffsetDirtyState().IsDirty(GlobalType))
                         {
-                            checkf(ControlElementToDirty->Pose.IsDirty(GlobalType) ||
-                                    ControlElementToDirty->Pose.IsDirty(LocalType),
+                            checkf(ControlElementToDirty->GetDirtyState().IsDirty(GlobalType) ||
+                                    ControlElementToDirty->GetDirtyState().IsDirty(LocalType),
                                     TEXT("Control '%s' %s Offset Cache is dirty, but the Pose is not."),
 									*ControlElementToDirty->GetKey().ToString(),
 									*TransformTypeString);
 						}
 
-                        if(ControlElementToDirty->Pose.IsDirty(GlobalType))
+                        if(ControlElementToDirty->GetDirtyState().IsDirty(GlobalType))
                         {
-                            checkf(ControlElementToDirty->Shape.IsDirty(GlobalType) ||
-                                    ControlElementToDirty->Shape.IsDirty(LocalType),
+                            checkf(ControlElementToDirty->GetShapeDirtyState().IsDirty(GlobalType) ||
+                                    ControlElementToDirty->GetShapeDirtyState().IsDirty(LocalType),
                                     TEXT("Control '%s' %s Pose Cache is dirty, but the Shape is not."),
 									*ControlElementToDirty->GetKey().ToString(),
 									*TransformTypeString);
@@ -6057,8 +6329,8 @@ void URigHierarchy::EnsureCacheValidityImpl()
                     }
                     else
                     {
-                        checkf(MultiParentElementToDirty->Pose.IsDirty(GlobalType) ||
-                                MultiParentElementToDirty->Pose.IsDirty(LocalType),
+                        checkf(MultiParentElementToDirty->GetDirtyState().IsDirty(GlobalType) ||
+                                MultiParentElementToDirty->GetDirtyState().IsDirty(LocalType),
                                 TEXT("MultiParent '%s' %s Parent Cache is dirty, but the Pose is not."),
 								*MultiParentElementToDirty->GetKey().ToString(),
 								*TransformTypeString);
@@ -6066,8 +6338,8 @@ void URigHierarchy::EnsureCacheValidityImpl()
 				}
 				else
 				{
-					checkf(ElementToDirty.Element->Pose.IsDirty(GlobalType) ||
-						ElementToDirty.Element->Pose.IsDirty(LocalType),
+					checkf(ElementToDirty.Element->GetDirtyState().IsDirty(GlobalType) ||
+						ElementToDirty.Element->GetDirtyState().IsDirty(LocalType),
 						TEXT("SingleParent '%s' %s Pose is not dirty in Local or Global"),
 						*ElementToDirty.Element->GetKey().ToString(),
 						*TransformTypeString);
@@ -6104,10 +6376,10 @@ void URigHierarchy::EnsureCacheValidityImpl()
 				const ERigTransformType::Type OpposedType = ERigTransformType::SwapLocalAndGlobal(TransformType);
 				const FString& TransformTypeString = TransformTypeStrings[TransformTypeIndex];
 
-				if(!ControlElement->Offset.IsDirty(TransformType) && !ControlElement->Offset.IsDirty(OpposedType))
+				if(!ControlElement->GetOffsetDirtyState().IsDirty(TransformType) && !ControlElement->GetOffsetDirtyState().IsDirty(OpposedType))
 				{
 					const FTransform CachedTransform = HierarchyForLambda->GetControlOffsetTransform(ControlElement, TransformType);
-					ControlElement->Offset.MarkDirty(TransformType);
+					ControlElement->GetOffsetDirtyState().MarkDirty(TransformType);
 					const FTransform ComputedTransform = HierarchyForLambda->GetControlOffsetTransform(ControlElement, TransformType);
 					checkf(FRigComputedTransform::Equals(CachedTransform, ComputedTransform),
 						TEXT("Element '%s' Offset %s Cached vs Computed doesn't match. ('%s' <-> '%s')"),
@@ -6127,10 +6399,10 @@ void URigHierarchy::EnsureCacheValidityImpl()
 				const ERigTransformType::Type OpposedType = ERigTransformType::SwapLocalAndGlobal(TransformType);
 				const FString& TransformTypeString = TransformTypeStrings[TransformTypeIndex];
 
-				if(!TransformElement->Pose.IsDirty(TransformType) && !TransformElement->Pose.IsDirty(OpposedType))
+				if(!TransformElement->GetDirtyState().IsDirty(TransformType) && !TransformElement->GetDirtyState().IsDirty(OpposedType))
 				{
 					const FTransform CachedTransform = HierarchyForLambda->GetTransform(TransformElement, TransformType);
-					TransformElement->Pose.MarkDirty(TransformType);
+					TransformElement->GetDirtyState().MarkDirty(TransformType);
 					const FTransform ComputedTransform = HierarchyForLambda->GetTransform(TransformElement, TransformType);
 					checkf(FRigComputedTransform::Equals(CachedTransform, ComputedTransform),
 						TEXT("Element '%s' Pose %s Cached vs Computed doesn't match. ('%s' <-> '%s')"),
@@ -6149,10 +6421,10 @@ void URigHierarchy::EnsureCacheValidityImpl()
 				const ERigTransformType::Type OpposedType = ERigTransformType::SwapLocalAndGlobal(TransformType);
 				const FString& TransformTypeString = TransformTypeStrings[TransformTypeIndex];
 
-				if(!ControlElement->Shape.IsDirty(TransformType) && !ControlElement->Shape.IsDirty(OpposedType))
+				if(!ControlElement->GetShapeDirtyState().IsDirty(TransformType) && !ControlElement->GetShapeDirtyState().IsDirty(OpposedType))
 				{
 					const FTransform CachedTransform = HierarchyForLambda->GetControlShapeTransform(ControlElement, TransformType);
-					ControlElement->Shape.MarkDirty(TransformType);
+					ControlElement->GetShapeDirtyState().MarkDirty(TransformType);
 					const FTransform ComputedTransform = HierarchyForLambda->GetControlShapeTransform(ControlElement, TransformType);
 					checkf(FRigComputedTransform::Equals(CachedTransform, ComputedTransform),
 						TEXT("Element '%s' Shape %s Cached vs Computed doesn't match. ('%s' <-> '%s')"),
@@ -7255,10 +7527,10 @@ FTransform URigHierarchy::LazilyComputeParentConstraint(
 		}
 
 		Transform.NormalizeRotation();
-		Constraint.Cache.Transform = Transform;
+		Constraint.Cache = Transform;
 		Constraint.bCacheIsDirty = false;
 	}
-	return Constraint.Cache.Transform;
+	return Constraint.Cache;
 }
 
 void URigHierarchy::ComputeParentConstraintIndices(
diff --git a/Engine/Plugins/Animation/ControlRig/Source/ControlRig/Private/Rigs/RigHierarchyController.cpp b/Engine/Plugins/Animation/ControlRig/Source/ControlRig/Private/Rigs/RigHierarchyController.cpp
--- a/Engine/Plugins/Animation/ControlRig/Source/ControlRig/Private/Rigs/RigHierarchyController.cpp
+++ b/Engine/Plugins/Animation/ControlRig/Source/ControlRig/Private/Rigs/RigHierarchyController.cpp
@@ -236,7 +236,8 @@ FRigElementKey URigHierarchyController::AddBone(FName InName, FRigElementKey InP
 			Hierarchy->SetTransform(NewElement, InTransform, ERigTransformType::CurrentLocal, true, false);
 		}
 
-		NewElement->Pose.Current = NewElement->Pose.Initial;
+		NewElement->GetTransform().Current = NewElement->GetTransform().Initial;
+		NewElement->GetDirtyState().Current = NewElement->GetDirtyState().Initial;
 	}
 
 #if WITH_EDITOR
@@ -297,7 +298,7 @@ FRigElementKey URigHierarchyController::AddNull(FName InName, FRigElementKey InP
 			Hierarchy->SetTransform(NewElement, InTransform, ERigTransformType::InitialLocal, true, false);
 		}
 
-		NewElement->Pose.Current = NewElement->Pose.Initial;
+		NewElement->GetTransform().Current = NewElement->GetTransform().Initial;
 	}
 
 #if WITH_EDITOR
@@ -374,20 +375,25 @@ FRigElementKey URigHierarchyController::AddControl(
 		
 		AddElement(NewElement, Hierarchy->Get(Hierarchy->GetIndex(InParent)), false, InName);
 		
-		NewElement->Offset.Set(ERigTransformType::InitialLocal, InOffsetTransform);  
-		NewElement->Shape.Set(ERigTransformType::InitialLocal, InShapeTransform);  
+		NewElement->GetOffsetTransform().Set(ERigTransformType::InitialLocal, InOffsetTransform);  
+		NewElement->GetOffsetDirtyState().MarkClean(ERigTransformType::InitialLocal);
+		NewElement->GetShapeTransform().Set(ERigTransformType::InitialLocal, InShapeTransform);
+		NewElement->GetShapeDirtyState().MarkClean(ERigTransformType::InitialLocal);
 		Hierarchy->SetControlValue(NewElement, InValue, ERigControlValueType::Initial, false);
 		const FTransform LocalTransform = Hierarchy->GetTransform(NewElement, ERigTransformType::InitialLocal);
 		static constexpr bool bInitial = true;
 		Hierarchy->SetControlPreferredEulerAngles(NewElement, LocalTransform, bInitial);
 
-		NewElement->Offset.MarkDirty(ERigTransformType::InitialGlobal);
-		NewElement->Pose.MarkDirty(ERigTransformType::InitialGlobal);
-		NewElement->Shape.MarkDirty(ERigTransformType::InitialGlobal);
-		NewElement->Offset.Current = NewElement->Offset.Initial;
-		NewElement->Pose.Current = NewElement->Pose.Initial;
+		NewElement->GetOffsetDirtyState().MarkDirty(ERigTransformType::InitialGlobal);
+		NewElement->GetDirtyState().MarkDirty(ERigTransformType::InitialGlobal);
+		NewElement->GetShapeDirtyState().MarkDirty(ERigTransformType::InitialGlobal);
+		NewElement->GetOffsetTransform().Current = NewElement->GetOffsetTransform().Initial;
+		NewElement->GetOffsetDirtyState().Current = NewElement->GetOffsetDirtyState().Initial; 
+		NewElement->GetTransform().Current = NewElement->GetTransform().Initial;
+		NewElement->GetDirtyState().Current = NewElement->GetDirtyState().Initial;
 		NewElement->PreferredEulerAngles.Current = NewElement->PreferredEulerAngles.Initial;
-		NewElement->Shape.Current = NewElement->Shape.Initial;
+		NewElement->GetShapeTransform().Current = NewElement->GetShapeTransform().Initial;
+		NewElement->GetShapeDirtyState().Current = NewElement->GetShapeDirtyState().Initial;
 	}
 
 #if WITH_EDITOR
@@ -458,8 +464,9 @@ FRigElementKey URigHierarchyController::AddCurve(FName InName, float InValue, bo
 		TGuardValue<bool> DisableCacheValidityChecks(Hierarchy->bEnableCacheValidityCheck, false);		
 		NewElement->Key.Type = ERigElementType::Curve;
 		NewElement->Key.Name = GetSafeNewName(InName, NewElement->Key.Type);
-		NewElement->Value = InValue;
 		AddElement(NewElement, nullptr, false, InName);
+		NewElement->Set(InValue);
+		NewElement->bIsValueSet = false;
 	}
 
 #if WITH_EDITOR
@@ -533,7 +540,8 @@ FRigElementKey URigHierarchyController::AddPhysicsElement(FName InName, FRigElem
 		AddElement(NewElement, Hierarchy->Get(Hierarchy->GetIndex(InParent)), true, InName);
 
 		Hierarchy->SetTransform(NewElement, InLocalTransform, ERigTransformType::InitialLocal, true, false);
-		NewElement->Pose.Current = NewElement->Pose.Initial;
+		NewElement->GetTransform().Current = NewElement->GetTransform().Initial;
+		NewElement->GetDirtyState().Current = NewElement->GetDirtyState().Initial;
 	}
 	
 #if WITH_EDITOR
@@ -587,7 +595,8 @@ FRigElementKey URigHierarchyController::AddReference(FName InName, FRigElementKe
 		AddElement(NewElement, Hierarchy->Get(Hierarchy->GetIndex(InParent)), true, InName);
 
 		Hierarchy->SetTransform(NewElement, FTransform::Identity, ERigTransformType::InitialLocal, true, false);
-		NewElement->Pose.Current = NewElement->Pose.Initial;
+		NewElement->GetTransform().Current = NewElement->GetTransform().Initial;
+		NewElement->GetDirtyState().Current = NewElement->GetDirtyState().Initial;
 	}
 
 #if WITH_EDITOR
@@ -720,7 +729,8 @@ FRigElementKey URigHierarchyController::AddSocket(FName InName, FRigElementKey I
 			CurrentHierarchy->SetTransform(NewElement, InTransform, ERigTransformType::CurrentLocal, true, false);
 		}
 
-		NewElement->Pose.Current = NewElement->Pose.Initial;
+		NewElement->GetTransform().Current = NewElement->GetTransform().Initial;
+		NewElement->GetDirtyState().Current = NewElement->GetDirtyState().Initial;
 
 		NewElement->SetColor(InColor, CurrentHierarchy);
 		NewElement->SetDescription(InDescription, CurrentHierarchy);
@@ -1214,10 +1224,14 @@ FString URigHierarchyController::ExportToText(TArray<FRigElementKey> InKeys) con
 
 		if(FRigTransformElement* TransformElement = Cast<FRigTransformElement>(Element))
 		{
-			PerElementData.Pose.Initial.Local.Set(Hierarchy->GetTransform(TransformElement, ERigTransformType::InitialLocal), PerElementData.Pose.Initial.bDirty[FRigLocalAndGlobalTransform::ELocal]);
-			PerElementData.Pose.Initial.Global.Set(Hierarchy->GetTransform(TransformElement, ERigTransformType::InitialGlobal), PerElementData.Pose.Initial.bDirty[FRigLocalAndGlobalTransform::EGlobal]);
-			PerElementData.Pose.Current.Local.Set(Hierarchy->GetTransform(TransformElement, ERigTransformType::CurrentLocal), PerElementData.Pose.Current.bDirty[FRigLocalAndGlobalTransform::ELocal]);
-			PerElementData.Pose.Current.Global.Set(Hierarchy->GetTransform(TransformElement, ERigTransformType::CurrentGlobal), PerElementData.Pose.Current.bDirty[FRigLocalAndGlobalTransform::EGlobal]);
+			PerElementData.Poses.Add(Hierarchy->GetTransform(TransformElement, ERigTransformType::InitialLocal));
+			PerElementData.Poses.Add(Hierarchy->GetTransform(TransformElement, ERigTransformType::CurrentLocal));
+			PerElementData.Poses.Add(Hierarchy->GetTransform(TransformElement, ERigTransformType::InitialGlobal));
+			PerElementData.Poses.Add(Hierarchy->GetTransform(TransformElement, ERigTransformType::CurrentGlobal));
+			PerElementData.DirtyStates.Add(TransformElement->GetDirtyState().GetDirtyFlag(ERigTransformType::InitialLocal));
+			PerElementData.DirtyStates.Add(TransformElement->GetDirtyState().GetDirtyFlag(ERigTransformType::CurrentLocal));
+			PerElementData.DirtyStates.Add(TransformElement->GetDirtyState().GetDirtyFlag(ERigTransformType::InitialGlobal));
+			PerElementData.DirtyStates.Add(TransformElement->GetDirtyState().GetDirtyFlag(ERigTransformType::CurrentGlobal));
 		}
 
 		switch (Key.Type)
@@ -1384,6 +1398,7 @@ TArray<FRigElementKey> URigHierarchyController::ImportFromText(FString InContent
 	{
 		KeyMap.Add(Element->GetKey(), Element->GetKey());
 	}
+	TArray<FRigElementKey> PreviouslyExistingKeys;
 
 	FRigHierarchyInteractionBracket InteractionBracket(Hierarchy);
 
@@ -1397,50 +1412,50 @@ TArray<FRigElementKey> URigHierarchyController::ImportFromText(FString InContent
 		{
 			case ERigElementType::Bone:
 			{
-				NewElement = MakeElement<FRigBoneElement>();
+				NewElement = MakeElement<FRigBoneElement>(true);
 				FRigBoneElement::StaticStruct()->ImportText(*PerElementData.Content, NewElement, nullptr, EPropertyPortFlags::PPF_None, &ErrorPipe, FRigBoneElement::StaticStruct()->GetName(), true);
 				CastChecked<FRigBoneElement>(NewElement)->BoneType = ERigBoneType::User;
 				break;
 			}
 			case ERigElementType::Null:
 			{
-				NewElement = MakeElement<FRigNullElement>();
+				NewElement = MakeElement<FRigNullElement>(true);
 				FRigNullElement::StaticStruct()->ImportText(*PerElementData.Content, NewElement, nullptr, EPropertyPortFlags::PPF_None, &ErrorPipe, FRigNullElement::StaticStruct()->GetName(), true);
 				break;
 			}
 			case ERigElementType::Control:
 			{
-				NewElement = MakeElement<FRigControlElement>();
+				NewElement = MakeElement<FRigControlElement>(true);
 				FRigControlElement::StaticStruct()->ImportText(*PerElementData.Content, NewElement, nullptr, EPropertyPortFlags::PPF_None, &ErrorPipe, FRigControlElement::StaticStruct()->GetName(), true);
 				break;
 			}
 			case ERigElementType::Curve:
 			{
-				NewElement = MakeElement<FRigCurveElement>();
+				NewElement = MakeElement<FRigCurveElement>(true);
 				FRigCurveElement::StaticStruct()->ImportText(*PerElementData.Content, NewElement, nullptr, EPropertyPortFlags::PPF_None, &ErrorPipe, FRigCurveElement::StaticStruct()->GetName(), true);
 				break;
 			}
 			case ERigElementType::Physics:
 			{
-				NewElement = MakeElement<FRigPhysicsElement>();
+				NewElement = MakeElement<FRigPhysicsElement>(true);
 				FRigPhysicsElement::StaticStruct()->ImportText(*PerElementData.Content, NewElement, nullptr, EPropertyPortFlags::PPF_None, &ErrorPipe, FRigPhysicsElement::StaticStruct()->GetName(), true);
 				break;
 			}
 			case ERigElementType::Reference:
 			{
-				NewElement = MakeElement<FRigReferenceElement>();
+				NewElement = MakeElement<FRigReferenceElement>(true);
 				FRigReferenceElement::StaticStruct()->ImportText(*PerElementData.Content, NewElement, nullptr, EPropertyPortFlags::PPF_None, &ErrorPipe, FRigReferenceElement::StaticStruct()->GetName(), true);
 				break;
 			}
 			case ERigElementType::Connector:
 			{
-				NewElement = MakeElement<FRigConnectorElement>();
+				NewElement = MakeElement<FRigConnectorElement>(true);
 				FRigConnectorElement::StaticStruct()->ImportText(*PerElementData.Content, NewElement, nullptr, EPropertyPortFlags::PPF_None, &ErrorPipe, FRigConnectorElement::StaticStruct()->GetName(), true);
 				break;
 			}
 			case ERigElementType::Socket:
 			{
-				NewElement = MakeElement<FRigSocketElement>();
+				NewElement = MakeElement<FRigSocketElement>(true);
 				FRigSocketElement::StaticStruct()->ImportText(*PerElementData.Content, NewElement, nullptr, EPropertyPortFlags::PPF_None, &ErrorPipe, FRigSocketElement::StaticStruct()->GetName(), true);
 				break;
 			}
@@ -1464,8 +1479,8 @@ TArray<FRigElementKey> URigHierarchyController::ImportFromText(FString InContent
 				{
 					Hierarchy->GetControlShapeTransform(ControlElement, ERigTransformType::CurrentLocal);
 					Hierarchy->GetControlShapeTransform(ControlElement, ERigTransformType::InitialLocal);
-					ControlElement->Shape.MarkDirty(ERigTransformType::CurrentGlobal);
-					ControlElement->Shape.MarkDirty(ERigTransformType::InitialGlobal);
+					ControlElement->GetShapeDirtyState().MarkDirty(ERigTransformType::CurrentGlobal);
+					ControlElement->GetShapeDirtyState().MarkDirty(ERigTransformType::InitialGlobal);
 				}
 				
 				TArray<FRigElementKey> CurrentParents = Hierarchy->GetParents(NewElement->GetKey());
@@ -1500,6 +1515,7 @@ TArray<FRigElementKey> URigHierarchyController::ImportFromText(FString InContent
 				}
 
 				PastedKeys.Add(ExistingElement->GetKey());
+				PreviouslyExistingKeys.Add(ExistingElement->GetKey());
 
 				Hierarchy->DestroyElement(NewElement);
 				continue;
@@ -1511,6 +1527,19 @@ TArray<FRigElementKey> URigHierarchyController::ImportFromText(FString InContent
 		AddElement(NewElement, nullptr, true, DesiredName);
 
 		KeyMap.FindOrAdd(PerElementData.Key) = NewElement->Key;
+	}
+
+	Hierarchy->UpdateElementStorage();
+	
+	for(const FRigHierarchyCopyPasteContentPerElement& PerElementData : Data.Elements)
+	{
+		if(PreviouslyExistingKeys.Contains(PerElementData.Key))
+		{
+			continue;
+		}
+		
+		FRigElementKey MappedKey = KeyMap.FindChecked(PerElementData.Key);
+		FRigBaseElement* NewElement = Hierarchy->FindChecked(MappedKey);
 
 		for(const FRigElementKey& OriginalParent : PerElementData.Parents)
 		{
@@ -1529,7 +1558,7 @@ TArray<FRigElementKey> URigHierarchyController::ImportFromText(FString InContent
 			Hierarchy->SetParentWeight(NewElement, ParentIndex, PerElementData.ParentWeights[ParentIndex], false, true);
 		}
 
-		PastedKeys.Add(NewElement->GetKey());
+		PastedKeys.AddUnique(NewElement->GetKey());
 	}
 
 	for(const FRigHierarchyCopyPasteContentPerElement& PerElementData : Data.Elements)
@@ -1539,8 +1568,11 @@ TArray<FRigElementKey> URigHierarchyController::ImportFromText(FString InContent
 
 		if(FRigTransformElement* TransformElement = Cast<FRigTransformElement>(Element))
 		{
-			Hierarchy->SetTransform(TransformElement, PerElementData.Pose.Initial.Local.Transform, ERigTransformType::InitialLocal, true, true);
-			Hierarchy->SetTransform(TransformElement, PerElementData.Pose.Current.Local.Transform, ERigTransformType::CurrentLocal, true, true);
+			if(PerElementData.Poses.Num() >= 2)
+			{
+				Hierarchy->SetTransform(TransformElement, PerElementData.Poses[0], ERigTransformType::InitialLocal, true, true);
+				Hierarchy->SetTransform(TransformElement, PerElementData.Poses[1], ERigTransformType::CurrentLocal, true, true);
+			}
 		}
 	}
 	
@@ -1775,7 +1807,7 @@ TArray<FString> URigHierarchyController::GetAddBonePythonCommands(FRigBoneElemen
 		return Commands;
 	}
 	
-	FString TransformStr = RigVMPythonUtils::TransformToPythonString(Bone->Pose.Initial.Local.Transform);
+	FString TransformStr = RigVMPythonUtils::TransformToPythonString(Bone->GetTransform().Initial.Local.Get());
 	FString ParentKeyStr = "''";
 	if (Bone->ParentElement)
 	{
@@ -1795,7 +1827,7 @@ TArray<FString> URigHierarchyController::GetAddBonePythonCommands(FRigBoneElemen
 
 TArray<FString> URigHierarchyController::GetAddNullPythonCommands(FRigNullElement* Null) const
 {
-	FString TransformStr = RigVMPythonUtils::TransformToPythonString(Null->Pose.Initial.Local.Transform);
+	FString TransformStr = RigVMPythonUtils::TransformToPythonString(Null->GetTransform().Initial.Local.Get());
 
 	FString ParentKeyStr = "''";
 	if (Null->ParentConstraints.Num() > 0)
@@ -1813,7 +1845,7 @@ TArray<FString> URigHierarchyController::GetAddNullPythonCommands(FRigNullElemen
 TArray<FString> URigHierarchyController::GetAddControlPythonCommands(FRigControlElement* Control) const
 {
 	TArray<FString> Commands;
-	FString TransformStr = RigVMPythonUtils::TransformToPythonString(Control->Pose.Initial.Local.Transform);
+	FString TransformStr = RigVMPythonUtils::TransformToPythonString(Control->GetTransform().Initial.Local.Get());
 
 	FString ParentKeyStr = "''";
 	if (Control->ParentConstraints.Num() > 0)
@@ -1843,10 +1875,10 @@ TArray<FString> URigHierarchyController::GetAddControlPythonCommands(FRigControl
 		*SettingsStr,
 		*ValueStr));
 
-	Commands.Append(GetSetControlShapeTransformPythonCommands(Control, Control->Shape.Initial.Local.Transform, true));
+	Commands.Append(GetSetControlShapeTransformPythonCommands(Control, Control->GetShapeTransform().Initial.Local.Get(), true));
 	Commands.Append(GetSetControlValuePythonCommands(Control, Settings.MinimumValue, ERigControlValueType::Minimum));
 	Commands.Append(GetSetControlValuePythonCommands(Control, Settings.MaximumValue, ERigControlValueType::Maximum));
-	Commands.Append(GetSetControlOffsetTransformPythonCommands(Control, Control->Offset.Initial.Local.Transform, true, true));
+	Commands.Append(GetSetControlOffsetTransformPythonCommands(Control, Control->GetOffsetTransform().Initial.Local.Get(), true, true));
 	Commands.Append(GetSetControlValuePythonCommands(Control, Value, ERigControlValueType::Current));
 
 	return Commands;
@@ -1865,7 +1897,7 @@ TArray<FString> URigHierarchyController::GetAddCurvePythonCommands(FRigCurveElem
 TArray<FString> URigHierarchyController::GetAddPhysicsElementPythonCommands(FRigPhysicsElement* PhysicsElement) const
 {
 	TArray<FString> Commands;
-	FString TransformStr = RigVMPythonUtils::TransformToPythonString(PhysicsElement->Pose.Initial.Local.Transform);
+	FString TransformStr = RigVMPythonUtils::TransformToPythonString(PhysicsElement->GetTransform().Initial.Local.Get());
 	
 	FString ParentKeyStr = "''";
 	if (PhysicsElement->ParentElement)
@@ -1924,7 +1956,7 @@ TArray<FString> URigHierarchyController::GetAddConnectorPythonCommands(FRigConne
 TArray<FString> URigHierarchyController::GetAddSocketPythonCommands(FRigSocketElement* Socket) const
 {
 	TArray<FString> Commands;
-	FString TransformStr = RigVMPythonUtils::TransformToPythonString(Socket->Pose.Initial.Local.Transform);
+	FString TransformStr = RigVMPythonUtils::TransformToPythonString(Socket->GetTransform().Initial.Local.Get());
 
 	FString ParentKeyStr = "''";
 	if (Socket->ParentElement)
@@ -2045,8 +2077,8 @@ int32 URigHierarchyController::AddElement(FRigBaseElement* InElementToAdd, FRigB
 	InElementToAdd->SubIndex = Hierarchy->Num(InElementToAdd->Key.Type);
 	InElementToAdd->Index = Hierarchy->Elements.Add(InElementToAdd);
 	Hierarchy->ElementsPerType[URigHierarchy::RigElementTypeToFlatIndex(InElementToAdd->GetKey().Type)].Add(InElementToAdd);
-
 	Hierarchy->IndexLookup.Add(InElementToAdd->Key, InElementToAdd->Index);
+	Hierarchy->AllocateDefaultElementStorage(InElementToAdd, true);
 	Hierarchy->IncrementTopologyVersion();
 
 	FRigName DesiredName = InDesiredName;
@@ -2103,8 +2135,8 @@ int32 URigHierarchyController::AddElement(FRigBaseElement* InElementToAdd, FRigB
 	{
 		Hierarchy->GetControlShapeTransform(ControlElement, ERigTransformType::CurrentLocal);
 		Hierarchy->GetControlShapeTransform(ControlElement, ERigTransformType::InitialLocal);
-		ControlElement->Shape.MarkDirty(ERigTransformType::CurrentGlobal);
-		ControlElement->Shape.MarkDirty(ERigTransformType::InitialGlobal);
+		ControlElement->GetShapeDirtyState().MarkDirty(ERigTransformType::CurrentGlobal);
+		ControlElement->GetShapeDirtyState().MarkDirty(ERigTransformType::InitialGlobal);
 	}
 
 	// only notify once at the end
@@ -2265,7 +2297,8 @@ bool URigHierarchyController::RemoveElement(FRigBaseElement* InElement)
 			ControlElement->Settings.DrivenControls.Remove(InElement->GetKey());
 		}
 	}
-	
+
+	Hierarchy->DeallocateElementStorage(InElement);
 	Hierarchy->IncrementTopologyVersion();
 
 	Notify(ERigHierarchyNotification::ElementRemoved, InElement);
@@ -2771,15 +2804,15 @@ bool URigHierarchyController::AddParent(FRigBaseElement* InChild, FRigBaseElemen
 			{
 				Hierarchy->GetTransform(TransformElement, ERigTransformType::CurrentGlobal);
 				Hierarchy->GetTransform(TransformElement, ERigTransformType::InitialGlobal);
-				TransformElement->Pose.MarkDirty(ERigTransformType::CurrentLocal);
-				TransformElement->Pose.MarkDirty(ERigTransformType::InitialLocal);
+				TransformElement->GetDirtyState().MarkDirty(ERigTransformType::CurrentLocal);
+				TransformElement->GetDirtyState().MarkDirty(ERigTransformType::InitialLocal);
 			}
 			else
 			{
 				Hierarchy->GetTransform(TransformElement, ERigTransformType::CurrentLocal);
 				Hierarchy->GetTransform(TransformElement, ERigTransformType::InitialLocal);
-				TransformElement->Pose.MarkDirty(ERigTransformType::CurrentGlobal);
-				TransformElement->Pose.MarkDirty(ERigTransformType::InitialGlobal);
+				TransformElement->GetDirtyState().MarkDirty(ERigTransformType::CurrentGlobal);
+				TransformElement->GetDirtyState().MarkDirty(ERigTransformType::InitialGlobal);
 			}
 		}
 
@@ -2846,10 +2879,10 @@ bool URigHierarchyController::AddParent(FRigBaseElement* InChild, FRigBaseElemen
 		{
 			if(FRigControlElement* ControlElement = Cast<FRigControlElement>(MultiParentElement))
 			{
-				ControlElement->Offset.MarkDirty(ERigTransformType::CurrentGlobal);  
-				ControlElement->Offset.MarkDirty(ERigTransformType::InitialGlobal);
-				ControlElement->Shape.MarkDirty(ERigTransformType::CurrentGlobal);  
-				ControlElement->Shape.MarkDirty(ERigTransformType::InitialGlobal);
+				ControlElement->GetOffsetDirtyState().MarkDirty(ERigTransformType::CurrentGlobal);  
+				ControlElement->GetOffsetDirtyState().MarkDirty(ERigTransformType::InitialGlobal);
+				ControlElement->GetShapeDirtyState().MarkDirty(ERigTransformType::CurrentGlobal);  
+				ControlElement->GetShapeDirtyState().MarkDirty(ERigTransformType::InitialGlobal);
 			}
 		}
 
@@ -2963,15 +2996,15 @@ bool URigHierarchyController::RemoveParent(FRigBaseElement* InChild, FRigBaseEle
 			{
 				Hierarchy->GetTransform(SingleParentElement, ERigTransformType::CurrentGlobal);
 				Hierarchy->GetTransform(SingleParentElement, ERigTransformType::InitialGlobal);
-				SingleParentElement->Pose.MarkDirty(ERigTransformType::CurrentLocal);
-				SingleParentElement->Pose.MarkDirty(ERigTransformType::InitialLocal);
+				SingleParentElement->GetDirtyState().MarkDirty(ERigTransformType::CurrentLocal);
+				SingleParentElement->GetDirtyState().MarkDirty(ERigTransformType::InitialLocal);
 			}
 			else
 			{
 				Hierarchy->GetTransform(SingleParentElement, ERigTransformType::CurrentLocal);
 				Hierarchy->GetTransform(SingleParentElement, ERigTransformType::InitialLocal);
-				SingleParentElement->Pose.MarkDirty(ERigTransformType::CurrentGlobal);
-				SingleParentElement->Pose.MarkDirty(ERigTransformType::InitialGlobal);
+				SingleParentElement->GetDirtyState().MarkDirty(ERigTransformType::CurrentGlobal);
+				SingleParentElement->GetDirtyState().MarkDirty(ERigTransformType::InitialGlobal);
 			}
 
 			const FRigElementKey PreviousParentKey = SingleParentElement->ParentElement->GetKey();
@@ -3022,15 +3055,15 @@ bool URigHierarchyController::RemoveParent(FRigBaseElement* InChild, FRigBaseEle
 			{
 				Hierarchy->GetTransform(MultiParentElement, ERigTransformType::CurrentGlobal);
 				Hierarchy->GetTransform(MultiParentElement, ERigTransformType::InitialGlobal);
-				MultiParentElement->Pose.MarkDirty(ERigTransformType::CurrentLocal);
-				MultiParentElement->Pose.MarkDirty(ERigTransformType::InitialLocal);
+				MultiParentElement->GetDirtyState().MarkDirty(ERigTransformType::CurrentLocal);
+				MultiParentElement->GetDirtyState().MarkDirty(ERigTransformType::InitialLocal);
 			}
 			else
 			{
 				Hierarchy->GetTransform(MultiParentElement, ERigTransformType::CurrentLocal);
 				Hierarchy->GetTransform(MultiParentElement, ERigTransformType::InitialLocal);
-				MultiParentElement->Pose.MarkDirty(ERigTransformType::CurrentGlobal);
-				MultiParentElement->Pose.MarkDirty(ERigTransformType::InitialGlobal);
+				MultiParentElement->GetDirtyState().MarkDirty(ERigTransformType::CurrentGlobal);
+				MultiParentElement->GetDirtyState().MarkDirty(ERigTransformType::InitialGlobal);
 			}
 
 			// remove the previous parent
@@ -3051,10 +3084,10 @@ bool URigHierarchyController::RemoveParent(FRigBaseElement* InChild, FRigBaseEle
 
 			if(FRigControlElement* ControlElement = Cast<FRigControlElement>(MultiParentElement))
 			{
-				ControlElement->Offset.MarkDirty(ERigTransformType::CurrentGlobal);  
-				ControlElement->Offset.MarkDirty(ERigTransformType::InitialGlobal);
-				ControlElement->Shape.MarkDirty(ERigTransformType::CurrentGlobal);  
-				ControlElement->Shape.MarkDirty(ERigTransformType::InitialGlobal);
+				ControlElement->GetOffsetDirtyState().MarkDirty(ERigTransformType::CurrentGlobal);  
+				ControlElement->GetOffsetDirtyState().MarkDirty(ERigTransformType::InitialGlobal);
+				ControlElement->GetShapeDirtyState().MarkDirty(ERigTransformType::CurrentGlobal);  
+				ControlElement->GetShapeDirtyState().MarkDirty(ERigTransformType::InitialGlobal);
 			}
 
 			Hierarchy->IncrementTopologyVersion();
diff --git a/Engine/Plugins/Animation/ControlRig/Source/ControlRig/Private/Rigs/RigHierarchyElements.cpp b/Engine/Plugins/Animation/ControlRig/Source/ControlRig/Private/Rigs/RigHierarchyElements.cpp
--- a/Engine/Plugins/Animation/ControlRig/Source/ControlRig/Private/Rigs/RigHierarchyElements.cpp
+++ b/Engine/Plugins/Animation/ControlRig/Source/ControlRig/Private/Rigs/RigHierarchyElements.cpp
@@ -216,53 +216,239 @@ void FRigBaseElement::CopyFrom(const FRigBaseElement* InOther)
 {
 }
 
+////////////////////////////////////////////////////////////////////////////////
+// FRigTransformDirtyState
+////////////////////////////////////////////////////////////////////////////////
+
+const bool& FRigTransformDirtyState::Get() const
+{
+	if(Storage)
+	{
+		return *Storage;
+	}
+	ensure(false);
+	static constexpr bool bDefaultDirtyFlag = false;
+	return bDefaultDirtyFlag;
+}
+
+bool& FRigTransformDirtyState::Get()
+{
+	if(Storage)
+	{
+		return *Storage;
+	}
+	ensure(false);
+	static bool bDefaultDirtyFlag = false;
+	return bDefaultDirtyFlag;
+}
+
+bool FRigTransformDirtyState::Set(bool InDirty)
+{
+	if(Storage)
+	{
+		if(*Storage != InDirty)
+		{
+			*Storage = InDirty;
+			return true;
+		}
+	}
+	return false;
+}
+
+FRigTransformDirtyState& FRigTransformDirtyState::operator=(const FRigTransformDirtyState& InOther)
+{
+	if(Storage)
+	{
+		*Storage = InOther.Get();
+	}
+	return *this;
+}
+
+void FRigTransformDirtyState::LinkStorage(const TArrayView<bool>& InStorage)
+{
+	if(InStorage.IsValidIndex(Index))
+	{
+		Storage = InStorage.GetData() + Index;
+	}
+}
+
+void FRigTransformDirtyState::UnlinkStorage(FRigReusableElementStorage<bool>& InStorage)
+{
+	InStorage.Deallocate(Index, &Storage);
+}
+
+////////////////////////////////////////////////////////////////////////////////
+// FRigLocalAndGlobalDirtyState
+////////////////////////////////////////////////////////////////////////////////
+
+FRigLocalAndGlobalDirtyState& FRigLocalAndGlobalDirtyState::operator=(const FRigLocalAndGlobalDirtyState& InOther)
+{
+	Local = InOther.Local;
+	Global = InOther.Global;
+	return *this;
+}
+
+void FRigLocalAndGlobalDirtyState::LinkStorage(const TArrayView<bool>& InStorage)
+{
+	Local.LinkStorage(InStorage);
+	Global.LinkStorage(InStorage);
+}
+
+void FRigLocalAndGlobalDirtyState::UnlinkStorage(FRigReusableElementStorage<bool>& InStorage)
+{
+	Local.UnlinkStorage(InStorage);
+	Global.UnlinkStorage(InStorage);
+}
+
+////////////////////////////////////////////////////////////////////////////////
+// FRigCurrentAndInitialDirtyState
+////////////////////////////////////////////////////////////////////////////////
+
+FRigCurrentAndInitialDirtyState& FRigCurrentAndInitialDirtyState::operator=(const FRigCurrentAndInitialDirtyState& InOther)
+{
+	Current = InOther.Current;
+	Initial = InOther.Initial;
+	return *this;
+}
+
+void FRigCurrentAndInitialDirtyState::LinkStorage(const TArrayView<bool>& InStorage)
+{
+	Current.LinkStorage(InStorage);
+	Initial.LinkStorage(InStorage);
+}
+
+void FRigCurrentAndInitialDirtyState::UnlinkStorage(FRigReusableElementStorage<bool>& InStorage)
+{
+	Current.UnlinkStorage(InStorage);
+	Initial.UnlinkStorage(InStorage);
+}
 
 ////////////////////////////////////////////////////////////////////////////////
 // FRigComputedTransform
 ////////////////////////////////////////////////////////////////////////////////
 
-void FRigComputedTransform::Save(FArchive& Ar, bool& bDirty)
+void FRigComputedTransform::Save(FArchive& Ar, const FRigTransformDirtyState& InDirtyState)
+{
+	FTransform Transform = Get();
+	bool bDirty = InDirtyState.Get();
+	
+	Ar << Transform;
+	Ar << bDirty;
+}
+
+void FRigComputedTransform::Load(FArchive& Ar, FRigTransformDirtyState& InDirtyState)
 {
+	FTransform Transform = FTransform::Identity;
+	bool bDirty = false;
+
 	Ar << Transform;
 	Ar << bDirty;
+
+	Set(Transform);
+	(void)InDirtyState.Set(bDirty);
 }
 
-void FRigComputedTransform::Load(FArchive& Ar, bool& bDirty)
+const FTransform& FRigComputedTransform::Get() const
 {
-	// load and save are identical
-	Save(Ar, bDirty);
+	if(Storage)
+	{
+		return *Storage;
+	}
+	ensure(false);
+	static const FTransform DefaultTransform = FTransform::Identity;
+	return DefaultTransform;
+}
+
+FRigComputedTransform& FRigComputedTransform::operator=(const FRigComputedTransform& InOther)
+{
+	if(Storage)
+	{
+		*Storage = InOther.Get();
+	}
+	return *this;
+}
+
+void FRigComputedTransform::LinkStorage(const TArrayView<FTransform>& InStorage)
+{
+	if(InStorage.IsValidIndex(Index))
+	{
+		Storage = InStorage.GetData() + Index;
+	}
+}
+
+void FRigComputedTransform::UnlinkStorage(FRigReusableElementStorage<FTransform>& InStorage)
+{
+	InStorage.Deallocate(Index, &Storage);
 }
 
 ////////////////////////////////////////////////////////////////////////////////
 // FRigLocalAndGlobalTransform
 ////////////////////////////////////////////////////////////////////////////////
 
-void FRigLocalAndGlobalTransform::Save(FArchive& Ar)
+void FRigLocalAndGlobalTransform::Save(FArchive& Ar, const FRigLocalAndGlobalDirtyState& InDirtyState)
+{
+	Local.Save(Ar, InDirtyState.Local);
+	Global.Save(Ar, InDirtyState.Global);
+}
+
+void FRigLocalAndGlobalTransform::Load(FArchive& Ar, FRigLocalAndGlobalDirtyState& OutDirtyState)
 {
-	Local.Save(Ar, bDirty[ELocal]);
-	Global.Save(Ar, bDirty[EGlobal]);
+	Local.Load(Ar, OutDirtyState.Local);
+	Global.Load(Ar, OutDirtyState.Global);
 }
 
-void FRigLocalAndGlobalTransform::Load(FArchive& Ar)
+FRigLocalAndGlobalTransform& FRigLocalAndGlobalTransform::operator=(const FRigLocalAndGlobalTransform& InOther)
 {
-	Local.Load(Ar, bDirty[ELocal]);
-	Global.Load(Ar, bDirty[EGlobal]);
+	Local = InOther.Local;
+	Global = InOther.Global;
+	return *this;
+}
+
+void FRigLocalAndGlobalTransform::LinkStorage(const TArrayView<FTransform>& InStorage)
+{
+	Local.LinkStorage(InStorage);
+	Global.LinkStorage(InStorage);
+}
+
+void FRigLocalAndGlobalTransform::UnlinkStorage(FRigReusableElementStorage<FTransform>& InStorage)
+{
+	Local.UnlinkStorage(InStorage);
+	Global.UnlinkStorage(InStorage);
 }
 
 ////////////////////////////////////////////////////////////////////////////////
 // FRigCurrentAndInitialTransform
 ////////////////////////////////////////////////////////////////////////////////
 
-void FRigCurrentAndInitialTransform::Save(FArchive& Ar)
+void FRigCurrentAndInitialTransform::Save(FArchive& Ar, const FRigCurrentAndInitialDirtyState& InDirtyState)
 {
-	Current.Save(Ar);
-	Initial.Save(Ar);
+	Current.Save(Ar, InDirtyState.Current);
+	Initial.Save(Ar, InDirtyState.Initial);
 }
 
-void FRigCurrentAndInitialTransform::Load(FArchive& Ar)
+void FRigCurrentAndInitialTransform::Load(FArchive& Ar, FRigCurrentAndInitialDirtyState& OutDirtyState)
 {
-	Current.Load(Ar);
-	Initial.Load(Ar);
+	Current.Load(Ar, OutDirtyState.Current);
+	Initial.Load(Ar, OutDirtyState.Initial);
+}
+
+FRigCurrentAndInitialTransform& FRigCurrentAndInitialTransform::operator=(const FRigCurrentAndInitialTransform& InOther)
+{
+	Current = InOther.Current;
+	Initial = InOther.Initial;
+	return *this;
+}
+
+void FRigCurrentAndInitialTransform::LinkStorage(const TArrayView<FTransform>& InStorage)
+{
+	Current.LinkStorage(InStorage);
+	Initial.LinkStorage(InStorage);
+}
+
+void FRigCurrentAndInitialTransform::UnlinkStorage(FRigReusableElementStorage<FTransform>& InStorage)
+{
+	Current.UnlinkStorage(InStorage);
+	Initial.UnlinkStorage(InStorage);
 }
 
 ////////////////////////////////////////////////////////////////////////////////
@@ -416,7 +602,7 @@ void FRigTransformElement::Save(FArchive& Ar, ESerializationPhase SerializationP
 
 	if(SerializationPhase == ESerializationPhase::StaticData)
 	{
-		Pose.Save(Ar);
+		GetTransform().Save(Ar, GetDirtyState());
 	}
 }
 
@@ -426,7 +612,7 @@ void FRigTransformElement::Load(FArchive& Ar, ESerializationPhase SerializationP
 
 	if(SerializationPhase == ESerializationPhase::StaticData)
 	{
-		Pose.Load(Ar);
+		GetTransform().Load(Ar, GetDirtyState());
 	}
 }
 
@@ -438,21 +624,44 @@ void FRigTransformElement::CopyPose(FRigBaseElement* InOther, bool bCurrent, boo
 	{
 		if(bCurrent)
 		{
-			Pose.Current = Other->Pose.Current;
+			GetTransform().Current = Other->GetTransform().Current;
+			GetDirtyState().Current = Other->GetDirtyState().Current; 
 		}
 		if(bInitial)
 		{
-			Pose.Initial = Other->Pose.Initial;
+			GetTransform().Initial = Other->GetTransform().Initial;
+			GetDirtyState().Initial = Other->GetDirtyState().Initial; 
 		}
 	}
 }
 
+const FRigCurrentAndInitialTransform& FRigTransformElement::GetTransform() const
+{
+	return PoseStorage;
+}
+
+FRigCurrentAndInitialTransform& FRigTransformElement::GetTransform()
+{
+	return PoseStorage;
+}
+
+const FRigCurrentAndInitialDirtyState& FRigTransformElement::GetDirtyState() const
+{
+	return PoseDirtyState;
+}
+
+FRigCurrentAndInitialDirtyState& FRigTransformElement::GetDirtyState()
+{
+	return PoseDirtyState;
+}
+
 void FRigTransformElement::CopyFrom(const FRigBaseElement* InOther)
 {
 	Super::CopyFrom(InOther);
 	
 	const FRigTransformElement* SourceTransform = CastChecked<FRigTransformElement>(InOther);
-	Pose = SourceTransform->Pose;
+	GetTransform() = SourceTransform->GetTransform();
+	GetDirtyState() = SourceTransform->GetDirtyState();
 
 	ElementsToDirty.Reset();
 	ElementsToDirty.Reserve(SourceTransform->ElementsToDirty.Num());
@@ -467,6 +676,22 @@ void FRigTransformElement::CopyFrom(const FRigBaseElement* InOther)
 	}
 }
 
+void FRigTransformElement::LinkStorage(const TArrayView<FTransform>& InTransforms, const TArrayView<bool>& InDirtyStates,
+	const TArrayView<float>& InCurves)
+{
+	FRigBaseElement::LinkStorage(InTransforms, InDirtyStates, InCurves);
+	PoseStorage.LinkStorage(InTransforms);
+	PoseDirtyState.LinkStorage(InDirtyStates);
+}
+
+void FRigTransformElement::UnlinkStorage(FRigReusableElementStorage<FTransform>& InTransforms,
+	FRigReusableElementStorage<bool>& InDirtyStates, FRigReusableElementStorage<float>& InCurves)
+{
+	FRigBaseElement::UnlinkStorage(InTransforms, InDirtyStates, InCurves);
+	PoseStorage.UnlinkStorage(InTransforms);
+	PoseDirtyState.UnlinkStorage(InDirtyStates);
+}
+
 ////////////////////////////////////////////////////////////////////////////////
 // FRigSingleParentElement
 ////////////////////////////////////////////////////////////////////////////////
@@ -561,7 +786,8 @@ void FRigMultiParentElement::Load(FArchive& Ar, ESerializationPhase Serializatio
 		if (Ar.CustomVer(FControlRigObjectVersion::GUID) < FControlRigObjectVersion::RemovedMultiParentParentCache)
 		{
 			FRigCurrentAndInitialTransform Parent;
-			Parent.Load(Ar);
+			FRigCurrentAndInitialDirtyState DirtyState;
+			Parent.Load(Ar, DirtyState);
 		}
 
 		int32 NumParents = 0;
@@ -1177,6 +1403,46 @@ void FRigControlSettings::SetupLimitArrayForType(bool bLimitTranslation, bool bL
 	}
 }
 
+const FRigCurrentAndInitialTransform& FRigControlElement::GetOffsetTransform() const
+{
+	return OffsetStorage;
+}
+
+FRigCurrentAndInitialTransform& FRigControlElement::GetOffsetTransform()
+{
+	return OffsetStorage;
+}
+
+const FRigCurrentAndInitialDirtyState& FRigControlElement::GetOffsetDirtyState() const
+{
+	return OffsetDirtyState;
+}
+
+FRigCurrentAndInitialDirtyState& FRigControlElement::GetOffsetDirtyState()
+{
+	return OffsetDirtyState;
+}
+
+const FRigCurrentAndInitialTransform& FRigControlElement::GetShapeTransform() const
+{
+	return ShapeStorage;
+}
+
+FRigCurrentAndInitialTransform& FRigControlElement::GetShapeTransform()
+{
+	return ShapeStorage;
+}
+
+const FRigCurrentAndInitialDirtyState& FRigControlElement::GetShapeDirtyState() const
+{
+	return ShapeDirtyState;
+}
+
+FRigCurrentAndInitialDirtyState& FRigControlElement::GetShapeDirtyState()
+{
+	return ShapeDirtyState;
+}
+
 ////////////////////////////////////////////////////////////////////////////////
 // FRigControlElement
 ////////////////////////////////////////////////////////////////////////////////
@@ -1190,8 +1456,8 @@ void FRigControlElement::Save(FArchive& Ar, ESerializationPhase SerializationPha
 	if(SerializationPhase == ESerializationPhase::StaticData)
 	{
 		Settings.Save(Ar);
-		Offset.Save(Ar);
-		Shape.Save(Ar);
+		GetOffsetTransform().Save(Ar, GetOffsetDirtyState());
+		GetShapeTransform().Save(Ar, GetShapeDirtyState());
 		PreferredEulerAngles.Save(Ar);
 	}
 }
@@ -1203,8 +1469,8 @@ void FRigControlElement::Load(FArchive& Ar, ESerializationPhase SerializationPha
 	if(SerializationPhase == ESerializationPhase::StaticData)
 	{
 		Settings.Load(Ar);
-		Offset.Load(Ar);
-		Shape.Load(Ar);
+		GetOffsetTransform().Load(Ar, GetOffsetDirtyState());
+		GetShapeTransform().Load(Ar, GetShapeDirtyState());
 
 		if (Ar.CustomVer(FControlRigObjectVersion::GUID) >= FControlRigObjectVersion::PreferredEulerAnglesForControls)
 		{
@@ -1224,8 +1490,10 @@ void FRigControlElement::CopyFrom(const FRigBaseElement* InOther)
 	
 	const FRigControlElement* Source = CastChecked<FRigControlElement>(InOther);
 	Settings = Source->Settings;
-	Offset = Source->Offset;
-	Shape = Source->Shape;
+	GetOffsetTransform() = Source->GetOffsetTransform();
+	GetOffsetDirtyState() = Source->GetOffsetDirtyState();
+	GetShapeTransform() = Source->GetShapeTransform();
+	GetShapeDirtyState() = Source->GetShapeDirtyState();
 	PreferredEulerAngles = Source->PreferredEulerAngles;
 }
 
@@ -1237,19 +1505,43 @@ void FRigControlElement::CopyPose(FRigBaseElement* InOther, bool bCurrent, bool
 	{
 		if(bCurrent)
 		{
-			Offset.Current = Other->Offset.Current;
-			Shape.Current = Other->Shape.Current;
+			GetOffsetTransform().Current = Other->GetOffsetTransform().Current;
+			GetOffsetDirtyState().Current = Other->GetOffsetDirtyState().Current;
+			GetShapeTransform().Current = Other->GetShapeTransform().Current;
+			GetShapeDirtyState().Current = Other->GetShapeDirtyState().Current;
 			PreferredEulerAngles.SetAngles(Other->PreferredEulerAngles.GetAngles(false), false);
 		}
 		if(bInitial)
 		{
-			Offset.Initial = Other->Offset.Initial;
-			Shape.Initial = Other->Shape.Initial;
+			GetOffsetTransform().Initial = Other->GetOffsetTransform().Initial;
+			GetOffsetDirtyState().Initial = Other->GetOffsetDirtyState().Initial;
+			GetShapeTransform().Initial = Other->GetShapeTransform().Initial;
+			GetShapeDirtyState().Initial = Other->GetShapeDirtyState().Initial;
 			PreferredEulerAngles.SetAngles(Other->PreferredEulerAngles.GetAngles(true), true);
 		}
 	}
 }
 
+void FRigControlElement::LinkStorage(const TArrayView<FTransform>& InTransforms, const TArrayView<bool>& InDirtyStates,
+	const TArrayView<float>& InCurves)
+{
+	FRigMultiParentElement::LinkStorage(InTransforms, InDirtyStates, InCurves);
+	OffsetStorage.LinkStorage(InTransforms);
+	ShapeStorage.LinkStorage(InTransforms);
+	OffsetDirtyState.LinkStorage(InDirtyStates);
+	ShapeDirtyState.LinkStorage(InDirtyStates);
+}
+
+void FRigControlElement::UnlinkStorage(FRigReusableElementStorage<FTransform>& InTransforms, FRigReusableElementStorage<bool>& InDirtyStates,
+	FRigReusableElementStorage<float>& InCurves)
+{
+	FRigMultiParentElement::UnlinkStorage(InTransforms, InDirtyStates, InCurves);
+	OffsetStorage.UnlinkStorage(InTransforms);
+	ShapeStorage.UnlinkStorage(InTransforms);
+	OffsetDirtyState.UnlinkStorage(InDirtyStates);
+	ShapeDirtyState.UnlinkStorage(InDirtyStates);
+}
+
 ////////////////////////////////////////////////////////////////////////////////
 // FRigCurveElement
 ////////////////////////////////////////////////////////////////////////////////
@@ -1262,6 +1554,7 @@ void FRigCurveElement::Save(FArchive& Ar, ESerializationPhase SerializationPhase
 
 	if(SerializationPhase == ESerializationPhase::StaticData)
 	{
+		float Value = Get();
 		Ar << bIsValueSet;
 		Ar << Value;
 	}
@@ -1281,7 +1574,13 @@ void FRigCurveElement::Load(FArchive& Ar, ESerializationPhase SerializationPhase
 		{
 			bIsValueSet = true;
 		}
+
+		float Value = 0.f;
 		Ar << Value;
+
+		// maintain the deserialized bIsValueSet flag
+		TGuardValue<bool> ValueIsSetGuard(bIsValueSet, bIsValueSet);
+		Set(Value);
 	}
 }
 
@@ -1291,19 +1590,56 @@ void FRigCurveElement::CopyPose(FRigBaseElement* InOther, bool bCurrent, bool bI
 	
 	if(const FRigCurveElement* Other = Cast<FRigCurveElement>(InOther))
 	{
+		Set(Other->Get());
 		bIsValueSet = Other->bIsValueSet;
-		Value = Other->Value;
 	}
 }
 
+const float& FRigCurveElement::Get() const
+{
+	if(Storage)
+	{
+		return *Storage;
+	}
+	ensure(false);
+	static constexpr float DefaultCurve = 0.f;
+	return DefaultCurve;
+}
+
+void FRigCurveElement::Set(const float& InValue)
+{
+	if(Storage)
+	{
+		*Storage = InValue;
+		bIsValueSet = true;
+	}
+}
+
+void FRigCurveElement::LinkStorage(const TArrayView<FTransform>& InTransforms, const TArrayView<bool>& InDirtyStates,
+	const TArrayView<float>& InCurves)
+{
+	FRigBaseElement::LinkStorage(InTransforms, InDirtyStates, InCurves);
+	if(InCurves.IsValidIndex(StorageIndex))
+	{
+		Storage = InCurves.GetData() + StorageIndex;
+	}
+}
+
+void FRigCurveElement::UnlinkStorage(FRigReusableElementStorage<FTransform>& InTransforms, FRigReusableElementStorage<bool>& InDirtyStates,
+	FRigReusableElementStorage<float>& InCurves)
+{
+	FRigBaseElement::UnlinkStorage(InTransforms, InDirtyStates, InCurves);
+	InCurves.Deallocate(StorageIndex, &Storage);
+}
+
 void FRigCurveElement::CopyFrom(const FRigBaseElement* InOther)
 {
 	Super::CopyFrom(InOther);
 	
 	if(const FRigCurveElement* Other = CastChecked<FRigCurveElement>(InOther))
 	{
+		Set(Other->Get());
 		bIsValueSet = Other->bIsValueSet;
-		Value = Other->Value;
 	}
 }
 
diff --git a/Engine/Plugins/Animation/ControlRig/Source/ControlRig/Private/Units/Execution/RigUnit_DynamicHierarchy.cpp b/Engine/Plugins/Animation/ControlRig/Source/ControlRig/Private/Units/Execution/RigUnit_DynamicHierarchy.cpp
--- a/Engine/Plugins/Animation/ControlRig/Source/ControlRig/Private/Units/Execution/RigUnit_DynamicHierarchy.cpp
+++ b/Engine/Plugins/Animation/ControlRig/Source/ControlRig/Private/Units/Execution/RigUnit_DynamicHierarchy.cpp
@@ -415,7 +415,7 @@ void FRigUnit_HierarchyAddControl_ShapeSettings::ConfigureFrom(const FRigControl
 	bVisible = InSettings.bShapeVisible;
 	Name = InSettings.ShapeName;
 	Color = InSettings.ShapeColor;
-	Transform = InControlElement->Shape.Get(ERigTransformType::InitialLocal);
+	Transform = InControlElement->GetShapeTransform().Get(ERigTransformType::InitialLocal);
 }
 
 void FRigUnit_HierarchyAddControl_ShapeSettings::Configure(FRigControlSettings& OutSettings) const
diff --git a/Engine/Plugins/Animation/ControlRig/Source/ControlRig/Public/ControlRigObjectVersion.h b/Engine/Plugins/Animation/ControlRig/Source/ControlRig/Public/ControlRigObjectVersion.h
--- a/Engine/Plugins/Animation/ControlRig/Source/ControlRig/Public/ControlRigObjectVersion.h
+++ b/Engine/Plugins/Animation/ControlRig/Source/ControlRig/Public/ControlRigObjectVersion.h
@@ -104,6 +104,9 @@ struct CONTROLRIG_API FControlRigObjectVersion
 		// Physics solvers stored on hierarchy
 		ControlRigStoresPhysicsSolvers,
 
+		// Moved the element storage into separate buffers
+		RigHierarchyIndirectElementStorage,
+
 		// -----<new versions can be added above this line>-------------------------------------------------
 		VersionPlusOne,
 		LatestVersion = VersionPlusOne - 1
diff --git a/Engine/Plugins/Animation/ControlRig/Source/ControlRig/Public/Rigs/RigHierarchy.h b/Engine/Plugins/Animation/ControlRig/Source/ControlRig/Public/Rigs/RigHierarchy.h
--- a/Engine/Plugins/Animation/ControlRig/Source/ControlRig/Public/Rigs/RigHierarchy.h
+++ b/Engine/Plugins/Animation/ControlRig/Source/ControlRig/Public/Rigs/RigHierarchy.h
@@ -4185,7 +4185,7 @@ private:
 	 * Templated helper function to create an element
 	 */
 	template<typename ElementType = FRigBaseElement>
-	ElementType* NewElement(int32 Num = 1)
+	ElementType* NewElement(int32 Num = 1, bool bAllocateStorage = false)
 	{
 		ElementType* NewElements = static_cast<ElementType*>(FMemory::Malloc(sizeof(ElementType) * Num));
 		for(int32 Index=0;Index<Num;Index++)
@@ -4193,6 +4193,13 @@ private:
 			new(&NewElements[Index]) ElementType(this);
 		}
 		NewElements[0].OwnedInstances = Num;
+		if(bAllocateStorage)
+		{
+			for(int32 Index=0;Index<Num;Index++)
+			{
+				AllocateDefaultElementStorage(&NewElements[Index], false);
+			}
+		}
 		return NewElements;
 	}
 
@@ -4283,14 +4290,38 @@ private:
 
 	TMap<FRigElementKey, FString> UserDefinedElementName;
 
-	// Element metadata storage. Storage is defined here rather than on the elements
+	// Per element pose storage. Storage is defined here rather than on the elements
+	// to reduce memory consumption. Only elements created by MakeElement point to
+	// the element storage. Copied elements via the copy constructor or copy operator
+	// do not have URigHierarchy as an owner and therefore do not carry poses with them.
+	FRigReusableElementStorage<FTransform> ElementTransforms;
+
+	// Per element dirty state storage. Storage is defined here rather than on the elements
 	// to reduce memory consumption. Only elements created by MakeElement point to
 	// the element storage. Copied elements via the copy constructor or copy operator
 	// do not have URigHierarchy as an owner and therefore do not carry metadata with them.
-	TArray<FMetadataStorage> ElementMetadata;
+	FRigReusableElementStorage<bool> ElementDirtyStates;
+
+	// Per element curve storage. Storage is defined here rather than on the elements
+	// to reduce memory consumption. Only elements created by MakeElement point to
+	// the element storage. Copied elements via the copy constructor or copy operator
+	// do not have URigHierarchy as an owner and therefore do not carry curves with them.
+	FRigReusableElementStorage<float> ElementCurves;
+
+	// Allocates the default element storage for an element
+	void AllocateDefaultElementStorage(FRigBaseElement* InElement, bool bUpdateAllElements);
+
+	// Deallocates the default element storage for an element
+	void DeallocateElementStorage(FRigBaseElement* InElement);
+
+	// Updates all storage pointers of the elements for poses and dirty states
+	void UpdateElementStorage();
 
-	// List of metadata storage entries that have been freed and can be recycled.
-	TArray<int32> ElementMetadataFreeList;
+	// Element metadata storage. Storage is defined here rather than on the elements
+	// to reduce memory consumption. Only elements created by MakeElement point to
+	// the element storage. Copied elements via the copy constructor or copy operator
+	// do not have URigHierarchy as an owner and therefore do not carry metadata with them.
+	FRigReusableElementStorage<FMetadataStorage> ElementMetadata;
 
 	// A quick-lookup cache for elements' children. Each element that has a ChildCacheIndex
 	// not equal to INDEX_NONE is an index into the offset and count cache below, which in
diff --git a/Engine/Plugins/Animation/ControlRig/Source/ControlRig/Public/Rigs/RigHierarchyController.h b/Engine/Plugins/Animation/ControlRig/Source/ControlRig/Public/Rigs/RigHierarchyController.h
--- a/Engine/Plugins/Animation/ControlRig/Source/ControlRig/Public/Rigs/RigHierarchyController.h
+++ b/Engine/Plugins/Animation/ControlRig/Source/ControlRig/Public/Rigs/RigHierarchyController.h
@@ -850,9 +850,9 @@ private:
 	TFunction<void(EMessageSeverity::Type,const FString&)> LogFunction = nullptr;
 
 	template<typename T>
-	T* MakeElement()
+	T* MakeElement(bool bAllocateStorage = false)
 	{
-		T* Element = GetHierarchy()->NewElement<T>();
+		T* Element = GetHierarchy()->NewElement<T>(1, bAllocateStorage);
 		Element->CreatedAtInstructionIndex = CurrentInstructionIndex;
 		return Element;
 	}
diff --git a/Engine/Plugins/Animation/ControlRig/Source/ControlRig/Public/Rigs/RigHierarchyElements.h b/Engine/Plugins/Animation/ControlRig/Source/ControlRig/Public/Rigs/RigHierarchyElements.h
--- a/Engine/Plugins/Animation/ControlRig/Source/ControlRig/Public/Rigs/RigHierarchyElements.h
+++ b/Engine/Plugins/Animation/ControlRig/Source/ControlRig/Public/Rigs/RigHierarchyElements.h
@@ -242,19 +242,305 @@ namespace ERigTransformType
 	}
 }
 
+template<typename T>
+struct CONTROLRIG_API FRigReusableElementStorage
+{
+	TArray<T> Storage;
+	TArray<int32> FreeList;
+
+	bool IsValidIndex(const int32& InIndex) const
+	{
+		return Storage.IsValidIndex(InIndex);
+	}
+
+	int32 Num() const
+	{
+		return Storage.Num();
+	}
+
+	T& operator[](int InIndex)
+	{
+		return Storage[InIndex];
+	}
+
+	const T& operator[](int InIndex) const
+	{
+		return Storage[InIndex];
+	}
+
+	typename TArray<T>::RangedForIteratorType begin() { return Storage.begin(); }
+	typename TArray<T>::RangedForIteratorType end() { return Storage.end(); }
+	typename TArray<T>::RangedForConstIteratorType begin() const { return Storage.begin(); }
+	typename TArray<T>::RangedForConstIteratorType end() const { return Storage.end(); }
+
+	const T* GetData() const
+	{
+		return Storage.GetData();
+	}
+
+	TArray<int32, TInlineAllocator<4>> Allocate(int32 InCount, const T& InDefault)
+	{
+		TArray<int32, TInlineAllocator<4>> Indices;
+	
+		const int32 NumToAllocate = InCount - FMath::Min(InCount, FreeList.Num());
+		if(NumToAllocate > 0)
+		{
+			Storage.Reserve(Storage.Num() + NumToAllocate);
+		}
+	
+		for(int32 Index = 0; Index < InCount; Index++)
+		{
+			if(FreeList.IsEmpty())
+			{
+				Indices.Push(Storage.Add(InDefault));
+			}
+			else
+			{
+				Indices.Push(FreeList.Pop(EAllowShrinking::No));
+			}
+		}
+
+		return Indices;
+	}
+	
+	void Deallocate(int32& InIndex, T** InStorage = nullptr)
+	{
+		if(InIndex == INDEX_NONE)
+		{
+			return;
+		}
+#if WITH_EDITOR
+		check(Storage.IsValidIndex(InIndex));
+		check(!FreeList.Contains(InIndex));
+#endif
+		FreeList.Add(InIndex);
+		InIndex = INDEX_NONE;
+		if(InStorage)
+		{
+			InStorage = nullptr;
+		}
+	}
+
+	void Reset(TFunction<void(int32, T&)> OnDestroyCallback = nullptr)
+	{
+		if(OnDestroyCallback)
+		{
+			for(int32 Index = 0; Index < Storage.Num(); Index++)
+			{
+				OnDestroyCallback(Index, Storage[Index]);
+			}
+		}
+		Storage.Reset();
+		FreeList.Reset();
+	}
+
+	bool Contains(int32 InIndex, const T* InStorage)
+	{
+		if(!IsValidIndex(InIndex))
+		{
+			return false;
+		}
+		return GetData() + InIndex == InStorage;
+	}
+
+	TMap<int32, int32> Shrink(TFunction<void(int32, T&)> OnDestroyCallback = nullptr)
+	{
+		TMap<int32, int32> OldToNew;
+		
+		if(!FreeList.IsEmpty())
+		{
+			TArray<bool> ToRemove;
+			ToRemove.AddZeroed(Storage.Num());
+
+			TArray<T> NewStorage;
+			
+			for(int32 FreeIndex : FreeList)
+			{
+				ToRemove[FreeIndex] = true;
+				if(OnDestroyCallback)
+				{
+					OnDestroyCallback(FreeIndex, Storage[FreeIndex]);
+				}
+			}
+			
+			for(int32 OldIndex = 0; OldIndex < Storage.Num(); OldIndex++)
+			{
+				if(!ToRemove[OldIndex])
+				{
+					OldToNew.Add(OldIndex, NewStorage.Add(Storage[OldIndex]));
+				}
+			}
+
+			Storage = NewStorage;
+			FreeList.Reset();
+		}
+
+		FreeList.Shrink();
+		Storage.Shrink();
+
+		return OldToNew;
+	}
+};
+
+USTRUCT(BlueprintType)
+struct CONTROLRIG_API FRigTransformDirtyState
+{
+public:
+
+	GENERATED_BODY()
+
+	FRigTransformDirtyState()
+	: Index(INDEX_NONE)
+	, Storage(nullptr)
+	{
+	}
+
+	const bool& Get() const;
+	bool& Get();
+	bool Set(bool InDirty);
+	FRigTransformDirtyState& operator =(const FRigTransformDirtyState& InOther);
+
+	UPROPERTY(BlueprintReadOnly, EditAnywhere, Category = "DirtyState")
+	int32 Index;
+
+private:
+
+	void LinkStorage(const TArrayView<bool>& InStorage);
+	void UnlinkStorage(FRigReusableElementStorage<bool>& InStorage);
+	bool* Storage;
+
+	friend struct FRigLocalAndGlobalDirtyState;
+	friend class URigHierarchy;
+	friend class URigHierarchyController;
+};
+
+USTRUCT(BlueprintType)
+struct CONTROLRIG_API FRigLocalAndGlobalDirtyState
+{
+	GENERATED_BODY()
+
+public:
+	
+	FRigLocalAndGlobalDirtyState()
+	{
+	}
+
+	FRigLocalAndGlobalDirtyState& operator =(const FRigLocalAndGlobalDirtyState& InOther);
+
+	UPROPERTY(BlueprintReadOnly, EditAnywhere, Category = "DirtyState")
+	FRigTransformDirtyState Global;
+
+	UPROPERTY(BlueprintReadOnly, EditAnywhere, Category = "DirtyState")
+	FRigTransformDirtyState Local;
+
+private:
+
+	void LinkStorage(const TArrayView<bool>& InStorage);
+	void UnlinkStorage(FRigReusableElementStorage<bool>& InStorage);
+
+	friend struct FRigCurrentAndInitialDirtyState;
+};
+
+USTRUCT(BlueprintType)
+struct CONTROLRIG_API FRigCurrentAndInitialDirtyState
+{
+	GENERATED_BODY()
+
+public:
+	
+	FRigCurrentAndInitialDirtyState()
+	{
+	}
+
+	bool& GetDirtyFlag(const ERigTransformType::Type InTransformType)
+	{
+		switch (InTransformType)
+		{
+			case ERigTransformType::CurrentLocal:
+				return Current.Local.Get();
+
+			case ERigTransformType::CurrentGlobal:
+				return Current.Global.Get();
+
+			case ERigTransformType::InitialLocal:
+				return Initial.Local.Get();
+				
+			default:
+				return Initial.Global.Get();
+		}
+	}
+
+	const bool& GetDirtyFlag(const ERigTransformType::Type InTransformType) const
+	{
+		switch (InTransformType)
+		{
+			case ERigTransformType::CurrentLocal:
+				return Current.Local.Get();
+
+			case ERigTransformType::CurrentGlobal:
+				return Current.Global.Get();
+
+			case ERigTransformType::InitialLocal:
+				return Initial.Local.Get();
+
+			default:
+				return Initial.Global.Get();
+		}
+	}
+
+	bool IsDirty(const ERigTransformType::Type InTransformType) const
+	{
+		return GetDirtyFlag(InTransformType);
+	}
+
+	void MarkDirty(const ERigTransformType::Type InTransformType)
+	{
+		ensure(!(GetDirtyFlag(ERigTransformType::SwapLocalAndGlobal(InTransformType))));
+		GetDirtyFlag(InTransformType) = true;
+	}
+
+	void MarkClean(const ERigTransformType::Type InTransformType)
+	{
+		GetDirtyFlag(InTransformType) = false;
+	}
+
+	FRigCurrentAndInitialDirtyState& operator =(const FRigCurrentAndInitialDirtyState& InOther);
+
+	UPROPERTY(BlueprintReadOnly, EditAnywhere, Category = "DirtyState")
+	FRigLocalAndGlobalDirtyState Current;
+
+	UPROPERTY(BlueprintReadOnly, EditAnywhere, Category = "DirtyState")
+	FRigLocalAndGlobalDirtyState Initial;
+
+private:
+
+	void LinkStorage(const TArrayView<bool>& InStorage);
+	void UnlinkStorage(FRigReusableElementStorage<bool>& InStorage);
+
+	friend class URigHierarchy;
+	friend class URigHierarchyController;
+	friend struct FRigTransformElement;
+	friend struct FRigControlElement;
+};
+
 USTRUCT(BlueprintType)
 struct CONTROLRIG_API FRigComputedTransform
 {
 	GENERATED_BODY()
 
+public:
+	
 	FRigComputedTransform()
-	: Transform(FTransform::Identity)
+	: Index(INDEX_NONE)
+	, Storage(nullptr)
 	{}
 
-	void Save(FArchive& Ar, bool& bDirty);
-	void Load(FArchive& Ar, bool& bDirty);
+	void Save(FArchive& Ar, const FRigTransformDirtyState& InDirtyState);
+	void Load(FArchive& Ar, FRigTransformDirtyState& InDirtyState);
+	
+	const FTransform& Get() const;
 
-	void Set(const FTransform& InTransform, bool& bDirty)
+	void Set(const FTransform& InTransform)
 	{
 #if WITH_EDITOR
 		ensure(InTransform.GetRotation().IsNormalized());
@@ -262,8 +548,10 @@ struct CONTROLRIG_API FRigComputedTransform
 		// ensure(!FMath::IsNearlyZero(InTransform.GetScale3D().X));
 		// ensure(!FMath::IsNearlyZero(InTransform.GetScale3D().Y));
 		// ensure(!FMath::IsNearlyZero(InTransform.GetScale3D().Z));
-		Transform = InTransform;
-		bDirty = false;
+		if(Storage)
+		{
+			*Storage = InTransform;
+		}
 	}
 
 	static bool Equals(const FTransform& A, const FTransform& B, const float InTolerance = 0.0001f)
@@ -275,32 +563,42 @@ struct CONTROLRIG_API FRigComputedTransform
 
 	bool operator == (const FRigComputedTransform& Other) const
 	{
-		return Equals(Transform, Other.Transform);
+		return Equals(Get(), Other.Get());
     }
 
-	UPROPERTY(BlueprintReadOnly, EditAnywhere, Category = "Pose")
-	FTransform Transform;
+	FRigComputedTransform& operator =(const FRigComputedTransform& InOther);
+
+	UPROPERTY(BlueprintReadOnly, EditAnywhere, Category = "Transform")
+	int32 Index;
+
+private:
+
+	void LinkStorage(const TArrayView<FTransform>& InStorage);
+	void UnlinkStorage(FRigReusableElementStorage<FTransform>& InStorage);
+	
+	FTransform* Storage;
+	
+	friend struct FRigLocalAndGlobalTransform;
+	friend class URigHierarchy;
+	friend class URigHierarchyController;
 };
 
 USTRUCT(BlueprintType)
 struct CONTROLRIG_API FRigLocalAndGlobalTransform
 {
 	GENERATED_BODY()
+
+public:
 	
 	FRigLocalAndGlobalTransform()
     : Local()
     , Global()
 	{}
 
-	enum EDirty
-	{
-		ELocal = 0,
-		EGlobal,
-		EDirtyMax
-	};
+	void Save(FArchive& Ar, const FRigLocalAndGlobalDirtyState& InDirtyState);
+	void Load(FArchive& Ar, FRigLocalAndGlobalDirtyState& OutDirtyState);
 
-	void Save(FArchive& Ar);
-	void Load(FArchive& Ar);
+	FRigLocalAndGlobalTransform& operator =(const FRigLocalAndGlobalTransform& InOther);
 
 	UPROPERTY(BlueprintReadOnly, EditAnywhere, Category = "Pose")
 	FRigComputedTransform Local;
@@ -308,7 +606,12 @@ struct CONTROLRIG_API FRigLocalAndGlobalTransform
 	UPROPERTY(BlueprintReadOnly, EditAnywhere, Category = "Pose")
 	FRigComputedTransform Global;
 
-	bool bDirty[FRigLocalAndGlobalTransform::EDirtyMax] = { false, false };
+private:
+	
+	void LinkStorage(const TArrayView<FTransform>& InStorage);
+	void UnlinkStorage(FRigReusableElementStorage<FTransform>& InStorage);
+
+	friend struct FRigCurrentAndInitialTransform;
 };
 
 USTRUCT(BlueprintType)
@@ -316,6 +619,8 @@ struct CONTROLRIG_API FRigCurrentAndInitialTransform
 {
 	GENERATED_BODY()
 
+public:
+	
 	FRigCurrentAndInitialTransform()
     : Current()
     , Initial()
@@ -369,79 +674,44 @@ struct CONTROLRIG_API FRigCurrentAndInitialTransform
 		return Initial.Global;
 	}
 
-	bool& GetDirtyFlag(const ERigTransformType::Type InTransformType)
-	{
-		switch (InTransformType)
-		{
-			case ERigTransformType::CurrentLocal:
-				return Current.bDirty[FRigLocalAndGlobalTransform::ELocal];
-
-			case ERigTransformType::CurrentGlobal:
-				return Current.bDirty[FRigLocalAndGlobalTransform::EGlobal];
-
-			case ERigTransformType::InitialLocal:
-				return Initial.bDirty[FRigLocalAndGlobalTransform::ELocal];
-			
-			default:
-				return Initial.bDirty[FRigLocalAndGlobalTransform::EGlobal];
-		}
-	}
-
-	const bool& GetDirtyFlag(const ERigTransformType::Type InTransformType) const
-	{
-		switch (InTransformType)
-		{
-		case ERigTransformType::CurrentLocal:
-			return Current.bDirty[FRigLocalAndGlobalTransform::ELocal];
-
-		case ERigTransformType::CurrentGlobal:
-			return Current.bDirty[FRigLocalAndGlobalTransform::EGlobal];
-
-		case ERigTransformType::InitialLocal:
-			return Initial.bDirty[FRigLocalAndGlobalTransform::ELocal];
-
-		default:
-			return Initial.bDirty[FRigLocalAndGlobalTransform::EGlobal];
-		}
-	}
-
 	const FTransform& Get(const ERigTransformType::Type InTransformType) const
 	{
-		return operator[](InTransformType).Transform;
+		return operator[](InTransformType).Get();
 	}
 
 	void Set(const ERigTransformType::Type InTransformType, const FTransform& InTransform)
 	{
-		operator[](InTransformType).Set(InTransform, GetDirtyFlag(InTransformType));
-	}
-
-	bool IsDirty(const ERigTransformType::Type InTransformType) const
-	{
-		return GetDirtyFlag(InTransformType);
+		operator[](InTransformType).Set(InTransform);
 	}
 
-	void MarkDirty(const ERigTransformType::Type InTransformType)
-	{
-		ensure(!(GetDirtyFlag(ERigTransformType::SwapLocalAndGlobal(InTransformType))));
-		GetDirtyFlag(InTransformType) = true;
-	}
-
-	void Save(FArchive& Ar);
-	void Load(FArchive& Ar);
+	void Save(FArchive& Ar, const FRigCurrentAndInitialDirtyState& InDirtyState);
+	void Load(FArchive& Ar, FRigCurrentAndInitialDirtyState& OutDirtyState);
 
 	bool operator == (const FRigCurrentAndInitialTransform& Other) const
 	{
-		return Current.bDirty[FRigLocalAndGlobalTransform::ELocal]  == Other.Current.bDirty[FRigLocalAndGlobalTransform::ELocal]  && Current.Local  == Other.Current.Local
-			&& Current.bDirty[FRigLocalAndGlobalTransform::EGlobal] == Other.Current.bDirty[FRigLocalAndGlobalTransform::EGlobal] && Current.Global == Other.Current.Global
-			&& Initial.bDirty[FRigLocalAndGlobalTransform::ELocal]  == Other.Initial.bDirty[FRigLocalAndGlobalTransform::ELocal]  && Initial.Local  == Other.Initial.Local
-			&& Initial.bDirty[FRigLocalAndGlobalTransform::EGlobal] == Other.Initial.bDirty[FRigLocalAndGlobalTransform::EGlobal] && Initial.Global == Other.Initial.Global;
+		return Current.Local  == Other.Current.Local
+			&& Current.Global == Other.Current.Global
+			&& Initial.Local  == Other.Initial.Local
+			&& Initial.Global == Other.Initial.Global;
 	}
 
+	FRigCurrentAndInitialTransform& operator =(const FRigCurrentAndInitialTransform& InOther);
+
 	UPROPERTY(BlueprintReadOnly, EditAnywhere, Category = "Pose")
 	FRigLocalAndGlobalTransform Current;
 
 	UPROPERTY(BlueprintReadOnly, EditAnywhere, Category = "Pose")
 	FRigLocalAndGlobalTransform Initial;
+	
+private:
+	
+	void LinkStorage(const TArrayView<FTransform>& InStorage);
+	void UnlinkStorage(FRigReusableElementStorage<FTransform>& InStorage);
+
+	friend class URigHierarchy;
+	friend class URigHierarchyController;
+	friend struct FRigTransformElement;
+	friend struct FRigControlElement;
 };
 
 USTRUCT(BlueprintType)
@@ -686,6 +956,9 @@ public:
 	bool RemoveAllMetadata();
 	
 	void NotifyMetadataTagChanged(const FName& InTag, bool bAdded);
+
+	virtual int32 GetNumTransforms() const { return 0; }
+	virtual int32 GetNumCurves() const { return 0; }
 	
 	template<typename T>
 	bool IsA() const { return T::IsClassOf(this); }
@@ -747,7 +1020,9 @@ protected:
 	// helper function to be called as part of URigHierarchy::CopyHierarchy
 	virtual void CopyFrom(const FRigBaseElement* InOther);
 
-	
+	virtual void LinkStorage(const TArrayView<FTransform>& InTransforms, const TArrayView<bool>& InDirtyStates, const TArrayView<float>& InCurves) {}
+	virtual void UnlinkStorage(FRigReusableElementStorage<FTransform>& InTransforms, FRigReusableElementStorage<bool>& InDirtyStates, FRigReusableElementStorage<float>& InCurves) {}
+
 	friend class FControlRigEditor;
 	friend class URigHierarchy;
 	friend class URigHierarchyController;
@@ -769,7 +1044,8 @@ struct CONTROLRIG_API FRigTransformElement : public FRigBaseElement
 	FRigTransformElement& operator=(const FRigTransformElement& InOther)
 	{
 		Super::operator=(InOther);
-		Pose = InOther.Pose;
+		GetTransform() = InOther.GetTransform();
+		GetDirtyState() = InOther.GetDirtyState();
 		return *this;
 	}
 	virtual ~FRigTransformElement() override {}
@@ -778,14 +1054,26 @@ struct CONTROLRIG_API FRigTransformElement : public FRigBaseElement
 	virtual void Load(FArchive& Ar, ESerializationPhase SerializationPhase) override;
 
 	virtual void CopyPose(FRigBaseElement* InOther, bool bCurrent, bool bInitial, bool bWeights) override;
-	
-	UPROPERTY(BlueprintReadOnly, EditAnywhere, Category = RigElement, meta = (DisplayAfter = "Index"))
-	FRigCurrentAndInitialTransform Pose;
-	
+
+	// Current and Initial, both in Local and Global
+	virtual int32 GetNumTransforms() const override { return 4; }
+
+	const FRigCurrentAndInitialTransform& GetTransform() const;
+	FRigCurrentAndInitialTransform& GetTransform();
+	const FRigCurrentAndInitialDirtyState& GetDirtyState() const;
+	FRigCurrentAndInitialDirtyState& GetDirtyState();
+
 protected:
+	
 	FRigTransformElement(URigHierarchy* InOwner, const ERigElementType InType) :
 		FRigBaseElement(InOwner, InType)
 	{}
+
+	// Pose storage for this element.
+	FRigCurrentAndInitialTransform PoseStorage;
+
+	// Dirty state storage for this element.
+	FRigCurrentAndInitialDirtyState PoseDirtyState;
 	
 	struct FElementToDirty
 	{
@@ -829,9 +1117,13 @@ protected:
 
 	virtual void CopyFrom(const FRigBaseElement* InOther) override;
 
+	virtual void LinkStorage(const TArrayView<FTransform>& InTransforms, const TArrayView<bool>& InDirtyStates, const TArrayView<float>& InCurves) override;
+	virtual void UnlinkStorage(FRigReusableElementStorage<FTransform>& InTransforms, FRigReusableElementStorage<bool>& InDirtyStates, FRigReusableElementStorage<float>& InCurves) override;
+
 	friend class URigHierarchy;
 	friend class URigHierarchyController;
 	friend struct FRigBaseElement;
+	friend class FControlRigEditor;
 };
 
 USTRUCT(BlueprintType)
@@ -949,12 +1241,13 @@ struct CONTROLRIG_API FRigElementParentConstraint
 	FRigTransformElement* ParentElement;
 	FRigElementWeight Weight;
 	FRigElementWeight InitialWeight;
-	mutable FRigComputedTransform Cache;
+	mutable FTransform Cache;
 	mutable bool bCacheIsDirty;
 		
 	FRigElementParentConstraint()
 		: ParentElement(nullptr)
 	{
+		Cache = FTransform::Identity;
 		bCacheIsDirty = true;
 	}
 
@@ -1399,15 +1692,20 @@ struct CONTROLRIG_API FRigControlElement final : public FRigMultiParentElement
 	UPROPERTY(BlueprintReadOnly, EditAnywhere, Category = Control)
 	FRigControlSettings Settings;
 
-	UPROPERTY(BlueprintReadOnly, EditAnywhere, Category = RigElement)
-	FRigCurrentAndInitialTransform Offset;
-
-	UPROPERTY(BlueprintReadOnly, EditAnywhere, Category = RigElement)
-	FRigCurrentAndInitialTransform Shape;
+	// Current and Initial, both in Local and Global, for Pose, Offset and Shape
+	virtual int32 GetNumTransforms() const override { return 12; }
+	
+	const FRigCurrentAndInitialTransform& GetOffsetTransform() const;
+	FRigCurrentAndInitialTransform& GetOffsetTransform();
+	const FRigCurrentAndInitialDirtyState& GetOffsetDirtyState() const;
+	FRigCurrentAndInitialDirtyState& GetOffsetDirtyState();
+	const FRigCurrentAndInitialTransform& GetShapeTransform() const;
+	FRigCurrentAndInitialTransform& GetShapeTransform();
+	const FRigCurrentAndInitialDirtyState& GetShapeDirtyState() const;
+	FRigCurrentAndInitialDirtyState& GetShapeDirtyState();
 
 	UPROPERTY(BlueprintReadOnly, EditAnywhere, Category = RigElement)
 	FRigPreferredEulerAngles PreferredEulerAngles;
-	
 
 	FRigControlElement()
 		: FRigControlElement(nullptr)
@@ -1422,8 +1720,8 @@ struct CONTROLRIG_API FRigControlElement final : public FRigMultiParentElement
 	{
 		Super::operator=(InOther);
 		Settings = InOther.Settings;
-		Offset = InOther.Offset;
-		Shape = InOther.Shape;
+		GetOffsetTransform() = InOther.GetOffsetTransform();
+		GetShapeTransform() = InOther.GetShapeTransform();
 		PreferredEulerAngles = InOther.PreferredEulerAngles;
 		return *this;
 	}
@@ -1460,6 +1758,11 @@ struct CONTROLRIG_API FRigControlElement final : public FRigMultiParentElement
 	virtual void Load(FArchive& Ar, ESerializationPhase SerializationPhase) override;
 
 	virtual void CopyPose(FRigBaseElement* InOther, bool bCurrent, bool bInitial, bool bWeights) override;
+
+protected:
+
+	virtual void LinkStorage(const TArrayView<FTransform>& InTransforms, const TArrayView<bool>& InDirtyStates, const TArrayView<float>& InCurves) override;
+	virtual void UnlinkStorage(FRigReusableElementStorage<FTransform>& InTransforms, FRigReusableElementStorage<bool>& InDirtyStates, FRigReusableElementStorage<float>& InCurves) override;
 	
 private:
 	explicit FRigControlElement(URigHierarchy* InOwner)
@@ -1473,8 +1776,18 @@ private:
 		return InElement->GetType() == ERigElementType::Control;
 	}
 
+	// Offset storage for this element.
+	FRigCurrentAndInitialTransform OffsetStorage;
+	FRigCurrentAndInitialDirtyState OffsetDirtyState;
+
+	// Shape storage for this element.
+	FRigCurrentAndInitialTransform ShapeStorage;
+	FRigCurrentAndInitialDirtyState ShapeDirtyState;
+
 	friend class URigHierarchy;
+	friend class URigHierarchyController;
 	friend struct FRigBaseElement;
+	friend class FControlRigEditor;
 };
 
 USTRUCT(BlueprintType)
@@ -1485,23 +1798,28 @@ struct CONTROLRIG_API FRigCurveElement final : public FRigBaseElement
 
 	static const EElementIndex ElementTypeIndex;
 
-	// Set to true if the value was actually set. Used to carry back and forth blend curve
-	// value validity state.
-	bool bIsValueSet = true;
-	
-	float Value = 0.0f;
-
-	
 	FRigCurveElement()
 		: FRigCurveElement(nullptr)
 	{}
 	
 	virtual ~FRigCurveElement() override {}
 
+	virtual int32 GetNumCurves() const override { return 1; }
+
 	virtual void Save(FArchive& A, ESerializationPhase SerializationPhase) override;
 	virtual void Load(FArchive& Ar, ESerializationPhase SerializationPhase) override;
 
 	virtual void CopyPose(FRigBaseElement* InOther, bool bCurrent, bool bInitial, bool bWeights) override;
+
+	const float& Get() const;
+	void Set(const float& InValue);
+
+	bool IsValueSet() const { return bIsValueSet; }
+
+protected:
+
+	virtual void LinkStorage(const TArrayView<FTransform>& InTransforms, const TArrayView<bool>& InDirtyStates, const TArrayView<float>& InCurves) override;
+	virtual void UnlinkStorage(FRigReusableElementStorage<FTransform>& InTransforms, FRigReusableElementStorage<bool>& InDirtyStates, FRigReusableElementStorage<float>& InCurves) override;
 	
 private:
 	FRigCurveElement(URigHierarchy* InOwner)
@@ -1515,6 +1833,13 @@ private:
 		return InElement->GetType() == ERigElementType::Curve;
 	}
 
+	// Set to true if the value was actually set. Used to carry back and forth blend curve
+	// value validity state.
+	bool bIsValueSet = true;
+
+	int32 StorageIndex = INDEX_NONE;
+	float* Storage = nullptr;
+
 	friend class URigHierarchy;
 	friend class URigHierarchyController;
 	friend struct FRigBaseElement;
@@ -1936,7 +2261,10 @@ struct CONTROLRIG_API FRigHierarchyCopyPasteContentPerElement
 	TArray<FRigElementWeight> ParentWeights;
 
 	UPROPERTY()
-	FRigCurrentAndInitialTransform Pose;
+	TArray<FTransform> Poses;
+
+	UPROPERTY()
+	TArray<bool> DirtyStates;
 };
 
 USTRUCT()
@@ -1959,4 +2287,4 @@ struct CONTROLRIG_API FRigHierarchyCopyPasteContent
 
 	UPROPERTY()
 	TArray<FTransform> GlobalTransforms;
-};
\ No newline at end of file
+};
diff --git a/Engine/Plugins/Animation/ControlRig/Source/ControlRigEditor/Private/ControlRigElementDetails.cpp b/Engine/Plugins/Animation/ControlRig/Source/ControlRigEditor/Private/ControlRigElementDetails.cpp
--- a/Engine/Plugins/Animation/ControlRig/Source/ControlRigEditor/Private/ControlRigElementDetails.cpp
+++ b/Engine/Plugins/Animation/ControlRig/Source/ControlRigEditor/Private/ControlRigElementDetails.cpp
@@ -5008,7 +5008,7 @@ void FRigControlElementDetails::OnCopyShapeProperties()
 			Value = FString::Printf(TEXT("(ShapeName=\"%s\",ShapeColor=%s,Transform=%s)"),
 				*ControlElement->Settings.ShapeName.ToString(),
 				*ControlElement->Settings.ShapeColor.ToString(),
-				*ControlElement->Shape.Initial.Local.Transform.ToString());
+				*ControlElement->GetShapeTransform().Initial.Local.Get().ToString());
 		}
 	}
 		
diff --git a/Engine/Plugins/Animation/ControlRig/Source/ControlRigEditor/Private/Editor/ControlRigEditor.cpp b/Engine/Plugins/Animation/ControlRig/Source/ControlRigEditor/Private/Editor/ControlRigEditor.cpp
--- a/Engine/Plugins/Animation/ControlRig/Source/ControlRigEditor/Private/Editor/ControlRigEditor.cpp
+++ b/Engine/Plugins/Animation/ControlRig/Source/ControlRigEditor/Private/Editor/ControlRigEditor.cpp
@@ -3327,9 +3327,9 @@ void FControlRigEditor::OnWrappedPropertyChangedChainEvent(URigVMDetailsViewWrap
 			}
 
 			static constexpr TCHAR PropertyChainElementFormat[] = TEXT("%s->");
-			static const FString PoseString = FString::Printf(PropertyChainElementFormat, GET_MEMBER_NAME_STRING_CHECKED(FRigTransformElement, Pose));
-			static const FString OffsetString = FString::Printf(PropertyChainElementFormat, GET_MEMBER_NAME_STRING_CHECKED(FRigControlElement, Offset));
-			static const FString ShapeString = FString::Printf(PropertyChainElementFormat, GET_MEMBER_NAME_STRING_CHECKED(FRigControlElement, Shape));
+			static const FString PoseString = FString::Printf(PropertyChainElementFormat, GET_MEMBER_NAME_STRING_CHECKED(FRigTransformElement, PoseStorage));
+			static const FString OffsetString = FString::Printf(PropertyChainElementFormat, GET_MEMBER_NAME_STRING_CHECKED(FRigControlElement, OffsetStorage));
+			static const FString ShapeString = FString::Printf(PropertyChainElementFormat, GET_MEMBER_NAME_STRING_CHECKED(FRigControlElement, ShapeStorage));
 			static const FString SettingsString = FString::Printf(PropertyChainElementFormat, GET_MEMBER_NAME_STRING_CHECKED(FRigControlElement, Settings));
 
 			struct Local
@@ -3384,7 +3384,7 @@ void FControlRigEditor::OnWrappedPropertyChangedChainEvent(URigVMDetailsViewWrap
 					return;
 				}
 
-				const FTransform Transform = InWrapperObject->GetContent<FRigTransformElement>().Pose.Get(TransformType);
+				const FTransform Transform = InWrapperObject->GetContent<FRigTransformElement>().GetTransform().Get(TransformType);
 
 				if(ERigTransformType::IsLocal(TransformType) && TransformElement->IsA<FRigControlElement>())
 				{
@@ -3415,7 +3415,7 @@ void FControlRigEditor::OnWrappedPropertyChangedChainEvent(URigVMDetailsViewWrap
 				ERigTransformType::Type TransformType = Local::GetTransformTypeFromPath(PropertyPath);
 				bIsInitial = bIsInitial || ERigTransformType::IsInitial(TransformType);
 
-				const FTransform Transform = GetWrapperObjects()[0]->GetContent<FRigControlElement>().Offset.Get(TransformType);
+				const FTransform Transform = GetWrapperObjects()[0]->GetContent<FRigControlElement>().GetOffsetTransform().Get(TransformType);
 				
 				ControlRigBP->Hierarchy->SetControlOffsetTransform(ControlElement, Transform, ERigTransformType::MakeInitial(TransformType), true, true, false, true);
 			}
@@ -3430,7 +3430,7 @@ void FControlRigEditor::OnWrappedPropertyChangedChainEvent(URigVMDetailsViewWrap
 				ERigTransformType::Type TransformType = Local::GetTransformTypeFromPath(PropertyPath);
 				bIsInitial = bIsInitial || ERigTransformType::IsInitial(TransformType);
 
-				const FTransform Transform = GetWrapperObjects()[0]->GetContent<FRigControlElement>().Shape.Get(TransformType);
+				const FTransform Transform = GetWrapperObjects()[0]->GetContent<FRigControlElement>().GetShapeTransform().Get(TransformType);
 				
 				ControlRigBP->Hierarchy->SetControlShapeTransform(ControlElement, Transform, ERigTransformType::MakeInitial(TransformType), true, false, true);
 			}
diff --git a/Engine/Plugins/Animation/ControlRig/Source/ControlRigEditor/Private/Editor/SRigHierarchy.cpp b/Engine/Plugins/Animation/ControlRig/Source/ControlRigEditor/Private/Editor/SRigHierarchy.cpp
--- a/Engine/Plugins/Animation/ControlRig/Source/ControlRigEditor/Private/Editor/SRigHierarchy.cpp
+++ b/Engine/Plugins/Animation/ControlRig/Source/ControlRigEditor/Private/Editor/SRigHierarchy.cpp
@@ -2573,7 +2573,7 @@ void SRigHierarchy::HandlePasteTransforms(ERigTransformType::Type InTransformTyp
 		for(int32 Index = 0; Index < Count; Index++)
 		{
 			const FRigHierarchyCopyPasteContentPerElement& PerElementData = Data.Elements[Index];
-			const FTransform Transform =  PerElementData.Pose.Get(InTransformType);
+			const FTransform Transform =  PerElementData.Poses[(int32)InTransformType];
 
 			if(FRigTransformElement* TransformElement = Hierarchy->Find<FRigTransformElement>(CurrentSelection[Index]))
 			{
