Control Rig: Add sorting capability to the element storage

#rnx

diff --git a/Engine/Plugins/Animation/ControlRig/Source/ControlRig/Private/Rigs/RigHierarchy.cpp b/Engine/Plugins/Animation/ControlRig/Source/ControlRig/Private/Rigs/RigHierarchy.cpp
--- a/Engine/Plugins/Animation/ControlRig/Source/ControlRig/Private/Rigs/RigHierarchy.cpp
+++ b/Engine/Plugins/Animation/ControlRig/Source/ControlRig/Private/Rigs/RigHierarchy.cpp
@@ -31,6 +31,7 @@ LLM_DEFINE_TAG(Animation_ControlRig);
 #include "Serialization/JsonSerializer.h"
 #include "Units/Execution/RigUnit_BeginExecution.h"
 #include "Algo/Transform.h"
+#include "Algo/Sort.h"
 
 static FCriticalSection GRigHierarchyStackTraceMutex;
 static char GRigHierarchyStackTrace[65536];
@@ -344,31 +345,31 @@ void URigHierarchy::Load(FArchive& Ar)
 			if(FRigTransformElement* TransformElement = Cast<FRigTransformElement>(Element))
 			{
 				TransformElement->PoseStorage.Initial.Local.Index = TransformIndices[UsedTransformIndex++];
-				TransformElement->PoseStorage.Initial.Global.Index = TransformIndices[UsedTransformIndex++];
 				TransformElement->PoseStorage.Current.Local.Index = TransformIndices[UsedTransformIndex++];
+				TransformElement->PoseStorage.Initial.Global.Index = TransformIndices[UsedTransformIndex++];
 				TransformElement->PoseStorage.Current.Global.Index = TransformIndices[UsedTransformIndex++];
 				TransformElement->PoseDirtyState.Initial.Local.Index = DirtyStateIndices[UsedDirtyStateIndex++];
-				TransformElement->PoseDirtyState.Initial.Global.Index = DirtyStateIndices[UsedDirtyStateIndex++];
 				TransformElement->PoseDirtyState.Current.Local.Index = DirtyStateIndices[UsedDirtyStateIndex++];
+				TransformElement->PoseDirtyState.Initial.Global.Index = DirtyStateIndices[UsedDirtyStateIndex++];
 				TransformElement->PoseDirtyState.Current.Global.Index = DirtyStateIndices[UsedDirtyStateIndex++];
 
 				if(FRigControlElement* ControlElement = Cast<FRigControlElement>(TransformElement))
 				{
 					ControlElement->OffsetStorage.Initial.Local.Index = TransformIndices[UsedTransformIndex++];
-					ControlElement->OffsetStorage.Initial.Global.Index = TransformIndices[UsedTransformIndex++];
 					ControlElement->OffsetStorage.Current.Local.Index = TransformIndices[UsedTransformIndex++];
+					ControlElement->OffsetStorage.Initial.Global.Index = TransformIndices[UsedTransformIndex++];
 					ControlElement->OffsetStorage.Current.Global.Index = TransformIndices[UsedTransformIndex++];
 					ControlElement->OffsetDirtyState.Initial.Local.Index = DirtyStateIndices[UsedDirtyStateIndex++];
-					ControlElement->OffsetDirtyState.Initial.Global.Index = DirtyStateIndices[UsedDirtyStateIndex++];
 					ControlElement->OffsetDirtyState.Current.Local.Index = DirtyStateIndices[UsedDirtyStateIndex++];
+					ControlElement->OffsetDirtyState.Initial.Global.Index = DirtyStateIndices[UsedDirtyStateIndex++];
 					ControlElement->OffsetDirtyState.Current.Global.Index = DirtyStateIndices[UsedDirtyStateIndex++];
 					ControlElement->ShapeStorage.Initial.Local.Index = TransformIndices[UsedTransformIndex++];
-					ControlElement->ShapeStorage.Initial.Global.Index = TransformIndices[UsedTransformIndex++];
 					ControlElement->ShapeStorage.Current.Local.Index = TransformIndices[UsedTransformIndex++];
+					ControlElement->ShapeStorage.Initial.Global.Index = TransformIndices[UsedTransformIndex++];
 					ControlElement->ShapeStorage.Current.Global.Index = TransformIndices[UsedTransformIndex++];
 					ControlElement->ShapeDirtyState.Initial.Local.Index = DirtyStateIndices[UsedDirtyStateIndex++];
-					ControlElement->ShapeDirtyState.Initial.Global.Index = DirtyStateIndices[UsedDirtyStateIndex++];
 					ControlElement->ShapeDirtyState.Current.Local.Index = DirtyStateIndices[UsedDirtyStateIndex++];
+					ControlElement->ShapeDirtyState.Initial.Global.Index = DirtyStateIndices[UsedDirtyStateIndex++];
 					ControlElement->ShapeDirtyState.Current.Global.Index = DirtyStateIndices[UsedDirtyStateIndex++];
 				}
 			}
@@ -442,6 +443,8 @@ void URigHierarchy::Load(FArchive& Ar)
 			ElementMetadata.Storage.Add(MoveTemp(Entry.Value));
 		}
 	}
+
+	SortElementStorage();
 }
 
 void URigHierarchy::PostLoad()
@@ -752,31 +755,31 @@ void URigHierarchy::CopyHierarchy(URigHierarchy* InHierarchy)
 			if(FRigTransformElement* TransformElement = Cast<FRigTransformElement>(Target))
 			{
 				TransformElement->PoseStorage.Initial.Local.Index = TransformIndices[UsedTransformIndex++];
-				TransformElement->PoseStorage.Initial.Global.Index = TransformIndices[UsedTransformIndex++];
 				TransformElement->PoseStorage.Current.Local.Index = TransformIndices[UsedTransformIndex++];
+				TransformElement->PoseStorage.Initial.Global.Index = TransformIndices[UsedTransformIndex++];
 				TransformElement->PoseStorage.Current.Global.Index = TransformIndices[UsedTransformIndex++];
 				TransformElement->PoseDirtyState.Initial.Local.Index = DirtyStateIndices[UsedDirtyStateIndex++];
-				TransformElement->PoseDirtyState.Initial.Global.Index = DirtyStateIndices[UsedDirtyStateIndex++];
 				TransformElement->PoseDirtyState.Current.Local.Index = DirtyStateIndices[UsedDirtyStateIndex++];
+				TransformElement->PoseDirtyState.Initial.Global.Index = DirtyStateIndices[UsedDirtyStateIndex++];
 				TransformElement->PoseDirtyState.Current.Global.Index = DirtyStateIndices[UsedDirtyStateIndex++];
 
 				if(FRigControlElement* ControlElement = Cast<FRigControlElement>(TransformElement))
 				{
 					ControlElement->OffsetStorage.Initial.Local.Index = TransformIndices[UsedTransformIndex++];
-					ControlElement->OffsetStorage.Initial.Global.Index = TransformIndices[UsedTransformIndex++];
 					ControlElement->OffsetStorage.Current.Local.Index = TransformIndices[UsedTransformIndex++];
+					ControlElement->OffsetStorage.Initial.Global.Index = TransformIndices[UsedTransformIndex++];
 					ControlElement->OffsetStorage.Current.Global.Index = TransformIndices[UsedTransformIndex++];
 					ControlElement->OffsetDirtyState.Initial.Local.Index = DirtyStateIndices[UsedDirtyStateIndex++];
-					ControlElement->OffsetDirtyState.Initial.Global.Index = DirtyStateIndices[UsedDirtyStateIndex++];
 					ControlElement->OffsetDirtyState.Current.Local.Index = DirtyStateIndices[UsedDirtyStateIndex++];
+					ControlElement->OffsetDirtyState.Initial.Global.Index = DirtyStateIndices[UsedDirtyStateIndex++];
 					ControlElement->OffsetDirtyState.Current.Global.Index = DirtyStateIndices[UsedDirtyStateIndex++];
 					ControlElement->ShapeStorage.Initial.Local.Index = TransformIndices[UsedTransformIndex++];
-					ControlElement->ShapeStorage.Initial.Global.Index = TransformIndices[UsedTransformIndex++];
 					ControlElement->ShapeStorage.Current.Local.Index = TransformIndices[UsedTransformIndex++];
+					ControlElement->ShapeStorage.Initial.Global.Index = TransformIndices[UsedTransformIndex++];
 					ControlElement->ShapeStorage.Current.Global.Index = TransformIndices[UsedTransformIndex++];
 					ControlElement->ShapeDirtyState.Initial.Local.Index = DirtyStateIndices[UsedDirtyStateIndex++];
-					ControlElement->ShapeDirtyState.Initial.Global.Index = DirtyStateIndices[UsedDirtyStateIndex++];
 					ControlElement->ShapeDirtyState.Current.Local.Index = DirtyStateIndices[UsedDirtyStateIndex++];
+					ControlElement->ShapeDirtyState.Initial.Global.Index = DirtyStateIndices[UsedDirtyStateIndex++];
 					ControlElement->ShapeDirtyState.Current.Global.Index = DirtyStateIndices[UsedDirtyStateIndex++];
 				}
 			}
@@ -791,6 +794,8 @@ void URigHierarchy::CopyHierarchy(URigHierarchy* InHierarchy)
 			CopyAllMetadataFromElement(Target, Source);
 		}
 
+		SortElementStorage();
+
 		PreviousNameMap.Append(InHierarchy->PreviousNameMap);
 
 		// copy the topology version from the hierarchy.
@@ -5915,31 +5920,31 @@ void URigHierarchy::AllocateDefaultElementStorage(FRigBaseElement* InElement, bo
 		// so if we are already linked to the default storage, exit early.
 		if(ElementTransforms.Contains(TransformElement->PoseStorage.Initial.Local.Index, TransformElement->PoseStorage.Initial.Local.Storage))
 		{
-			check(ElementTransforms.Contains(TransformElement->PoseStorage.Initial.Global.Index, TransformElement->PoseStorage.Initial.Global.Storage));
 			check(ElementTransforms.Contains(TransformElement->PoseStorage.Current.Local.Index, TransformElement->PoseStorage.Current.Local.Storage));
+			check(ElementTransforms.Contains(TransformElement->PoseStorage.Initial.Global.Index, TransformElement->PoseStorage.Initial.Global.Storage));
 			check(ElementTransforms.Contains(TransformElement->PoseStorage.Current.Global.Index, TransformElement->PoseStorage.Current.Global.Storage));
 			check(ElementDirtyStates.Contains(TransformElement->PoseDirtyState.Initial.Local.Index, TransformElement->PoseDirtyState.Initial.Local.Storage));
-			check(ElementDirtyStates.Contains(TransformElement->PoseDirtyState.Initial.Global.Index, TransformElement->PoseDirtyState.Initial.Global.Storage));
 			check(ElementDirtyStates.Contains(TransformElement->PoseDirtyState.Current.Local.Index, TransformElement->PoseDirtyState.Current.Local.Storage));
+			check(ElementDirtyStates.Contains(TransformElement->PoseDirtyState.Initial.Global.Index, TransformElement->PoseDirtyState.Initial.Global.Storage));
 			check(ElementDirtyStates.Contains(TransformElement->PoseDirtyState.Current.Global.Index, TransformElement->PoseDirtyState.Current.Global.Storage));
 
 			if(const FRigControlElement* ControlElement = Cast<FRigControlElement>(TransformElement))
 			{
 				check(ElementTransforms.Contains(ControlElement->OffsetStorage.Initial.Local.Index, ControlElement->OffsetStorage.Initial.Local.Storage));
-				check(ElementTransforms.Contains(ControlElement->OffsetStorage.Initial.Global.Index, ControlElement->OffsetStorage.Initial.Global.Storage));
 				check(ElementTransforms.Contains(ControlElement->OffsetStorage.Current.Local.Index, ControlElement->OffsetStorage.Current.Local.Storage));
+				check(ElementTransforms.Contains(ControlElement->OffsetStorage.Initial.Global.Index, ControlElement->OffsetStorage.Initial.Global.Storage));
 				check(ElementTransforms.Contains(ControlElement->OffsetStorage.Current.Global.Index, ControlElement->OffsetStorage.Current.Global.Storage));
 				check(ElementDirtyStates.Contains(ControlElement->OffsetDirtyState.Initial.Local.Index, ControlElement->OffsetDirtyState.Initial.Local.Storage));
-				check(ElementDirtyStates.Contains(ControlElement->OffsetDirtyState.Initial.Global.Index, ControlElement->OffsetDirtyState.Initial.Global.Storage));
 				check(ElementDirtyStates.Contains(ControlElement->OffsetDirtyState.Current.Local.Index, ControlElement->OffsetDirtyState.Current.Local.Storage));
+				check(ElementDirtyStates.Contains(ControlElement->OffsetDirtyState.Initial.Global.Index, ControlElement->OffsetDirtyState.Initial.Global.Storage));
 				check(ElementDirtyStates.Contains(ControlElement->OffsetDirtyState.Current.Global.Index, ControlElement->OffsetDirtyState.Current.Global.Storage));
 				check(ElementTransforms.Contains(ControlElement->ShapeStorage.Initial.Local.Index, ControlElement->ShapeStorage.Initial.Local.Storage));
-				check(ElementTransforms.Contains(ControlElement->ShapeStorage.Initial.Global.Index, ControlElement->ShapeStorage.Initial.Global.Storage));
 				check(ElementTransforms.Contains(ControlElement->ShapeStorage.Current.Local.Index, ControlElement->ShapeStorage.Current.Local.Storage));
+				check(ElementTransforms.Contains(ControlElement->ShapeStorage.Initial.Global.Index, ControlElement->ShapeStorage.Initial.Global.Storage));
 				check(ElementTransforms.Contains(ControlElement->ShapeStorage.Current.Global.Index, ControlElement->ShapeStorage.Current.Global.Storage));
 				check(ElementDirtyStates.Contains(ControlElement->ShapeDirtyState.Initial.Local.Index, ControlElement->ShapeDirtyState.Initial.Local.Storage));
-				check(ElementDirtyStates.Contains(ControlElement->ShapeDirtyState.Initial.Global.Index, ControlElement->ShapeDirtyState.Initial.Global.Storage));
 				check(ElementDirtyStates.Contains(ControlElement->ShapeDirtyState.Current.Local.Index, ControlElement->ShapeDirtyState.Current.Local.Storage));
+				check(ElementDirtyStates.Contains(ControlElement->ShapeDirtyState.Initial.Global.Index, ControlElement->ShapeDirtyState.Initial.Global.Storage));
 				check(ElementDirtyStates.Contains(ControlElement->ShapeDirtyState.Current.Global.Index, ControlElement->ShapeDirtyState.Current.Global.Storage));
 			}
 			return;
@@ -5949,12 +5954,12 @@ void URigHierarchy::AllocateDefaultElementStorage(FRigBaseElement* InElement, bo
 		const TArray<int32, TInlineAllocator<4>> DirtyStateIndices = ElementDirtyStates.Allocate(TransformElement->GetNumTransforms(), false);
 		check(DirtyStateIndices.Num() >= 4);
 		TransformElement->PoseStorage.Initial.Local.Index = TransformIndices[0];
-		TransformElement->PoseStorage.Initial.Global.Index = TransformIndices[1];
-		TransformElement->PoseStorage.Current.Local.Index = TransformIndices[2];
+		TransformElement->PoseStorage.Current.Local.Index = TransformIndices[1];
+		TransformElement->PoseStorage.Initial.Global.Index = TransformIndices[2];
 		TransformElement->PoseStorage.Current.Global.Index = TransformIndices[3];
 		TransformElement->PoseDirtyState.Initial.Local.Index = DirtyStateIndices[0];
-		TransformElement->PoseDirtyState.Initial.Global.Index = DirtyStateIndices[1];
-		TransformElement->PoseDirtyState.Current.Local.Index = DirtyStateIndices[2];
+		TransformElement->PoseDirtyState.Current.Local.Index = DirtyStateIndices[1];
+		TransformElement->PoseDirtyState.Initial.Global.Index = DirtyStateIndices[2];
 		TransformElement->PoseDirtyState.Current.Global.Index = DirtyStateIndices[3];
 
 		if(FRigControlElement* ControlElement = Cast<FRigControlElement>(TransformElement))
@@ -5962,20 +5967,20 @@ void URigHierarchy::AllocateDefaultElementStorage(FRigBaseElement* InElement, bo
 			check(TransformIndices.Num() >= 12);
 			check(DirtyStateIndices.Num() >= 12);
 			ControlElement->OffsetStorage.Initial.Local.Index = TransformIndices[4];
-			ControlElement->OffsetStorage.Initial.Global.Index = TransformIndices[5];
-			ControlElement->OffsetStorage.Current.Local.Index = TransformIndices[6];
+			ControlElement->OffsetStorage.Current.Local.Index = TransformIndices[5];
+			ControlElement->OffsetStorage.Initial.Global.Index = TransformIndices[6];
 			ControlElement->OffsetStorage.Current.Global.Index = TransformIndices[7];
 			ControlElement->OffsetDirtyState.Initial.Local.Index = DirtyStateIndices[4];
-			ControlElement->OffsetDirtyState.Initial.Global.Index = DirtyStateIndices[5];
-			ControlElement->OffsetDirtyState.Current.Local.Index = DirtyStateIndices[6];
+			ControlElement->OffsetDirtyState.Current.Local.Index = DirtyStateIndices[5];
+			ControlElement->OffsetDirtyState.Initial.Global.Index = DirtyStateIndices[6];
 			ControlElement->OffsetDirtyState.Current.Global.Index = DirtyStateIndices[7];
 			ControlElement->ShapeStorage.Initial.Local.Index = TransformIndices[8];
-			ControlElement->ShapeStorage.Initial.Global.Index = TransformIndices[9];
-			ControlElement->ShapeStorage.Current.Local.Index = TransformIndices[10];
+			ControlElement->ShapeStorage.Current.Local.Index = TransformIndices[9];
+			ControlElement->ShapeStorage.Initial.Global.Index = TransformIndices[10];
 			ControlElement->ShapeStorage.Current.Global.Index = TransformIndices[11];
 			ControlElement->ShapeDirtyState.Initial.Local.Index = DirtyStateIndices[8];
-			ControlElement->ShapeDirtyState.Initial.Global.Index = DirtyStateIndices[9];
-			ControlElement->ShapeDirtyState.Current.Local.Index = DirtyStateIndices[10];
+			ControlElement->ShapeDirtyState.Current.Local.Index = DirtyStateIndices[9];
+			ControlElement->ShapeDirtyState.Initial.Global.Index = DirtyStateIndices[10];
 			ControlElement->ShapeDirtyState.Current.Global.Index = DirtyStateIndices[11];
 		}
 	}
@@ -6012,6 +6017,196 @@ void URigHierarchy::UpdateElementStorage()
 	{
 		Element->LinkStorage(ElementTransforms.Storage, ElementDirtyStates.Storage, ElementCurves.Storage);
 	}
+	ElementTransformRanges.Reset();
+}
+
+void URigHierarchy::SortElementStorage()
+{
+	if(Elements.IsEmpty())
+	{
+		return;
+	}
+	
+	TMap<int32, FRigComputedTransform*> TransformIndexToComputedTransform;
+	TMap<int32, FRigTransformDirtyState*> TransformIndexToDirtyState;
+
+	struct FTypeInfo
+	{
+		FTypeInfo()
+		: Index(INDEX_NONE)
+		, Element(ERigElementType::Bone)
+		, Transform(ERigTransformType::InitialLocal)
+		, Storage(ERigTransformStorageType::Pose)
+		{
+		}
+
+		FTypeInfo(int32 InIndex, ERigElementType InElement, ERigTransformType::Type InTransform, ERigTransformStorageType::Type InStorage)
+		: Index(InIndex)
+		, Element(InElement)
+		, Transform(InTransform)
+		, Storage(InStorage)
+		{
+		}
+
+		uint32 GetSortKey(uint32 InNumElements) const
+		{
+			const uint32 StorageMultiplier = InNumElements; 
+			const uint32 ElementMultiplier = StorageMultiplier * (uint32)ERigTransformStorageType::NumStorageTypes;
+			const uint32 TransformMultiplier = ElementMultiplier * uint32(RigElementTypeToFlatIndex(ERigElementType::Last) + 1);
+			// we sort elements by laying them out first by transform type, then element type, then storage type
+			// and finally in the order of their original index in the hierarchy
+			return Index +
+				int32(Transform) * TransformMultiplier +
+				int32(Element) * ElementMultiplier +
+				int32(Storage) * StorageMultiplier;
+		}
+
+		int32 Index;
+		ERigElementType Element;
+		ERigTransformType::Type Transform;
+		ERigTransformStorageType::Type Storage;
+	};
+
+	struct FTransformInfo
+	{
+		FTransformInfo()
+		: ComputedTransform(nullptr)
+		, DirtyState(nullptr)
+		{
+		}
+
+		FTransformInfo(const FTypeInfo& InTypeInfo, FRigComputedTransform* InComputedTransform, FRigTransformDirtyState* InDirtyState)
+		: TypeInfo(InTypeInfo)
+		, ComputedTransform(InComputedTransform)
+		, DirtyState(InDirtyState)
+		{
+		}
+		
+		FTypeInfo TypeInfo;
+		FRigComputedTransform* ComputedTransform;
+		FRigTransformDirtyState* DirtyState;
+	};
+
+	TArray<FTransformInfo> TransformInfos;
+	TransformInfos.Reserve(Elements.Num() * 4);
+
+	for(FRigBaseElement* Element : Elements)
+	{
+		if(FRigTransformElement* TransformElement = Cast<FRigTransformElement>(Element))
+		{
+			TransformInfos.Emplace(FTypeInfo(TransformElement->GetIndex(), TransformElement->GetType(), ERigTransformType::InitialLocal, ERigTransformStorageType::Pose), &TransformElement->PoseStorage.Initial.Local, &TransformElement->PoseDirtyState.Initial.Local);
+			TransformInfos.Emplace(FTypeInfo(TransformElement->GetIndex(), TransformElement->GetType(), ERigTransformType::InitialGlobal, ERigTransformStorageType::Pose), &TransformElement->PoseStorage.Initial.Global, &TransformElement->PoseDirtyState.Initial.Global);
+			TransformInfos.Emplace(FTypeInfo(TransformElement->GetIndex(), TransformElement->GetType(), ERigTransformType::CurrentLocal, ERigTransformStorageType::Pose), &TransformElement->PoseStorage.Current.Local, &TransformElement->PoseDirtyState.Current.Local);
+			TransformInfos.Emplace(FTypeInfo(TransformElement->GetIndex(), TransformElement->GetType(), ERigTransformType::CurrentGlobal, ERigTransformStorageType::Pose), &TransformElement->PoseStorage.Current.Global, &TransformElement->PoseDirtyState.Current.Global);
+
+			if(FRigControlElement* ControlElement = Cast<FRigControlElement>(Element))
+			{
+				TransformInfos.Emplace(FTypeInfo(ControlElement->GetIndex(), ControlElement->GetType(), ERigTransformType::InitialLocal, ERigTransformStorageType::Offset), &ControlElement->OffsetStorage.Initial.Local, &ControlElement->OffsetDirtyState.Initial.Local);
+				TransformInfos.Emplace(FTypeInfo(ControlElement->GetIndex(), ControlElement->GetType(), ERigTransformType::InitialGlobal, ERigTransformStorageType::Offset), &ControlElement->OffsetStorage.Initial.Global, &ControlElement->OffsetDirtyState.Initial.Global);
+				TransformInfos.Emplace(FTypeInfo(ControlElement->GetIndex(), ControlElement->GetType(), ERigTransformType::CurrentLocal, ERigTransformStorageType::Offset), &ControlElement->OffsetStorage.Current.Local, &ControlElement->OffsetDirtyState.Current.Local);
+				TransformInfos.Emplace(FTypeInfo(ControlElement->GetIndex(), ControlElement->GetType(), ERigTransformType::CurrentGlobal, ERigTransformStorageType::Offset), &ControlElement->OffsetStorage.Current.Global, &ControlElement->OffsetDirtyState.Current.Global);
+				TransformInfos.Emplace(FTypeInfo(ControlElement->GetIndex(), ControlElement->GetType(), ERigTransformType::InitialLocal, ERigTransformStorageType::Shape), &ControlElement->ShapeStorage.Initial.Local, &ControlElement->ShapeDirtyState.Initial.Local);
+				TransformInfos.Emplace(FTypeInfo(ControlElement->GetIndex(), ControlElement->GetType(), ERigTransformType::InitialGlobal, ERigTransformStorageType::Shape), &ControlElement->ShapeStorage.Initial.Global, &ControlElement->ShapeDirtyState.Initial.Global);
+				TransformInfos.Emplace(FTypeInfo(ControlElement->GetIndex(), ControlElement->GetType(), ERigTransformType::CurrentLocal, ERigTransformStorageType::Shape), &ControlElement->ShapeStorage.Current.Local, &ControlElement->ShapeDirtyState.Current.Local);
+				TransformInfos.Emplace(FTypeInfo(ControlElement->GetIndex(), ControlElement->GetType(), ERigTransformType::CurrentGlobal, ERigTransformStorageType::Shape), &ControlElement->ShapeStorage.Current.Global, &ControlElement->ShapeDirtyState.Current.Global);
+			}
+		}
+	}
+
+	uint32 NumElements = (uint32)Elements.Num();
+	Algo::SortBy(TransformInfos, [NumElements](const FTransformInfo& InInfo) -> uint32
+	{
+		return InInfo.TypeInfo.GetSortKey(NumElements);
+	});
+
+	// check if the sorting is at all required
+	bool bRequiresSort = false;
+	for(int32 Index = 0; Index < TransformInfos.Num(); Index++)
+	{
+		const int32 CurrentTransformIndex = TransformInfos[Index].ComputedTransform->Index;
+		if(CurrentTransformIndex == INDEX_NONE)
+		{
+			// we cannot sort this if some of the storage is no longer local
+			ElementTransformRanges.Reset();
+			return;
+		}
+		if(CurrentTransformIndex != Index)
+		{
+			bRequiresSort = true;
+			break;
+		}
+
+		const int32 CurrentDirtyStateIndex = TransformInfos[Index].DirtyState->Index;
+		if(CurrentDirtyStateIndex == INDEX_NONE)
+		{
+			// we cannot sort this if some of the storage is no longer local
+			ElementTransformRanges.Reset();
+			return;
+		}
+		
+		if(CurrentDirtyStateIndex != Index)
+		{
+			bRequiresSort = true;
+			break;
+		}
+	}
+
+	if(!bRequiresSort)
+	{
+		return;
+	}
+
+	FRigReusableElementStorage<FTransform> NewElementTransforms;
+	FRigReusableElementStorage<bool> NewElementDirtyStates;
+	NewElementDirtyStates.Storage.Reserve(ElementDirtyStates.Storage.Num() - ElementDirtyStates.FreeList.Num());
+	NewElementDirtyStates.Storage.Reserve(ElementDirtyStates.Storage.Num() - ElementDirtyStates.FreeList.Num());
+
+	TMap<ERigTransformType::Type, TTuple<int32, int32>> NewRanges;
+
+	bool bHasMatchingRanges = true;
+	for(int32 Index = 0; Index < TransformInfos.Num(); Index++)
+	{
+		const int32 CurrentTransformIndex = TransformInfos[Index].ComputedTransform->Index;
+		const int32 NewTransformIndex = NewElementTransforms.Storage.Add(ElementTransforms[CurrentTransformIndex]);
+		TransformInfos[Index].ComputedTransform->Index = NewTransformIndex;
+
+		const int32 CurrentDirtyStateIndex = TransformInfos[Index].DirtyState->Index;
+		const int32 NewDirtyStateIndex = NewElementDirtyStates.Storage.Add(ElementDirtyStates[CurrentDirtyStateIndex]);
+		TransformInfos[Index].DirtyState->Index = NewDirtyStateIndex;
+
+		if(NewTransformIndex == NewDirtyStateIndex)
+		{
+			TTuple<int32,int32>& Range = NewRanges.FindOrAdd(TransformInfos[Index].TypeInfo.Transform, {INT32_MAX, INDEX_NONE});
+			Range.Get<0>() = FMath::Min(NewTransformIndex, Range.Get<0>());
+			Range.Get<1>() = FMath::Max(NewTransformIndex, Range.Get<1>());
+		}
+		else
+		{
+			bHasMatchingRanges = false;
+		}
+	}
+
+	ElementTransforms = NewElementTransforms;
+	ElementDirtyStates = NewElementDirtyStates;
+	UpdateElementStorage();
+
+	if(bHasMatchingRanges)
+	{
+		ElementTransformRanges = NewRanges;
+	}
+	else
+	{
+		ElementTransformRanges.Reset();
+	}
+}
+
+TOptional<TTuple<int32, int32>> URigHierarchy::GetElementStorageRange(ERigTransformType::Type InTransformType) const
+{
+	if(const TTuple<int32,int32>* Range = ElementTransformRanges.Find(InTransformType))
+	{
+		return *Range;
+	}
+	return TOptional<TTuple<int32, int32>>();
 }
 
 void URigHierarchy::PropagateMetadata(const FRigElementKey& InKey, const FName& InName, bool bNotify)
diff --git a/Engine/Plugins/Animation/ControlRig/Source/ControlRig/Public/Rigs/RigHierarchy.h b/Engine/Plugins/Animation/ControlRig/Source/ControlRig/Public/Rigs/RigHierarchy.h
--- a/Engine/Plugins/Animation/ControlRig/Source/ControlRig/Public/Rigs/RigHierarchy.h
+++ b/Engine/Plugins/Animation/ControlRig/Source/ControlRig/Public/Rigs/RigHierarchy.h
@@ -4308,6 +4308,9 @@ private:
 	// do not have URigHierarchy as an owner and therefore do not carry curves with them.
 	FRigReusableElementStorage<float> ElementCurves;
 
+	// A list of ranges which can be used to copy all poses from initial to current, for example 
+	TMap<ERigTransformType::Type, TTuple<int32, int32>> ElementTransformRanges;
+
 	// Allocates the default element storage for an element
 	void AllocateDefaultElementStorage(FRigBaseElement* InElement, bool bUpdateAllElements);
 
@@ -4317,6 +4320,15 @@ private:
 	// Updates all storage pointers of the elements for poses and dirty states
 	void UpdateElementStorage();
 
+	// Orders the element storage by storing first initial, then current,
+	// within each first local, then global, and within each of those lists
+	// we'll place bones, nulls, controls etc in that order
+	void SortElementStorage();
+
+	// Returns the range of the element transform / dirty state storage for a given
+	// transform type. this is only valid if the element storage has been sorted. 
+	TOptional<TTuple<int32,int32>> GetElementStorageRange(ERigTransformType::Type InTransformType) const;
+
 	// Element metadata storage. Storage is defined here rather than on the elements
 	// to reduce memory consumption. Only elements created by MakeElement point to
 	// the element storage. Copied elements via the copy constructor or copy operator
@@ -4969,6 +4981,7 @@ private:
 	friend struct FRigDispatch_SetMetadata;
 	friend struct FRigDispatch_GetModuleMetadata;
 	friend struct FRigDispatch_SetModuleMetadata;
+	friend class FControlRigHierarchySortElementStorage;
 };
 
 struct CONTROLRIG_API FRigHierarchyInteractionBracket
diff --git a/Engine/Plugins/Animation/ControlRig/Source/ControlRig/Public/Rigs/RigHierarchyElements.h b/Engine/Plugins/Animation/ControlRig/Source/ControlRig/Public/Rigs/RigHierarchyElements.h
--- a/Engine/Plugins/Animation/ControlRig/Source/ControlRig/Public/Rigs/RigHierarchyElements.h
+++ b/Engine/Plugins/Animation/ControlRig/Source/ControlRig/Public/Rigs/RigHierarchyElements.h
@@ -242,6 +242,18 @@ namespace ERigTransformType
 	}
 }
 
+UENUM()
+namespace ERigTransformStorageType
+{
+	enum Type : int
+	{
+		Pose,
+		Offset,
+		Shape,
+		NumStorageTypes
+	};
+}
+
 template<typename T>
 struct CONTROLRIG_API FRigReusableElementStorage
 {
