[ControlRig] Commented out outdated code to avoid a crash, the underlying issue should be properly fixed in a future CL

diff --git a/Engine/Plugins/Animation/ControlRig/Source/ControlRigEditor/Private/ControlRigElementDetails.cpp b/Engine/Plugins/Animation/ControlRig/Source/ControlRigEditor/Private/ControlRigElementDetails.cpp
--- a/Engine/Plugins/Animation/ControlRig/Source/ControlRigEditor/Private/ControlRigElementDetails.cpp
+++ b/Engine/Plugins/Animation/ControlRig/Source/ControlRigEditor/Private/ControlRigElementDetails.cpp
@@ -4343,10 +4343,10 @@ void FRigControlElementDetails::CustomizeShape(IDetailLayoutBuilder& DetailBuild
 	const bool bIsProcedural = IsAnyElementProcedural();
 	const bool bIsEnabled = !bIsProcedural;
 
-	TSharedPtr<IPropertyHandle> ShapeHandle = DetailBuilder.GetProperty(TEXT("Shape"));
-	TSharedPtr<IPropertyHandle> InitialHandle = ShapeHandle->GetChildHandle(TEXT("Initial"));
-	TSharedPtr<IPropertyHandle> LocalHandle = InitialHandle->GetChildHandle(TEXT("Local"));
-	ShapeTransformHandle = LocalHandle->GetChildHandle(TEXT("Transform"));
+	// TSharedPtr<IPropertyHandle> ShapeHandle = DetailBuilder.GetProperty(TEXT("Shape"));
+	// TSharedPtr<IPropertyHandle> InitialHandle = ShapeHandle->GetChildHandle(TEXT("Initial"));
+	// TSharedPtr<IPropertyHandle> LocalHandle = InitialHandle->GetChildHandle(TEXT("Local"));
+	// ShapeTransformHandle = LocalHandle->GetChildHandle(TEXT("Transform"));
 	
 	ShapeNameList.Reset();
 	
@@ -5060,22 +5060,22 @@ void FRigControlElementDetails::OnPasteShapeProperties()
 			ShapeColorHandle->NotifyPostChange(EPropertyChangeType::ValueSet);
 		}
 
-		// Transform
-		{
-			ShapeTransformHandle->NotifyPreChange();
-			TArray<void*> RawDataPtrs;
-			ShapeTransformHandle->AccessRawData(RawDataPtrs);
-			for (void* RawPtr: RawDataPtrs)
-			{
-				bSuccessful &= static_cast<FTransform*>(RawPtr)->InitFromString(TransformStr);
-				if (!bSuccessful)
-				{
-					Transaction.Cancel();
-					return;
-				}
-			}		
-			ShapeTransformHandle->NotifyPostChange(EPropertyChangeType::ValueSet);
-		}		
+		// // Transform
+		// {
+		// 	ShapeTransformHandle->NotifyPreChange();
+		// 	TArray<void*> RawDataPtrs;
+		// 	ShapeTransformHandle->AccessRawData(RawDataPtrs);
+		// 	for (void* RawPtr: RawDataPtrs)
+		// 	{
+		// 		bSuccessful &= static_cast<FTransform*>(RawPtr)->InitFromString(TransformStr);
+		// 		if (!bSuccessful)
+		// 		{
+		// 			Transaction.Cancel();
+		// 			return;
+		// 		}
+		// 	}		
+		// 	ShapeTransformHandle->NotifyPostChange(EPropertyChangeType::ValueSet);
+		// }		
 	}
 }
 
