[RigHierarchy] added dirty flag copy for AddNull

diff --git a/Engine/Plugins/Animation/ControlRig/Source/ControlRig/Private/Rigs/RigHierarchyController.cpp b/Engine/Plugins/Animation/ControlRig/Source/ControlRig/Private/Rigs/RigHierarchyController.cpp
--- a/Engine/Plugins/Animation/ControlRig/Source/ControlRig/Private/Rigs/RigHierarchyController.cpp
+++ b/Engine/Plugins/Animation/ControlRig/Source/ControlRig/Private/Rigs/RigHierarchyController.cpp
@@ -299,6 +299,7 @@ FRigElementKey URigHierarchyController::AddNull(FName InName, FRigElementKey InP
 		}
 
 		NewElement->GetTransform().Current = NewElement->GetTransform().Initial;
+		NewElement->GetDirtyState().Current = NewElement->GetDirtyState().Initial;
 	}
 
 #if WITH_EDITOR
