[ControlRigEditor] removed old shape transform related logic that are now handled by a dedicated transform widget

diff --git a/Engine/Plugins/Animation/ControlRig/Source/ControlRigEditor/Private/ControlRigElementDetails.cpp b/Engine/Plugins/Animation/ControlRig/Source/ControlRigEditor/Private/ControlRigElementDetails.cpp
--- a/Engine/Plugins/Animation/ControlRig/Source/ControlRigEditor/Private/ControlRigElementDetails.cpp
+++ b/Engine/Plugins/Animation/ControlRig/Source/ControlRigEditor/Private/ControlRigElementDetails.cpp
@@ -4343,11 +4343,6 @@ void FRigControlElementDetails::CustomizeShape(IDetailLayoutBuilder& DetailBuild
 	const bool bIsProcedural = IsAnyElementProcedural();
 	const bool bIsEnabled = !bIsProcedural;
 
-	// TSharedPtr<IPropertyHandle> ShapeHandle = DetailBuilder.GetProperty(TEXT("Shape"));
-	// TSharedPtr<IPropertyHandle> InitialHandle = ShapeHandle->GetChildHandle(TEXT("Initial"));
-	// TSharedPtr<IPropertyHandle> LocalHandle = InitialHandle->GetChildHandle(TEXT("Local"));
-	// ShapeTransformHandle = LocalHandle->GetChildHandle(TEXT("Transform"));
-	
 	ShapeNameList.Reset();
 	
 	if (UControlRigBlueprint* Blueprint = PerElementInfos[0].GetBlueprint())
@@ -5005,10 +5000,9 @@ void FRigControlElementDetails::OnCopyShapeProperties()
 	{
 		if(const FRigControlElement* ControlElement = PerElementInfos[0].GetElement<FRigControlElement>())
 		{
-			Value = FString::Printf(TEXT("(ShapeName=\"%s\",ShapeColor=%s,Transform=%s)"),
+			Value = FString::Printf(TEXT("(ShapeName=\"%s\",ShapeColor=%s)"),
 				*ControlElement->Settings.ShapeName.ToString(),
-				*ControlElement->Settings.ShapeColor.ToString(),
-				*ControlElement->GetShapeTransform().Initial.Local.Get().ToString());
+				*ControlElement->Settings.ShapeColor.ToString());
 		}
 	}
 		
@@ -5027,10 +5021,8 @@ void FRigControlElementDetails::OnPasteShapeProperties()
 	FString TrimmedText = PastedText.LeftChop(1).RightChop(1);
 	FString ShapeName;
 	FString ShapeColorStr;
-	FString TransformStr;
 	bool bSuccessful = FParse::Value(*TrimmedText, TEXT("ShapeName="), ShapeName) &&
-					   FParse::Value(*TrimmedText, TEXT("ShapeColor="), ShapeColorStr, false) &&
-					   FParse::Value(*TrimmedText, TEXT("Transform="), TransformStr, false);
+					   FParse::Value(*TrimmedText, TEXT("ShapeColor="), ShapeColorStr, false);
 
 	if (bSuccessful)
 	{
@@ -5059,23 +5051,6 @@ void FRigControlElementDetails::OnPasteShapeProperties()
 			}		
 			ShapeColorHandle->NotifyPostChange(EPropertyChangeType::ValueSet);
 		}
-
-		// // Transform
-		// {
-		// 	ShapeTransformHandle->NotifyPreChange();
-		// 	TArray<void*> RawDataPtrs;
-		// 	ShapeTransformHandle->AccessRawData(RawDataPtrs);
-		// 	for (void* RawPtr: RawDataPtrs)
-		// 	{
-		// 		bSuccessful &= static_cast<FTransform*>(RawPtr)->InitFromString(TransformStr);
-		// 		if (!bSuccessful)
-		// 		{
-		// 			Transaction.Cancel();
-		// 			return;
-		// 		}
-		// 	}		
-		// 	ShapeTransformHandle->NotifyPostChange(EPropertyChangeType::ValueSet);
-		// }		
 	}
 }
 
diff --git a/Engine/Plugins/Animation/ControlRig/Source/ControlRigEditor/Private/ControlRigElementDetails.h b/Engine/Plugins/Animation/ControlRig/Source/ControlRigEditor/Private/ControlRigElementDetails.h
--- a/Engine/Plugins/Animation/ControlRig/Source/ControlRigEditor/Private/ControlRigElementDetails.h
+++ b/Engine/Plugins/Animation/ControlRig/Source/ControlRigEditor/Private/ControlRigElementDetails.h
@@ -761,7 +761,6 @@ private:
 
 	TSharedPtr<IPropertyHandle> ShapeNameHandle;
 	TSharedPtr<IPropertyHandle> ShapeColorHandle;
-	TSharedPtr<IPropertyHandle> ShapeTransformHandle;
 
 	TSharedPtr<SControlRigShapeNameList> ShapeNameListWidget; 
 	static TSharedPtr<TArray<ERigControlValueType>> PickedValueTypes;
