Control Rig: Introducing pose adapters - first steps

#rnx

diff --git a/Engine/Plugins/Animation/ControlRig/Source/ControlRig/Private/Rigs/RigHierarchy.cpp b/Engine/Plugins/Animation/ControlRig/Source/ControlRig/Private/Rigs/RigHierarchy.cpp
--- a/Engine/Plugins/Animation/ControlRig/Source/ControlRig/Private/Rigs/RigHierarchy.cpp
+++ b/Engine/Plugins/Animation/ControlRig/Source/ControlRig/Private/Rigs/RigHierarchy.cpp
@@ -344,33 +344,33 @@ void URigHierarchy::Load(FArchive& Ar)
 			
 			if(FRigTransformElement* TransformElement = Cast<FRigTransformElement>(Element))
 			{
-				TransformElement->PoseStorage.Initial.Local.Index = TransformIndices[UsedTransformIndex++];
-				TransformElement->PoseStorage.Current.Local.Index = TransformIndices[UsedTransformIndex++];
-				TransformElement->PoseStorage.Initial.Global.Index = TransformIndices[UsedTransformIndex++];
-				TransformElement->PoseStorage.Current.Global.Index = TransformIndices[UsedTransformIndex++];
-				TransformElement->PoseDirtyState.Initial.Local.Index = DirtyStateIndices[UsedDirtyStateIndex++];
-				TransformElement->PoseDirtyState.Current.Local.Index = DirtyStateIndices[UsedDirtyStateIndex++];
-				TransformElement->PoseDirtyState.Initial.Global.Index = DirtyStateIndices[UsedDirtyStateIndex++];
-				TransformElement->PoseDirtyState.Current.Global.Index = DirtyStateIndices[UsedDirtyStateIndex++];
+				TransformElement->PoseStorage.Initial.Local.StorageIndex = TransformIndices[UsedTransformIndex++];
+				TransformElement->PoseStorage.Current.Local.StorageIndex = TransformIndices[UsedTransformIndex++];
+				TransformElement->PoseStorage.Initial.Global.StorageIndex = TransformIndices[UsedTransformIndex++];
+				TransformElement->PoseStorage.Current.Global.StorageIndex = TransformIndices[UsedTransformIndex++];
+				TransformElement->PoseDirtyState.Initial.Local.StorageIndex = DirtyStateIndices[UsedDirtyStateIndex++];
+				TransformElement->PoseDirtyState.Current.Local.StorageIndex = DirtyStateIndices[UsedDirtyStateIndex++];
+				TransformElement->PoseDirtyState.Initial.Global.StorageIndex = DirtyStateIndices[UsedDirtyStateIndex++];
+				TransformElement->PoseDirtyState.Current.Global.StorageIndex = DirtyStateIndices[UsedDirtyStateIndex++];
 
 				if(FRigControlElement* ControlElement = Cast<FRigControlElement>(TransformElement))
 				{
-					ControlElement->OffsetStorage.Initial.Local.Index = TransformIndices[UsedTransformIndex++];
-					ControlElement->OffsetStorage.Current.Local.Index = TransformIndices[UsedTransformIndex++];
-					ControlElement->OffsetStorage.Initial.Global.Index = TransformIndices[UsedTransformIndex++];
-					ControlElement->OffsetStorage.Current.Global.Index = TransformIndices[UsedTransformIndex++];
-					ControlElement->OffsetDirtyState.Initial.Local.Index = DirtyStateIndices[UsedDirtyStateIndex++];
-					ControlElement->OffsetDirtyState.Current.Local.Index = DirtyStateIndices[UsedDirtyStateIndex++];
-					ControlElement->OffsetDirtyState.Initial.Global.Index = DirtyStateIndices[UsedDirtyStateIndex++];
-					ControlElement->OffsetDirtyState.Current.Global.Index = DirtyStateIndices[UsedDirtyStateIndex++];
-					ControlElement->ShapeStorage.Initial.Local.Index = TransformIndices[UsedTransformIndex++];
-					ControlElement->ShapeStorage.Current.Local.Index = TransformIndices[UsedTransformIndex++];
-					ControlElement->ShapeStorage.Initial.Global.Index = TransformIndices[UsedTransformIndex++];
-					ControlElement->ShapeStorage.Current.Global.Index = TransformIndices[UsedTransformIndex++];
-					ControlElement->ShapeDirtyState.Initial.Local.Index = DirtyStateIndices[UsedDirtyStateIndex++];
-					ControlElement->ShapeDirtyState.Current.Local.Index = DirtyStateIndices[UsedDirtyStateIndex++];
-					ControlElement->ShapeDirtyState.Initial.Global.Index = DirtyStateIndices[UsedDirtyStateIndex++];
-					ControlElement->ShapeDirtyState.Current.Global.Index = DirtyStateIndices[UsedDirtyStateIndex++];
+					ControlElement->OffsetStorage.Initial.Local.StorageIndex = TransformIndices[UsedTransformIndex++];
+					ControlElement->OffsetStorage.Current.Local.StorageIndex = TransformIndices[UsedTransformIndex++];
+					ControlElement->OffsetStorage.Initial.Global.StorageIndex = TransformIndices[UsedTransformIndex++];
+					ControlElement->OffsetStorage.Current.Global.StorageIndex = TransformIndices[UsedTransformIndex++];
+					ControlElement->OffsetDirtyState.Initial.Local.StorageIndex = DirtyStateIndices[UsedDirtyStateIndex++];
+					ControlElement->OffsetDirtyState.Current.Local.StorageIndex = DirtyStateIndices[UsedDirtyStateIndex++];
+					ControlElement->OffsetDirtyState.Initial.Global.StorageIndex = DirtyStateIndices[UsedDirtyStateIndex++];
+					ControlElement->OffsetDirtyState.Current.Global.StorageIndex = DirtyStateIndices[UsedDirtyStateIndex++];
+					ControlElement->ShapeStorage.Initial.Local.StorageIndex = TransformIndices[UsedTransformIndex++];
+					ControlElement->ShapeStorage.Current.Local.StorageIndex = TransformIndices[UsedTransformIndex++];
+					ControlElement->ShapeStorage.Initial.Global.StorageIndex = TransformIndices[UsedTransformIndex++];
+					ControlElement->ShapeStorage.Current.Global.StorageIndex = TransformIndices[UsedTransformIndex++];
+					ControlElement->ShapeDirtyState.Initial.Local.StorageIndex = DirtyStateIndices[UsedDirtyStateIndex++];
+					ControlElement->ShapeDirtyState.Current.Local.StorageIndex = DirtyStateIndices[UsedDirtyStateIndex++];
+					ControlElement->ShapeDirtyState.Initial.Global.StorageIndex = DirtyStateIndices[UsedDirtyStateIndex++];
+					ControlElement->ShapeDirtyState.Current.Global.StorageIndex = DirtyStateIndices[UsedDirtyStateIndex++];
 				}
 			}
 			else if(FRigCurveElement* CurveElement = Cast<FRigCurveElement>(Element))
@@ -444,7 +444,7 @@ void URigHierarchy::Load(FArchive& Ar)
 		}
 	}
 
-	SortElementStorage();
+	(void)SortElementStorage();
 }
 
 void URigHierarchy::PostLoad()
@@ -754,33 +754,33 @@ void URigHierarchy::CopyHierarchy(URigHierarchy* InHierarchy)
 
 			if(FRigTransformElement* TransformElement = Cast<FRigTransformElement>(Target))
 			{
-				TransformElement->PoseStorage.Initial.Local.Index = TransformIndices[UsedTransformIndex++];
-				TransformElement->PoseStorage.Current.Local.Index = TransformIndices[UsedTransformIndex++];
-				TransformElement->PoseStorage.Initial.Global.Index = TransformIndices[UsedTransformIndex++];
-				TransformElement->PoseStorage.Current.Global.Index = TransformIndices[UsedTransformIndex++];
-				TransformElement->PoseDirtyState.Initial.Local.Index = DirtyStateIndices[UsedDirtyStateIndex++];
-				TransformElement->PoseDirtyState.Current.Local.Index = DirtyStateIndices[UsedDirtyStateIndex++];
-				TransformElement->PoseDirtyState.Initial.Global.Index = DirtyStateIndices[UsedDirtyStateIndex++];
-				TransformElement->PoseDirtyState.Current.Global.Index = DirtyStateIndices[UsedDirtyStateIndex++];
+				TransformElement->PoseStorage.Initial.Local.StorageIndex = TransformIndices[UsedTransformIndex++];
+				TransformElement->PoseStorage.Current.Local.StorageIndex = TransformIndices[UsedTransformIndex++];
+				TransformElement->PoseStorage.Initial.Global.StorageIndex = TransformIndices[UsedTransformIndex++];
+				TransformElement->PoseStorage.Current.Global.StorageIndex = TransformIndices[UsedTransformIndex++];
+				TransformElement->PoseDirtyState.Initial.Local.StorageIndex = DirtyStateIndices[UsedDirtyStateIndex++];
+				TransformElement->PoseDirtyState.Current.Local.StorageIndex = DirtyStateIndices[UsedDirtyStateIndex++];
+				TransformElement->PoseDirtyState.Initial.Global.StorageIndex = DirtyStateIndices[UsedDirtyStateIndex++];
+				TransformElement->PoseDirtyState.Current.Global.StorageIndex = DirtyStateIndices[UsedDirtyStateIndex++];
 
 				if(FRigControlElement* ControlElement = Cast<FRigControlElement>(TransformElement))
 				{
-					ControlElement->OffsetStorage.Initial.Local.Index = TransformIndices[UsedTransformIndex++];
-					ControlElement->OffsetStorage.Current.Local.Index = TransformIndices[UsedTransformIndex++];
-					ControlElement->OffsetStorage.Initial.Global.Index = TransformIndices[UsedTransformIndex++];
-					ControlElement->OffsetStorage.Current.Global.Index = TransformIndices[UsedTransformIndex++];
-					ControlElement->OffsetDirtyState.Initial.Local.Index = DirtyStateIndices[UsedDirtyStateIndex++];
-					ControlElement->OffsetDirtyState.Current.Local.Index = DirtyStateIndices[UsedDirtyStateIndex++];
-					ControlElement->OffsetDirtyState.Initial.Global.Index = DirtyStateIndices[UsedDirtyStateIndex++];
-					ControlElement->OffsetDirtyState.Current.Global.Index = DirtyStateIndices[UsedDirtyStateIndex++];
-					ControlElement->ShapeStorage.Initial.Local.Index = TransformIndices[UsedTransformIndex++];
-					ControlElement->ShapeStorage.Current.Local.Index = TransformIndices[UsedTransformIndex++];
-					ControlElement->ShapeStorage.Initial.Global.Index = TransformIndices[UsedTransformIndex++];
-					ControlElement->ShapeStorage.Current.Global.Index = TransformIndices[UsedTransformIndex++];
-					ControlElement->ShapeDirtyState.Initial.Local.Index = DirtyStateIndices[UsedDirtyStateIndex++];
-					ControlElement->ShapeDirtyState.Current.Local.Index = DirtyStateIndices[UsedDirtyStateIndex++];
-					ControlElement->ShapeDirtyState.Initial.Global.Index = DirtyStateIndices[UsedDirtyStateIndex++];
-					ControlElement->ShapeDirtyState.Current.Global.Index = DirtyStateIndices[UsedDirtyStateIndex++];
+					ControlElement->OffsetStorage.Initial.Local.StorageIndex = TransformIndices[UsedTransformIndex++];
+					ControlElement->OffsetStorage.Current.Local.StorageIndex = TransformIndices[UsedTransformIndex++];
+					ControlElement->OffsetStorage.Initial.Global.StorageIndex = TransformIndices[UsedTransformIndex++];
+					ControlElement->OffsetStorage.Current.Global.StorageIndex = TransformIndices[UsedTransformIndex++];
+					ControlElement->OffsetDirtyState.Initial.Local.StorageIndex = DirtyStateIndices[UsedDirtyStateIndex++];
+					ControlElement->OffsetDirtyState.Current.Local.StorageIndex = DirtyStateIndices[UsedDirtyStateIndex++];
+					ControlElement->OffsetDirtyState.Initial.Global.StorageIndex = DirtyStateIndices[UsedDirtyStateIndex++];
+					ControlElement->OffsetDirtyState.Current.Global.StorageIndex = DirtyStateIndices[UsedDirtyStateIndex++];
+					ControlElement->ShapeStorage.Initial.Local.StorageIndex = TransformIndices[UsedTransformIndex++];
+					ControlElement->ShapeStorage.Current.Local.StorageIndex = TransformIndices[UsedTransformIndex++];
+					ControlElement->ShapeStorage.Initial.Global.StorageIndex = TransformIndices[UsedTransformIndex++];
+					ControlElement->ShapeStorage.Current.Global.StorageIndex = TransformIndices[UsedTransformIndex++];
+					ControlElement->ShapeDirtyState.Initial.Local.StorageIndex = DirtyStateIndices[UsedDirtyStateIndex++];
+					ControlElement->ShapeDirtyState.Current.Local.StorageIndex = DirtyStateIndices[UsedDirtyStateIndex++];
+					ControlElement->ShapeDirtyState.Initial.Global.StorageIndex = DirtyStateIndices[UsedDirtyStateIndex++];
+					ControlElement->ShapeDirtyState.Current.Global.StorageIndex = DirtyStateIndices[UsedDirtyStateIndex++];
 				}
 			}
 			else if(FRigCurveElement* CurveElement = Cast<FRigCurveElement>(Target))
@@ -794,7 +794,7 @@ void URigHierarchy::CopyHierarchy(URigHierarchy* InHierarchy)
 			CopyAllMetadataFromElement(Target, Source);
 		}
 
-		SortElementStorage();
+		(void)SortElementStorage();
 
 		PreviousNameMap.Append(InHierarchy->PreviousNameMap);
 
@@ -3418,6 +3418,21 @@ void URigHierarchy::SetPose(const FRigPose& InPose, ERigTransformType::Type InTr
 	}
 }
 
+void URigHierarchy::LinkPoseAdapter(TSharedPtr<FRigHierarchyPoseAdapter> InPoseAdapter)
+{
+	if(PoseAdapter)
+	{
+		PoseAdapter->PreUnlinked(this);
+		PoseAdapter.Reset();
+	}
+
+	if(InPoseAdapter)
+	{
+		PoseAdapter = InPoseAdapter;
+		PoseAdapter->PostLinked(this);
+	}
+}
+
 void URigHierarchy::Notify(ERigHierarchyNotification InNotifType, const FRigBaseElement* InElement)
 {
 	if(bSuspendNotifications)
@@ -5570,6 +5585,8 @@ void URigHierarchy::DestroyElement(FRigBaseElement*& InElement)
 
 	FMemory::Free(InElement);
 	InElement = nullptr;
+	
+	ElementTransformRanges.Reset();
 }
 
 void URigHierarchy::PropagateDirtyFlags(FRigTransformElement* InTransformElement, bool bInitial, bool bAffectChildren, bool bComputeOpposed, bool bMarkDirty) const
@@ -5918,34 +5935,34 @@ void URigHierarchy::AllocateDefaultElementStorage(FRigBaseElement* InElement, bo
 	{
 		// only for the default allocation we want to catch double allocations.
 		// so if we are already linked to the default storage, exit early.
-		if(ElementTransforms.Contains(TransformElement->PoseStorage.Initial.Local.Index, TransformElement->PoseStorage.Initial.Local.Storage))
+		if(ElementTransforms.Contains(TransformElement->PoseStorage.Initial.Local.StorageIndex, TransformElement->PoseStorage.Initial.Local.Storage))
 		{
-			check(ElementTransforms.Contains(TransformElement->PoseStorage.Current.Local.Index, TransformElement->PoseStorage.Current.Local.Storage));
-			check(ElementTransforms.Contains(TransformElement->PoseStorage.Initial.Global.Index, TransformElement->PoseStorage.Initial.Global.Storage));
-			check(ElementTransforms.Contains(TransformElement->PoseStorage.Current.Global.Index, TransformElement->PoseStorage.Current.Global.Storage));
-			check(ElementDirtyStates.Contains(TransformElement->PoseDirtyState.Initial.Local.Index, TransformElement->PoseDirtyState.Initial.Local.Storage));
-			check(ElementDirtyStates.Contains(TransformElement->PoseDirtyState.Current.Local.Index, TransformElement->PoseDirtyState.Current.Local.Storage));
-			check(ElementDirtyStates.Contains(TransformElement->PoseDirtyState.Initial.Global.Index, TransformElement->PoseDirtyState.Initial.Global.Storage));
-			check(ElementDirtyStates.Contains(TransformElement->PoseDirtyState.Current.Global.Index, TransformElement->PoseDirtyState.Current.Global.Storage));
+			check(ElementTransforms.Contains(TransformElement->PoseStorage.Current.Local.StorageIndex, TransformElement->PoseStorage.Current.Local.Storage));
+			check(ElementTransforms.Contains(TransformElement->PoseStorage.Initial.Global.StorageIndex, TransformElement->PoseStorage.Initial.Global.Storage));
+			check(ElementTransforms.Contains(TransformElement->PoseStorage.Current.Global.StorageIndex, TransformElement->PoseStorage.Current.Global.Storage));
+			check(ElementDirtyStates.Contains(TransformElement->PoseDirtyState.Initial.Local.StorageIndex, TransformElement->PoseDirtyState.Initial.Local.Storage));
+			check(ElementDirtyStates.Contains(TransformElement->PoseDirtyState.Current.Local.StorageIndex, TransformElement->PoseDirtyState.Current.Local.Storage));
+			check(ElementDirtyStates.Contains(TransformElement->PoseDirtyState.Initial.Global.StorageIndex, TransformElement->PoseDirtyState.Initial.Global.Storage));
+			check(ElementDirtyStates.Contains(TransformElement->PoseDirtyState.Current.Global.StorageIndex, TransformElement->PoseDirtyState.Current.Global.Storage));
 
 			if(const FRigControlElement* ControlElement = Cast<FRigControlElement>(TransformElement))
 			{
-				check(ElementTransforms.Contains(ControlElement->OffsetStorage.Initial.Local.Index, ControlElement->OffsetStorage.Initial.Local.Storage));
-				check(ElementTransforms.Contains(ControlElement->OffsetStorage.Current.Local.Index, ControlElement->OffsetStorage.Current.Local.Storage));
-				check(ElementTransforms.Contains(ControlElement->OffsetStorage.Initial.Global.Index, ControlElement->OffsetStorage.Initial.Global.Storage));
-				check(ElementTransforms.Contains(ControlElement->OffsetStorage.Current.Global.Index, ControlElement->OffsetStorage.Current.Global.Storage));
-				check(ElementDirtyStates.Contains(ControlElement->OffsetDirtyState.Initial.Local.Index, ControlElement->OffsetDirtyState.Initial.Local.Storage));
-				check(ElementDirtyStates.Contains(ControlElement->OffsetDirtyState.Current.Local.Index, ControlElement->OffsetDirtyState.Current.Local.Storage));
-				check(ElementDirtyStates.Contains(ControlElement->OffsetDirtyState.Initial.Global.Index, ControlElement->OffsetDirtyState.Initial.Global.Storage));
-				check(ElementDirtyStates.Contains(ControlElement->OffsetDirtyState.Current.Global.Index, ControlElement->OffsetDirtyState.Current.Global.Storage));
-				check(ElementTransforms.Contains(ControlElement->ShapeStorage.Initial.Local.Index, ControlElement->ShapeStorage.Initial.Local.Storage));
-				check(ElementTransforms.Contains(ControlElement->ShapeStorage.Current.Local.Index, ControlElement->ShapeStorage.Current.Local.Storage));
-				check(ElementTransforms.Contains(ControlElement->ShapeStorage.Initial.Global.Index, ControlElement->ShapeStorage.Initial.Global.Storage));
-				check(ElementTransforms.Contains(ControlElement->ShapeStorage.Current.Global.Index, ControlElement->ShapeStorage.Current.Global.Storage));
-				check(ElementDirtyStates.Contains(ControlElement->ShapeDirtyState.Initial.Local.Index, ControlElement->ShapeDirtyState.Initial.Local.Storage));
-				check(ElementDirtyStates.Contains(ControlElement->ShapeDirtyState.Current.Local.Index, ControlElement->ShapeDirtyState.Current.Local.Storage));
-				check(ElementDirtyStates.Contains(ControlElement->ShapeDirtyState.Initial.Global.Index, ControlElement->ShapeDirtyState.Initial.Global.Storage));
-				check(ElementDirtyStates.Contains(ControlElement->ShapeDirtyState.Current.Global.Index, ControlElement->ShapeDirtyState.Current.Global.Storage));
+				check(ElementTransforms.Contains(ControlElement->OffsetStorage.Initial.Local.StorageIndex, ControlElement->OffsetStorage.Initial.Local.Storage));
+				check(ElementTransforms.Contains(ControlElement->OffsetStorage.Current.Local.StorageIndex, ControlElement->OffsetStorage.Current.Local.Storage));
+				check(ElementTransforms.Contains(ControlElement->OffsetStorage.Initial.Global.StorageIndex, ControlElement->OffsetStorage.Initial.Global.Storage));
+				check(ElementTransforms.Contains(ControlElement->OffsetStorage.Current.Global.StorageIndex, ControlElement->OffsetStorage.Current.Global.Storage));
+				check(ElementDirtyStates.Contains(ControlElement->OffsetDirtyState.Initial.Local.StorageIndex, ControlElement->OffsetDirtyState.Initial.Local.Storage));
+				check(ElementDirtyStates.Contains(ControlElement->OffsetDirtyState.Current.Local.StorageIndex, ControlElement->OffsetDirtyState.Current.Local.Storage));
+				check(ElementDirtyStates.Contains(ControlElement->OffsetDirtyState.Initial.Global.StorageIndex, ControlElement->OffsetDirtyState.Initial.Global.Storage));
+				check(ElementDirtyStates.Contains(ControlElement->OffsetDirtyState.Current.Global.StorageIndex, ControlElement->OffsetDirtyState.Current.Global.Storage));
+				check(ElementTransforms.Contains(ControlElement->ShapeStorage.Initial.Local.StorageIndex, ControlElement->ShapeStorage.Initial.Local.Storage));
+				check(ElementTransforms.Contains(ControlElement->ShapeStorage.Current.Local.StorageIndex, ControlElement->ShapeStorage.Current.Local.Storage));
+				check(ElementTransforms.Contains(ControlElement->ShapeStorage.Initial.Global.StorageIndex, ControlElement->ShapeStorage.Initial.Global.Storage));
+				check(ElementTransforms.Contains(ControlElement->ShapeStorage.Current.Global.StorageIndex, ControlElement->ShapeStorage.Current.Global.Storage));
+				check(ElementDirtyStates.Contains(ControlElement->ShapeDirtyState.Initial.Local.StorageIndex, ControlElement->ShapeDirtyState.Initial.Local.Storage));
+				check(ElementDirtyStates.Contains(ControlElement->ShapeDirtyState.Current.Local.StorageIndex, ControlElement->ShapeDirtyState.Current.Local.Storage));
+				check(ElementDirtyStates.Contains(ControlElement->ShapeDirtyState.Initial.Global.StorageIndex, ControlElement->ShapeDirtyState.Initial.Global.Storage));
+				check(ElementDirtyStates.Contains(ControlElement->ShapeDirtyState.Current.Global.StorageIndex, ControlElement->ShapeDirtyState.Current.Global.Storage));
 			}
 			return;
 		}
@@ -5953,35 +5970,35 @@ void URigHierarchy::AllocateDefaultElementStorage(FRigBaseElement* InElement, bo
 		check(TransformIndices.Num() >= 4);
 		const TArray<int32, TInlineAllocator<4>> DirtyStateIndices = ElementDirtyStates.Allocate(TransformElement->GetNumTransforms(), false);
 		check(DirtyStateIndices.Num() >= 4);
-		TransformElement->PoseStorage.Initial.Local.Index = TransformIndices[0];
-		TransformElement->PoseStorage.Current.Local.Index = TransformIndices[1];
-		TransformElement->PoseStorage.Initial.Global.Index = TransformIndices[2];
-		TransformElement->PoseStorage.Current.Global.Index = TransformIndices[3];
-		TransformElement->PoseDirtyState.Initial.Local.Index = DirtyStateIndices[0];
-		TransformElement->PoseDirtyState.Current.Local.Index = DirtyStateIndices[1];
-		TransformElement->PoseDirtyState.Initial.Global.Index = DirtyStateIndices[2];
-		TransformElement->PoseDirtyState.Current.Global.Index = DirtyStateIndices[3];
+		TransformElement->PoseStorage.Initial.Local.StorageIndex = TransformIndices[0];
+		TransformElement->PoseStorage.Current.Local.StorageIndex = TransformIndices[1];
+		TransformElement->PoseStorage.Initial.Global.StorageIndex = TransformIndices[2];
+		TransformElement->PoseStorage.Current.Global.StorageIndex = TransformIndices[3];
+		TransformElement->PoseDirtyState.Initial.Local.StorageIndex = DirtyStateIndices[0];
+		TransformElement->PoseDirtyState.Current.Local.StorageIndex = DirtyStateIndices[1];
+		TransformElement->PoseDirtyState.Initial.Global.StorageIndex = DirtyStateIndices[2];
+		TransformElement->PoseDirtyState.Current.Global.StorageIndex = DirtyStateIndices[3];
 
 		if(FRigControlElement* ControlElement = Cast<FRigControlElement>(TransformElement))
 		{
 			check(TransformIndices.Num() >= 12);
 			check(DirtyStateIndices.Num() >= 12);
-			ControlElement->OffsetStorage.Initial.Local.Index = TransformIndices[4];
-			ControlElement->OffsetStorage.Current.Local.Index = TransformIndices[5];
-			ControlElement->OffsetStorage.Initial.Global.Index = TransformIndices[6];
-			ControlElement->OffsetStorage.Current.Global.Index = TransformIndices[7];
-			ControlElement->OffsetDirtyState.Initial.Local.Index = DirtyStateIndices[4];
-			ControlElement->OffsetDirtyState.Current.Local.Index = DirtyStateIndices[5];
-			ControlElement->OffsetDirtyState.Initial.Global.Index = DirtyStateIndices[6];
-			ControlElement->OffsetDirtyState.Current.Global.Index = DirtyStateIndices[7];
-			ControlElement->ShapeStorage.Initial.Local.Index = TransformIndices[8];
-			ControlElement->ShapeStorage.Current.Local.Index = TransformIndices[9];
-			ControlElement->ShapeStorage.Initial.Global.Index = TransformIndices[10];
-			ControlElement->ShapeStorage.Current.Global.Index = TransformIndices[11];
-			ControlElement->ShapeDirtyState.Initial.Local.Index = DirtyStateIndices[8];
-			ControlElement->ShapeDirtyState.Current.Local.Index = DirtyStateIndices[9];
-			ControlElement->ShapeDirtyState.Initial.Global.Index = DirtyStateIndices[10];
-			ControlElement->ShapeDirtyState.Current.Global.Index = DirtyStateIndices[11];
+			ControlElement->OffsetStorage.Initial.Local.StorageIndex = TransformIndices[4];
+			ControlElement->OffsetStorage.Current.Local.StorageIndex = TransformIndices[5];
+			ControlElement->OffsetStorage.Initial.Global.StorageIndex = TransformIndices[6];
+			ControlElement->OffsetStorage.Current.Global.StorageIndex = TransformIndices[7];
+			ControlElement->OffsetDirtyState.Initial.Local.StorageIndex = DirtyStateIndices[4];
+			ControlElement->OffsetDirtyState.Current.Local.StorageIndex = DirtyStateIndices[5];
+			ControlElement->OffsetDirtyState.Initial.Global.StorageIndex = DirtyStateIndices[6];
+			ControlElement->OffsetDirtyState.Current.Global.StorageIndex = DirtyStateIndices[7];
+			ControlElement->ShapeStorage.Initial.Local.StorageIndex = TransformIndices[8];
+			ControlElement->ShapeStorage.Current.Local.StorageIndex = TransformIndices[9];
+			ControlElement->ShapeStorage.Initial.Global.StorageIndex = TransformIndices[10];
+			ControlElement->ShapeStorage.Current.Global.StorageIndex = TransformIndices[11];
+			ControlElement->ShapeDirtyState.Initial.Local.StorageIndex = DirtyStateIndices[8];
+			ControlElement->ShapeDirtyState.Current.Local.StorageIndex = DirtyStateIndices[9];
+			ControlElement->ShapeDirtyState.Initial.Global.StorageIndex = DirtyStateIndices[10];
+			ControlElement->ShapeDirtyState.Current.Global.StorageIndex = DirtyStateIndices[11];
 		}
 	}
 	else if(FRigCurveElement* CurveElement = Cast<FRigCurveElement>(InElement))
@@ -6004,6 +6021,7 @@ void URigHierarchy::AllocateDefaultElementStorage(FRigBaseElement* InElement, bo
 	{
 		InElement->LinkStorage(ElementTransforms.Storage, ElementDirtyStates.Storage, ElementCurves.Storage);
 	}
+	ElementTransformRanges.Reset();
 }
 
 void URigHierarchy::DeallocateElementStorage(FRigBaseElement* InElement)
@@ -6020,11 +6038,13 @@ void URigHierarchy::UpdateElementStorage()
 	ElementTransformRanges.Reset();
 }
 
-void URigHierarchy::SortElementStorage()
+bool URigHierarchy::SortElementStorage()
 {
+	ElementTransformRanges.Reset();
+
 	if(Elements.IsEmpty())
 	{
-		return;
+		return false;
 	}
 	
 	TMap<int32, FRigComputedTransform*> TransformIndexToComputedTransform;
@@ -6090,28 +6110,18 @@ void URigHierarchy::SortElementStorage()
 	TArray<FTransformInfo> TransformInfos;
 	TransformInfos.Reserve(Elements.Num() * 4);
 
-	for(FRigBaseElement* Element : Elements)
-	{
-		if(FRigTransformElement* TransformElement = Cast<FRigTransformElement>(Element))
-		{
-			TransformInfos.Emplace(FTypeInfo(TransformElement->GetIndex(), TransformElement->GetType(), ERigTransformType::InitialLocal, ERigTransformStorageType::Pose), &TransformElement->PoseStorage.Initial.Local, &TransformElement->PoseDirtyState.Initial.Local);
-			TransformInfos.Emplace(FTypeInfo(TransformElement->GetIndex(), TransformElement->GetType(), ERigTransformType::InitialGlobal, ERigTransformStorageType::Pose), &TransformElement->PoseStorage.Initial.Global, &TransformElement->PoseDirtyState.Initial.Global);
-			TransformInfos.Emplace(FTypeInfo(TransformElement->GetIndex(), TransformElement->GetType(), ERigTransformType::CurrentLocal, ERigTransformStorageType::Pose), &TransformElement->PoseStorage.Current.Local, &TransformElement->PoseDirtyState.Current.Local);
-			TransformInfos.Emplace(FTypeInfo(TransformElement->GetIndex(), TransformElement->GetType(), ERigTransformType::CurrentGlobal, ERigTransformStorageType::Pose), &TransformElement->PoseStorage.Current.Global, &TransformElement->PoseDirtyState.Current.Global);
-
-			if(FRigControlElement* ControlElement = Cast<FRigControlElement>(Element))
-			{
-				TransformInfos.Emplace(FTypeInfo(ControlElement->GetIndex(), ControlElement->GetType(), ERigTransformType::InitialLocal, ERigTransformStorageType::Offset), &ControlElement->OffsetStorage.Initial.Local, &ControlElement->OffsetDirtyState.Initial.Local);
-				TransformInfos.Emplace(FTypeInfo(ControlElement->GetIndex(), ControlElement->GetType(), ERigTransformType::InitialGlobal, ERigTransformStorageType::Offset), &ControlElement->OffsetStorage.Initial.Global, &ControlElement->OffsetDirtyState.Initial.Global);
-				TransformInfos.Emplace(FTypeInfo(ControlElement->GetIndex(), ControlElement->GetType(), ERigTransformType::CurrentLocal, ERigTransformStorageType::Offset), &ControlElement->OffsetStorage.Current.Local, &ControlElement->OffsetDirtyState.Current.Local);
-				TransformInfos.Emplace(FTypeInfo(ControlElement->GetIndex(), ControlElement->GetType(), ERigTransformType::CurrentGlobal, ERigTransformStorageType::Offset), &ControlElement->OffsetStorage.Current.Global, &ControlElement->OffsetDirtyState.Current.Global);
-				TransformInfos.Emplace(FTypeInfo(ControlElement->GetIndex(), ControlElement->GetType(), ERigTransformType::InitialLocal, ERigTransformStorageType::Shape), &ControlElement->ShapeStorage.Initial.Local, &ControlElement->ShapeDirtyState.Initial.Local);
-				TransformInfos.Emplace(FTypeInfo(ControlElement->GetIndex(), ControlElement->GetType(), ERigTransformType::InitialGlobal, ERigTransformStorageType::Shape), &ControlElement->ShapeStorage.Initial.Global, &ControlElement->ShapeDirtyState.Initial.Global);
-				TransformInfos.Emplace(FTypeInfo(ControlElement->GetIndex(), ControlElement->GetType(), ERigTransformType::CurrentLocal, ERigTransformStorageType::Shape), &ControlElement->ShapeStorage.Current.Local, &ControlElement->ShapeDirtyState.Current.Local);
-				TransformInfos.Emplace(FTypeInfo(ControlElement->GetIndex(), ControlElement->GetType(), ERigTransformType::CurrentGlobal, ERigTransformStorageType::Shape), &ControlElement->ShapeStorage.Current.Global, &ControlElement->ShapeDirtyState.Current.Global);
-			}
-		}
-	}
+	ForEachTransformElementStorage([&TransformInfos](
+		FRigTransformElement* InTransformElement,
+		ERigTransformType::Type InTransformType,
+		ERigTransformStorageType::Type InStorageType,
+		FRigComputedTransform* InComputedTransform,
+		FRigTransformDirtyState* InDirtyState)
+	{
+		TransformInfos.Emplace(
+			FTypeInfo(InTransformElement->GetIndex(), InTransformElement->GetType(), InTransformType, InStorageType),
+			InComputedTransform,
+			InDirtyState);
+	});
 
 	uint32 NumElements = (uint32)Elements.Num();
 	Algo::SortBy(TransformInfos, [NumElements](const FTransformInfo& InInfo) -> uint32
@@ -6119,41 +6129,57 @@ void URigHierarchy::SortElementStorage()
 		return InInfo.TypeInfo.GetSortKey(NumElements);
 	});
 
+	bool bHasMatchingRanges = true;
+	TMap<ERigTransformType::Type, TTuple<int32, int32>> NewRanges;
+	
 	// check if the sorting is at all required
 	bool bRequiresSort = false;
 	for(int32 Index = 0; Index < TransformInfos.Num(); Index++)
 	{
-		const int32 CurrentTransformIndex = TransformInfos[Index].ComputedTransform->Index;
+		const int32 CurrentTransformIndex = TransformInfos[Index].ComputedTransform->StorageIndex;
 		if(CurrentTransformIndex == INDEX_NONE)
 		{
 			// we cannot sort this if some of the storage is no longer local
 			ElementTransformRanges.Reset();
-			return;
+			return false;
 		}
 		if(CurrentTransformIndex != Index)
 		{
 			bRequiresSort = true;
-			break;
 		}
 
-		const int32 CurrentDirtyStateIndex = TransformInfos[Index].DirtyState->Index;
+		const int32 CurrentDirtyStateIndex = TransformInfos[Index].DirtyState->StorageIndex;
 		if(CurrentDirtyStateIndex == INDEX_NONE)
 		{
 			// we cannot sort this if some of the storage is no longer local
 			ElementTransformRanges.Reset();
-			return;
+			return false;
 		}
 		
 		if(CurrentDirtyStateIndex != Index)
 		{
 			bRequiresSort = true;
-			break;
+		}
+
+		if(CurrentTransformIndex != CurrentDirtyStateIndex)
+		{
+			bHasMatchingRanges = false;
+		}
+		if(bHasMatchingRanges)
+		{
+			TTuple<int32,int32>& Range = NewRanges.FindOrAdd(TransformInfos[Index].TypeInfo.Transform, {INT32_MAX, INDEX_NONE});
+			Range.Get<0>() = FMath::Min(CurrentTransformIndex, Range.Get<0>());
+			Range.Get<1>() = FMath::Max(CurrentTransformIndex, Range.Get<1>());
 		}
 	}
 
 	if(!bRequiresSort)
 	{
-		return;
+		if(bHasMatchingRanges)
+		{
+			ElementTransformRanges = NewRanges;
+		}
+		return false;
 	}
 
 	FRigReusableElementStorage<FTransform> NewElementTransforms;
@@ -6161,18 +6187,17 @@ void URigHierarchy::SortElementStorage()
 	NewElementDirtyStates.Storage.Reserve(ElementDirtyStates.Storage.Num() - ElementDirtyStates.FreeList.Num());
 	NewElementDirtyStates.Storage.Reserve(ElementDirtyStates.Storage.Num() - ElementDirtyStates.FreeList.Num());
 
-	TMap<ERigTransformType::Type, TTuple<int32, int32>> NewRanges;
-
-	bool bHasMatchingRanges = true;
+	bHasMatchingRanges = true;
+	NewRanges.Reset();
 	for(int32 Index = 0; Index < TransformInfos.Num(); Index++)
 	{
-		const int32 CurrentTransformIndex = TransformInfos[Index].ComputedTransform->Index;
+		const int32 CurrentTransformIndex = TransformInfos[Index].ComputedTransform->StorageIndex;
 		const int32 NewTransformIndex = NewElementTransforms.Storage.Add(ElementTransforms[CurrentTransformIndex]);
-		TransformInfos[Index].ComputedTransform->Index = NewTransformIndex;
+		TransformInfos[Index].ComputedTransform->StorageIndex = NewTransformIndex;
 
-		const int32 CurrentDirtyStateIndex = TransformInfos[Index].DirtyState->Index;
+		const int32 CurrentDirtyStateIndex = TransformInfos[Index].DirtyState->StorageIndex;
 		const int32 NewDirtyStateIndex = NewElementDirtyStates.Storage.Add(ElementDirtyStates[CurrentDirtyStateIndex]);
-		TransformInfos[Index].DirtyState->Index = NewDirtyStateIndex;
+		TransformInfos[Index].DirtyState->StorageIndex = NewDirtyStateIndex;
 
 		if(NewTransformIndex == NewDirtyStateIndex)
 		{
@@ -6188,6 +6213,7 @@ void URigHierarchy::SortElementStorage()
 
 	ElementTransforms = NewElementTransforms;
 	ElementDirtyStates = NewElementDirtyStates;
+	
 	UpdateElementStorage();
 
 	if(bHasMatchingRanges)
@@ -6198,6 +6224,141 @@ void URigHierarchy::SortElementStorage()
 	{
 		ElementTransformRanges.Reset();
 	}
+
+	return true;
+}
+
+bool URigHierarchy::ShrinkElementStorage()
+{
+	const TMap<int32, int32> TransformLookup = ElementTransforms.Shrink();
+	const TMap<int32, int32> DirtyStateLookup = ElementDirtyStates.Shrink();
+	const TMap<int32, int32> CurveLookup = ElementCurves.Shrink();
+
+	if(!TransformLookup.IsEmpty() || !DirtyStateLookup.IsEmpty())
+	{
+		ForEachTransformElementStorage([&TransformLookup, &DirtyStateLookup](
+			FRigTransformElement* InTransformElement,
+			ERigTransformType::Type InTransformType,
+			ERigTransformStorageType::Type InStorageType,
+			FRigComputedTransform* InComputedTransform,
+			FRigTransformDirtyState* InDirtyState)
+		{
+			if(const int32* NewTransformIndex = TransformLookup.Find(InComputedTransform->StorageIndex))
+			{
+				InComputedTransform->StorageIndex = *NewTransformIndex;
+			}
+			if(const int32* NewDirtyStateIndex = DirtyStateLookup.Find(InDirtyState->StorageIndex))
+			{
+				InDirtyState->StorageIndex = *NewDirtyStateIndex;
+			}
+		});
+	}
+
+	if(!CurveLookup.IsEmpty())
+	{
+		static const int32 CurveElementIndex = RigElementTypeToFlatIndex(ERigElementType::Curve); 
+		for(FRigBaseElement* Element : ElementsPerType[CurveElementIndex])
+		{
+			FRigCurveElement* CurveElement = CastChecked<FRigCurveElement>(Element);
+			if(const int32* NewIndex = CurveLookup.Find(CurveElement->StorageIndex))
+			{
+				CurveElement->StorageIndex = *NewIndex;
+			}
+		}
+	}
+
+	if(!TransformLookup.IsEmpty() || !DirtyStateLookup.IsEmpty() || !CurveLookup.IsEmpty())
+	{
+		UpdateElementStorage();
+		(void)SortElementStorage();
+		return true;
+	}
+	
+	return false;
+}
+
+void URigHierarchy::ForEachTransformElementStorage(
+	TFunction<void(FRigTransformElement*, ERigTransformType::Type, ERigTransformStorageType::Type, FRigComputedTransform*, FRigTransformDirtyState*)>
+	InCallback)
+{
+	check(InCallback);
+	
+	for(FRigBaseElement* Element : Elements)
+	{
+		if(FRigTransformElement* TransformElement = Cast<FRigTransformElement>(Element))
+		{
+			InCallback(TransformElement, ERigTransformType::InitialLocal, ERigTransformStorageType::Pose, &TransformElement->PoseStorage.Initial.Local, &TransformElement->PoseDirtyState.Initial.Local);
+			InCallback(TransformElement, ERigTransformType::InitialGlobal, ERigTransformStorageType::Pose, &TransformElement->PoseStorage.Initial.Global, &TransformElement->PoseDirtyState.Initial.Global);
+			InCallback(TransformElement, ERigTransformType::CurrentLocal, ERigTransformStorageType::Pose, &TransformElement->PoseStorage.Current.Local, &TransformElement->PoseDirtyState.Current.Local);
+			InCallback(TransformElement, ERigTransformType::CurrentGlobal, ERigTransformStorageType::Pose, &TransformElement->PoseStorage.Current.Global, &TransformElement->PoseDirtyState.Current.Global);
+
+			if(FRigControlElement* ControlElement = Cast<FRigControlElement>(Element))
+			{
+				InCallback(ControlElement, ERigTransformType::InitialLocal, ERigTransformStorageType::Offset, &ControlElement->OffsetStorage.Initial.Local, &ControlElement->OffsetDirtyState.Initial.Local);
+				InCallback(ControlElement, ERigTransformType::InitialGlobal, ERigTransformStorageType::Offset, &ControlElement->OffsetStorage.Initial.Global, &ControlElement->OffsetDirtyState.Initial.Global);
+				InCallback(ControlElement, ERigTransformType::CurrentLocal, ERigTransformStorageType::Offset, &ControlElement->OffsetStorage.Current.Local, &ControlElement->OffsetDirtyState.Current.Local);
+				InCallback(ControlElement, ERigTransformType::CurrentGlobal, ERigTransformStorageType::Offset, &ControlElement->OffsetStorage.Current.Global, &ControlElement->OffsetDirtyState.Current.Global);
+				InCallback(ControlElement, ERigTransformType::InitialLocal, ERigTransformStorageType::Shape, &ControlElement->ShapeStorage.Initial.Local, &ControlElement->ShapeDirtyState.Initial.Local);
+				InCallback(ControlElement, ERigTransformType::InitialGlobal, ERigTransformStorageType::Shape, &ControlElement->ShapeStorage.Initial.Global, &ControlElement->ShapeDirtyState.Initial.Global);
+				InCallback(ControlElement, ERigTransformType::CurrentLocal, ERigTransformStorageType::Shape, &ControlElement->ShapeStorage.Current.Local, &ControlElement->ShapeDirtyState.Current.Local);
+				InCallback(ControlElement, ERigTransformType::CurrentGlobal, ERigTransformStorageType::Shape, &ControlElement->ShapeStorage.Current.Global, &ControlElement->ShapeDirtyState.Current.Global);
+			}
+		}
+	}
+}
+
+TTuple<FRigComputedTransform*, FRigTransformDirtyState*> URigHierarchy::GetElementTransformStorage(const FRigElementKeyAndIndex& InKey,
+	ERigTransformType::Type InTransformType, ERigTransformStorageType::Type InStorageType)
+{
+	FRigTransformElement* TransformElement = Get<FRigTransformElement>(InKey);
+	if(TransformElement && TransformElement->GetKey() == InKey.Key)
+	{
+		FRigCurrentAndInitialTransform& Transform = TransformElement->PoseStorage;
+		FRigCurrentAndInitialDirtyState& DirtyState = TransformElement->PoseDirtyState;
+		if(InStorageType == ERigTransformStorageType::Offset || InStorageType == ERigTransformStorageType::Shape)
+		{
+			FRigControlElement* ControlElement = Cast<FRigControlElement>(TransformElement);
+			if(ControlElement == nullptr)
+			{
+				return {nullptr, nullptr};
+			}
+			if(InStorageType == ERigTransformStorageType::Offset)
+			{
+				Transform = ControlElement->OffsetStorage;
+				DirtyState = ControlElement->OffsetDirtyState;
+			}
+			else
+			{
+				Transform = ControlElement->ShapeStorage;
+				DirtyState = ControlElement->ShapeDirtyState;
+			}
+		}
+		
+		switch(InTransformType)
+		{
+			case ERigTransformType::InitialLocal:
+			{
+				return {&Transform.Initial.Local, &DirtyState.Initial.Local};
+			}
+			case ERigTransformType::CurrentLocal:
+			{
+				return {&Transform.Current.Local, &DirtyState.Current.Local};
+			}
+			case ERigTransformType::InitialGlobal:
+			{
+				return {&Transform.Initial.Global, &DirtyState.Initial.Global};
+			}
+			case ERigTransformType::CurrentGlobal:
+			{
+				return {&Transform.Current.Global, &DirtyState.Current.Global};
+			}
+			default:
+			{
+				break;
+			}
+		}
+	}
+	return {nullptr, nullptr};
 }
 
 TOptional<TTuple<int32, int32>> URigHierarchy::GetElementStorageRange(ERigTransformType::Type InTransformType) const
@@ -6397,7 +6558,7 @@ bool URigHierarchy::RemoveMetadataForElement(FRigBaseElement* InElement, const F
 	// metadata storage can just recycle that.
 	if (Storage.MetadataMap.IsEmpty())
 	{
-		ElementMetadata.Deallocate(InElement->MetadataStorageIndex);
+		ElementMetadata.Deallocate(InElement->MetadataStorageIndex, nullptr);
 		InElement->MetadataStorageIndex = INDEX_NONE;
 	}
 	else if (Storage.LastAccessName == InName)
@@ -6428,7 +6589,7 @@ bool URigHierarchy::RemoveAllMetadataForElement(FRigBaseElement* InElement)
 	// Clear the storage for the next user.
 	Storage.Reset();
 	
-	ElementMetadata.Deallocate(InElement->MetadataStorageIndex);
+	ElementMetadata.Deallocate(InElement->MetadataStorageIndex, nullptr);
 	InElement->MetadataStorageIndex = INDEX_NONE;
 
 	if(ElementBeingDestroyed != InElement)
diff --git a/Engine/Plugins/Animation/ControlRig/Source/ControlRig/Private/Rigs/RigHierarchyDefines.cpp b/Engine/Plugins/Animation/ControlRig/Source/ControlRig/Private/Rigs/RigHierarchyDefines.cpp
--- a/Engine/Plugins/Animation/ControlRig/Source/ControlRig/Private/Rigs/RigHierarchyDefines.cpp
+++ b/Engine/Plugins/Animation/ControlRig/Source/ControlRig/Private/Rigs/RigHierarchyDefines.cpp
@@ -94,6 +94,16 @@ FString FRigElementKey::ToPythonString() const
 #endif
 }
 
+////////////////////////////////////////////////////////////////////////////////
+// FRigElementKeyAndIndex
+////////////////////////////////////////////////////////////////////////////////
+
+FRigElementKeyAndIndex::FRigElementKeyAndIndex(const FRigBaseElement* InElement)
+: Key(InElement->Key)
+, Index(InElement->Index)
+{
+}
+
 ////////////////////////////////////////////////////////////////////////////////
 // FRigElementKeyCollection
 ////////////////////////////////////////////////////////////////////////////////
diff --git a/Engine/Plugins/Animation/ControlRig/Source/ControlRig/Private/Rigs/RigHierarchyElements.cpp b/Engine/Plugins/Animation/ControlRig/Source/ControlRig/Private/Rigs/RigHierarchyElements.cpp
--- a/Engine/Plugins/Animation/ControlRig/Source/ControlRig/Private/Rigs/RigHierarchyElements.cpp
+++ b/Engine/Plugins/Animation/ControlRig/Source/ControlRig/Private/Rigs/RigHierarchyElements.cpp
@@ -266,15 +266,15 @@ FRigTransformDirtyState& FRigTransformDirtyState::operator=(const FRigTransformD
 
 void FRigTransformDirtyState::LinkStorage(const TArrayView<bool>& InStorage)
 {
-	if(InStorage.IsValidIndex(Index))
+	if(InStorage.IsValidIndex(StorageIndex))
 	{
-		Storage = InStorage.GetData() + Index;
+		Storage = InStorage.GetData() + StorageIndex;
 	}
 }
 
 void FRigTransformDirtyState::UnlinkStorage(FRigReusableElementStorage<bool>& InStorage)
 {
-	InStorage.Deallocate(Index, &Storage);
+	InStorage.Deallocate(StorageIndex, &Storage);
 }
 
 ////////////////////////////////////////////////////////////////////////////////
@@ -370,15 +370,15 @@ FRigComputedTransform& FRigComputedTransform::operator=(const FRigComputedTransf
 
 void FRigComputedTransform::LinkStorage(const TArrayView<FTransform>& InStorage)
 {
-	if(InStorage.IsValidIndex(Index))
+	if(InStorage.IsValidIndex(StorageIndex))
 	{
-		Storage = InStorage.GetData() + Index;
+		Storage = InStorage.GetData() + StorageIndex;
 	}
 }
 
 void FRigComputedTransform::UnlinkStorage(FRigReusableElementStorage<FTransform>& InStorage)
 {
-	InStorage.Deallocate(Index, &Storage);
+	InStorage.Deallocate(StorageIndex, &Storage);
 }
 
 ////////////////////////////////////////////////////////////////////////////////
diff --git a/Engine/Plugins/Animation/ControlRig/Source/ControlRig/Private/Rigs/RigHierarchyPoseAdapter.cpp b/Engine/Plugins/Animation/ControlRig/Source/ControlRig/Private/Rigs/RigHierarchyPoseAdapter.cpp
new file mode 100644
--- /dev/null
+++ b/Engine/Plugins/Animation/ControlRig/Source/ControlRig/Private/Rigs/RigHierarchyPoseAdapter.cpp
@@ -0,0 +1,244 @@
+// Copyright Epic Games, Inc. All Rights Reserved.
+
+#include "Rigs/RigHierarchyPoseAdapter.h"
+#include "Rigs/RigHierarchy.h"
+
+URigHierarchy* FRigHierarchyPoseAdapter::GetHierarchy() const
+{
+	if(WeakHierarchy.IsValid())
+	{
+		return WeakHierarchy.Get();
+	}
+	return nullptr;
+}
+
+void FRigHierarchyPoseAdapter::PostLinked(URigHierarchy* InHierarchy)
+{
+	WeakHierarchy = InHierarchy; 
+}
+
+void FRigHierarchyPoseAdapter::PreUnlinked(URigHierarchy* InHierarchy)
+{
+}
+
+void FRigHierarchyPoseAdapter::PostUnlinked(URigHierarchy* InHierarchy)
+{
+	WeakHierarchy.Reset();
+}
+
+bool FRigHierarchyPoseAdapter::RelinkTransformStorage(const FRigElementKeyAndIndex& InKeyAndIndex, ERigTransformType::Type InTransformType,
+	ERigTransformStorageType::Type InStorageType, FTransform* InTransformStorage, bool* InDirtyFlagStorage)
+{
+	TArray<TTuple<FRigElementKeyAndIndex, ERigTransformType::Type, ERigTransformStorageType::Type, FTransform*, bool*>> Data =
+		{{InKeyAndIndex, InTransformType, InStorageType, InTransformStorage, InDirtyFlagStorage}};
+	return RelinkTransformStorage(Data);
+}
+
+bool FRigHierarchyPoseAdapter::RestoreTransformStorage(const FRigElementKeyAndIndex& InKeyAndIndex, ERigTransformType::Type InTransformType,
+	ERigTransformStorageType::Type InStorageType, bool bUpdateElementStorage)
+{
+	TArray<TTuple<FRigElementKeyAndIndex, ERigTransformType::Type, ERigTransformStorageType::Type>> Data =
+		{{InKeyAndIndex, InTransformType, InStorageType}};
+	return RestoreTransformStorage(Data, bUpdateElementStorage);
+}
+
+bool FRigHierarchyPoseAdapter::RelinkTransformStorage(
+	const TArrayView<TTuple<FRigElementKeyAndIndex, ERigTransformType::Type, ERigTransformStorageType::Type, FTransform*, bool*>>& InData)
+{
+	if(URigHierarchy* Hierarchy = GetHierarchy())
+	{
+		TArray<int32> TransformIndicesToDeallocate;
+		TArray<int32> DirtyStateIndicesToDeallocate;
+		TransformIndicesToDeallocate.Reserve(InData.Num());
+		DirtyStateIndicesToDeallocate.Reserve(InData.Num());
+		
+		bool bPerformedChange = false;
+		for(const TTuple<FRigElementKeyAndIndex, ERigTransformType::Type, ERigTransformStorageType::Type, FTransform*, bool*>& Tuple : InData)
+		{
+			auto CurrentStorage = Hierarchy->GetElementTransformStorage(Tuple.Get<0>(), Tuple.Get<1>(), Tuple.Get<2>());
+
+			if(FTransform* NewTransformStorage = Tuple.Get<3>())
+			{
+				if(Hierarchy->ElementTransforms.Contains(CurrentStorage.Get<0>()))
+				{
+					TransformIndicesToDeallocate.Add(CurrentStorage.Get<0>()->GetStorageIndex());
+				}
+				CurrentStorage.Get<0>()->StorageIndex = INDEX_NONE;
+				CurrentStorage.Get<0>()->Storage = NewTransformStorage;
+				bPerformedChange = true;
+			}
+			if(bool* NewDirtyStateStorage = Tuple.Get<4>())
+			{
+				if(Hierarchy->ElementDirtyStates.Contains(CurrentStorage.Get<1>()))
+				{
+					DirtyStateIndicesToDeallocate.Add(CurrentStorage.Get<1>()->GetStorageIndex());
+				}
+				CurrentStorage.Get<1>()->StorageIndex = INDEX_NONE;
+				CurrentStorage.Get<1>()->Storage = NewDirtyStateStorage;
+				bPerformedChange = true;
+			}
+		}
+
+		Hierarchy->ElementTransforms.Deallocate(TransformIndicesToDeallocate);
+		Hierarchy->ElementDirtyStates.Deallocate(DirtyStateIndicesToDeallocate);
+		return bPerformedChange;
+	}
+	return false;
+}
+
+bool FRigHierarchyPoseAdapter::RestoreTransformStorage(
+	const TArrayView<TTuple<FRigElementKeyAndIndex, ERigTransformType::Type, ERigTransformStorageType::Type>>& InData, bool bUpdateElementStorage)
+{
+	if(URigHierarchy* Hierarchy = GetHierarchy())
+	{
+		TArray<TTuple<FRigComputedTransform*, FRigTransformDirtyState*>> StoragePerElement;
+		StoragePerElement.Reserve(InData.Num());
+		for(const TTuple<FRigElementKeyAndIndex, ERigTransformType::Type, ERigTransformStorageType::Type>& Tuple : InData)
+		{
+			const TTuple<FRigComputedTransform*, FRigTransformDirtyState*> CurrentStorage =
+				Hierarchy->GetElementTransformStorage(Tuple.Get<0>(), Tuple.Get<1>(), Tuple.Get<2>());
+
+			if(CurrentStorage.Get<0>() == nullptr || CurrentStorage.Get<1>() == nullptr)
+			{
+				continue;
+			}
+			if(Hierarchy->ElementTransforms.Contains(CurrentStorage.Get<0>()) ||
+				Hierarchy->ElementDirtyStates.Contains(CurrentStorage.Get<1>()))
+			{
+				continue;
+			}
+			StoragePerElement.Add(CurrentStorage);
+		}
+
+		if(StoragePerElement.IsEmpty())
+		{
+			return false;
+		}
+
+		const TArray<int32, TInlineAllocator<4>> NewTransformIndices = Hierarchy->ElementTransforms.Allocate(StoragePerElement.Num(), FTransform::Identity);
+		const TArray<int32, TInlineAllocator<4>> NewDirtyStateIndices = Hierarchy->ElementDirtyStates.Allocate(StoragePerElement.Num(), false);
+		check(StoragePerElement.Num() == NewTransformIndices.Num());
+		check(StoragePerElement.Num() == NewDirtyStateIndices.Num());
+		for(int32 Index = 0; Index < StoragePerElement.Num(); Index++)
+		{
+			StoragePerElement[Index].Get<0>()->StorageIndex = NewTransformIndices[Index];
+			StoragePerElement[Index].Get<1>()->StorageIndex = NewDirtyStateIndices[Index];
+		}
+		
+		if(bUpdateElementStorage)
+		{
+			(void)UpdateElementStorage();
+			(void)SortStorage();
+		}
+		return true;
+	}
+	return false;
+}
+
+bool FRigHierarchyPoseAdapter::RelinkCurveStorage(const FRigElementKeyAndIndex& InKeyAndIndex, float* InCurveStorage)
+{
+	TArray<TTuple<FRigElementKeyAndIndex, float*>> Data = {{InKeyAndIndex, InCurveStorage}};
+	return RelinkCurveStorage(Data);
+}
+
+bool FRigHierarchyPoseAdapter::RestoreCurveStorage(const FRigElementKeyAndIndex& InKeyAndIndex, bool bUpdateElementStorage)
+{
+	TArray<FRigElementKeyAndIndex> Data = {InKeyAndIndex};
+	return RestoreCurveStorage(Data, bUpdateElementStorage);
+}
+
+bool FRigHierarchyPoseAdapter::RelinkCurveStorage(const TArrayView<TTuple<FRigElementKeyAndIndex, float*>>& InData)
+{
+	if(URigHierarchy* Hierarchy = GetHierarchy())
+	{
+		TArray<int32> CurveIndicesToDeallocate;
+		CurveIndicesToDeallocate.Reserve(InData.Num());
+		
+		bool bPerformedChange = false;
+		for(const TTuple<FRigElementKeyAndIndex, float*>& Tuple : InData)
+		{
+			FRigCurveElement* CurveElement = Hierarchy->Get<FRigCurveElement>(Tuple.Get<0>());
+
+			if(float* NewCurveStorage = Tuple.Get<1>())
+			{
+				if(Hierarchy->ElementCurves.Contains(CurveElement))
+				{
+					CurveIndicesToDeallocate.Add(CurveElement->GetStorageIndex());
+				}
+				CurveElement->StorageIndex = INDEX_NONE;
+				CurveElement->Storage = NewCurveStorage;
+				bPerformedChange = true;
+			}
+		}
+
+		Hierarchy->ElementCurves.Deallocate(CurveIndicesToDeallocate);
+		return bPerformedChange;
+	}
+	return false;
+}
+
+bool FRigHierarchyPoseAdapter::RestoreCurveStorage(const TArrayView<FRigElementKeyAndIndex>& InData, bool bUpdateElementStorage)
+{
+	if(URigHierarchy* Hierarchy = GetHierarchy())
+	{
+		TArray<FRigCurveElement*> Curves;
+		Curves.Reserve(InData.Num());
+		for(const FRigElementKeyAndIndex& KeyAndIndex : InData)
+		{
+			if(FRigCurveElement* CurveElement = Hierarchy->Get<FRigCurveElement>(KeyAndIndex))
+			{
+				if(!Hierarchy->ElementCurves.Contains(CurveElement))
+				{
+					Curves.Add(CurveElement);
+				}
+			}
+		}
+
+		if(Curves.IsEmpty())
+		{
+			return false;
+		}
+
+		const TArray<int32, TInlineAllocator<4>> NewCurveIndices = Hierarchy->ElementCurves.Allocate(Curves.Num(), 0.f);
+		check(Curves.Num() == NewCurveIndices.Num());
+		for(int32 Index = 0; Index < Curves.Num(); Index++)
+		{
+			Curves[Index]->StorageIndex = NewCurveIndices[Index];
+		}
+
+		if(bUpdateElementStorage)
+		{
+			(void)UpdateElementStorage();
+		}
+		return true;
+	}
+	return false;
+}
+
+bool FRigHierarchyPoseAdapter::SortStorage()
+{
+	if(URigHierarchy* Hierarchy = GetHierarchy())
+	{
+		return Hierarchy->SortElementStorage();
+	}
+	return false;
+}
+
+bool FRigHierarchyPoseAdapter::ShrinkStorage()
+{
+	if(URigHierarchy* Hierarchy = GetHierarchy())
+	{
+		return Hierarchy->ShrinkElementStorage();
+	}
+	return false;
+}
+
+bool FRigHierarchyPoseAdapter::UpdateElementStorage()
+{
+	if(URigHierarchy* Hierarchy = GetHierarchy())
+	{
+		Hierarchy->UpdateElementStorage();
+		return true;
+	}
+	return false;
+}
diff --git a/Engine/Plugins/Animation/ControlRig/Source/ControlRig/Public/ControlRigDefines.h b/Engine/Plugins/Animation/ControlRig/Source/ControlRig/Public/ControlRigDefines.h
--- a/Engine/Plugins/Animation/ControlRig/Source/ControlRig/Public/ControlRigDefines.h
+++ b/Engine/Plugins/Animation/ControlRig/Source/ControlRig/Public/ControlRigDefines.h
@@ -13,6 +13,10 @@
 //#define DECLARE_SCOPE_HIERARCHICAL_COUNTER_RIGUNIT() \
 //	DECLARE_SCOPE_HIERARCHICAL_COUNTER_FUNC()
 
+#define UE_CONTROLRIG_LLM_SCOPE_BYNAME(InLabel)
+//#define UE_CONTROLRIG_LLM_SCOPE_BYNAME(InLabel) \
+//	LLM_SCOPE_BYNAME(TEXT(InLabel));
+
 UENUM()
 enum class ETransformSpaceMode : uint8
 {
diff --git a/Engine/Plugins/Animation/ControlRig/Source/ControlRig/Public/Rigs/RigHierarchy.h b/Engine/Plugins/Animation/ControlRig/Source/ControlRig/Public/Rigs/RigHierarchy.h
--- a/Engine/Plugins/Animation/ControlRig/Source/ControlRig/Public/Rigs/RigHierarchy.h
+++ b/Engine/Plugins/Animation/ControlRig/Source/ControlRig/Public/Rigs/RigHierarchy.h
@@ -9,6 +9,7 @@
 #include "RigHierarchyElements.h"
 #include "RigHierarchyCache.h"
 #include "RigHierarchyPose.h"
+#include "RigHierarchyPoseAdapter.h"
 #include "UObject/WeakObjectPtrTemplates.h"
 #include "EdGraph/EdGraphPin.h"
 #include "RigHierarchyDefines.h"
@@ -432,6 +433,30 @@ public:
 		return INDEX_NONE;
 	}
 
+	/**
+	 * Returns the key and index pair of an element given its key
+	 * @param InKey The key of the element to retrieve the information for
+	 * @return The key and index pair of the element
+	 */
+	FRigElementKeyAndIndex GetKeyAndIndex(const FRigElementKey& InKey) const
+	{
+		return GetKeyAndIndex(GetIndex(InKey));
+	};
+
+	/**
+	 * Returns the key and index pair of an element given its index
+	 * @param InIndex The index of the element to retrieve the information for
+	 * @return The key and index pair of the element
+	 */
+	FRigElementKeyAndIndex GetKeyAndIndex(int32 InIndex) const
+	{
+		if(const FRigBaseElement* Element = Get(InIndex))
+		{
+			return FRigElementKeyAndIndex(Element);
+		}
+		return FRigElementKeyAndIndex();
+	};
+
 	/**
 	 * Returns the index of an element given its key within its default parent (or root)
 	 * @param InKey The key of the element to retrieve the index for
@@ -3690,6 +3715,20 @@ public:
 		return SetPose(InPose);
 	}
 
+	/**
+	 * Sets the pose adapter used for storage of pose data
+	 * @param InPoseAdapter The pose adapter to set on the hierarchy
+	 */
+	void LinkPoseAdapter(TSharedPtr<FRigHierarchyPoseAdapter> InPoseAdapter);
+
+	/**
+	 * Clears the pose adapter used for storage of pose data
+	 */
+	void UnlinkPoseAdapter()
+	{
+		return LinkPoseAdapter(nullptr);
+	}
+
 	/**
 	 * Creates a rig control value from a bool value
 	 * @param InValue The value to create the rig control value from
@@ -3889,6 +3928,8 @@ private:
 	FRigHierarchyMetadataTagChangedDelegate MetadataTagChangedDelegate;
 	FRigEventDelegate EventDelegate;
 
+	TSharedPtr<FRigHierarchyPoseAdapter> PoseAdapter;
+
 public:
 
 	void Notify(ERigHierarchyNotification InNotifType, const FRigBaseElement* InElement);
@@ -4323,7 +4364,19 @@ private:
 	// Orders the element storage by storing first initial, then current,
 	// within each first local, then global, and within each of those lists
 	// we'll place bones, nulls, controls etc in that order
-	void SortElementStorage();
+	bool SortElementStorage();
+
+	// Compacts the element storage
+	bool ShrinkElementStorage();
+
+	// Helper function to iterate all transform element storage
+	void ForEachTransformElementStorage(TFunction<void(FRigTransformElement*,ERigTransformType::Type,ERigTransformStorageType::Type,FRigComputedTransform*,FRigTransformDirtyState*)> InCallback);
+
+	// Returns the computed transform and dirty state for a given element
+	TTuple<FRigComputedTransform*,FRigTransformDirtyState*> GetElementTransformStorage(
+		const FRigElementKeyAndIndex& InKey,
+		ERigTransformType::Type InTransformType,
+		ERigTransformStorageType::Type InStorageType = ERigTransformStorageType::Pose);
 
 	// Returns the range of the element transform / dirty state storage for a given
 	// transform type. this is only valid if the element storage has been sorted. 
@@ -4982,6 +5035,9 @@ private:
 	friend struct FRigDispatch_GetModuleMetadata;
 	friend struct FRigDispatch_SetModuleMetadata;
 	friend class FControlRigHierarchySortElementStorage;
+	friend class FControlRigHierarchyShrinkElementStorage;
+	friend class FControlRigHierarchyRelinkElementStorage;
+	friend class FRigHierarchyPoseAdapter;
 };
 
 struct CONTROLRIG_API FRigHierarchyInteractionBracket
diff --git a/Engine/Plugins/Animation/ControlRig/Source/ControlRig/Public/Rigs/RigHierarchyDefines.h b/Engine/Plugins/Animation/ControlRig/Source/ControlRig/Public/Rigs/RigHierarchyDefines.h
--- a/Engine/Plugins/Animation/ControlRig/Source/ControlRig/Public/Rigs/RigHierarchyDefines.h
+++ b/Engine/Plugins/Animation/ControlRig/Source/ControlRig/Public/Rigs/RigHierarchyDefines.h
@@ -8,6 +8,7 @@
 #include "RigHierarchyDefines.generated.h"
 
 class URigHierarchy;
+struct FRigBaseElement;
 
 // Debug define which performs a full check on the cache validity for all elements of the hierarchy.
 // This can be useful for debugging cache validity bugs.
@@ -1604,6 +1605,46 @@ public:
 	FString ToPythonString() const;
 };
 
+struct CONTROLRIG_API FRigElementKeyAndIndex
+{
+public:
+
+	inline static const FRigElementKey InvalidKey = FRigElementKey(NAME_None, ERigElementType::Bone);
+	inline static constexpr int32 InvalidIndex = INDEX_NONE;
+
+	FRigElementKeyAndIndex()
+		: Key(InvalidKey)
+		, Index(InvalidIndex)
+	{
+	}
+
+	explicit FRigElementKeyAndIndex(const FRigBaseElement* InElement);
+
+	FRigElementKeyAndIndex(const FRigElementKey& InKey, const int32& InIndex)
+		: Key(InKey)
+		, Index(InIndex)
+	{
+	}
+
+	bool IsValid() const
+	{
+		return Key.IsValid() && Index != INDEX_NONE;
+	}
+
+	operator int32() const
+	{
+		return Index;
+	}
+
+	friend uint32 GetTypeHash(const FRigElementKeyAndIndex& InKeyAndIndex)
+	{
+		return GetTypeHash(InKeyAndIndex.Index);
+	}
+	
+	const FRigElementKey& Key;
+	const int32& Index;
+};
+
 USTRUCT(BlueprintType)
 struct CONTROLRIG_API FRigElementKeyCollection
 {
diff --git a/Engine/Plugins/Animation/ControlRig/Source/ControlRig/Public/Rigs/RigHierarchyElements.h b/Engine/Plugins/Animation/ControlRig/Source/ControlRig/Public/Rigs/RigHierarchyElements.h
--- a/Engine/Plugins/Animation/ControlRig/Source/ControlRig/Public/Rigs/RigHierarchyElements.h
+++ b/Engine/Plugins/Animation/ControlRig/Source/ControlRig/Public/Rigs/RigHierarchyElements.h
@@ -314,8 +314,8 @@ struct CONTROLRIG_API FRigReusableElementStorage
 
 		return Indices;
 	}
-	
-	void Deallocate(int32& InIndex, T** InStorage = nullptr)
+
+	void Deallocate(int32& InIndex, T** InStorage)
 	{
 		if(InIndex == INDEX_NONE)
 		{
@@ -329,8 +329,31 @@ struct CONTROLRIG_API FRigReusableElementStorage
 		InIndex = INDEX_NONE;
 		if(InStorage)
 		{
-			InStorage = nullptr;
+			*InStorage = nullptr;
+		}
+	}
+
+	void Deallocate(const TConstArrayView<int32>& InIndices)
+	{
+		if(InIndices.IsEmpty())
+		{
+			return;
 		}
+		FreeList.Reserve(FreeList.Num() + InIndices.Num());
+		for(int32 Index : InIndices)
+		{
+			if(Index != INDEX_NONE && !FreeList.Contains(Index))
+			{
+				Deallocate(Index, nullptr);
+			}
+		}
+	}
+
+	template<typename OwnerType>
+	void Deallocate(OwnerType* InOwner)
+	{
+		check(InOwner);
+		Deallocate(InOwner->StorageIndex, &InOwner->Storage);
 	}
 
 	void Reset(TFunction<void(int32, T&)> OnDestroyCallback = nullptr)
@@ -355,6 +378,13 @@ struct CONTROLRIG_API FRigReusableElementStorage
 		return GetData() + InIndex == InStorage;
 	}
 
+	template<typename OwnerType>
+	bool Contains(const OwnerType* InOwner)
+	{
+		check(InOwner);
+		return Contains(InOwner->StorageIndex, InOwner->Storage);
+	}
+
 	TMap<int32, int32> Shrink(TFunction<void(int32, T&)> OnDestroyCallback = nullptr)
 	{
 		TMap<int32, int32> OldToNew;
@@ -402,7 +432,7 @@ public:
 	GENERATED_BODY()
 
 	FRigTransformDirtyState()
-	: Index(INDEX_NONE)
+	: StorageIndex(INDEX_NONE)
 	, Storage(nullptr)
 	{
 	}
@@ -412,18 +442,25 @@ public:
 	bool Set(bool InDirty);
 	FRigTransformDirtyState& operator =(const FRigTransformDirtyState& InOther);
 
-	UPROPERTY(BlueprintReadOnly, EditAnywhere, Category = "DirtyState")
-	int32 Index;
+	int32 GetStorageIndex() const
+	{
+		return StorageIndex;
+	}
 
 private:
 
 	void LinkStorage(const TArrayView<bool>& InStorage);
 	void UnlinkStorage(FRigReusableElementStorage<bool>& InStorage);
+
+	int32 StorageIndex;
 	bool* Storage;
 
 	friend struct FRigLocalAndGlobalDirtyState;
 	friend class URigHierarchy;
 	friend class URigHierarchyController;
+	friend class FRigHierarchyPoseAdapter;
+	friend struct FRigReusableElementStorage<bool>;
+	friend class FControlRigHierarchyRelinkElementStorage;
 };
 
 USTRUCT(BlueprintType)
@@ -543,7 +580,7 @@ struct CONTROLRIG_API FRigComputedTransform
 public:
 	
 	FRigComputedTransform()
-	: Index(INDEX_NONE)
+	: StorageIndex(INDEX_NONE)
 	, Storage(nullptr)
 	{}
 
@@ -576,23 +613,29 @@ public:
 	bool operator == (const FRigComputedTransform& Other) const
 	{
 		return Equals(Get(), Other.Get());
-    }
+	}
 
 	FRigComputedTransform& operator =(const FRigComputedTransform& InOther);
 
-	UPROPERTY(BlueprintReadOnly, EditAnywhere, Category = "Transform")
-	int32 Index;
+	int32 GetStorageIndex() const
+	{
+		return StorageIndex;
+	}
 
 private:
 
 	void LinkStorage(const TArrayView<FTransform>& InStorage);
 	void UnlinkStorage(FRigReusableElementStorage<FTransform>& InStorage);
 	
+	int32 StorageIndex;
 	FTransform* Storage;
 	
 	friend struct FRigLocalAndGlobalTransform;
 	friend class URigHierarchy;
 	friend class URigHierarchyController;
+	friend class FRigHierarchyPoseAdapter;
+	friend struct FRigReusableElementStorage<FTransform>;
+	friend class FControlRigHierarchyRelinkElementStorage;
 };
 
 USTRUCT(BlueprintType)
@@ -1038,6 +1081,7 @@ protected:
 	friend class FControlRigEditor;
 	friend class URigHierarchy;
 	friend class URigHierarchyController;
+	friend struct FRigElementKeyAndIndex;
 };
 
 USTRUCT(BlueprintType)
@@ -1828,6 +1872,11 @@ struct CONTROLRIG_API FRigCurveElement final : public FRigBaseElement
 
 	bool IsValueSet() const { return bIsValueSet; }
 
+	int32 GetStorageIndex() const
+	{
+		return StorageIndex;
+	}
+
 protected:
 
 	virtual void LinkStorage(const TArrayView<FTransform>& InTransforms, const TArrayView<bool>& InDirtyStates, const TArrayView<float>& InCurves) override;
@@ -1855,6 +1904,9 @@ private:
 	friend class URigHierarchy;
 	friend class URigHierarchyController;
 	friend struct FRigBaseElement;
+	friend class FRigHierarchyPoseAdapter;
+	friend struct FRigReusableElementStorage<float>;
+	friend class FControlRigHierarchyRelinkElementStorage;
 };
 
 USTRUCT(BlueprintType)
diff --git a/Engine/Plugins/Animation/ControlRig/Source/ControlRig/Public/Rigs/RigHierarchyPoseAdapter.h b/Engine/Plugins/Animation/ControlRig/Source/ControlRig/Public/Rigs/RigHierarchyPoseAdapter.h
new file mode 100644
--- /dev/null
+++ b/Engine/Plugins/Animation/ControlRig/Source/ControlRig/Public/Rigs/RigHierarchyPoseAdapter.h
@@ -0,0 +1,45 @@
+// Copyright Epic Games, Inc. All Rights Reserved.
+
+#pragma once
+
+#include "CoreMinimal.h"
+#include "RigHierarchyElements.h"
+
+class URigHierarchy;
+
+class CONTROLRIG_API FRigHierarchyPoseAdapter : public TSharedFromThis<FRigHierarchyPoseAdapter>
+{
+public:
+
+	virtual ~FRigHierarchyPoseAdapter()
+	{
+	}
+
+protected:
+
+	URigHierarchy* GetHierarchy() const;
+
+	virtual void PostLinked(URigHierarchy* InHierarchy);
+	virtual void PreUnlinked(URigHierarchy* InHierarchy);
+	virtual void PostUnlinked(URigHierarchy* InHierarchy);
+
+	bool RelinkTransformStorage(const FRigElementKeyAndIndex& InKeyAndIndex, ERigTransformType::Type InTransformType, ERigTransformStorageType::Type InStorageType, FTransform* InTransformStorage, bool* InDirtyFlagStorage);
+	bool RestoreTransformStorage(const FRigElementKeyAndIndex& InKeyAndIndex, ERigTransformType::Type InTransformType, ERigTransformStorageType::Type InStorageType, bool bUpdateElementStorage);
+	bool RelinkTransformStorage(const TArrayView<TTuple<FRigElementKeyAndIndex,ERigTransformType::Type,ERigTransformStorageType::Type,FTransform*,bool*>>& InData);
+	bool RestoreTransformStorage(const TArrayView<TTuple<FRigElementKeyAndIndex,ERigTransformType::Type,ERigTransformStorageType::Type>>& InData, bool bUpdateElementStorage);
+
+	bool RelinkCurveStorage(const FRigElementKeyAndIndex& InKeyAndIndex, float* InCurveStorage);
+	bool RestoreCurveStorage(const FRigElementKeyAndIndex& InKeyAndIndex, bool bUpdateElementStorage);
+	bool RelinkCurveStorage(const TArrayView<TTuple<FRigElementKeyAndIndex,float*>>& InData);
+	bool RestoreCurveStorage(const TArrayView<FRigElementKeyAndIndex>& InData, bool bUpdateElementStorage);
+
+	bool SortStorage();
+	bool ShrinkStorage();
+	bool UpdateElementStorage();
+
+	TWeakObjectPtr<URigHierarchy> WeakHierarchy;
+
+	friend class URigHierarchy;
+	friend class URigHierarchyController;
+	friend class FControlRigHierarchyPoseAdapter;
+};
