[ControlRig AnimNode] Added support for enum property

diff --git a/Engine/Plugins/Animation/ControlRig/Source/ControlRig/Private/AnimNode_ControlRig.cpp b/Engine/Plugins/Animation/ControlRig/Source/ControlRig/Private/AnimNode_ControlRig.cpp
--- a/Engine/Plugins/Animation/ControlRig/Source/ControlRig/Private/AnimNode_ControlRig.cpp
+++ b/Engine/Plugins/Animation/ControlRig/Source/ControlRig/Private/AnimNode_ControlRig.cpp
@@ -838,6 +838,13 @@ void FAnimNode_ControlRig::PropagateInputProperties(const UObject* InSourceInsta
 						ObjectProperty->CopyCompleteValue(Variable.Memory, SrcPtr);
 					}
 				}
+				else if(FEnumProperty* EnumProperty = CastField<FEnumProperty>(CallerProperty))
+				{
+					if(ensure(EnumProperty->SameType(Variable.Property)))
+					{
+						EnumProperty->CopyCompleteValue(Variable.Memory, SrcPtr);
+					}
+				}
 				else
 				{
 					ensureMsgf(false, TEXT("Property %s type %s not recognized"), *CallerProperty->GetName(), *CallerProperty->GetCPPType());
