[ControlRig Track] unregister from endloadpackage delegate to avoid the loading of other packages causing infinite loop

diff --git a/Engine/Plugins/Animation/ControlRig/Source/ControlRig/Private/Sequencer/MovieSceneControlRigParameterTrack.cpp b/Engine/Plugins/Animation/ControlRig/Source/ControlRig/Private/Sequencer/MovieSceneControlRigParameterTrack.cpp
--- a/Engine/Plugins/Animation/ControlRig/Source/ControlRig/Private/Sequencer/MovieSceneControlRigParameterTrack.cpp
+++ b/Engine/Plugins/Animation/ControlRig/Source/ControlRig/Private/Sequencer/MovieSceneControlRigParameterTrack.cpp
@@ -608,11 +608,16 @@ void UMovieSceneControlRigParameterTrack::HandlePackageDone(const FEndLoadPackag
 			}
 		}
 
+		// All dependent packages ready, no need to wait/check for any other packages
+		// ReconstructControlRig may trigger loading of packages that we don't care about, so unregister from the delegate
+		// to avoid infinite loop
+		FCoreUObjectDelegates::OnEndLoadPackage.RemoveAll(this);
+		
 		// Only reconstruct in case it is not a native ControlRig class
 		ReconstructControlRig();
 	}
 
-	FCoreUObjectDelegates::OnEndLoadPackage.RemoveAll(this);
+
 }
 
 void UMovieSceneControlRigParameterTrack::HandleControlRigPackageDone(URigVMHost* InControlRig)
