Control Rig: Fix nullptr crash in metadata

#rnx

diff --git a/Engine/Plugins/Animation/ControlRig/Source/ControlRig/Private/Rigs/RigHierarchy.cpp b/Engine/Plugins/Animation/ControlRig/Source/ControlRig/Private/Rigs/RigHierarchy.cpp
--- a/Engine/Plugins/Animation/ControlRig/Source/ControlRig/Private/Rigs/RigHierarchy.cpp
+++ b/Engine/Plugins/Animation/ControlRig/Source/ControlRig/Private/Rigs/RigHierarchy.cpp
@@ -6504,7 +6504,7 @@ FRigBaseMetadata* URigHierarchy::FindMetadataForElement(const FRigBaseElement* I
 	}
 
 	FMetadataStorage& Storage = ElementMetadata[InElement->MetadataStorageIndex];
-	if (InName == Storage.LastAccessName && (InType == ERigMetadataType::Invalid || Storage.LastAccessMetadata->GetType() == InType))
+	if (InName == Storage.LastAccessName && (InType == ERigMetadataType::Invalid || (Storage.LastAccessMetadata && Storage.LastAccessMetadata->GetType() == InType)))
 	{
 		return Storage.LastAccessMetadata;
 	}
