Control Rig: Crash when undoing deletion of rig elements

#rnx

diff --git a/Engine/Plugins/Animation/ControlRig/Source/ControlRig/Private/Rigs/RigHierarchy.cpp b/Engine/Plugins/Animation/ControlRig/Source/ControlRig/Private/Rigs/RigHierarchy.cpp
--- a/Engine/Plugins/Animation/ControlRig/Source/ControlRig/Private/Rigs/RigHierarchy.cpp
+++ b/Engine/Plugins/Animation/ControlRig/Source/ControlRig/Private/Rigs/RigHierarchy.cpp
@@ -251,6 +251,9 @@ void URigHierarchy::Load(FArchive& Ar)
 {
 	FScopeLock Lock(&ElementsLock);
 	
+	// unlink any existing pose adapter
+	UnlinkPoseAdapter();
+
 	TArray<FRigElementKey> SelectedKeys;
 	if(Ar.IsTransacting())
 	{
@@ -579,6 +582,9 @@ void URigHierarchy::CopyHierarchy(URigHierarchy* InHierarchy)
 		return;
 	}
 
+	// unlink any existing pose adapter
+	UnlinkPoseAdapter();
+
 	// check if we really need to do a deep copy all over again.
 	// for rigs which contain more elements (likely procedural elements)
 	// we'll assume we can just remove superfluous elements (from the end of the lists).
@@ -594,7 +600,7 @@ void URigHierarchy::CopyHierarchy(URigHierarchy* InHierarchy)
 			}
 			check(ElementsPerType.IsValidIndex(ElementTypeIndex));
 			check(InHierarchy->ElementsPerType.IsValidIndex(ElementTypeIndex));
-			if(ElementsPerType[ElementTypeIndex].Num() < InHierarchy->ElementsPerType[ElementTypeIndex].Num())
+			if(ElementsPerType[ElementTypeIndex].Num() != InHierarchy->ElementsPerType[ElementTypeIndex].Num())
 			{
 				bReallocateElements = true;
 				break;
@@ -6137,28 +6143,21 @@ bool URigHierarchy::SortElementStorage()
 	for(int32 Index = 0; Index < TransformInfos.Num(); Index++)
 	{
 		const int32 CurrentTransformIndex = TransformInfos[Index].ComputedTransform->StorageIndex;
-		if(CurrentTransformIndex == INDEX_NONE)
-		{
-			// we cannot sort this if some of the storage is no longer local
-			ElementTransformRanges.Reset();
-			return false;
-		}
-		if(CurrentTransformIndex != Index)
+		if(CurrentTransformIndex != INDEX_NONE && CurrentTransformIndex != Index)
 		{
 			bRequiresSort = true;
 		}
 
 		const int32 CurrentDirtyStateIndex = TransformInfos[Index].DirtyState->StorageIndex;
-		if(CurrentDirtyStateIndex == INDEX_NONE)
+		if(CurrentDirtyStateIndex != INDEX_NONE && CurrentDirtyStateIndex != Index)
 		{
-			// we cannot sort this if some of the storage is no longer local
-			ElementTransformRanges.Reset();
-			return false;
+			bRequiresSort = true;
 		}
-		
-		if(CurrentDirtyStateIndex != Index)
+
+		if(CurrentTransformIndex == INDEX_NONE && CurrentDirtyStateIndex == INDEX_NONE)
 		{
-			bRequiresSort = true;
+			// skip over entries with external storage
+			continue;
 		}
 
 		if(CurrentTransformIndex != CurrentDirtyStateIndex)
@@ -6192,12 +6191,28 @@ bool URigHierarchy::SortElementStorage()
 	for(int32 Index = 0; Index < TransformInfos.Num(); Index++)
 	{
 		const int32 CurrentTransformIndex = TransformInfos[Index].ComputedTransform->StorageIndex;
-		const int32 NewTransformIndex = NewElementTransforms.Storage.Add(ElementTransforms[CurrentTransformIndex]);
-		TransformInfos[Index].ComputedTransform->StorageIndex = NewTransformIndex;
-
 		const int32 CurrentDirtyStateIndex = TransformInfos[Index].DirtyState->StorageIndex;
-		const int32 NewDirtyStateIndex = NewElementDirtyStates.Storage.Add(ElementDirtyStates[CurrentDirtyStateIndex]);
-		TransformInfos[Index].DirtyState->StorageIndex = NewDirtyStateIndex;
+
+		int32 NewTransformIndex = INDEX_NONE;
+		int32 NewDirtyStateIndex = INDEX_NONE;
+
+		if(CurrentTransformIndex != INDEX_NONE)
+		{
+			NewTransformIndex = NewElementTransforms.Storage.Add(ElementTransforms[CurrentTransformIndex]);
+			TransformInfos[Index].ComputedTransform->StorageIndex = NewTransformIndex;
+		}
+
+		if(CurrentDirtyStateIndex != INDEX_NONE)
+		{
+			NewDirtyStateIndex = NewElementDirtyStates.Storage.Add(ElementDirtyStates[CurrentDirtyStateIndex]);
+			TransformInfos[Index].DirtyState->StorageIndex = NewDirtyStateIndex;
+		}
+
+		if(CurrentTransformIndex == INDEX_NONE && CurrentDirtyStateIndex)
+		{
+			// ignore entries with external storage
+			continue;
+		}
 
 		if(NewTransformIndex == NewDirtyStateIndex)
 		{
diff --git a/Engine/Plugins/Animation/ControlRig/Source/ControlRig/Public/Rigs/RigHierarchy.h b/Engine/Plugins/Animation/ControlRig/Source/ControlRig/Public/Rigs/RigHierarchy.h
--- a/Engine/Plugins/Animation/ControlRig/Source/ControlRig/Public/Rigs/RigHierarchy.h
+++ b/Engine/Plugins/Animation/ControlRig/Source/ControlRig/Public/Rigs/RigHierarchy.h
@@ -452,7 +452,7 @@ public:
 	{
 		if(const FRigBaseElement* Element = Get(InIndex))
 		{
-			return FRigElementKeyAndIndex(Element);
+			return Element->GetKeyAndIndex();
 		}
 		return FRigElementKeyAndIndex();
 	};
@@ -939,6 +939,14 @@ public:
 		return GetElementsOfType<FRigCurveElement>();
 	}
 
+	/**
+	 * Returns all Curve elements without traversing the hierarchy
+	 */
+	TArray<FRigBaseElement*>& GetCurvesFast() const
+	{
+		return ElementsPerType[RigElementTypeToFlatIndex(ERigElementType::Curve)];
+	}
+
 	/**
 	 * Returns all Curve elements
 	 */
diff --git a/Engine/Plugins/Animation/ControlRig/Source/ControlRig/Public/Rigs/RigHierarchyElements.h b/Engine/Plugins/Animation/ControlRig/Source/ControlRig/Public/Rigs/RigHierarchyElements.h
--- a/Engine/Plugins/Animation/ControlRig/Source/ControlRig/Public/Rigs/RigHierarchyElements.h
+++ b/Engine/Plugins/Animation/ControlRig/Source/ControlRig/Public/Rigs/RigHierarchyElements.h
@@ -392,28 +392,38 @@ struct CONTROLRIG_API FRigReusableElementStorage
 		if(!FreeList.IsEmpty())
 		{
 			TArray<bool> ToRemove;
-			ToRemove.AddZeroed(Storage.Num());
 
-			TArray<T> NewStorage;
-			
-			for(int32 FreeIndex : FreeList)
+			if(FreeList.Num() != Storage.Num() || OnDestroyCallback != nullptr)
 			{
-				ToRemove[FreeIndex] = true;
-				if(OnDestroyCallback)
+				ToRemove.AddZeroed(Storage.Num());
+				for(int32 FreeIndex : FreeList)
 				{
-					OnDestroyCallback(FreeIndex, Storage[FreeIndex]);
+					ToRemove[FreeIndex] = true;
+					if(OnDestroyCallback)
+					{
+						OnDestroyCallback(FreeIndex, Storage[FreeIndex]);
+					}
 				}
 			}
-			
-			for(int32 OldIndex = 0; OldIndex < Storage.Num(); OldIndex++)
+
+			if(FreeList.Num() != Storage.Num())
 			{
-				if(!ToRemove[OldIndex])
+				TArray<T> NewStorage;
+				NewStorage.Reserve(FMath::Max(Storage.Num() - FreeList.Num(), 0));
+				for(int32 OldIndex = 0; OldIndex < Storage.Num(); OldIndex++)
 				{
-					OldToNew.Add(OldIndex, NewStorage.Add(Storage[OldIndex]));
+					if(!ToRemove[OldIndex])
+					{
+						OldToNew.Add(OldIndex, NewStorage.Add(Storage[OldIndex]));
+					}
 				}
+				Storage = NewStorage;
+			}
+			else
+			{
+				Storage.Reset();
 			}
 
-			Storage = NewStorage;
 			FreeList.Reset();
 		}
 
@@ -994,6 +1004,7 @@ public:
 	virtual const FName& GetDisplayName() const { return GetFName(); }
 	ERigElementType GetType() const { return Key.Type; }
 	const FRigElementKey& GetKey() const { return Key; }
+	FRigElementKeyAndIndex GetKeyAndIndex() const { return {Key, Index}; };
 	int32 GetIndex() const { return Index; }
 	int32 GetSubIndex() const { return SubIndex; }
 	bool IsSelected() const { return bSelected; }
