Sequencer: Control Rig: Disable control rig not letting other objects to key with 'S'

diff --git a/Engine/Plugins/Animation/ControlRig/Source/ControlRigEditor/Private/Sequencer/ControlRigParameterTrackEditor.cpp b/Engine/Plugins/Animation/ControlRig/Source/ControlRigEditor/Private/Sequencer/ControlRigParameterTrackEditor.cpp
--- a/Engine/Plugins/Animation/ControlRig/Source/ControlRigEditor/Private/Sequencer/ControlRigParameterTrackEditor.cpp
+++ b/Engine/Plugins/Animation/ControlRig/Source/ControlRigEditor/Private/Sequencer/ControlRigParameterTrackEditor.cpp
@@ -1847,9 +1847,9 @@ void FControlRigParameterTrackEditor::AddControlRigFromComponent(FGuid InGuid)
 
 bool FControlRigParameterTrackEditor::HasTransformKeyOverridePriority() const
 {
-	return CanAddTransformKeysForSelectedObjects();
-
+	return false; 
 }
+
 bool FControlRigParameterTrackEditor::CanAddTransformKeysForSelectedObjects() const
 {
 	// WASD hotkeys to fly the viewport can conflict with hotkeys for setting keyframes (ie. s). 
