Control Rig: AnimNode to use pose adapters - first steps

#rnx

diff --git a/Engine/Plugins/Animation/ControlRig/Source/ControlRig/Private/Rigs/RigHierarchy.cpp b/Engine/Plugins/Animation/ControlRig/Source/ControlRig/Private/Rigs/RigHierarchy.cpp
--- a/Engine/Plugins/Animation/ControlRig/Source/ControlRig/Private/Rigs/RigHierarchy.cpp
+++ b/Engine/Plugins/Animation/ControlRig/Source/ControlRig/Private/Rigs/RigHierarchy.cpp
@@ -105,6 +105,7 @@ URigHierarchy::URigHierarchy()
 , MetadataTagVersion(0)
 , bEnableDirtyPropagation(true)
 , Elements()
+, bRecordCurveChanges(true)
 , IndexLookup()
 , TransformStackIndex(0)
 , bTransactingForTransformChange(false)
@@ -524,6 +525,7 @@ void URigHierarchy::Reset_Impl(bool bResetElements)
 	OrderedSelection.Reset();
 	PoseVersionPerElement.Reset();
 	ElementDependencyCache.Reset();
+	ResetChangedCurveIndices();
 
 	ChildElementOffsetAndCountCache.Reset();
 	ChildElementCache.Reset();
@@ -804,6 +806,9 @@ void URigHierarchy::CopyHierarchy(URigHierarchy* InHierarchy)
 
 		PreviousNameMap.Append(InHierarchy->PreviousNameMap);
 
+		// make sure the curves are marked as unset after copying
+		UnsetCurveValues();
+
 		// copy the topology version from the hierarchy.
 		// for this we'll include the hash of the hierarchy to make sure we
 		// are deterministic for different hierarchies.
@@ -1062,6 +1067,7 @@ void URigHierarchy::ResetPoseToInitial(ERigElementType InTypeFilter)
 	bool bPerformFiltering = InTypeFilter != ERigElementType::All;
 	
 	FScopeLock Lock(&ElementsLock);
+	const TGuardValue<bool> DisableRecordingOfCurveChanges(bRecordCurveChanges, false);
 	
 	// if we are resetting the pose on some elements, we need to check if
 	// any of affected elements has any children that would not be affected
@@ -1195,9 +1201,13 @@ void URigHierarchy::ResetPoseToInitial(ERigElementType InTypeFilter)
 void URigHierarchy::ResetCurveValues()
 {
 	LLM_SCOPE_BYNAME(TEXT("Animation/ControlRig"));
-	for(int32 ElementIndex=0; ElementIndex<Elements.Num(); ElementIndex++)
+
+	const TGuardValue<bool> DisableRecordingOfCurveChanges(bRecordCurveChanges, false);
+
+	TArray<FRigBaseElement*> Curves = GetCurvesFast();
+	for(FRigBaseElement* Element : Curves)
 	{
-		if(FRigCurveElement* CurveElement = Cast<FRigCurveElement>(Elements[ElementIndex]))
+		if(FRigCurveElement* CurveElement = CastChecked<FRigCurveElement>(Element))
 		{
 			SetCurveValue(CurveElement, 0.f);
 		}
@@ -1207,13 +1217,25 @@ void URigHierarchy::ResetCurveValues()
 void URigHierarchy::UnsetCurveValues(bool bSetupUndo)
 {
 	LLM_SCOPE_BYNAME(TEXT("Animation/ControlRig"));
-	for(int32 ElementIndex=0; ElementIndex<Elements.Num(); ElementIndex++)
+	TArray<FRigBaseElement*> Curves = GetCurvesFast();
+	for(FRigBaseElement* Element : Curves)
 	{
-		if(FRigCurveElement* CurveElement = Cast<FRigCurveElement>(Elements[ElementIndex]))
+		if(FRigCurveElement* CurveElement = CastChecked<FRigCurveElement>(Element))
 		{
 			UnsetCurveValue(CurveElement, bSetupUndo);
 		}
 	}
+	ResetChangedCurveIndices();
+}
+
+const TArray<int32>& URigHierarchy::GetChangedCurveIndices() const
+{
+	return ChangedCurveIndices;
+}
+
+void URigHierarchy::ResetChangedCurveIndices()
+{
+	ChangedCurveIndices.Reset();
 }
 
 int32 URigHierarchy::Num(ERigElementType InElementType) const
@@ -4864,7 +4886,11 @@ void URigHierarchy::SetCurveValue(FRigCurveElement* InCurveElement, float InValu
 		return;
 	}
 
-	InCurveElement->Set(InValue);
+	InCurveElement->Set(InValue, bRecordCurveChanges);
+	if(bRecordCurveChanges)
+	{
+		ChangedCurveIndices.Add(InCurveElement->GetIndex());
+	}
 
 #if WITH_EDITOR
 	if(bSetupUndo || IsTracingChanges())
@@ -4912,6 +4938,7 @@ void URigHierarchy::UnsetCurveValue(FRigCurveElement* InCurveElement, bool bSetu
 	}
 
 	InCurveElement->bIsValueSet = false;
+	ChangedCurveIndices.Remove(InCurveElement->GetIndex());
 
 #if WITH_EDITOR
 	if(bSetupUndo || IsTracingChanges())
diff --git a/Engine/Plugins/Animation/ControlRig/Source/ControlRig/Private/Rigs/RigHierarchyElements.cpp b/Engine/Plugins/Animation/ControlRig/Source/ControlRig/Private/Rigs/RigHierarchyElements.cpp
--- a/Engine/Plugins/Animation/ControlRig/Source/ControlRig/Private/Rigs/RigHierarchyElements.cpp
+++ b/Engine/Plugins/Animation/ControlRig/Source/ControlRig/Private/Rigs/RigHierarchyElements.cpp
@@ -1578,9 +1578,7 @@ void FRigCurveElement::Load(FArchive& Ar, ESerializationPhase SerializationPhase
 		float Value = 0.f;
 		Ar << Value;
 
-		// maintain the deserialized bIsValueSet flag
-		TGuardValue<bool> ValueIsSetGuard(bIsValueSet, bIsValueSet);
-		Set(Value);
+		Set(Value, bIsValueSet);
 	}
 }
 
@@ -1606,12 +1604,12 @@ const float& FRigCurveElement::Get() const
 	return DefaultCurve;
 }
 
-void FRigCurveElement::Set(const float& InValue)
+void FRigCurveElement::Set(const float& InValue, bool InValueIsSet)
 {
 	if(Storage)
 	{
 		*Storage = InValue;
-		bIsValueSet = true;
+		bIsValueSet = InValueIsSet;
 	}
 }
 
diff --git a/Engine/Plugins/Animation/ControlRig/Source/ControlRig/Private/Rigs/RigHierarchyPoseAdapter.cpp b/Engine/Plugins/Animation/ControlRig/Source/ControlRig/Private/Rigs/RigHierarchyPoseAdapter.cpp
--- a/Engine/Plugins/Animation/ControlRig/Source/ControlRig/Private/Rigs/RigHierarchyPoseAdapter.cpp
+++ b/Engine/Plugins/Animation/ControlRig/Source/ControlRig/Private/Rigs/RigHierarchyPoseAdapter.cpp
@@ -26,8 +26,18 @@ void FRigHierarchyPoseAdapter::PostUnlinked(URigHierarchy* InHierarchy)
 	WeakHierarchy.Reset();
 }
 
+TTuple<FRigComputedTransform*, FRigTransformDirtyState*> FRigHierarchyPoseAdapter::GetElementTransformStorage(
+	const FRigElementKeyAndIndex& InKeyAndIndex, ERigTransformType::Type InTransformType, ERigTransformStorageType::Type InStorageType) const
+{
+	if(URigHierarchy* Hierarchy = GetHierarchy())
+	{
+		return Hierarchy->GetElementTransformStorage(InKeyAndIndex, InTransformType, InStorageType);
+	}
+	return {nullptr, nullptr};
+}
+
 bool FRigHierarchyPoseAdapter::RelinkTransformStorage(const FRigElementKeyAndIndex& InKeyAndIndex, ERigTransformType::Type InTransformType,
-	ERigTransformStorageType::Type InStorageType, FTransform* InTransformStorage, bool* InDirtyFlagStorage)
+                                                      ERigTransformStorageType::Type InStorageType, FTransform* InTransformStorage, bool* InDirtyFlagStorage)
 {
 	TArray<TTuple<FRigElementKeyAndIndex, ERigTransformType::Type, ERigTransformStorageType::Type, FTransform*, bool*>> Data =
 		{{InKeyAndIndex, InTransformType, InStorageType, InTransformStorage, InDirtyFlagStorage}};
@@ -59,22 +69,26 @@ bool FRigHierarchyPoseAdapter::RelinkTransformStorage(
 
 			if(FTransform* NewTransformStorage = Tuple.Get<3>())
 			{
+				const FTransform PreviousTransform = CurrentStorage.Get<0>()->Get();
 				if(Hierarchy->ElementTransforms.Contains(CurrentStorage.Get<0>()))
 				{
 					TransformIndicesToDeallocate.Add(CurrentStorage.Get<0>()->GetStorageIndex());
 				}
 				CurrentStorage.Get<0>()->StorageIndex = INDEX_NONE;
 				CurrentStorage.Get<0>()->Storage = NewTransformStorage;
+				CurrentStorage.Get<0>()->Set(PreviousTransform);
 				bPerformedChange = true;
 			}
 			if(bool* NewDirtyStateStorage = Tuple.Get<4>())
 			{
+				const bool bPreviousState = CurrentStorage.Get<1>()->Get();
 				if(Hierarchy->ElementDirtyStates.Contains(CurrentStorage.Get<1>()))
 				{
 					DirtyStateIndicesToDeallocate.Add(CurrentStorage.Get<1>()->GetStorageIndex());
 				}
 				CurrentStorage.Get<1>()->StorageIndex = INDEX_NONE;
 				CurrentStorage.Get<1>()->Storage = NewDirtyStateStorage;
+				CurrentStorage.Get<1>()->Set(bPreviousState);
 				bPerformedChange = true;
 			}
 		}
@@ -121,14 +135,26 @@ bool FRigHierarchyPoseAdapter::RestoreTransformStorage(
 		check(StoragePerElement.Num() == NewDirtyStateIndices.Num());
 		for(int32 Index = 0; Index < StoragePerElement.Num(); Index++)
 		{
-			StoragePerElement[Index].Get<0>()->StorageIndex = NewTransformIndices[Index];
-			StoragePerElement[Index].Get<1>()->StorageIndex = NewDirtyStateIndices[Index];
+			FRigComputedTransform* ComputedTransform = StoragePerElement[Index].Get<0>();
+			FRigTransformDirtyState* DirtyState = StoragePerElement[Index].Get<1>();
+
+			const FTransform PreviousTransform = ComputedTransform->Get();
+			const bool bPreviousState = DirtyState->Get();
+
+			ComputedTransform->StorageIndex = NewTransformIndices[Index];
+			ComputedTransform->Storage = &Hierarchy->ElementTransforms[ComputedTransform->StorageIndex];
+
+			DirtyState->StorageIndex = NewDirtyStateIndices[Index];
+			DirtyState->Storage = &Hierarchy->ElementDirtyStates[DirtyState->StorageIndex];
+			
+			ComputedTransform->Set(PreviousTransform);
+			DirtyState->Set(bPreviousState);
 		}
 		
 		if(bUpdateElementStorage)
 		{
-			(void)UpdateElementStorage();
-			(void)SortStorage();
+			(void)UpdateHierarchyStorage();
+			(void)SortHierarchyStorage();
 		}
 		return true;
 	}
@@ -158,6 +184,7 @@ bool FRigHierarchyPoseAdapter::RelinkCurveStorage(const TArrayView<TTuple<FRigEl
 		for(const TTuple<FRigElementKeyAndIndex, float*>& Tuple : InData)
 		{
 			FRigCurveElement* CurveElement = Hierarchy->Get<FRigCurveElement>(Tuple.Get<0>());
+			const float PreviousValue = CurveElement->Get();
 
 			if(float* NewCurveStorage = Tuple.Get<1>())
 			{
@@ -169,6 +196,8 @@ bool FRigHierarchyPoseAdapter::RelinkCurveStorage(const TArrayView<TTuple<FRigEl
 				CurveElement->Storage = NewCurveStorage;
 				bPerformedChange = true;
 			}
+
+			CurveElement->Set(PreviousValue, CurveElement->bIsValueSet);
 		}
 
 		Hierarchy->ElementCurves.Deallocate(CurveIndicesToDeallocate);
@@ -203,19 +232,23 @@ bool FRigHierarchyPoseAdapter::RestoreCurveStorage(const TArrayView<FRigElementK
 		check(Curves.Num() == NewCurveIndices.Num());
 		for(int32 Index = 0; Index < Curves.Num(); Index++)
 		{
-			Curves[Index]->StorageIndex = NewCurveIndices[Index];
+			FRigCurveElement* CurveElement = Curves[Index];
+			const float PreviousValue = CurveElement->Get();
+			CurveElement->StorageIndex = NewCurveIndices[Index];
+			CurveElement->Storage = &Hierarchy->ElementCurves[CurveElement->StorageIndex];
+			CurveElement->Set(PreviousValue, CurveElement->bIsValueSet);
 		}
 
 		if(bUpdateElementStorage)
 		{
-			(void)UpdateElementStorage();
+			(void)UpdateHierarchyStorage();
 		}
 		return true;
 	}
 	return false;
 }
 
-bool FRigHierarchyPoseAdapter::SortStorage()
+bool FRigHierarchyPoseAdapter::SortHierarchyStorage()
 {
 	if(URigHierarchy* Hierarchy = GetHierarchy())
 	{
@@ -224,7 +257,7 @@ bool FRigHierarchyPoseAdapter::SortStorage()
 	return false;
 }
 
-bool FRigHierarchyPoseAdapter::ShrinkStorage()
+bool FRigHierarchyPoseAdapter::ShrinkHierarchyStorage()
 {
 	if(URigHierarchy* Hierarchy = GetHierarchy())
 	{
@@ -233,7 +266,7 @@ bool FRigHierarchyPoseAdapter::ShrinkStorage()
 	return false;
 }
 
-bool FRigHierarchyPoseAdapter::UpdateElementStorage()
+bool FRigHierarchyPoseAdapter::UpdateHierarchyStorage()
 {
 	if(URigHierarchy* Hierarchy = GetHierarchy())
 	{
diff --git a/Engine/Plugins/Animation/ControlRig/Source/ControlRig/Public/Rigs/RigHierarchy.h b/Engine/Plugins/Animation/ControlRig/Source/ControlRig/Public/Rigs/RigHierarchy.h
--- a/Engine/Plugins/Animation/ControlRig/Source/ControlRig/Public/Rigs/RigHierarchy.h
+++ b/Engine/Plugins/Animation/ControlRig/Source/ControlRig/Public/Rigs/RigHierarchy.h
@@ -283,6 +283,21 @@ public:
 	 */
 	void UnsetCurveValues(bool bSetupUndo = false);
 
+	/**
+	 * Returns all changed curve values
+	 */
+	const TArray<int32>& GetChangedCurveIndices() const;
+
+	/**
+	 * Returns all changed curve values
+	 */
+	void ResetChangedCurveIndices();
+
+	/**
+	 * Returns the flag used decide if we should be recording curve changes
+	 */
+	bool& GetRecordCurveChangesFlag() { return bRecordCurveChanges; }
+	
 	/**
 	 * Returns the number of elements in the Hierarchy.
 	 * @return The number of elements in the Hierarchy
@@ -4317,6 +4332,9 @@ private:
 	// Storage for the elements
 	mutable TArray<TArray<FRigBaseElement*>> ElementsPerType;
 
+	TArray<int32> ChangedCurveIndices;
+	bool bRecordCurveChanges;
+	
 	//
 	struct FMetadataStorage
 	{
diff --git a/Engine/Plugins/Animation/ControlRig/Source/ControlRig/Public/Rigs/RigHierarchyElements.h b/Engine/Plugins/Animation/ControlRig/Source/ControlRig/Public/Rigs/RigHierarchyElements.h
--- a/Engine/Plugins/Animation/ControlRig/Source/ControlRig/Public/Rigs/RigHierarchyElements.h
+++ b/Engine/Plugins/Animation/ControlRig/Source/ControlRig/Public/Rigs/RigHierarchyElements.h
@@ -1879,7 +1879,7 @@ struct CONTROLRIG_API FRigCurveElement final : public FRigBaseElement
 	virtual void CopyPose(FRigBaseElement* InOther, bool bCurrent, bool bInitial, bool bWeights) override;
 
 	const float& Get() const;
-	void Set(const float& InValue);
+	void Set(const float& InValue, bool InValueIsSet = true);
 
 	bool IsValueSet() const { return bIsValueSet; }
 
diff --git a/Engine/Plugins/Animation/ControlRig/Source/ControlRig/Public/Rigs/RigHierarchyPoseAdapter.h b/Engine/Plugins/Animation/ControlRig/Source/ControlRig/Public/Rigs/RigHierarchyPoseAdapter.h
--- a/Engine/Plugins/Animation/ControlRig/Source/ControlRig/Public/Rigs/RigHierarchyPoseAdapter.h
+++ b/Engine/Plugins/Animation/ControlRig/Source/ControlRig/Public/Rigs/RigHierarchyPoseAdapter.h
@@ -23,6 +23,8 @@ protected:
 	virtual void PreUnlinked(URigHierarchy* InHierarchy);
 	virtual void PostUnlinked(URigHierarchy* InHierarchy);
 
+	TTuple<FRigComputedTransform*, FRigTransformDirtyState*> GetElementTransformStorage(const FRigElementKeyAndIndex& InKeyAndIndex, ERigTransformType::Type InTransformType, ERigTransformStorageType::Type InStorageType) const;
+
 	bool RelinkTransformStorage(const FRigElementKeyAndIndex& InKeyAndIndex, ERigTransformType::Type InTransformType, ERigTransformStorageType::Type InStorageType, FTransform* InTransformStorage, bool* InDirtyFlagStorage);
 	bool RestoreTransformStorage(const FRigElementKeyAndIndex& InKeyAndIndex, ERigTransformType::Type InTransformType, ERigTransformStorageType::Type InStorageType, bool bUpdateElementStorage);
 	bool RelinkTransformStorage(const TArrayView<TTuple<FRigElementKeyAndIndex,ERigTransformType::Type,ERigTransformStorageType::Type,FTransform*,bool*>>& InData);
@@ -33,9 +35,9 @@ protected:
 	bool RelinkCurveStorage(const TArrayView<TTuple<FRigElementKeyAndIndex,float*>>& InData);
 	bool RestoreCurveStorage(const TArrayView<FRigElementKeyAndIndex>& InData, bool bUpdateElementStorage);
 
-	bool SortStorage();
-	bool ShrinkStorage();
-	bool UpdateElementStorage();
+	bool SortHierarchyStorage();
+	bool ShrinkHierarchyStorage();
+	bool UpdateHierarchyStorage();
 
 	TWeakObjectPtr<URigHierarchy> WeakHierarchy;
 
