Control Rig: Fix RigUnit_GetCurveValue unit test

#rnx

diff --git a/Engine/Plugins/Animation/ControlRig/Source/ControlRig/Private/Units/Hierarchy/RigUnit_GetCurveValue.cpp b/Engine/Plugins/Animation/ControlRig/Source/ControlRig/Private/Units/Hierarchy/RigUnit_GetCurveValue.cpp
--- a/Engine/Plugins/Animation/ControlRig/Source/ControlRig/Private/Units/Hierarchy/RigUnit_GetCurveValue.cpp
+++ b/Engine/Plugins/Animation/ControlRig/Source/ControlRig/Private/Units/Hierarchy/RigUnit_GetCurveValue.cpp
@@ -38,10 +38,16 @@ IMPLEMENT_RIGUNIT_AUTOMATION_TEST(FRigUnit_GetCurveValue)
 	
 	Unit.Curve = ValidCurveName;
 	Execute();
-	AddErrorIfFalse(Unit.Valid, TEXT("Expected curve hold a valid value"));
+	AddErrorIfFalse(!Unit.Valid, TEXT("Expected curve hold a unset value"));
 	Execute();
 	AddErrorIfFalse(Unit.Value != 1.0f, TEXT("Expected curve's value to be 1.0"));
 
+	Hierarchy->SetCurveValue({ValidCurveName, ERigElementType::Curve}, 2.0f);
+	Execute();
+	AddErrorIfFalse(Unit.Valid, TEXT("Expected curve hold a valid value"));
+	Execute();
+	AddErrorIfFalse(Unit.Value == 2.0f, TEXT("Expected curve's value to be 2.0"));
+
 	Unit.Curve = InvalidCurveName;
 	Execute();
 	AddErrorIfFalse(!Unit.Valid, TEXT("Expected curve hold an invalid value"));
