Control Rig: Fix RigHierarchyController AddControl dirty state evaluation order

#rnx

diff --git a/Engine/Plugins/Animation/ControlRig/Source/ControlRig/Public/Rigs/RigHierarchyElements.h b/Engine/Plugins/Animation/ControlRig/Source/ControlRig/Public/Rigs/RigHierarchyElements.h
--- a/Engine/Plugins/Animation/ControlRig/Source/ControlRig/Public/Rigs/RigHierarchyElements.h
+++ b/Engine/Plugins/Animation/ControlRig/Source/ControlRig/Public/Rigs/RigHierarchyElements.h
@@ -309,6 +309,7 @@ struct CONTROLRIG_API FRigReusableElementStorage
 			else
 			{
 				Indices.Push(FreeList.Pop(EAllowShrinking::No));
+				Storage[Indices.Last()] = InDefault;
 			}
 		}
 
