[MovieSceneControlRig] make sure we unregister endload callback for when the rig is a native rig like FK rig

diff --git a/Engine/Plugins/Animation/ControlRig/Source/ControlRig/Private/Sequencer/MovieSceneControlRigParameterTrack.cpp b/Engine/Plugins/Animation/ControlRig/Source/ControlRig/Private/Sequencer/MovieSceneControlRigParameterTrack.cpp
--- a/Engine/Plugins/Animation/ControlRig/Source/ControlRig/Private/Sequencer/MovieSceneControlRigParameterTrack.cpp
+++ b/Engine/Plugins/Animation/ControlRig/Source/ControlRig/Private/Sequencer/MovieSceneControlRigParameterTrack.cpp
@@ -592,32 +592,34 @@ void UMovieSceneControlRigParameterTrack::DeclareConstructClasses(TArray<FTopLev
 #if WITH_EDITOR
 void UMovieSceneControlRigParameterTrack::HandlePackageDone(const FEndLoadPackageContext& Context)
 {
+	if (!ControlRig || ControlRig->GetClass()->IsNative())
+	{
+		// EndLoad is never called for native packages, so skip work
+		FCoreUObjectDelegates::OnEndLoadPackage.RemoveAll(this);
+		return;
+	}
+	
+	// ensure both the track package and the control rig package are fully end-loaded	
 	if (!GetPackage()->GetHasBeenEndLoaded())
 	{
 		return;
 	}
-
-	// ensure both packages are fully end-loaded
-	if (ControlRig && !ControlRig->GetClass()->IsNative())
+	
+	if (const UPackage* ControlRigPackage = Cast<UPackage>(ControlRig->GetClass()->GetOutermost()))
 	{
-		if (const UPackage* ControlRigPackage = Cast<UPackage>(ControlRig->GetClass()->GetOutermost()))
+		if (!ControlRigPackage->GetHasBeenEndLoaded())
 		{
-			if (!ControlRigPackage->GetHasBeenEndLoaded())
-			{
-				return;
-			}
+			return;
 		}
-
-		// All dependent packages ready, no need to wait/check for any other packages
-		// ReconstructControlRig may trigger loading of packages that we don't care about, so unregister from the delegate
-		// to avoid infinite loop
-		FCoreUObjectDelegates::OnEndLoadPackage.RemoveAll(this);
-		
-		// Only reconstruct in case it is not a native ControlRig class
-		ReconstructControlRig();
 	}
 
+	// All dependent packages ready, no need to wait/check for any other packages
+	// ReconstructControlRig may trigger loading of packages that we don't care about, so unregister from the delegate
+	// before reconstruction to avoid infinite loop
+	FCoreUObjectDelegates::OnEndLoadPackage.RemoveAll(this);
 
+	// Only reconstruct in case it is not a native ControlRig class
+	ReconstructControlRig();	
 }
 
 void UMovieSceneControlRigParameterTrack::HandleControlRigPackageDone(URigVMHost* InControlRig)
