Control Rig: Optimize offset / modify transform node to be smart about using global or local

#rnx

diff --git a/Engine/Plugins/Animation/ControlRig/Source/ControlRig/Private/Units/Hierarchy/RigUnit_OffsetTransform.cpp b/Engine/Plugins/Animation/ControlRig/Source/ControlRig/Private/Units/Hierarchy/RigUnit_OffsetTransform.cpp
--- a/Engine/Plugins/Animation/ControlRig/Source/ControlRig/Private/Units/Hierarchy/RigUnit_OffsetTransform.cpp
+++ b/Engine/Plugins/Animation/ControlRig/Source/ControlRig/Private/Units/Hierarchy/RigUnit_OffsetTransform.cpp
@@ -20,9 +20,28 @@ FRigUnit_OffsetTransformForItem_Execute()
 	FTransform PreviousTransform = FTransform::Identity;
 	FTransform GlobalTransform = FTransform::Identity;
 
-	FRigUnit_GetTransform::StaticExecute(ExecuteContext, Item, ERigVMTransformSpace::GlobalSpace, false, PreviousTransform, CachedIndex);
+	if(!CachedIndex.UpdateCache(Item, ExecuteContext.Hierarchy))
+	{
+		return;
+	}
+
+	const FRigTransformElement* TransformElement = Cast<FRigTransformElement>(CachedIndex.GetElement());
+	if(TransformElement == nullptr)
+	{
+		return;
+	}
+
+	// figure out which transform type has already been computed (is clean) to avoid compute
+	ERigVMTransformSpace TransformTypeToUse = ERigVMTransformSpace::GlobalSpace;
+	if(TransformElement->GetDirtyState().IsDirty(ERigTransformType::CurrentGlobal))
+	{
+		check(!TransformElement->GetDirtyState().IsDirty(ERigTransformType::CurrentLocal));
+		TransformTypeToUse = ERigVMTransformSpace::LocalSpace;
+	}
+
+	FRigUnit_GetTransform::StaticExecute(ExecuteContext, Item, TransformTypeToUse, false, PreviousTransform, CachedIndex);
 	FRigVMFunction_MathTransformMakeAbsolute::StaticExecute(ExecuteContext, OffsetTransform, PreviousTransform, GlobalTransform);
-	FRigUnit_SetTransform::StaticExecute(ExecuteContext, Item, ERigVMTransformSpace::GlobalSpace, false, GlobalTransform, Weight, bPropagateToChildren, CachedIndex);
+	FRigUnit_SetTransform::StaticExecute(ExecuteContext, Item, TransformTypeToUse, false, GlobalTransform, Weight, bPropagateToChildren, CachedIndex);
 }
 
 #if WITH_EDITOR
diff --git a/Engine/Plugins/Animation/ControlRig/Source/ControlRig/Private/Units/Highlevel/Hierarchy/RigUnit_ModifyTransforms.cpp b/Engine/Plugins/Animation/ControlRig/Source/ControlRig/Private/Units/Highlevel/Hierarchy/RigUnit_ModifyTransforms.cpp
--- a/Engine/Plugins/Animation/ControlRig/Source/ControlRig/Private/Units/Highlevel/Hierarchy/RigUnit_ModifyTransforms.cpp
+++ b/Engine/Plugins/Animation/ControlRig/Source/ControlRig/Private/Units/Highlevel/Hierarchy/RigUnit_ModifyTransforms.cpp
@@ -70,8 +70,23 @@ FRigUnit_ModifyTransforms_Execute()
 						{
 							Transform = FControlRigMathLibrary::LerpTransform(FTransform::Identity, Transform, T);
 						}
-						Transform = Transform * Hierarchy->GetLocalTransform(CachedItem);
-						Hierarchy->SetLocalTransform(CachedItem, Transform, true);
+							
+						FRigTransformElement* TransformElement = const_cast<FRigTransformElement*>(Cast<FRigTransformElement>(CachedItem.GetElement()));
+						if(TransformElement == nullptr)
+						{
+							return;
+						}
+
+						// figure out which transform type has already been computed (is clean) to avoid compute
+						ERigTransformType::Type TransformTypeToUse = ERigTransformType::CurrentLocal;
+						if(TransformElement->GetDirtyState().IsDirty(ERigTransformType::CurrentLocal))
+						{
+							check(!TransformElement->GetDirtyState().IsDirty(ERigTransformType::CurrentGlobal));
+							TransformTypeToUse = ERigTransformType::CurrentGlobal;
+						}
+
+						Transform = Transform * Hierarchy->GetTransform(TransformElement, TransformTypeToUse);
+						Hierarchy->SetTransform(TransformElement, Transform, TransformTypeToUse, true);
 						break;
 					}
 					case EControlRigModifyBoneMode::AdditiveGlobal:
