Control Rig: Show pretty short names also if the element is missing (missing skeletal mesh)

#rnx

diff --git a/Engine/Plugins/Animation/ControlRig/Source/ControlRig/Private/ModularRig.cpp b/Engine/Plugins/Animation/ControlRig/Source/ControlRig/Private/ModularRig.cpp
--- a/Engine/Plugins/Animation/ControlRig/Source/ControlRig/Private/ModularRig.cpp
+++ b/Engine/Plugins/Animation/ControlRig/Source/ControlRig/Private/ModularRig.cpp
@@ -874,7 +874,16 @@ void UModularRig::RecomputeShortestDisplayPathCache() const
 	const URigHierarchy* Hierarchy = GetHierarchy();
 	check(Hierarchy);
 
-	const TArray<FRigElementKey> AllKeys = Hierarchy->GetAllKeys();
+	TArray<FRigElementKey> AllKeys = Hierarchy->GetAllKeys();
+
+	// add all of the keys from the connections. we may be in a situation where
+	// the keys in the connections are actually referring to invalid / unloaded elements
+	// but we still want to show pretty short names.
+	for(const FModularRigSingleConnection& Connection : ModularRigModel.Connections)
+	{
+		AllKeys.AddUnique(Connection.Connector);
+		AllKeys.AddUnique(Connection.Target);
+	}
 
 	auto GetNameForElement = [Hierarchy](const FRigElementKey& InElementKey)
 	{
@@ -913,7 +922,13 @@ void UModularRig::RecomputeShortestDisplayPathCache() const
 		{
 			if(const FRigModuleReference* Module = Model.FindModule(ModulePath.ToString()))
 			{
-				const FString NameString = Name.ToString();
+				const FString ModulePathString = ModulePath.ToString();
+				const FString ModulePathPrefix = ModulePath.ToString() + NamespaceSeparator;
+				FString NameString = Name.ToString();
+				if(NameString.StartsWith(ModulePathPrefix))
+				{
+					NameString = NameString.Mid(ModulePathPrefix.Len());
+				}
 				const FString ModuleShortName = Module->GetShortName();
 				const FString NameSpacedName = URigHierarchy::JoinNameSpace(ModuleShortName, NameString);
 				ElementKeyToShortestDisplayPath.Add(Key, { NameSpacedName, IsNameUniqueInHierarchy.FindChecked(Name) ? NameString : NameSpacedName });
diff --git a/Engine/Plugins/Animation/ControlRig/Source/ControlRigEditor/Private/Editor/SModularRigTreeView.cpp b/Engine/Plugins/Animation/ControlRig/Source/ControlRigEditor/Private/Editor/SModularRigTreeView.cpp
--- a/Engine/Plugins/Animation/ControlRig/Source/ControlRigEditor/Private/Editor/SModularRigTreeView.cpp
+++ b/Engine/Plugins/Animation/ControlRig/Source/ControlRigEditor/Private/Editor/SModularRigTreeView.cpp
@@ -515,7 +515,9 @@ TSharedRef<SWidget> SModularRigModelItem::GenerateWidgetForColumn(const FName& C
 			CurrentTargetKey = *Key;
 		}
 		TPair<const FSlateBrush*, FSlateColor> IconAndColor = SRigHierarchyItem::GetBrushForElementType(ModularRig->GetHierarchy(), CurrentTargetKey);
-		PopulateConnectorCurrentTarget(ComboButtonBox, ConnectorKey, CurrentTargetKey, IconAndColor.Key, IconAndColor.Value, FText::FromName(CurrentTargetKey.Name));
+
+		const FText CurrentTargetShortName = FText::FromString(ModularRig->GetShortestDisplayPathForElement(CurrentTargetKey, false));
+		PopulateConnectorCurrentTarget(ComboButtonBox, ConnectorKey, CurrentTargetKey, IconAndColor.Key, IconAndColor.Value, CurrentTargetShortName);
 
 		return Widget;
 	}
