Sequencer: Control Rig: Fix rotation orders issue, was incorrectly setting the preferred angles directly from the rotator, insead of converting it into rotation order space

diff --git a/Engine/Plugins/Animation/ControlRig/Source/ControlRig/Private/Rigs/RigHierarchy.cpp b/Engine/Plugins/Animation/ControlRig/Source/ControlRig/Private/Rigs/RigHierarchy.cpp
--- a/Engine/Plugins/Animation/ControlRig/Source/ControlRig/Private/Rigs/RigHierarchy.cpp
+++ b/Engine/Plugins/Animation/ControlRig/Source/ControlRig/Private/Rigs/RigHierarchy.cpp
@@ -3947,7 +3947,8 @@ void URigHierarchy::SetTransform(FRigTransformElement* InTransformElement, const
 		if(bUsePreferredEulerAngles && ERigTransformType::IsLocal(InTransformType))
 		{
 			const bool bInitial = ERigTransformType::IsInitial(InTransformType);
-			ControlElement->PreferredEulerAngles.SetRotator(InTransform.Rotator(), bInitial, true);
+			const FVector Angle = GetControlAnglesFromQuat(ControlElement, InTransform.GetRotation(), true);
+			ControlElement->PreferredEulerAngles.SetAngles(Angle, bInitial, ControlElement->PreferredEulerAngles.RotationOrder, true);
 		}
 	}
 
@@ -8130,4 +8131,4 @@ FRigHierarchyRedirectorGuard::FRigHierarchyRedirectorGuard(UControlRig* InContro
 : Guard(InControlRig->GetHierarchy()->ElementKeyRedirector, &InControlRig->GetElementKeyRedirector())
 {
 }
-;
\ No newline at end of file
+;
diff --git a/Engine/Plugins/Animation/ControlRig/Source/ControlRig/Private/Sequencer/MovieSceneControlRigParameterTemplate.cpp b/Engine/Plugins/Animation/ControlRig/Source/ControlRig/Private/Sequencer/MovieSceneControlRigParameterTemplate.cpp
--- a/Engine/Plugins/Animation/ControlRig/Source/ControlRig/Private/Sequencer/MovieSceneControlRigParameterTemplate.cpp
+++ b/Engine/Plugins/Animation/ControlRig/Source/ControlRig/Private/Sequencer/MovieSceneControlRigParameterTemplate.cpp
@@ -1814,6 +1814,7 @@ struct TControlRigParameterActuatorTransform : TMovieSceneBlendingActuator<FCont
 						FEulerTransform Transform = InFinalValue.Value;
 						Transform.Rotation = UERotator;
 						ControlRig->SetControlValue<FRigControlValue::FEulerTransform_Float>(ParameterName, Transform, true, EControlRigSetKey::Never, bSetupUndo);	
+						Hierarchy->SetControlSpecifiedEulerAngle(ControlElement, EulerAngle);
 					}
 				}
 			}
