Control Rig: Avoid resetting pose after construction event

#rnx

diff --git a/Engine/Plugins/Animation/ControlRig/Source/ControlRig/Private/ControlRig.cpp b/Engine/Plugins/Animation/ControlRig/Source/ControlRig/Private/ControlRig.cpp
--- a/Engine/Plugins/Animation/ControlRig/Source/ControlRig/Private/ControlRig.cpp
+++ b/Engine/Plugins/Animation/ControlRig/Source/ControlRig/Private/ControlRig.cpp
@@ -970,6 +970,7 @@ bool UControlRig::Execute(const FName& InEventName)
 	}
 
 	bool bSuccess = true;
+	bool bPostConstructionEventWasRun = false;
 
 	// we'll special case the construction event here
 	if (bIsConstructionEvent)
@@ -1055,9 +1056,10 @@ bool UControlRig::Execute(const FName& InEventName)
 					
 				} // destroy FTransientControlScope
 
-				if(!bPostConstructionEventInQueue)
+				if(!bPostConstructionEventInQueue && !bPostConstructionEventWasRun)
 				{
 					RunPostConstructionEvent();
+					bPostConstructionEventWasRun = true;
 				}
 
 				// Reset the input pose after construction
@@ -1073,9 +1075,10 @@ bool UControlRig::Execute(const FName& InEventName)
 			
 		} // destroy DisableSelectionNotifications
 
-		if (bIsPostConstructionEvent || (bIsConstructionEvent && !bPostConstructionEventInQueue))
+		if ((bIsPostConstructionEvent || (bIsConstructionEvent && !bPostConstructionEventInQueue)) && !bPostConstructionEventWasRun)
 		{
 			RunPostConstructionEvent();
+			bPostConstructionEventWasRun = true;
 		}
 
 		if (bConstructionModeEnabled)
diff --git a/Engine/Plugins/Animation/ControlRig/Source/ControlRigDeveloper/Private/ControlRigBlueprint.cpp b/Engine/Plugins/Animation/ControlRig/Source/ControlRigDeveloper/Private/ControlRigBlueprint.cpp
--- a/Engine/Plugins/Animation/ControlRig/Source/ControlRigDeveloper/Private/ControlRigBlueprint.cpp
+++ b/Engine/Plugins/Animation/ControlRig/Source/ControlRigDeveloper/Private/ControlRigBlueprint.cpp
@@ -1773,6 +1773,33 @@ FName UControlRigBlueprint::RemoveTransientControl(const FRigElementKey& InEleme
 
 void UControlRigBlueprint::ClearTransientControls()
 {
+	bool bHasAnyTransientControls = false;
+	
+	if (URigVMBlueprintGeneratedClass* RigClass = GetRigVMBlueprintGeneratedClass())
+	{
+		UControlRig* CDO = Cast<UControlRig>(RigClass->GetDefaultObject(true /* create if needed */));
+
+		TArray<UObject*> ArchetypeInstances;
+		CDO->GetArchetypeInstances(ArchetypeInstances);
+		for (UObject* ArchetypeInstance : ArchetypeInstances)
+		{
+			UControlRig* InstancedControlRig = Cast<UControlRig>(ArchetypeInstance);
+			if (InstancedControlRig)
+			{
+				if(!InstancedControlRig->GetHierarchy()->GetTransientControls().IsEmpty())
+				{
+					bHasAnyTransientControls = true;
+					break;
+				}
+			}
+		}
+	}
+
+	if(!bHasAnyTransientControls)
+	{
+		return;
+	}
+
 	TUniquePtr<FControlValueScope> ValueScope;
 	if (!UControlRigEditorSettings::Get()->bResetControlsOnPinValueInteraction) // if we need to retain the controls
 	{
