Control Rig: Show pretty short names also if the element is missing (missing skeletal mesh)

#rnx

diff --git a/Engine/Plugins/Animation/ControlRig/Source/ControlRigEditor/Private/Editor/SModularRigTreeView.cpp b/Engine/Plugins/Animation/ControlRig/Source/ControlRigEditor/Private/Editor/SModularRigTreeView.cpp
--- a/Engine/Plugins/Animation/ControlRig/Source/ControlRigEditor/Private/Editor/SModularRigTreeView.cpp
+++ b/Engine/Plugins/Animation/ControlRig/Source/ControlRigEditor/Private/Editor/SModularRigTreeView.cpp
@@ -516,7 +516,11 @@ TSharedRef<SWidget> SModularRigModelItem::GenerateWidgetForColumn(const FName& C
 		}
 		TPair<const FSlateBrush*, FSlateColor> IconAndColor = SRigHierarchyItem::GetBrushForElementType(ModularRig->GetHierarchy(), CurrentTargetKey);
 
-		const FText CurrentTargetShortName = FText::FromString(ModularRig->GetShortestDisplayPathForElement(CurrentTargetKey, false));
+		FText CurrentTargetShortName = FText::FromString(ModularRig->GetShortestDisplayPathForElement(CurrentTargetKey, false));
+		if(CurrentTargetShortName.IsEmpty())
+		{
+			CurrentTargetShortName = FText::FromName(CurrentTargetKey.Name);
+		}
 		PopulateConnectorCurrentTarget(ComboButtonBox, ConnectorKey, CurrentTargetKey, IconAndColor.Key, IconAndColor.Value, CurrentTargetShortName);
 
 		return Widget;
