Control Rig: AnimNode to use pose adapters

#rnx

diff --git a/Engine/Plugins/Animation/ControlRig/Source/ControlRig/Private/AnimNode_ControlRig.cpp b/Engine/Plugins/Animation/ControlRig/Source/ControlRig/Private/AnimNode_ControlRig.cpp
--- a/Engine/Plugins/Animation/ControlRig/Source/ControlRig/Private/AnimNode_ControlRig.cpp
+++ b/Engine/Plugins/Animation/ControlRig/Source/ControlRig/Private/AnimNode_ControlRig.cpp
@@ -232,7 +232,7 @@ UClass* FAnimNode_ControlRig::GetTargetClass() const
 	return DefaultControlRigClass;
 }
 
-void FAnimNode_ControlRig::UpdateInput(UControlRig* InControlRig, const FPoseContext& InOutput)
+void FAnimNode_ControlRig::UpdateInput(UControlRig* InControlRig, FPoseContext& InOutput)
 {
 	DECLARE_SCOPE_HIERARCHICAL_COUNTER_FUNC()
 
diff --git a/Engine/Plugins/Animation/ControlRig/Source/ControlRig/Private/AnimNode_ControlRigBase.cpp b/Engine/Plugins/Animation/ControlRig/Source/ControlRig/Private/AnimNode_ControlRigBase.cpp
--- a/Engine/Plugins/Animation/ControlRig/Source/ControlRig/Private/AnimNode_ControlRigBase.cpp
+++ b/Engine/Plugins/Animation/ControlRig/Source/ControlRig/Private/AnimNode_ControlRigBase.cpp
@@ -20,7 +20,6 @@ DECLARE_CYCLE_STAT(TEXT("ControlRig_UpdateInput"), STAT_ControlRig_UpdateInput,
 DECLARE_CYCLE_STAT(TEXT("ControlRig_Evaluate"), STAT_ControlRig_Evaluate, STATGROUP_Anim);
 DECLARE_CYCLE_STAT(TEXT("ControlRig_UpdateOutput"), STAT_ControlRig_UpdateOutput, STATGROUP_Anim);
 
-
 #if ENABLE_ANIM_DEBUG
 TAutoConsoleVariable<int32> CVarAnimNodeControlRigDebug(TEXT("a.AnimNode.ControlRig.Debug"), 0, TEXT("Set to 1 to turn on debug drawing for AnimNode_ControlRigBase"));
 #endif
@@ -28,12 +27,329 @@ TAutoConsoleVariable<int32> CVarAnimNodeControlRigDebug(TEXT("a.AnimNode.Control
 // CVar to disable control rig execution within an anim node
 static TAutoConsoleVariable<int32> CVarControlRigDisableExecutionAnimNode(TEXT("ControlRig.DisableExecutionInAnimNode"), 0, TEXT("if nonzero we disable the execution of Control Rigs inside an anim node."));
 
+void FAnimNode_ControlRig_PoseAdapter::PostLinked(URigHierarchy* InHierarchy)
+{
+	FRigHierarchyPoseAdapter::PostLinked(InHierarchy);
+
+	// 1. sort the hierarchies storage so that we can make sure initial and local is grouped correctly
+	SortHierarchyStorage();
+
+	// 2. make sure to compute all local transforms in initial and local
+	const TArray<FRigTransformElement*> TransformElements = InHierarchy->GetElementsOfType<FRigTransformElement>(false);
+	for(FRigTransformElement* TransformElement : TransformElements)
+	{
+		(void)InHierarchy->GetTransform(TransformElement, ERigTransformType::InitialLocal);
+		(void)InHierarchy->GetTransform(TransformElement, ERigTransformType::CurrentLocal);
+	}
+	
+	// 3. set up a list of dependents in the hierarchy of the rig to be reset to initial on execute
+	TArray<FRigTransformElement*> DependentTransformElements;
+	for(FRigTransformElement* TransformElement : TransformElements)
+	{
+		if(ElementIndexToPoseIndex.Contains((uint16)TransformElement->GetIndex()))
+		{
+			continue;
+		}
+
+		const FRigBaseElementParentArray& ParentElements = InHierarchy->GetParents(TransformElement);
+		for(const FRigBaseElement* ParentElement : ParentElements)
+		{
+			if(ElementIndexToPoseIndex.Contains((uint16)ParentElement->GetIndex()))
+			{
+				DependentTransformElements.AddUnique(TransformElement);
+			}
+			else if(DependentTransformElements.Contains(ParentElement))
+			{
+				DependentTransformElements.AddUnique(TransformElement);
+			}
+		}
+	}
+	Dependents.Reset();
+	Dependents.Reserve(DependentTransformElements.Num() * 3);
+	
+	for(FRigTransformElement* DependentTransformElement : DependentTransformElements)
+	{
+		// skip bones - since they are taking care of by BonesToResetToInitial
+		if(const FRigBoneElement* BoneElement = Cast<FRigBoneElement>(DependentTransformElement))
+		{
+			// skip bones if they are not user defined - and if the parent of the bone is not a dependent as well
+			if(BoneElement->BoneType == ERigBoneType::Imported &&
+				!DependentTransformElements.Contains(BoneElement->ParentElement))
+			{
+				continue;
+			}
+		}
+
+		FRigControlElement* DependentControlElement = Cast<FRigControlElement>(DependentTransformElement);
+		if(DependentControlElement && DependentControlElement->IsAnimationChannel())
+		{
+			continue;
+		}
+		
+		Dependents.Emplace(DependentTransformElement->GetKeyAndIndex(), ERigTransformType::CurrentGlobal, ERigTransformStorageType::Pose, &DependentTransformElement->GetDirtyState().Current);
+		if(DependentControlElement)
+		{
+			Dependents.Emplace(DependentTransformElement->GetKeyAndIndex(), ERigTransformType::CurrentGlobal, ERigTransformStorageType::Offset, &DependentControlElement->GetOffsetDirtyState().Current);
+			Dependents.Emplace(DependentTransformElement->GetKeyAndIndex(), ERigTransformType::CurrentGlobal, ERigTransformStorageType::Shape, &DependentControlElement->GetShapeDirtyState().Current);
+		}
+	}
+
+	// 4. relink the storage for the transforms (local or global or both) and dirty states to our local storage
+	static constexpr bool bLocalIsPrimary = true;
+	UpdateDirtyStates(bLocalIsPrimary);
+	for(int32 PoseIndex = 0; PoseIndex < PoseIndexToElementIndex.Num(); PoseIndex++)
+	{
+		const int32& TransformElementIndex = PoseIndexToElementIndex[PoseIndex];
+		if(TransformElementIndex != INDEX_NONE)
+		{
+			const FRigElementKeyAndIndex& KeyAndIndex = InHierarchy->GetKeyAndIndex(TransformElementIndex);
+			RelinkTransformStorage(KeyAndIndex, ERigTransformType::CurrentLocal, ERigTransformStorageType::Pose, &LocalPose[PoseIndex], &LocalPoseIsDirty[PoseIndex]);
+			RelinkTransformStorage(KeyAndIndex, ERigTransformType::CurrentGlobal, ERigTransformStorageType::Pose, &GlobalPose[PoseIndex], &GlobalPoseIsDirty[PoseIndex]);
+		}
+	}
+	UpdateDirtyStates(bLocalIsPrimary); // do this again to make sure local is flagged as clean and global as dirty
+
+	// 5. Shrink the storage on the hierarchy now that we've relinked it
+	ShrinkHierarchyStorage();
+	
+	LastTopologyVersion = InHierarchy->GetTopologyVersion();
+}
+
+void FAnimNode_ControlRig_PoseAdapter::PreUnlinked(URigHierarchy* InHierarchy)
+{
+	UnlinkTransformStorage();
+	UnlinkCurveStorage();
+	UpdateHierarchyStorage();
+	
+	LastTopologyVersion = UINT32_MAX;
+	FRigHierarchyPoseAdapter::PreUnlinked(InHierarchy);
+}
+
+void FAnimNode_ControlRig_PoseAdapter::UnlinkTransformStorage()
+{
+	if(const URigHierarchy* Hierarchy = GetHierarchy())
+	{
+		for(int32 PoseIndex = 0; PoseIndex < PoseIndexToElementIndex.Num(); PoseIndex++)
+		{
+			const int32& TransformElementIndex = PoseIndexToElementIndex[PoseIndex];
+			if(TransformElementIndex != INDEX_NONE)
+			{
+				const FRigElementKeyAndIndex& KeyAndIndex = Hierarchy->GetKeyAndIndex(TransformElementIndex);
+				RestoreTransformStorage(KeyAndIndex, ERigTransformType::CurrentLocal, ERigTransformStorageType::Pose, false);
+				RestoreTransformStorage(KeyAndIndex, ERigTransformType::CurrentGlobal, ERigTransformStorageType::Pose, false);
+			}
+		}
+	}
+
+	ElementIndexToPoseIndex.Reset();
+	PoseIndexToElementIndex.Reset();
+}
+
+void FAnimNode_ControlRig_PoseAdapter::UnlinkCurveStorage()
+{
+	if(const URigHierarchy* Hierarchy = GetHierarchy())
+	{
+		const TArray<FRigBaseElement*> CurveElements = Hierarchy->GetCurvesFast();
+		for(const FRigBaseElement* CurveElement : CurveElements)
+		{
+			RestoreCurveStorage(CurveElement->GetKeyAndIndex(), false);
+		}
+	}
+}
+
+void FAnimNode_ControlRig_PoseAdapter::ConvertToLocalPose()
+{
+	check(ParentPoseIndices.Num() == GlobalPose.Num());
+	LocalPose.SetNum(GlobalPose.Num());
+	for(int32 Index = 0; Index < ParentPoseIndices.Num(); Index++)
+	{
+		(void)GetLocalTransform(Index);
+	}
+}
+void FAnimNode_ControlRig_PoseAdapter::ConvertToGlobalPose()
+{
+	check(ParentPoseIndices.Num() == LocalPose.Num());
+	GlobalPose.SetNum(LocalPose.Num());
+
+	for(int32 Index = 0; Index < ParentPoseIndices.Num(); Index++)
+	{
+		(void)GetGlobalTransform(Index);
+	}
+}
+
+const FTransform& FAnimNode_ControlRig_PoseAdapter::GetLocalTransform(int32 InIndex)
+{
+	check(LocalPose.IsValidIndex(InIndex));
+	check(LocalPoseIsDirty.IsValidIndex(InIndex));
+	check(ParentPoseIndices.IsValidIndex(InIndex));
+	check(RequiresHierarchyForSpaceConversion.IsValidIndex(InIndex));
+
+	const URigHierarchy* Hierarchy = GetHierarchy();
+	check(Hierarchy);
+
+	if(LocalPoseIsDirty[InIndex])
+	{
+		check(GlobalPoseIsDirty[InIndex] == false);
+		if(RequiresHierarchyForSpaceConversion[InIndex] && PoseIndexToElementIndex.IsValidIndex(InIndex))
+		{
+			LocalPose[InIndex] = Hierarchy->GetLocalTransformByIndex(PoseIndexToElementIndex[InIndex]);
+		}
+		else
+		{
+			const int32 ParentIndex = ParentPoseIndices[InIndex];
+			if(ParentIndex == INDEX_NONE)
+			{
+				LocalPose[InIndex] = GetGlobalTransform(InIndex);
+			}
+			else
+			{
+				LocalPose[InIndex] = GetGlobalTransform(InIndex).GetRelativeTransform(GetGlobalTransform(ParentIndex));
+			}
+			LocalPose[InIndex].NormalizeRotation();
+		}
+		LocalPoseIsDirty[InIndex] = false;
+	}
+	return LocalPose[InIndex];
+}
+
+const FTransform& FAnimNode_ControlRig_PoseAdapter::GetGlobalTransform(int32 InIndex)
+{
+	check(GlobalPose.IsValidIndex(InIndex));
+	check(GlobalPoseIsDirty.IsValidIndex(InIndex));
+	check(ParentPoseIndices.IsValidIndex(InIndex));
+	check(RequiresHierarchyForSpaceConversion.IsValidIndex(InIndex));
+
+	const URigHierarchy* Hierarchy = GetHierarchy();
+	check(Hierarchy);
+
+	if(GlobalPoseIsDirty[InIndex])
+	{
+		check(LocalPoseIsDirty[InIndex] == false);
+		if(RequiresHierarchyForSpaceConversion[InIndex] && PoseIndexToElementIndex.IsValidIndex(InIndex))
+		{
+			GlobalPose[InIndex] = Hierarchy->GetGlobalTransformByIndex(PoseIndexToElementIndex[InIndex]);
+		}
+		else
+		{
+			const int32 ParentIndex = ParentPoseIndices[InIndex];
+			if(ParentIndex == INDEX_NONE)
+			{
+				GlobalPose[InIndex] = GetLocalTransform(InIndex);
+			}
+			else
+			{
+				GlobalPose[InIndex] = GetLocalTransform(InIndex) * GetGlobalTransform(ParentIndex);
+			}
+			GlobalPose[InIndex].NormalizeRotation();
+		}
+		GlobalPoseIsDirty[InIndex] = false;
+	}
+	return GlobalPose[InIndex];
+}
+
+void FAnimNode_ControlRig_PoseAdapter::UpdateDirtyStates(const TOptional<bool> InLocalIsPrimary)
+{
+	bool bLocalIsPrimary = bTransferInLocalSpace;
+	if(InLocalIsPrimary.IsSet())
+	{
+		bLocalIsPrimary = InLocalIsPrimary.GetValue();
+	}
+	
+	if(LocalPoseIsDirty.Num() != LocalPose.Num())
+	{
+		LocalPoseIsDirty.Reset();
+		LocalPoseIsDirty.AddZeroed(LocalPose.Num());
+	}
+	else if(bLocalIsPrimary)
+	{
+		FMemory::Memzero(LocalPoseIsDirty.GetData(), LocalPoseIsDirty.GetAllocatedSize());
+	}
+	if(!bLocalIsPrimary)
+	{
+		for(bool& Flag : LocalPoseIsDirty)
+		{
+			Flag = true;
+		}
+	}
+
+	if(GlobalPoseIsDirty.Num() != GlobalPose.Num())
+	{
+		GlobalPoseIsDirty.Reset();
+		GlobalPoseIsDirty.AddZeroed(GlobalPose.Num());
+	}
+	else if(!bLocalIsPrimary)
+	{
+		FMemory::Memzero(GlobalPoseIsDirty.GetData(), GlobalPoseIsDirty.GetAllocatedSize());
+	}
+	if(bLocalIsPrimary)
+	{
+		for(bool& Flag : GlobalPoseIsDirty)
+		{
+			Flag = true;
+		}
+	}
+}
+
+void FAnimNode_ControlRig_PoseAdapter::ComputeDependentTransforms()
+{
+	if(URigHierarchy* Hierarchy = GetHierarchy())
+	{
+		// ensure to compute all local transforms
+		for(const FDependentTransform& Dependent : Dependents)
+		{
+			if(FRigTransformElement* TransformElement = Hierarchy->Get<FRigTransformElement>(Dependent.KeyAndIndex.Index))
+			{
+				switch(Dependent.StorageType)
+				{
+					case ERigTransformStorageType::Pose:
+					{
+						(void)Hierarchy->GetTransform(TransformElement, MakeLocal(Dependent.TransformType));
+						break;
+					}
+					case ERigTransformStorageType::Offset:
+					{
+						FRigControlElement* ControlElement = CastChecked<FRigControlElement>(TransformElement);
+						(void)Hierarchy->GetControlOffsetTransform(ControlElement, MakeLocal(Dependent.TransformType));
+						break;
+					}
+					case ERigTransformStorageType::Shape:
+					{
+						FRigControlElement* ControlElement = CastChecked<FRigControlElement>(TransformElement);
+						(void)Hierarchy->GetControlShapeTransform(ControlElement, MakeLocal(Dependent.TransformType));
+						break;
+					}
+					default:
+					{
+						break;
+					}
+				}
+				check(Dependent.DirtyState->Local.Get() == false);
+			}
+		}
+	}
+}
+
+void FAnimNode_ControlRig_PoseAdapter::MarkDependentsDirty()
+{
+	ComputeDependentTransforms();
+
+	if(URigHierarchy* Hierarchy = GetHierarchy())
+	{
+		// mark the global dependent as dirty
+		for(const FDependentTransform& Dependent : Dependents)
+		{
+			check(Dependent.DirtyState->Local.Get() == false);
+			Dependent.DirtyState->Global.Set(true);
+		}
+	}
+}
+
 FAnimNode_ControlRigBase::FAnimNode_ControlRigBase()
 	: FAnimNode_CustomProperty()
 	, bResetInputPoseToInitial(true) 
 	, bTransferInputPose(true)
 	, bTransferInputCurves(true)
-	, bTransferPoseInGlobalSpace(true)
+	, bTransferPoseInGlobalSpace(false)
 	, InputSettings(FControlRigIOSettings())
 	, OutputSettings(FControlRigIOSettings())
 	, bExecute(true)
@@ -50,6 +366,7 @@ void FAnimNode_ControlRigBase::OnInitializeAnimInstance(const FAnimInstanceProxy
 	FAnimNode_CustomProperty::OnInitializeAnimInstance(InProxy, InAnimInstance);
 
 	WeakAnimInstanceObject = TWeakObjectPtr<const UAnimInstance>(InAnimInstance);
+	PoseAdapter = MakeShareable(new FAnimNode_ControlRig_PoseAdapter());
 
 	USkeletalMeshComponent* Component = InAnimInstance->GetOwningComponent();
 	UControlRig* ControlRig = GetControlRig();
@@ -110,6 +427,11 @@ bool FAnimNode_ControlRigBase::CanExecute()
 		return false;
 	}
 
+	if(!PoseAdapter.IsValid())
+	{
+		return false;
+	}
+
 	if (UControlRig* ControlRig = GetControlRig())
 	{
 		return ControlRig->CanExecute(); 
@@ -118,11 +440,10 @@ bool FAnimNode_ControlRigBase::CanExecute()
 	return false;
 }
 
-void FAnimNode_ControlRigBase::UpdateInput(UControlRig* ControlRig, const FPoseContext& InOutput)
+void FAnimNode_ControlRigBase::UpdateInput(UControlRig* ControlRig, FPoseContext& InOutput)
 {
 	SCOPE_CYCLE_COUNTER(STAT_ControlRig_UpdateInput);
 
-
 	if(!CanExecute())
 	{
 		return;
@@ -144,108 +465,71 @@ void FAnimNode_ControlRigBase::UpdateInput(UControlRig* ControlRig, const FPoseC
 
 	DECLARE_SCOPE_HIERARCHICAL_COUNTER_FUNC()
 
-	if (InputSettings.bUpdatePose && bTransferInputPose)
+	if (PoseAdapter && InputSettings.bUpdatePose)
 	{
-		const FBoneContainer& RequiredBones = InOutput.Pose.GetBoneContainer();
-
-		// reset transforms here to prevent additive transforms from accumulating to INF
-		// we only update transforms from the mesh pose for bones in the current LOD, 
-		// so the reset here ensures excluded bones are also reset
-		if(!ControlRigBoneInputMappingByName.IsEmpty() || bResetInputPoseToInitial)
-		{
-			FRigHierarchyValidityBracket ValidityBracket(Hierarchy);
-
-			{
-#if WITH_EDITOR
-				// make sure transient controls don't get reset
-				UControlRig::FTransientControlPoseScope PoseScope(ControlRig);
-#endif 
-				Hierarchy->ResetPoseToInitial(ERigElementType::Bone);
-			}
-		}
-
-		if(bTransferPoseInGlobalSpace || NodeMappingContainer.IsValid())
+		if(InputSettings.bUpdatePose)
 		{
-			// get component pose from control rig
-			FCSPose<FCompactPose> MeshPoses;
-			// first I need to convert to local pose
-			MeshPoses.InitPose(InOutput.Pose);
-
-			if(!ControlRigBoneInputMappingByIndex.IsEmpty())
+			// reset transforms here to prevent additive transforms from accumulating to INF
+			// we only update transforms from the mesh pose for bones in the current LOD, 
+			// so the reset here ensures excluded bones are also reset
+			if(!PoseAdapter->BonesToResetToInitial.IsEmpty())
 			{
-				for (const TPair<uint16, uint16>& Pair : ControlRigBoneInputMappingByIndex)
+				FRigHierarchyValidityBracket ValidityBracket(Hierarchy);
 				{
-					const uint16 ControlRigIndex = Pair.Key;
-					const uint16 SkeletonIndex = Pair.Value;
-					
-					FCompactPoseBoneIndex CompactPoseIndex(SkeletonIndex);
-					const FTransform& ComponentTransform = MeshPoses.GetComponentSpaceTransform(CompactPoseIndex);
-					Hierarchy->SetGlobalTransformByIndex(ControlRigIndex, ComponentTransform, false);
-				}
-			}
-			else
-			{
-				for (auto Iter = ControlRigBoneInputMappingByName.CreateConstIterator(); Iter; ++Iter)
-				{
-					const FName& Name = Iter.Key();
-					const uint16 Index = Iter.Value();
-					const FRigElementKey Key(Name, ERigElementType::Bone);
-
-					FCompactPoseBoneIndex CompactPoseIndex(Index);
-
-					const FTransform& ComponentTransform = MeshPoses.GetComponentSpaceTransform(CompactPoseIndex);
-					if (NodeMappingContainer.IsValid())
-					{
-						const FTransform& RelativeTransformReverse = NodeMappingContainer->GetSourceToTargetTransform(Name).GetRelativeTransformReverse(ComponentTransform);
-						Hierarchy->SetGlobalTransform(Key, RelativeTransformReverse, false);
-					}
-					else
+#if WITH_EDITOR
+					// make sure transient controls don't get reset
+					UControlRig::FTransientControlPoseScope PoseScope(ControlRig);
+#endif
+					for(const int32& BoneElementIndex : PoseAdapter->BonesToResetToInitial)
 					{
-						Hierarchy->SetGlobalTransform(Key, ComponentTransform, false);
+						if(FRigTransformElement* BoneElement = Hierarchy->Get<FRigBoneElement>(BoneElementIndex))
+						{
+							const FTransform InitialLocalTransform = Hierarchy->GetTransform(BoneElement, ERigTransformType::InitialLocal);
+							BoneElement->GetTransform().Set(ERigTransformType::CurrentLocal, InitialLocalTransform);
+							BoneElement->GetDirtyState().MarkClean(ERigTransformType::CurrentLocal);
+							BoneElement->GetDirtyState().MarkDirty(ERigTransformType::CurrentGlobal);
+						}
 					}
-					
 				}
 			}
+			
+			if(bTransferInputPose && InOutput.Pose.GetNumBones() == PoseAdapter->LocalPose.Num())
+			{
+				PoseAdapter->MarkDependentsDirty();
+				InOutput.Pose.CopyBonesTo(PoseAdapter->LocalPose);
+				PoseAdapter->UpdateDirtyStates(true);
+			}
 		}
-		else
+		
+		if(InputSettings.bUpdateCurves || OutputSettings.bUpdateCurves)
 		{
-			if(!ControlRigBoneInputMappingByIndex.IsEmpty())
+			Hierarchy->ResetChangedCurveIndices();
+			
+			if(PoseAdapter->PoseCurveToStorageIndex.Num() != InOutput.Curve.Num())
 			{
-				for (const TPair<uint16, uint16>& Pair : ControlRigBoneInputMappingByIndex)
+				PoseAdapter->PoseCurveToStorageIndex.Reset();
+				PoseAdapter->PoseCurveToStorageIndex.Reserve(InOutput.Curve.Num());
+				InOutput.Curve.ForEachElement([this, Hierarchy](const UE::Anim::FCurveElement& InCurveElement)
 				{
-					const uint16 ControlRigIndex = Pair.Key;
-					const uint16 SkeletonIndex = Pair.Value;
-					
-					FCompactPoseBoneIndex CompactPoseIndex(SkeletonIndex);
-					const FTransform& LocalTransform = InOutput.Pose[CompactPoseIndex];
-					Hierarchy->SetLocalTransformByIndex(ControlRigIndex, LocalTransform, false);
-				}
+					const int32 StorageIndex = Hierarchy->GetIndex({InCurveElement.Name, ERigElementType::Curve});
+					PoseAdapter->PoseCurveToStorageIndex.Add(StorageIndex);
+				});
 			}
-			else
+
+			int32 CurveIndex = 0;
+			InOutput.Curve.ForEachElement([this, &CurveIndex, Hierarchy](const UE::Anim::FCurveElement& InCurveElement)
 			{
-				for (auto Iter = ControlRigBoneInputMappingByName.CreateConstIterator(); Iter; ++Iter)
+				const int32& StorageIndex = PoseAdapter->PoseCurveToStorageIndex[CurveIndex];
+				if(StorageIndex != INDEX_NONE)
 				{
-					const FName& Name = Iter.Key();
-					const uint16 SkeletonIndex = Iter.Value();
-					const FRigElementKey Key(Name, ERigElementType::Bone);
-
-					FCompactPoseBoneIndex CompactPoseIndex(SkeletonIndex);
-					const FTransform& LocalTransform = InOutput.Pose[CompactPoseIndex];
-					Hierarchy->SetLocalTransform(Key, LocalTransform, false);
+					if(FRigCurveElement* HierarchyCurve = Hierarchy->Get<FRigCurveElement>(StorageIndex))
+					{
+						HierarchyCurve->Set(InCurveElement.Value, false);
+					}
 				}
-			}
+				CurveIndex++;
+			});
 		}
-		
-	}
-
-	if (InputSettings.bUpdateCurves && bTransferInputCurves)
-	{
-		Hierarchy->UnsetCurveValues();
-		InOutput.Curve.ForEachElement([Hierarchy](const UE::Anim::FCurveElement& InCurveElement)
-		{
-			const FRigElementKey Key(InCurveElement.Name, ERigElementType::Curve);
-			Hierarchy->SetCurveValue(Key, InCurveElement.Value);
-		});
 	}
 
 #if WITH_EDITOR
@@ -275,105 +559,100 @@ void FAnimNode_ControlRigBase::UpdateOutput(UControlRig* ControlRig, FPoseContex
 
 	if (OutputSettings.bUpdatePose)
 	{
-		// copy output of the rig
-		const FBoneContainer& RequiredBones = InOutput.Pose.GetBoneContainer();
-
-		TMap<FName, uint16>& NameBasedMapping = ControlRigBoneOutputMappingByName;
-		TArray<TPair<uint16, uint16>>& IndexBasedMapping = ControlRigBoneOutputMappingByIndex;
-
-		// if we don't have a different mapping for outputs, use the input mapping
-		if(NameBasedMapping.IsEmpty() && IndexBasedMapping.IsEmpty())
-		{
-			NameBasedMapping = ControlRigBoneInputMappingByName;
-			IndexBasedMapping = ControlRigBoneInputMappingByIndex;
-		}
-
-		if(bTransferPoseInGlobalSpace || NodeMappingContainer.IsValid())
+		if(PoseAdapter)
 		{
-			// get component pose from control rig
-			FCSPose<FCompactPose> MeshPoses;
-			MeshPoses.InitPose(InOutput.Pose);
+			// make sure the local respectively global transforms are all ready
+			if(PoseAdapter->bTransferInLocalSpace)
+			{
+				PoseAdapter->ConvertToLocalPose();
+			}
+			else
+			{
+				PoseAdapter->ConvertToGlobalPose();
+			}
 
-			if(!IndexBasedMapping.IsEmpty())
+			// if we transfered in global - let's convert back to local
+			if(!PoseAdapter->bTransferInLocalSpace)
 			{
-				for (const TPair<uint16, uint16>& Pair : IndexBasedMapping)
+				if(PoseAdapter->GlobalPose.Num() == InOutput.Pose.GetNumBones())
 				{
-					const uint16 ControlRigIndex = Pair.Key;
-					const uint16 SkeletonIndex = Pair.Value;
-
-					FCompactPoseBoneIndex CompactPoseIndex(SkeletonIndex);
-					FTransform ComponentTransform = Hierarchy->GetGlobalTransformByIndex(ControlRigIndex);
-					MeshPoses.SetComponentSpaceTransform(CompactPoseIndex, ComponentTransform);
+					PoseAdapter->ConvertToLocalPose();
 				}
 			}
-			else
+			
+			InOutput.Pose.CopyBonesFrom(PoseAdapter->LocalPose);
+		}
+	}
+
+	if(OutputSettings.bUpdateCurves)
+	{
+		const TArray<int32>& ChangedCurveIndices = Hierarchy->GetChangedCurveIndices();
+		if(ChangedCurveIndices.Num() > 0)
+		{
+			const TArray<FRigBaseElement*> HierarchyCurves = Hierarchy->GetCurvesFast();
+			PoseAdapter->HierarchyCurveCopied.SetNumUninitialized(HierarchyCurves.Num());
+			FMemory::Memzero(PoseAdapter->HierarchyCurveCopied.GetData(), PoseAdapter->HierarchyCurveCopied.GetAllocatedSize());
+
+			int32 CurveIndex = 0;
+			int32 CurvesCopied = 0;
+			InOutput.Curve.ForEachElement([this, &CurveIndex, &CurvesCopied, &HierarchyCurves](const UE::Anim::FCurveElement& InCurveElement)
 			{
-				for (auto Iter = NameBasedMapping.CreateConstIterator(); Iter; ++Iter)
+				const int32& StorageIndex = PoseAdapter->PoseCurveToStorageIndex[CurveIndex];
+				if(StorageIndex != INDEX_NONE)
 				{
-					const FName& Name = Iter.Key();
-					const uint16 SkeletonIndex = Iter.Value();
-					const FRigElementKey Key(Name, ERigElementType::Bone);
-
-					FCompactPoseBoneIndex CompactPoseIndex(SkeletonIndex);
-					FTransform ComponentTransform = Hierarchy->GetGlobalTransform(Key);
-					if (NodeMappingContainer.IsValid())
+					const FRigCurveElement* HierarchyCurve = CastChecked<FRigCurveElement>(HierarchyCurves[StorageIndex]);
+					if(HierarchyCurve->IsValueSet())
 					{
-						ComponentTransform = NodeMappingContainer->GetSourceToTargetTransform(Name) * ComponentTransform;
+						const_cast<UE::Anim::FCurveElement*>(&InCurveElement)->Value = HierarchyCurve->Get();
+					}
+					if(!PoseAdapter->HierarchyCurveCopied[HierarchyCurve->GetSubIndex()])
+					{
+						PoseAdapter->HierarchyCurveCopied[HierarchyCurve->GetSubIndex()] = true;
+						CurvesCopied++;
 					}
-
-					MeshPoses.SetComponentSpaceTransform(CompactPoseIndex, ComponentTransform);
 				}
-			}
+				CurveIndex++;
+			});
 
-			FCSPose<FCompactPose>::ConvertComponentPosesToLocalPosesSafe(MeshPoses, InOutput.Pose);
-			InOutput.Pose.NormalizeRotations();
-		}
-		else
-		{
-			if(!IndexBasedMapping.IsEmpty())
+			if(CurvesCopied < HierarchyCurves.Num())
 			{
-				for (const TPair<uint16, uint16>& Pair : IndexBasedMapping)
+				FBlendedCurve ControlRigCurves;
+				ControlRigCurves.Reserve(ChangedCurveIndices.Num());
+				for(const int32& ChangedCurveIndex : ChangedCurveIndices)
 				{
-					const uint16 ControlRigIndex = Pair.Key;
-					const uint16 SkeletonIndex = Pair.Value;
+					if(!Hierarchy->IsValidIndex(ChangedCurveIndex))
+					{
+						continue;
+					}
+					const FRigCurveElement* HierarchyCurve = Hierarchy->Get<FRigCurveElement>(ChangedCurveIndex);
+					if(PoseAdapter->HierarchyCurveCopied[HierarchyCurve->GetSubIndex()])
+					{
+						continue;
+					}
+					if(HierarchyCurve->IsValueSet())
+					{
+						ControlRigCurves.Add(HierarchyCurve->GetFName(), HierarchyCurve->Get());
+						PoseAdapter->HierarchyCurveCopied[HierarchyCurve->GetSubIndex()] = true;
+					}
+				}
 
-					FCompactPoseBoneIndex CompactPoseIndex(SkeletonIndex);
-					FTransform LocalTransform = Hierarchy->GetLocalTransformByIndex(ControlRigIndex);
-					InOutput.Pose[CompactPoseIndex] = LocalTransform;
+				if(InOutput.Curve.Num() > ControlRigCurves.Num())
+				{
+					InOutput.Curve.Combine(ControlRigCurves);
 				}
-			}
-			else
-			{
-				for (auto Iter = NameBasedMapping.CreateConstIterator(); Iter; ++Iter)
+				else if(InOutput.Curve.Num() == 0)
 				{
-					const FName& Name = Iter.Key();
-					const uint16 Index = Iter.Value();
-					const FRigElementKey Key(Name, ERigElementType::Bone);
-
-					FCompactPoseBoneIndex CompactPoseIndex(Index);
-					FTransform LocalTransform = Hierarchy->GetLocalTransform(Key);
-					InOutput.Pose[CompactPoseIndex] = LocalTransform;
+					InOutput.Curve = ControlRigCurves;
+				}
+				else
+				{
+					Swap(InOutput.Curve, ControlRigCurves);
+					InOutput.Curve.Combine(ControlRigCurves);
 				}
 			}
 		}
 	}
 
-	if (OutputSettings.bUpdateCurves)
-	{
-		FBlendedCurve ControlRigCurves;
-		ControlRigCurves.Reserve(Hierarchy->Num(ERigElementType::Curve));
-		Hierarchy->ForEach<FRigCurveElement>([&ControlRigCurves](const FRigCurveElement* InElement)
-		{
-			if(InElement->IsValueSet())
-			{
-				ControlRigCurves.Add(InElement->GetFName(), InElement->Get());
-			}
-			return true;
-		});
-
-		InOutput.Curve.Combine(ControlRigCurves);
-	}
-
 #if WITH_EDITOR
 	if(bExecute && Hierarchy->IsTracingChanges())
 	{
@@ -387,6 +666,11 @@ void FAnimNode_ControlRigBase::Evaluate_AnyThread(FPoseContext& Output)
 {
 	DECLARE_SCOPE_HIERARCHICAL_COUNTER_FUNC()
 
+	if(!PoseAdapter.IsValid())
+	{
+		return;
+	}
+
 	FPoseContext SourcePose(Output);
 
 	if (Source.GetLinkNode())
@@ -446,7 +730,21 @@ void FAnimNode_ControlRigBase::ExecuteControlRig(FPoseContext& InOutput)
 		// temporarily give control rig access to the stack allocated attribute container
 		// control rig may have rig units that can add/get attributes to/from this container
 		UControlRig::FAnimAttributeContainerPtrScope AttributeScope(ControlRig, InOutput.CustomAttributes);
-		
+
+		URigHierarchy* Hierarchy = ControlRig->GetHierarchy();
+		if(Hierarchy == nullptr)
+		{
+			return;
+		}
+
+		const uint32 TopologyVersion = Hierarchy->GetTopologyVersion();
+		if(PoseAdapter && PoseAdapter->LastTopologyVersion != TopologyVersion)
+		{
+			Hierarchy->UnlinkPoseAdapter();
+			UpdateInputOutputMappingIfRequired(ControlRig, InOutput.Pose.GetBoneContainer());
+			Hierarchy->LinkPoseAdapter(PoseAdapter);
+		}
+
 		// first update input to the system
 		UpdateInput(ControlRig, InOutput);
 		
@@ -455,12 +753,6 @@ void FAnimNode_ControlRigBase::ExecuteControlRig(FPoseContext& InOutput)
 			TGuardValue<bool> ResetCurrentTransfromsAfterConstructionGuard(ControlRig->bResetCurrentTransformsAfterConstruction, true);
 			
 #if WITH_EDITOR
-			URigHierarchy* Hierarchy = ControlRig->GetHierarchy();
-			if(Hierarchy == nullptr)
-			{
-				return;
-			}
-
 			if(Hierarchy->IsTracingChanges())
 			{
 				Hierarchy->StorePoseForTrace(TEXT("FAnimNode_ControlRigBase::BeforeEvaluate"));
@@ -519,38 +811,6 @@ void FAnimNode_ControlRigBase::ExecuteControlRig(FPoseContext& InOutput)
 	}
 }
 
-struct FControlRigControlScope
-{
-	FControlRigControlScope(UControlRig* InControlRig)
-		: ControlRig(InControlRig)
-	{
-		if (ControlRig.IsValid())
-		{
-			URigHierarchy* Hierarchy = ControlRig->GetHierarchy();
-			Hierarchy->ForEach<FRigControlElement>([this, Hierarchy](FRigControlElement* ControlElement) -> bool
-			{
-				ControlValues.Add(ControlElement->GetKey(), Hierarchy->GetControlValueByIndex(ControlElement->GetIndex()));
-				return true; // continue
-			});
-		}
-	}
-
-	~FControlRigControlScope()
-	{
-		if (ControlRig.IsValid())
-		{
-			URigHierarchy* Hierarchy = ControlRig->GetHierarchy();
-			for (const TPair<FRigElementKey, FRigControlValue>& Pair: ControlValues)
-			{
-				Hierarchy->SetControlValue(Pair.Key, Pair.Value);
-			}
-		}
-	}
-
-	TMap<FRigElementKey, FRigControlValue> ControlValues;
-	TWeakObjectPtr<UControlRig> ControlRig;
-};
-
 void FAnimNode_ControlRigBase::CacheBones_AnyThread(const FAnimationCacheBonesContext& Context)
 {
 	DECLARE_SCOPE_HIERARCHICAL_COUNTER_FUNC()
@@ -593,30 +853,69 @@ void FAnimNode_ControlRigBase::CacheBones_AnyThread(const FAnimationCacheBonesCo
 		}
 		
 		LastBonesSerialNumberForCacheBones = BonesSerialNumber;
+
+		if(URigHierarchy* Hierarchy = ControlRig->GetHierarchy())
+		{
+			Hierarchy->LinkPoseAdapter(PoseAdapter);
+		}
 	}
 }
 
 void FAnimNode_ControlRigBase::UpdateInputOutputMappingIfRequired(UControlRig* InControlRig, const FBoneContainer& InRequiredBones)
 {
-	const URigHierarchy* Hierarchy = InControlRig->GetHierarchy();
+	URigHierarchy* Hierarchy = InControlRig->GetHierarchy();
 	if(Hierarchy == nullptr)
 	{
 		return;
 	}
+	Hierarchy->UnlinkPoseAdapter();
 
-	ControlRigBoneInputMappingByIndex.Reset();
-	ControlRigBoneOutputMappingByIndex.Reset();
-	ControlRigCurveMappingByIndex.Reset();
-	ControlRigBoneInputMappingByName.Reset();
-	ControlRigBoneOutputMappingByName.Reset();
-	ControlRigCurveMappingByName.Reset();
+	if(!PoseAdapter.IsValid())
+	{
+		return;
+	}
 
+	PoseAdapter->ParentPoseIndices.Reset();
+	PoseAdapter->RequiresHierarchyForSpaceConversion.Reset();
+	PoseAdapter->ElementIndexToPoseIndex.Reset();
+
+	const int32 NumBonesInPose = InRequiredBones.GetCompactPoseNumBones();
+
+	PoseAdapter->GlobalPose.Reset();
+	PoseAdapter->LocalPose.Reset();
+	PoseAdapter->ParentPoseIndices.Reserve(NumBonesInPose);
+	PoseAdapter->RequiresHierarchyForSpaceConversion.Reserve(NumBonesInPose);
+	PoseAdapter->GlobalPose.AddDefaulted(NumBonesInPose);
+	PoseAdapter->LocalPose.AddDefaulted(NumBonesInPose);
+	PoseAdapter->PoseCurveToStorageIndex.Reset();
+	PoseAdapter->HierarchyCurveCopied.Reset();
+
+	PoseAdapter->bTransferInLocalSpace = !(bTransferPoseInGlobalSpace || NodeMappingContainer.IsValid());
+
+	for(int32 Index = 0; Index < NumBonesInPose; Index++)
+	{
+		PoseAdapter->ParentPoseIndices.Add(INDEX_NONE);
+
+		const FCompactPoseBoneIndex ParentBoneIndex = InRequiredBones.GetParentBoneIndex(FCompactPoseBoneIndex(Index));
+		if(ParentBoneIndex.IsValid())
+		{
+			PoseAdapter->ParentPoseIndices[Index] = ParentBoneIndex.GetInt();
+		}
+		PoseAdapter->RequiresHierarchyForSpaceConversion.Add(false);
+	}
+	PoseAdapter->UpdateDirtyStates();
+
+	TArray<int32> MappedBoneElementIndices;
 	if(InRequiredBones.IsValid())
 	{
-		const TArray<FBoneIndexType>& RequiredBonesArray = InRequiredBones.GetBoneIndicesArray();
-		const int32 NumBones = RequiredBonesArray.Num();
+		PoseAdapter->ElementIndexToPoseIndex.Reserve(NumBonesInPose);
+		PoseAdapter->PoseIndexToElementIndex.Reserve(NumBonesInPose);
 
-		const FReferenceSkeleton& RefSkeleton = InRequiredBones.GetReferenceSkeleton();
+		const FReferenceSkeleton* RefSkeleton = &InRequiredBones.GetReferenceSkeleton();
+		if(const USkeleton* Skeleton = InRequiredBones.GetSkeletonAsset())
+		{
+			RefSkeleton = &Skeleton->GetReferenceSkeleton();
+		}			
 
 		// @todo: thread-safe? probably not in editor, but it may not be a big issue in editor
 		if (NodeMappingContainer.IsValid())
@@ -626,15 +925,29 @@ void FAnimNode_ControlRigBase::UpdateInputOutputMappingIfRequired(UControlRig* I
 			NodeMappingContainer->GetTargetToSourceMappingTable(TargetToSourceMappingTable);
 
 			// now fill up node name
-			for (uint16 Index = 0; Index < NumBones; ++Index)
+			for (uint16 Index = 0; Index < NumBonesInPose; ++Index)
 			{
 				// get bone name, and find reverse mapping
-				FName TargetNodeName = RefSkeleton.GetBoneName(RequiredBonesArray[Index]);
-				FName* SourceName = TargetToSourceMappingTable.Find(TargetNodeName);
-				if (SourceName)
+				const FSkeletonPoseBoneIndex BoneIndex = InRequiredBones.GetSkeletonPoseIndexFromCompactPoseIndex(FCompactPoseBoneIndex(Index));
+				if(BoneIndex.IsValid())
 				{
-					ControlRigBoneInputMappingByName.Add(*SourceName, Index);
+					FName TargetNodeName = RefSkeleton->GetBoneName(BoneIndex.GetInt());
+					FName* SourceName = TargetToSourceMappingTable.Find(TargetNodeName);
+					if (SourceName)
+					{
+						const int32 ElementIndex = Hierarchy->GetIndex({*SourceName, ERigElementType::Bone});
+						if(ElementIndex != INDEX_NONE)
+						{
+							MappedBoneElementIndices.Add(ElementIndex);
+							PoseAdapter->ElementIndexToPoseIndex.Add(static_cast<uint16>(ElementIndex), Index);
+							PoseAdapter->PoseIndexToElementIndex.Add(ElementIndex);
+							PoseAdapter->LocalPose[Index] = Hierarchy->GetLocalTransform(ElementIndex);
+							PoseAdapter->GlobalPose[Index] = Hierarchy->GetGlobalTransform(ElementIndex);
+							continue;
+						}
+					}
 				}
+				PoseAdapter->PoseIndexToElementIndex.Add(INDEX_NONE);
 			}
 		}
 		else
@@ -644,132 +957,90 @@ void FAnimNode_ControlRigBase::UpdateInputOutputMappingIfRequired(UControlRig* I
 			InControlRig->GetMappableNodeData(NodeNames, NodeItems);
 
 			// even if not mapped, we map only node that exists in the controlrig
-			for (uint16 Index = 0; Index < NumBones; ++Index)
+			for (uint16 Index = 0; Index < NumBonesInPose; ++Index)
 			{
-				const FName& BoneName = RefSkeleton.GetBoneName(RequiredBonesArray[Index]);
-				if (NodeNames.Contains(BoneName))
+				const FSkeletonPoseBoneIndex BoneIndex = InRequiredBones.GetSkeletonPoseIndexFromCompactPoseIndex(FCompactPoseBoneIndex(Index));
+				if(BoneIndex.IsValid())
 				{
-					ControlRigBoneInputMappingByName.Add(BoneName, Index);
+					const FName& BoneName = RefSkeleton->GetBoneName(BoneIndex.GetInt());
+					if (NodeNames.Contains(BoneName))
+					{
+						const int32 ElementIndex = Hierarchy->GetIndex({BoneName, ERigElementType::Bone});
+						if(ElementIndex != INDEX_NONE)
+						{
+							MappedBoneElementIndices.Add(ElementIndex);
+							PoseAdapter->ElementIndexToPoseIndex.Add(static_cast<uint16>(ElementIndex), Index);
+							PoseAdapter->PoseIndexToElementIndex.Add(ElementIndex);
+							PoseAdapter->LocalPose[Index] = Hierarchy->GetLocalTransform(ElementIndex);
+							PoseAdapter->GlobalPose[Index] = Hierarchy->GetGlobalTransform(ElementIndex);
+							continue;
+						}
+					}
 				}
+				PoseAdapter->PoseIndexToElementIndex.Add(INDEX_NONE);
 			}
 		}
 
-		auto UpdatingMappingFromSpecificTransferList = [] (
-			TArray<FBoneReference>& InTransferList,
-			const TWeakObjectPtr<UNodeMappingContainer>& InMappingContainer,
-			const FBoneContainer& InRequiredBones,
-			const FReferenceSkeleton& InRefSkeleton,
-			const TArray<FBoneIndexType>& InRequiredBonesArray,
-			const UControlRig* InControlRig,
-			TMap<FName, uint16>& OutMapping
-		) {
-			OutMapping.Reset();
-			
-			if (InMappingContainer.IsValid())
+		// once we know all of the bones we are going to transfer - we can check if any of these bones has a different parenting
+		// relationship in the skeleton used in the anim graph vs the hierarchy in the rig. in that case we have to transfer in global
+		if(PoseAdapter->bTransferInLocalSpace)
+		{
+			for(const int32& BoneElementIndex : MappedBoneElementIndices)
 			{
-				// get target to source mapping table - this is reversed mapping table
-				TMap<FName, FName> TargetToSourceMappingTable;
-				InMappingContainer->GetTargetToSourceMappingTable(TargetToSourceMappingTable);
+				const int32 HierarchyParentIndex = Hierarchy->GetFirstParent(BoneElementIndex);
+				const int16 PoseIndex = PoseAdapter->ElementIndexToPoseIndex.FindChecked(BoneElementIndex); 
+				const FCompactPoseBoneIndex CompactPoseParentIndex(PoseAdapter->ParentPoseIndices[PoseIndex]);
 
-				for(FBoneReference& InputBoneToTransfer : InTransferList)
+				FName HierarchyParentName(NAME_None);
+				FName PoseParentName(NAME_None);
+				
+				if(HierarchyParentIndex != INDEX_NONE)
 				{
-					if(!InputBoneToTransfer.Initialize(InRequiredBones))
-					{
-						continue;
-					}
-					const FName TargetNodeName = InRefSkeleton.GetBoneName(InputBoneToTransfer.BoneIndex);
-					if (const FName* SourceName = TargetToSourceMappingTable.Find(TargetNodeName))
-					{
-						OutMapping.Add(*SourceName, InputBoneToTransfer.BoneIndex);
-					}
+					HierarchyParentName = Hierarchy->Get(HierarchyParentIndex)->GetFName();
+				}
+				if(CompactPoseParentIndex.IsValid())
+				{
+					const FSkeletonPoseBoneIndex SkeletonIndex = InRequiredBones.GetSkeletonPoseIndexFromCompactPoseIndex(CompactPoseParentIndex);
+					PoseParentName = RefSkeleton->GetBoneName(SkeletonIndex.GetInt());
 				}
-			}
-			else
-			{
-				TArray<FName> NodeNames;
-				TArray<FNodeItem> NodeItems;
-				InControlRig->GetMappableNodeData(NodeNames, NodeItems);
 
-				for(FBoneReference& InputBoneToTransfer : InTransferList)
+				if(HierarchyParentName.IsEqual(PoseParentName, ENameCase::CaseSensitive))
 				{
-					if(!InputBoneToTransfer.Initialize(InRequiredBones))
-					{
-						continue;
-					}
-					if (InRequiredBonesArray.IsValidIndex(InputBoneToTransfer.BoneIndex))
-					{
-						const FName& BoneName = InRefSkeleton.GetBoneName(InRequiredBonesArray[InputBoneToTransfer.BoneIndex]);
-						if (NodeNames.Contains(BoneName))
-						{
-							OutMapping.Add(BoneName, InputBoneToTransfer.BoneIndex);
-						}
-					}
+					continue;
 				}
+				
+				PoseAdapter->RequiresHierarchyForSpaceConversion[PoseIndex] = true;
+				check(PoseAdapter->PoseIndexToElementIndex[PoseIndex] != INDEX_NONE);
+				PoseAdapter->bTransferInLocalSpace = false;
 			}
-		};
-		
-		if(!InputBonesToTransfer.IsEmpty())
-		{
-			ControlRigBoneOutputMappingByName = ControlRigBoneInputMappingByName;
-
-			UpdatingMappingFromSpecificTransferList(
-				InputBonesToTransfer,
-				NodeMappingContainer,
-				InRequiredBones,
-				RefSkeleton,
-				RequiredBonesArray,
-				InControlRig,
-				ControlRigBoneInputMappingByName);
 		}
 
-		if(!OutputBonesToTransfer.IsEmpty())
-		{
-			UpdatingMappingFromSpecificTransferList(
-				OutputBonesToTransfer,
-				NodeMappingContainer,
-				InRequiredBones,
-				RefSkeleton,
-				RequiredBonesArray,
-				InControlRig,
-				ControlRigBoneOutputMappingByName);
-		}
+		// only reset the full pose if we are not mapping all bones
+		const TArray<FRigBaseElement*>& HierarchyBones = Hierarchy->GetBonesFast();
+		const bool bMapsAllBones = MappedBoneElementIndices.Num() == HierarchyBones.Num();
+		PoseAdapter->BonesToResetToInitial.Reset();
+		PoseAdapter->bRequiresResetPoseToInitial = bResetInputPoseToInitial && !bMapsAllBones;
 
-		// check if we can switch the bones to an index based mapping.
-		// we can only do that if there is no node mapping container set.
-		if(!NodeMappingContainer.IsValid())
+		if(PoseAdapter->bRequiresResetPoseToInitial)
 		{
-			for(int32 InputOutput = 0; InputOutput < 2; InputOutput++)
+			PoseAdapter->BonesToResetToInitial.Reserve(HierarchyBones.Num() - MappedBoneElementIndices.Num());
+
+			// bone is mapped stores sub indices (bone index within the list of bones)
+			TArray<bool> BoneIsMapped;
+			BoneIsMapped.AddZeroed(HierarchyBones.Num());
+			for(const int32& MappedTransformIndex : MappedBoneElementIndices)
 			{
-				bool bIsMappingByIndex = true;
-				TMap<FName, uint16>& NameBasedMapping = InputOutput == 0 ? ControlRigBoneInputMappingByName : ControlRigBoneOutputMappingByName;
-				if(NameBasedMapping.IsEmpty())
-				{
-					continue;
-				}
-				
-				TArray<TPair<uint16, uint16>>& IndexBasedMapping = InputOutput == 0 ? ControlRigBoneInputMappingByIndex : ControlRigBoneOutputMappingByIndex;
-				
-				for (auto Iter = NameBasedMapping.CreateConstIterator(); Iter; ++Iter)
-				{
-					const uint16 SkeletonIndex = Iter.Value();
-					const int32 ControlRigIndex = Hierarchy->GetIndex(FRigElementKey(Iter.Key(), ERigElementType::Bone));
-					if(ControlRigIndex != INDEX_NONE)
-					{
-						IndexBasedMapping.Add(TPair<uint16, uint16>((uint16)ControlRigIndex, SkeletonIndex));
-					}
-					else
-					{
-						bIsMappingByIndex = false;
-					}
-				}
+				const FRigBaseElement* MappedElement = Hierarchy->Get(MappedTransformIndex);
+				check(MappedElement);
+				BoneIsMapped[MappedElement->GetSubIndex()] = true;
+			}
 
-				if(bIsMappingByIndex)
-				{
-					NameBasedMapping.Reset();
-				}
-				else
+			// when we want to know which bones to reset we want to convert back to a global index
+			for(int32 UnmappedBoneIndex = 0; UnmappedBoneIndex < BoneIsMapped.Num(); UnmappedBoneIndex++)
+			{
+				if(!BoneIsMapped[UnmappedBoneIndex])
 				{
-					IndexBasedMapping.Reset();
+					PoseAdapter->BonesToResetToInitial.Add(HierarchyBones[UnmappedBoneIndex]->GetIndex());
 				}
 			}
 		}
@@ -837,6 +1108,10 @@ void FAnimNode_ControlRigBase::QueueControlRigDrawInstructions(UControlRig* Cont
 					// TODO: Add support for this if anyone is actually using it. Currently it is only defined and referenced in an unused API, DrawCone in Control Rig.
 					break;
 				}
+				default:
+				{
+					break;
+				}
 			}
 		}
 	}
diff --git a/Engine/Plugins/Animation/ControlRig/Source/ControlRig/Private/ControlRig.cpp b/Engine/Plugins/Animation/ControlRig/Source/ControlRig/Private/ControlRig.cpp
--- a/Engine/Plugins/Animation/ControlRig/Source/ControlRig/Private/ControlRig.cpp
+++ b/Engine/Plugins/Animation/ControlRig/Source/ControlRig/Private/ControlRig.cpp
@@ -936,7 +936,9 @@ bool UControlRig::Execute(const FName& InEventName)
 #if WITH_EDITOR
 	TSharedPtr<TGuardValue<bool>> RecordTransformsPerInstructionGuard;
 #endif
-	if(URigHierarchy* Hierarchy = GetHierarchy())
+
+	URigHierarchy* Hierarchy = GetHierarchy();
+	if(Hierarchy)
 	{
 		Hierarchy->UpdateReferences(&PublicContext);
 
@@ -975,8 +977,10 @@ bool UControlRig::Execute(const FName& InEventName)
 	// we'll special case the construction event here
 	if (bIsConstructionEvent)
 	{
+		check(Hierarchy);
+		
 		// remember the previous selection
-		const TArray<FRigElementKey> PreviousSelection = GetHierarchy()->GetSelectedKeys();
+		const TArray<FRigElementKey> PreviousSelection = Hierarchy->GetSelectedKeys();
 
 		// construction mode means that we are running the construction event
 		// constantly for testing purposes.
@@ -998,36 +1002,37 @@ bool UControlRig::Execute(const FName& InEventName)
 #endif
 
 			// disable selection notifications from the hierarchy
-			TGuardValue<bool> DisableSelectionNotifications(GetHierarchy()->GetController(true)->bSuspendSelectionNotifications, true);
+			TGuardValue<bool> DisableSelectionNotifications(Hierarchy->GetController(true)->bSuspendSelectionNotifications, true);
 			{
 				FRigPose CurrentPose;
 				// We might want to reset the input pose after construction
 				if (bResetCurrentTransformsAfterConstruction)
 				{
-					CurrentPose = GetHierarchy()->GetPose(false, ERigElementType::ToResetAfterConstructionEvent, FRigElementKeyCollection());
+					CurrentPose = Hierarchy->GetPose(false, ERigElementType::ToResetAfterConstructionEvent, FRigElementKeyCollection());
 				}
 				
 				{
 					// Copy the hierarchy from the default object onto this one
 #if WITH_EDITOR
-					FTransientControlScope TransientControlScope(GetHierarchy());
+					FTransientControlScope TransientControlScope(Hierarchy);
 	#endif
 					{
 						// maintain the initial pose if it ever was set by the client
 						FRigPose InitialPose;
 						if(!bResetInitialTransformsBeforeConstruction)
 						{
-							InitialPose = GetHierarchy()->GetPose(true, ERigElementType::ToResetAfterConstructionEvent, FRigElementKeyCollection());
+							InitialPose = Hierarchy->GetPose(true, ERigElementType::ToResetAfterConstructionEvent, FRigElementKeyCollection());
 						}
 
 						if(bCopyHierarchyBeforeConstruction)
 						{
-							GetHierarchy()->ResetToDefault();
+							Hierarchy->ResetToDefault();
 						}
 
 						if(InitialPose.Num() > 0)
 						{
-							GetHierarchy()->SetPose(InitialPose, ERigTransformType::InitialLocal);
+							const TGuardValue<bool> DisableRecordingCurveChanges(Hierarchy->GetRecordCurveChangesFlag(), false);
+							Hierarchy->SetPose(InitialPose, ERigTransformType::InitialLocal);
 						}
 					}
 
@@ -1041,7 +1046,7 @@ bool UControlRig::Execute(const FName& InEventName)
 						}
 	#endif
 						// reset the pose to initial such that construction event can run from a deterministic initial state
-						GetHierarchy()->ResetPoseToInitial(ERigElementType::All);
+						Hierarchy->ResetPoseToInitial(ERigElementType::All);
 					}
 
 					RestoreShapeLibrariesFromCDO();
@@ -1065,13 +1070,14 @@ bool UControlRig::Execute(const FName& InEventName)
 				// Reset the input pose after construction
 				if (CurrentPose.Num() > 0)
 				{
-					GetHierarchy()->SetPose(CurrentPose, ERigTransformType::CurrentLocal);
+					const TGuardValue<bool> DisableRecordingCurveChanges(Hierarchy->GetRecordCurveChangesFlag(), false);
+					Hierarchy->SetPose(CurrentPose, ERigTransformType::CurrentLocal);
 				}
 			}
 			
 			// set it here to reestablish the selection. the notifications
 			// will be eaten since we still have the bSuspend flag on in the controller.
-			GetHierarchy()->GetController()->SetSelection(PreviousSelection);
+			Hierarchy->GetController()->SetSelection(PreviousSelection);
 			
 		} // destroy DisableSelectionNotifications
 
@@ -1091,25 +1097,25 @@ bool UControlRig::Execute(const FName& InEventName)
 				TransientControlPoseScope = MakeUnique<FTransientControlPoseScope>(this);
 			}
 #endif
-			GetHierarchy()->ResetPoseToInitial(ERigElementType::Bone);
+			Hierarchy->ResetPoseToInitial(ERigElementType::Bone);
 		}
 
 		// synchronize the selection now with the new hierarchy after running construction
-		const TArray<const FRigBaseElement*> CurrentSelection = GetHierarchy()->GetSelectedElements();
+		const TArray<const FRigBaseElement*> CurrentSelection = Hierarchy->GetSelectedElements();
 		for(const FRigBaseElement* SelectedElement : CurrentSelection)
 		{
 			if(!PreviousSelection.Contains(SelectedElement->GetKey()))
 			{
-				GetHierarchy()->Notify(ERigHierarchyNotification::ElementSelected, SelectedElement);
+				Hierarchy->Notify(ERigHierarchyNotification::ElementSelected, SelectedElement);
 			}
 		}
 		for(const FRigElementKey& PreviouslySelectedKey : PreviousSelection)
 		{
-			if(const FRigBaseElement* PreviouslySelectedElement = GetHierarchy()->Find(PreviouslySelectedKey))
+			if(const FRigBaseElement* PreviouslySelectedElement = Hierarchy->Find(PreviouslySelectedKey))
 			{
 				if(!CurrentSelection.Contains(PreviouslySelectedElement))
 				{
-					GetHierarchy()->Notify(ERigHierarchyNotification::ElementDeselected, PreviouslySelectedElement);
+					Hierarchy->Notify(ERigHierarchyNotification::ElementDeselected, PreviouslySelectedElement);
 				}
 			}
 		}
@@ -1215,16 +1221,16 @@ bool UControlRig::Execute(const FName& InEventName)
 		EventContext.Event = ERigEvent::CloseUndoBracket;
 		EventContext.SourceEventName = InEventName;
 		EventContext.LocalTime = PublicContext.GetAbsoluteTime();
-		HandleHierarchyEvent(GetHierarchy(), EventContext);
+		HandleHierarchyEvent(Hierarchy, EventContext);
 	}
 
 	if (PublicContext.GetDrawInterface() && PublicContext.GetDrawContainer() && bIsEventLastInQueue) 
 	{
 		PublicContext.GetDrawInterface()->Instructions.Append(PublicContext.GetDrawContainer()->Instructions);
 
-		FRigHierarchyValidityBracket ValidityBracket(GetHierarchy());
+		FRigHierarchyValidityBracket ValidityBracket(Hierarchy);
 		
-		GetHierarchy()->ForEach<FRigControlElement>([this](FRigControlElement* ControlElement) -> bool
+		Hierarchy->ForEach<FRigControlElement>([this, Hierarchy](FRigControlElement* ControlElement) -> bool
 		{
 			const FRigControlSettings& Settings = ControlElement->Settings;
 
@@ -1233,7 +1239,7 @@ bool UControlRig::Execute(const FName& InEventName)
 				Settings.bDrawLimits &&
 				Settings.LimitEnabled.Contains(FRigControlLimitEnabled(true, true)))
 			{
-				FTransform Transform = GetHierarchy()->GetGlobalControlOffsetTransformByIndex(ControlElement->GetIndex());
+				FTransform Transform = Hierarchy->GetGlobalControlOffsetTransformByIndex(ControlElement->GetIndex());
 				FRigVMDrawInstruction Instruction(ERigVMDrawSettings::Lines, Settings.ShapeColor, 0.f, Transform);
 
 				switch (Settings.ControlType)
diff --git a/Engine/Plugins/Animation/ControlRig/Source/ControlRig/Public/AnimNode_ControlRig.h b/Engine/Plugins/Animation/ControlRig/Source/ControlRig/Public/AnimNode_ControlRig.h
--- a/Engine/Plugins/Animation/ControlRig/Source/ControlRig/Public/AnimNode_ControlRig.h
+++ b/Engine/Plugins/Animation/ControlRig/Source/ControlRig/Public/AnimNode_ControlRig.h
@@ -134,7 +134,7 @@ private:
 
 protected:
 	virtual UClass* GetTargetClass() const override;
-	virtual void UpdateInput(UControlRig* InControlRig, const FPoseContext& InOutput) override;
+	virtual void UpdateInput(UControlRig* InControlRig, FPoseContext& InOutput) override;
 	virtual void UpdateOutput(UControlRig* InControlRig, FPoseContext& InOutput) override;
 
 	void SetControlRigClass(TSubclassOf<UControlRig> InControlRigClass);
diff --git a/Engine/Plugins/Animation/ControlRig/Source/ControlRig/Public/AnimNode_ControlRigBase.h b/Engine/Plugins/Animation/ControlRig/Source/ControlRig/Public/AnimNode_ControlRigBase.h
--- a/Engine/Plugins/Animation/ControlRig/Source/ControlRig/Public/AnimNode_ControlRigBase.h
+++ b/Engine/Plugins/Animation/ControlRig/Source/ControlRig/Public/AnimNode_ControlRigBase.h
@@ -4,6 +4,8 @@
 
 #include "Animation/AnimBulkCurves.h"
 #include "Animation/AnimNode_CustomProperty.h"
+#include "Rigs/RigHierarchyPoseAdapter.h"
+#include "Rigs/RigHierarchy.h"
 #include "AnimNode_ControlRigBase.generated.h"
 
 class UNodeMappingContainer;
@@ -52,6 +54,66 @@ struct CONTROLRIG_API FControlRigAnimNodeEventName
 	FName EventName;	
 };
 
+class CONTROLRIG_API FAnimNode_ControlRig_PoseAdapter : public FRigHierarchyPoseAdapter
+{
+public:
+
+	FAnimNode_ControlRig_PoseAdapter()
+	: bTransferInLocalSpace(true)
+	, LastTopologyVersion(UINT32_MAX)
+	, bRequiresResetPoseToInitial(true)
+	, bEnabled(true)
+	{
+	}
+	
+protected:
+
+	virtual void PostLinked(URigHierarchy* InHierarchy) override;
+	virtual void PreUnlinked(URigHierarchy* InHierarchy) override;
+
+	void ConvertToLocalPose();
+	void ConvertToGlobalPose();
+	const FTransform& GetLocalTransform(int32 InIndex);
+	const FTransform& GetGlobalTransform(int32 InIndex);
+	
+	void UpdateDirtyStates(const TOptional<bool> InLocalIsPrimary = TOptional<bool>());
+	void ComputeDependentTransforms();
+	void MarkDependentsDirty();
+	
+	void UnlinkTransformStorage();
+	void UnlinkCurveStorage();
+
+	bool bTransferInLocalSpace;
+
+	TArray<int32> ParentPoseIndices;
+	TArray<bool> RequiresHierarchyForSpaceConversion;
+	TArray<FTransform> LocalPose;
+	TArray<FTransform> GlobalPose;
+	TArray<bool> LocalPoseIsDirty;
+	TArray<bool> GlobalPoseIsDirty;
+	TArray<int32> PoseCurveToStorageIndex;
+	TArray<bool> HierarchyCurveCopied;
+	uint32 LastTopologyVersion;
+	bool bRequiresResetPoseToInitial;
+	TArray<int32> BonesToResetToInitial;
+
+	TMap<uint16, uint16> ElementIndexToPoseIndex;
+	TArray<int32> PoseIndexToElementIndex;
+
+	struct FDependentTransform
+	{
+		FRigElementKeyAndIndex KeyAndIndex;
+		ERigTransformType::Type TransformType;
+		ERigTransformStorageType::Type StorageType;
+		FRigLocalAndGlobalDirtyState* DirtyState;
+	};
+	TArray<FDependentTransform> Dependents;
+
+	bool bEnabled;;
+
+	friend struct FAnimNode_ControlRigBase;
+};
+
 /**
  * Animation node that allows animation ControlRig output to be used in an animation graph
  */
@@ -118,34 +180,20 @@ protected:
 	 * of the input pose transfer phase.
 	 * If this list is empty all bones will be transferred.
 	 */
-	UPROPERTY(EditAnywhere, AdvancedDisplay, Category = Settings)
-	TArray<FBoneReference> InputBonesToTransfer;
+	UPROPERTY()
+	TArray<FBoneReference> InputBonesToTransfer_DEPRECATED;
 
 	/**
 	 * An inclusive list of bones to transfer as part
 	 * of the output pose transfer phase.
 	 * If this list is empty all bones will be transferred.
 	 */
-	UPROPERTY(EditAnywhere, AdvancedDisplay, Category = Settings)
-	TArray<FBoneReference> OutputBonesToTransfer;
+	UPROPERTY()
+	TArray<FBoneReference> OutputBonesToTransfer_DEPRECATED;
 
 	UPROPERTY(EditAnywhere, AdvancedDisplay, Category = Settings)
 	TArray<TObjectPtr<UAssetUserData>> AssetUserData;
 
-	/** Complete mapping from skeleton to control rig bone index */
-	TArray<TPair<uint16, uint16>> ControlRigBoneInputMappingByIndex;
-	TArray<TPair<uint16, uint16>> ControlRigBoneOutputMappingByIndex;
-
-	/** Complete mapping from skeleton to curve name */
-	TArray<TPair<uint16, FName>> ControlRigCurveMappingByIndex;
-
-	/** Rig Hierarchy bone name to required array index mapping */
-	TMap<FName, uint16> ControlRigBoneInputMappingByName;
-	TMap<FName, uint16> ControlRigBoneOutputMappingByName;
-
-	/** Rig Curve name to Curve mapping */
-	TMap<FName, FName> ControlRigCurveMappingByName;
-	
 	TMap<FName, int32> InputToControlIndex;
 	
 	/** Node Mapping Container */
@@ -158,6 +206,8 @@ protected:
 	UPROPERTY(transient)
 	FControlRigIOSettings OutputSettings;
 
+	TSharedPtr<FAnimNode_ControlRig_PoseAdapter> PoseAdapter;
+
 	UPROPERTY(transient)
 	bool bExecute;
 
@@ -172,7 +222,7 @@ protected:
 
 	virtual bool CanExecute();
 	// update input/output to control rig
-	virtual void UpdateInput(UControlRig* ControlRig, const FPoseContext& InOutput);
+	virtual void UpdateInput(UControlRig* ControlRig, FPoseContext& InOutput);
 	virtual void UpdateOutput(UControlRig* ControlRig, FPoseContext& InOutput);
 	virtual UClass* GetTargetClass() const override;
 	
diff --git a/Engine/Plugins/Animation/ControlRig/Source/ControlRig/Public/Rigs/RigHierarchy.h b/Engine/Plugins/Animation/ControlRig/Source/ControlRig/Public/Rigs/RigHierarchy.h
--- a/Engine/Plugins/Animation/ControlRig/Source/ControlRig/Public/Rigs/RigHierarchy.h
+++ b/Engine/Plugins/Animation/ControlRig/Source/ControlRig/Public/Rigs/RigHierarchy.h
@@ -886,6 +886,14 @@ public:
 	{
 		return GetElementsOfType<FRigBoneElement>(bTraverse);
 	}
+	
+	/**
+   	 * Returns all Bone elements without traversing the hierarchy
+   	 */
+   	TArray<FRigBaseElement*>& GetBonesFast() const
+   	{
+   		return ElementsPerType[RigElementTypeToFlatIndex(ERigElementType::Bone)];
+   	}
 
 	/**
 	 * Returns all Bone elements
