Control Rig: fix tooltip for post construction node

#rnx

diff --git a/Engine/Plugins/Animation/ControlRig/Source/ControlRig/Public/Units/Execution/RigUnit_PrepareForExecution.h b/Engine/Plugins/Animation/ControlRig/Source/ControlRig/Public/Units/Execution/RigUnit_PrepareForExecution.h
--- a/Engine/Plugins/Animation/ControlRig/Source/ControlRig/Public/Units/Execution/RigUnit_PrepareForExecution.h
+++ b/Engine/Plugins/Animation/ControlRig/Source/ControlRig/Public/Units/Execution/RigUnit_PrepareForExecution.h
@@ -28,7 +28,7 @@ struct CONTROLRIG_API FRigUnit_PrepareForExecution : public FRigUnit
 };
 
 /**
- * Event to create / configure elements before any other event
+ * Event to further configure elements. Runs after the Construction Event
  */
 USTRUCT(meta=(DisplayName="Post Construction", Category="Events", NodeColor="0.6, 0, 1", Keywords="Create,Build,Spawn,Setup,Init,Fit"))
 struct CONTROLRIG_API FRigUnit_PostPrepareForExecution : public FRigUnit
