Control Rig: Fix build breakage on ARM

#rnx

diff --git a/Engine/Plugins/Animation/ControlRig/Source/ControlRig/Public/AnimNode_ControlRigBase.h b/Engine/Plugins/Animation/ControlRig/Source/ControlRig/Public/AnimNode_ControlRigBase.h
--- a/Engine/Plugins/Animation/ControlRig/Source/ControlRig/Public/AnimNode_ControlRigBase.h
+++ b/Engine/Plugins/Animation/ControlRig/Source/ControlRig/Public/AnimNode_ControlRigBase.h
@@ -102,6 +102,20 @@ protected:
 
 	struct FDependentTransform
 	{
+		FDependentTransform()
+			: KeyAndIndex()
+			, TransformType(ERigTransformType::InitialLocal)
+			, StorageType(ERigTransformStorageType::Pose)
+			, DirtyState(nullptr)
+		{}
+		
+		FDependentTransform(const FRigElementKeyAndIndex& InKeyAndIndex, ERigTransformType::Type InTransformType, ERigTransformStorageType::Type InStorageType, FRigLocalAndGlobalDirtyState* InDirtyState)
+			: KeyAndIndex(InKeyAndIndex)
+			, TransformType(InTransformType)
+			, StorageType(InStorageType)
+			, DirtyState(InDirtyState)
+		{}
+
 		FRigElementKeyAndIndex KeyAndIndex;
 		ERigTransformType::Type TransformType;
 		ERigTransformStorageType::Type StorageType;
