Control Rig Edit Mode: fixed hovered shape

#rnx

diff --git a/Engine/Plugins/Animation/ControlRig/Source/ControlRigEditor/Private/EditMode/ControlRigEditMode.cpp b/Engine/Plugins/Animation/ControlRig/Source/ControlRigEditor/Private/EditMode/ControlRigEditMode.cpp
--- a/Engine/Plugins/Animation/ControlRig/Source/ControlRigEditor/Private/EditMode/ControlRigEditMode.cpp
+++ b/Engine/Plugins/Animation/ControlRig/Source/ControlRigEditor/Private/EditMode/ControlRigEditMode.cpp
@@ -3671,18 +3671,20 @@ bool FControlRigEditMode::MouseMove(FEditorViewportClient* InViewportClient, FVi
 	
 	if (HasAnyHoverableShapeActor())
 	{	
-		HActor* ActorHitProxy = HitProxyCast<HActor>(InViewport->GetHitProxy(InX, InY));
-		if (ActorHitProxy && ActorHitProxy->Actor)
+		const HActor* ActorHitProxy = HitProxyCast<HActor>(InViewport->GetHitProxy(InX, InY));
+		const AControlRigShapeActor* HitShape = ActorHitProxy && ActorHitProxy->Actor ? Cast<AControlRigShapeActor>(ActorHitProxy->Actor) : nullptr;
+		auto IsHovered = [HitShape](const AControlRigShapeActor* InShapeActor)
 		{
-			if (ActorHitProxy->Actor->IsA<AControlRigShapeActor>())
+			return HitShape ? InShapeActor == HitShape : false;
+		};
+
+		for (const auto& [ControlRig, Shapes] : ControlRigShapeActors)
+		{
+			for (AControlRigShapeActor* ShapeActor : Shapes)
 			{
-				for (auto& ShapeActors : ControlRigShapeActors)
+				if (ShapeActor)
 				{
-					for (AControlRigShapeActor* ShapeActor : ShapeActors.Value)
-					{
-						ShapeActor->SetHovered(ShapeActor == ActorHitProxy->Actor);
-						return false;
-					}
+					ShapeActor->SetHovered(IsHovered(ShapeActor));
 				}
 			}
 		}
