[Sequencer Control Rig] 2 changes
+ change to use asset picker instead of class picker for control rig, and added support for filtering control rig assets via tags
+ minor: generate unique name for new rig track that shares the same rig as existing tracks

diff --git a/Engine/Plugins/Animation/ControlRig/Source/ControlRigEditor/Private/Sequencer/ControlRigParameterTrackEditor.cpp b/Engine/Plugins/Animation/ControlRig/Source/ControlRigEditor/Private/Sequencer/ControlRigParameterTrackEditor.cpp
--- a/Engine/Plugins/Animation/ControlRig/Source/ControlRigEditor/Private/Sequencer/ControlRigParameterTrackEditor.cpp
+++ b/Engine/Plugins/Animation/ControlRig/Source/ControlRigEditor/Private/Sequencer/ControlRigParameterTrackEditor.cpp
@@ -103,6 +103,7 @@
 #include "LevelSequence.h"
 #include "ISequencerModule.h"
 #include "CurveModel.h"
+#include "FrontendFilterBase.h"
 
 #define LOCTEXT_NAMESPACE "FControlRigParameterTrackEditor"
 
@@ -226,6 +227,12 @@ static USkeleton* AcquireSkeletonFromObjectGuid(const FGuid& Guid, UObject** Obj
 	return nullptr;
 }
 
+static USkeletalMeshComponent* AcquireSkeletalMeshFromObjectGuid(const FGuid& Guid, TSharedPtr<ISequencer> SequencerPtr)
+{
+	UObject* BoundObject = SequencerPtr.IsValid() ? SequencerPtr->FindSpawnedObjectOrTemplate(Guid) : nullptr;
+	return AcquireSkeletalMeshFromObject(BoundObject, SequencerPtr);
+}
+
 static bool DoesControlRigAllowMultipleInstances(const FTopLevelAssetPath& InGeneratedClassPath)
 {
 	const IAssetRegistry& AssetRegistry = FModuleManager::GetModuleChecked<FAssetRegistryModule>(TEXT("AssetRegistry")).Get();
@@ -241,6 +248,45 @@ static bool DoesControlRigAllowMultipleInstances(const FTopLevelAssetPath& InGen
 	return false;
 }
 
+
+class FFrontendFilter_ControlRigFilterByAssetTag: public FFrontendFilter
+{
+public:
+	FFrontendFilter_ControlRigFilterByAssetTag(TSharedPtr<FFrontendFilterCategory> InCategory, const FRigVMTag& InTag) :
+		FFrontendFilter(InCategory),
+		Tag(InTag)
+	{}
+	
+	FString GetName() const override { return Tag.Name.ToString(); }
+	FText GetDisplayName() const override { return FText::FromString(Tag.GetLabel()); }
+	FText GetToolTipText() const override { return Tag.ToolTip; }
+	FLinearColor GetColor() const override { return Tag.Color; };
+	
+	bool PassesFilter(const FContentBrowserItem& InItem) const override
+	{
+		FAssetData AssetData;
+		if (InItem.Legacy_TryGetAssetData(AssetData))
+		{
+			static const FName AssetVariantPropertyName = GET_MEMBER_NAME_CHECKED(URigVMBlueprint, AssetVariant);
+			const FProperty* AssetVariantProperty = CastField<FProperty>(URigVMBlueprint::StaticClass()->FindPropertyByName(AssetVariantPropertyName));
+			const FString VariantStr = AssetData.GetTagValueRef<FString>(AssetVariantPropertyName);
+			if(!VariantStr.IsEmpty())
+			{
+				FRigVMVariant AssetVariant;
+				AssetVariantProperty->ImportText_Direct(*VariantStr, &AssetVariant, nullptr, EPropertyPortFlags::PPF_None);
+				if (AssetVariant.Tags.Contains(Tag))
+				{
+					return true;
+				}
+			}
+		}
+		return false;
+	}
+
+private:
+	FRigVMTag Tag;
+};
+
 bool FControlRigParameterTrackEditor::bAutoGenerateControlRigTrack = true;
 FCriticalSection FControlRigParameterTrackEditor::ControlUndoTransactionMutex;
 
@@ -1263,18 +1309,7 @@ void FControlRigParameterTrackEditor::BakeToControlRig(UClass* InClass, FGuid Ob
 	}
 }
 
-void FControlRigParameterTrackEditor::IterateTracks(TFunctionRef<bool(UMovieSceneControlRigParameterTrack*)> Callback) const
-{
-	UMovieScene* MovieScene = GetSequencer().IsValid() && GetSequencer()->GetFocusedMovieSceneSequence() ? GetSequencer()->GetFocusedMovieSceneSequence()->GetMovieScene() : nullptr;
-	if (!MovieScene)
-	{
-		return;
-	}
-
-	IterateTracksInMovieScene(*MovieScene, Callback);
-}
-
-void FControlRigParameterTrackEditor::IterateTracksInMovieScene(UMovieScene& MovieScene, TFunctionRef<bool(UMovieSceneControlRigParameterTrack*)> Callback) const
+static void IterateTracksInMovieScene(UMovieScene& MovieScene, TFunctionRef<bool(UMovieSceneControlRigParameterTrack*)> Callback)
 {
 	TArray<UMovieSceneControlRigParameterTrack*> Tracks;
 	
@@ -1300,6 +1335,18 @@ void FControlRigParameterTrackEditor::IterateTracksInMovieScene(UMovieScene& Mov
 	}
 }
 
+void FControlRigParameterTrackEditor::IterateTracks(TFunctionRef<bool(UMovieSceneControlRigParameterTrack*)> Callback) const
+{
+	UMovieScene* MovieScene = GetSequencer().IsValid() && GetSequencer()->GetFocusedMovieSceneSequence() ? GetSequencer()->GetFocusedMovieSceneSequence()->GetMovieScene() : nullptr;
+	if (!MovieScene)
+	{
+		return;
+	}
+
+	IterateTracksInMovieScene(*MovieScene, Callback);
+}
+
+
 void FControlRigParameterTrackEditor::BakeInvertedPose(UControlRig* InControlRig, UMovieSceneControlRigParameterTrack* Track)
 {
 	if (!InControlRig->IsAdditive())
@@ -1434,6 +1481,38 @@ void FControlRigParameterTrackEditor::BuildObjectBindingTrackMenu(FMenuBuilder&
 	}
 }
 
+void FControlRigParameterTrackEditor::ToggleIsAdditiveControlRig()
+{
+	bIsLayeredControlRig = bIsLayeredControlRig ? false : true;
+	RefreshControlRigPickerDelegate.ExecuteIfBound(true);
+}
+
+bool FControlRigParameterTrackEditor::IsToggleIsAdditiveControlRig()
+{
+	return bIsLayeredControlRig;
+}
+
+void FControlRigParameterTrackEditor::ToggleFilterAssetBySkeleton()
+{
+	bFilterAssetBySkeleton = bFilterAssetBySkeleton ? false : true;
+	RefreshControlRigPickerDelegate.ExecuteIfBound(true);
+}
+
+bool FControlRigParameterTrackEditor::IsToggleFilterAssetBySkeleton()
+{
+	return bFilterAssetBySkeleton;
+}
+
+void FControlRigParameterTrackEditor::ToggleFilterAssetByAnimatableControls()
+{
+	bFilterAssetByAnimatableControls = bFilterAssetByAnimatableControls ? false : true;
+	RefreshControlRigPickerDelegate.ExecuteIfBound(true);
+}
+
+bool FControlRigParameterTrackEditor::IsToggleFilterAssetByAnimatableControls()
+{
+	return bFilterAssetByAnimatableControls;
+}
 
 void FControlRigParameterTrackEditor::HandleAddTrackSubMenu(FMenuBuilder& MenuBuilder, TArray<FGuid> ObjectBindings, UMovieSceneTrack* Track)
 {
@@ -1476,96 +1555,13 @@ void FControlRigParameterTrackEditor::HandleAddTrackSubMenu(FMenuBuilder& MenuBu
 		NAME_None,
 		EUserInterfaceActionType::ToggleButton);
 
-	MenuBuilder.AddSubMenu(
-		LOCTEXT("AddControlRigClass", "Control Rig Classes"),
-		LOCTEXT("AddControlRigClassTooltip", "Adds a Control Rig track based on selected class"),
-		FNewMenuDelegate::CreateRaw(this, &FControlRigParameterTrackEditor::HandleAddControlRigSubMenu, ObjectBindings, Track)
-	);
-}
-
-void FControlRigParameterTrackEditor::ToggleIsAdditiveControlRig()
-{
-	bIsLayeredControlRig = bIsLayeredControlRig ? false : true;
-}
-
-bool FControlRigParameterTrackEditor::IsToggleIsAdditiveControlRig()
-{
-	return bIsLayeredControlRig;
-}
-
-void FControlRigParameterTrackEditor::ToggleFilterAssetBySkeleton()
-{
-	bFilterAssetBySkeleton = bFilterAssetBySkeleton ? false : true;
-}
-
-bool FControlRigParameterTrackEditor::IsToggleFilterAssetBySkeleton()
-{
-	return bFilterAssetBySkeleton;
-}
-
-void FControlRigParameterTrackEditor::ToggleFilterAssetByAnimatableControls()
-{
-	bFilterAssetByAnimatableControls = bFilterAssetByAnimatableControls ? false : true;
-
-}
-
-bool FControlRigParameterTrackEditor::IsToggleFilterAssetByAnimatableControls()
-{
-	return bFilterAssetByAnimatableControls;
-}
-
-void FControlRigParameterTrackEditor::HandleAddControlRigSubMenu(FMenuBuilder& MenuBuilder, TArray<FGuid> ObjectBindings, UMovieSceneTrack* Track)
-{
-	/*
-	MenuBuilder.BeginSection(TEXT("ChooseSequence"), LOCTEXT("ChooseSequence", "Choose Sequence"));
-	{
-	FAssetPickerConfig AssetPickerConfig;
-	{
-	AssetPickerConfig.OnAssetSelected = FOnAssetSelected::CreateRaw(this, &FControlRigTrackEditor::OnSequencerAssetSelected, ObjectBindings, Track);
-	AssetPickerConfig.OnAssetEnterPressed = FOnAssetEnterPressed::CreateRaw(this, &FControlRigTrackEditor::OnSequencerAssetEnterPressed, ObjectBindings, Track);
-	AssetPickerConfig.bAllowNullSelection = false;
-	AssetPickerConfig.InitialAssetViewType = EAssetViewType::List;
-	AssetPickerConfig.Filter.bRecursiveClasses = true;
-	AssetPickerConfig.Filter.ClassPaths.Add(UControlRigSequence::StaticClass()->GetClassPathName());
-	AssetPickerConfig.SaveSettingsName = TEXT("SequencerAssetPicker");
-	}
-
-	FContentBrowserModule& ContentBrowserModule = FModuleManager::Get().LoadModuleChecked<FContentBrowserModule>(TEXT("ContentBrowser"));
-
-	TSharedPtr<SBox> MenuEntry = SNew(SBox)
-	.WidthOverride(300.0f)
-	.HeightOverride(300.0f)
-	[
-	ContentBrowserModule.Get().CreateAssetPicker(AssetPickerConfig)
-	];
-
-	MenuBuilder.AddWidget(MenuEntry.ToSharedRef(), FText::GetEmpty(), true);
-	}
-	MenuBuilder.EndSection();
-	*/
-
-
 	const TSharedPtr<ISequencer> ParentSequencer = GetSequencer();
 	UObject* BoundObject = nullptr;
 	//todo support multiple bindings?
-	USkeleton* Skeleton = AcquireSkeletonFromObjectGuid(ObjectBindings[0], &BoundObject, GetSequencer());
+	USkeleton* Skeleton = AcquireSkeletonFromObjectGuid(ObjectBindings[0], &BoundObject, ParentSequencer);
 
 	if (Skeleton)
 	{
-		//MenuBuilder.AddSubMenu(
-		//	LOCTEXT("AddControlRigTrack", "ControlRigTrack"), NSLOCTEXT("ControlRig", "AddControlRigTrack", "Adds a Control Rigtrack."),
-		//	FNewMenuDelegate::CreateRaw(this, &FControlRigParameterTrackEditor::AddAnimationSubMenu, BoundObject, ObjectBindings[0], Skeleton)
-		//);
-
-		FClassViewerInitializationOptions Options;
-		Options.bShowUnloadedBlueprints = true;
-		Options.NameTypeToDisplay = EClassViewerNameTypeToDisplay::DisplayName;
-
-		const bool bCheckInversion = bIsLayeredControlRig;
-		TSharedPtr<FControlRigClassFilter> ClassFilter = MakeShareable(new FControlRigClassFilter(bFilterAssetBySkeleton, bFilterAssetByAnimatableControls, bCheckInversion, Skeleton));
-		Options.ClassFilters.Add(ClassFilter.ToSharedRef());
-		Options.bShowNoneOption = false;
-		Options.ClassViewerSortPredicate = ClassViewerSortPredicate;
 		TArray<UClass*> ExistingRigs;
 		USkeletalMeshComponent* SkeletalMeshComponent = AcquireSkeletalMeshFromObject(BoundObject, ParentSequencer);
 		IterateTracks([&ExistingRigs, SkeletalMeshComponent](UMovieSceneControlRigParameterTrack* Track) -> bool
@@ -1582,52 +1578,180 @@ void FControlRigParameterTrackEditor::HandleAddControlRigSubMenu(FMenuBuilder& M
 			}
 			return true;
 		});
-		TSharedPtr<FClassViewerHideAlreadyAddedRigsFilter> ExistingClassesFilter = MakeShareable(new FClassViewerHideAlreadyAddedRigsFilter(ExistingRigs));
-		Options.ClassFilters.Add(ExistingClassesFilter.ToSharedRef());
-		if (!ExistingRigs.Contains(UFKControlRig::StaticClass()))
-		{
-			Options.ExtraPickerCommonClasses.Add(UFKControlRig::StaticClass());
-		}
-
-		UMovieSceneSequence* Sequence = GetSequencer() ? GetSequencer()->GetFocusedMovieSceneSequence() : nullptr;
-		Options.AdditionalReferencingAssets.Add(FAssetData(Sequence));
-
-		FClassViewerModule& ClassViewerModule = FModuleManager::LoadModuleChecked<FClassViewerModule>("ClassViewer");
-		//	FNewMenuDelegate::CreateRaw(this, &FControlRigParameterTrackEditor::AddAnimationSubMenu, BoundObject, ObjectBindings[0], Skeleton)
-
-
-
-		TSharedRef<SWidget> ClassViewer = ClassViewerModule.CreateClassViewer(Options, FOnClassPicked::CreateRaw(this, &FControlRigParameterTrackEditor::AddControlRig, BoundObject, ObjectBindings[0]));
-		MenuBuilder.AddWidget(ClassViewer, FText::GetEmpty(), true);
-
-		/*
+		MenuBuilder.AddSeparator();
+		MenuBuilder.AddMenuEntry(
+			LOCTEXT("FKControlRig", "FK Control Rig"),
+			LOCTEXT("FKControlRigTooltip", "Adds the FK Control Rig"),
+			FSlateIcon(),
+			FUIAction(
+				FExecuteAction::CreateSP(this, &FControlRigParameterTrackEditor::AddFKControlRig, BoundObject, ObjectBindings[0]),
+				FCanExecuteAction::CreateLambda([ExistingRigs]()
+				{
+					if (!ExistingRigs.Contains(UFKControlRig::StaticClass()))
+					{
+						return true;
+					}
+					return false;
+				})
+			),
+			NAME_None,
+			EUserInterfaceActionType::Button);
+		
 		FAssetPickerConfig AssetPickerConfig;
 		{
-		AssetPickerConfig.OnAssetSelected = FOnAssetSelected::CreateRaw(this, &FControlRigParameterTrackEditor::OnControlRigAssetSelected, ObjectBindings, Skeleton);
-		AssetPickerConfig.OnAssetEnterPressed = FOnAssetEnterPressed::CreateRaw(this, &FControlRigParameterTrackEditor::OnControlRigAssetEnterPressed, ObjectBindings, Skeleton);
-		AssetPickerConfig.bAllowNullSelection = false;
-		AssetPickerConfig.InitialAssetViewType = EAssetViewType::List;
-		AssetPickerConfig.OnShouldFilterAsset = FOnShouldFilterAsset::CreateRaw(this, &FControlRigParameterTrackEditor::ShouldFilterAsset);
-		AssetPickerConfig.Filter.bRecursiveClasses = true;
-		AssetPickerConfig.Filter.ClassPaths.Add((UControlRig::StaticClass())->GetClassPathName());
-		AssetPickerConfig.Filter.TagsAndValues.Add(TEXT("Skeleton"), FAssetData(Skeleton).GetExportTextName());
-		AssetPickerConfig.SaveSettingsName = TEXT("SequencerAssetPicker");
+			AssetPickerConfig.SelectionMode = ESelectionMode::Single;
+			AssetPickerConfig.bAddFilterUI = true;
+			AssetPickerConfig.bFocusSearchBoxWhenOpened = true;
+			AssetPickerConfig.OnAssetSelected = FOnAssetSelected::CreateRaw(this, &FControlRigParameterTrackEditor::AddControlRig, BoundObject, ObjectBindings[0]);
+			AssetPickerConfig.OnAssetEnterPressed = FOnAssetEnterPressed::CreateRaw(this, &FControlRigParameterTrackEditor::AddControlRig, BoundObject, ObjectBindings[0]);
+			AssetPickerConfig.RefreshAssetViewDelegates.Add(&RefreshControlRigPickerDelegate);
+			AssetPickerConfig.OnShouldFilterAsset = FOnShouldFilterAsset::CreateLambda([this, ExistingRigs, Skeleton](const FAssetData& AssetData)
+			{
+				if (!IsControlRigAllowed(AssetData, ExistingRigs, Skeleton))
+				{
+					// Should be filtered out
+					return true;
+				}
 
+				return false;
+			});
+			AssetPickerConfig.bAllowNullSelection = false;
+			AssetPickerConfig.InitialAssetViewType = EAssetViewType::List;
+			AssetPickerConfig.Filter.bRecursiveClasses = true;
+			AssetPickerConfig.Filter.ClassPaths.Add((UControlRigBlueprint::StaticClass())->GetClassPathName());
+			AssetPickerConfig.SaveSettingsName = TEXT("SequencerControlRigTrackAssetPicker");
+			TSharedRef<FFrontendFilterCategory>	ControlRigFilterCategory = MakeShared<FFrontendFilterCategory>(LOCTEXT("ControlRigFilterCategoryName", "ControlRig Tags"), LOCTEXT("ControlRigFilterCategoryToolTip", "Filter ControlRigs by tags specified in ControlRig Blueprint class settings"));
+			const URigVMProjectSettings* Settings = GetDefault<URigVMProjectSettings>(URigVMProjectSettings::StaticClass());
+			TArray<FRigVMTag> AvailableTags = Settings->VariantTags;
+
+			for (const FRigVMTag& Tag : AvailableTags)
+			{
+				AssetPickerConfig.ExtraFrontendFilters.Add(MakeShared<FFrontendFilter_ControlRigFilterByAssetTag>(ControlRigFilterCategory, Tag));
+			}
 		}
 
 		FContentBrowserModule& ContentBrowserModule = FModuleManager::Get().LoadModuleChecked<FContentBrowserModule>(TEXT("ContentBrowser"));
 
 		TSharedPtr<SBox> MenuEntry = SNew(SBox)
-		.WidthOverride(300.0f)
+		// Extra space to display filter capsules horizontally
+		.WidthOverride(600.f)
 		.HeightOverride(300.f)
 		[
-		ContentBrowserModule.Get().CreateAssetPicker(AssetPickerConfig)
+			ContentBrowserModule.Get().CreateAssetPicker(AssetPickerConfig)
 		];
 
 		MenuBuilder.AddWidget(MenuEntry.ToSharedRef(), FText::GetEmpty(), true);
-		*/
 	}
+}
 
+bool FControlRigParameterTrackEditor::IsControlRigAllowed(const FAssetData& AssetData, TArray<UClass*> ExistingRigs, USkeleton* Skeleton)
+{
+	if (UControlRigBlueprint* LoadedControlRig = Cast<UControlRigBlueprint>(AssetData.FastGetAsset()))
+	{
+		if (ExistingRigs.Contains(LoadedControlRig->GetRigVMBlueprintGeneratedClass()))
+		{
+			if (!AssetData.GetTagValueRef<bool>(GET_MEMBER_NAME_CHECKED(UControlRigBlueprint, bAllowMultipleInstances)))
+			{
+				return false;
+			}	
+		}
+	}
+
+	const IAssetRegistry& AssetRegistry = FModuleManager::GetModuleChecked<FAssetRegistryModule>(TEXT("AssetRegistry")).Get();
+	
+	const bool bExposesAnimatableControls = AssetData.GetTagValueRef<bool>(TEXT("bExposesAnimatableControls"));
+	if (bFilterAssetByAnimatableControls == true && bExposesAnimatableControls == false)
+	{
+		return false;
+	}
+	if (bIsLayeredControlRig)
+	{		
+		FAssetDataTagMapSharedView::FFindTagResult Tag = AssetData.TagsAndValues.FindTag(TEXT("SupportedEventNames"));
+		if (Tag.IsSet())
+		{
+			bool bHasInversion = false;
+			FString EventString = FRigUnit_InverseExecution::EventName.ToString();
+			FString OldEventString = FString(TEXT("Inverse"));
+			TArray<FString> SupportedEventNames;
+			Tag.GetValue().ParseIntoArray(SupportedEventNames, TEXT(","), true);
+
+			for (const FString& Name : SupportedEventNames)
+			{
+				if (Name.Contains(EventString) || Name.Contains(OldEventString))
+				{
+					bHasInversion = true;
+					break;
+				}
+			}
+			if (bHasInversion == false)
+			{
+				return false;
+			}
+		}
+	}
+	if (bFilterAssetBySkeleton)
+	{
+		FString SkeletonName;
+		if (Skeleton)
+		{
+			SkeletonName = FAssetData(Skeleton).GetExportTextName();
+		}
+		FString PreviewSkeletalMesh = AssetData.GetTagValueRef<FString>(TEXT("PreviewSkeletalMesh"));
+		if (PreviewSkeletalMesh.Len() > 0)
+		{
+			FAssetData SkelMeshData = AssetRegistry.GetAssetByObjectPath(FSoftObjectPath(PreviewSkeletalMesh));
+			FString PreviewSkeleton = SkelMeshData.GetTagValueRef<FString>(TEXT("Skeleton"));
+			if (PreviewSkeleton == SkeletonName)
+			{
+				return true;
+			}
+			else if(Skeleton)
+			{
+				if (Skeleton->IsCompatibleForEditor(PreviewSkeleton))
+				{
+					return true;
+				}
+			}
+		}
+		FString PreviewSkeleton = AssetData.GetTagValueRef<FString>(TEXT("PreviewSkeleton"));
+		if (PreviewSkeleton == SkeletonName)
+		{
+			return true;
+		}
+		else if (Skeleton)
+		{
+			if (Skeleton->IsCompatibleForEditor(PreviewSkeleton))
+			{
+				return true;
+			}
+		}
+		FString SourceHierarchyImport = AssetData.GetTagValueRef<FString>(TEXT("SourceHierarchyImport"));
+		if (SourceHierarchyImport == SkeletonName)
+		{
+			return true;
+		}
+		else if (Skeleton)
+		{
+			if (Skeleton->IsCompatibleForEditor(SourceHierarchyImport))
+			{
+				return true;
+			}
+		}
+		FString SourceCurveImport = AssetData.GetTagValueRef<FString>(TEXT("SourceCurveImport"));
+		if (SourceCurveImport == SkeletonName)
+		{
+			return true;
+		}
+		else if (Skeleton)
+		{
+			if (Skeleton->IsCompatibleForEditor(SourceCurveImport))
+			{
+				return true;
+			}
+		}
+		return false;
+	}
+	return true;	
 }
 
 /*
@@ -1661,6 +1785,28 @@ void FControlRigParameterTrackEditor::OnControlRigAssetSelected(const FAssetData
 		}
 	}
 */
+
+static TArray<UMovieSceneControlRigParameterTrack*> GetExistingControlRigTracksForSkeletalMeshComponent(UMovieScene* MovieScene, USkeletalMeshComponent* SkeletalMeshComponent)
+{
+	TArray<UMovieSceneControlRigParameterTrack*> ExistingControlRigTracks;
+	IterateTracksInMovieScene(*MovieScene, [&ExistingControlRigTracks, SkeletalMeshComponent](UMovieSceneControlRigParameterTrack* Track) -> bool
+	{
+		if (UControlRig* ControlRig = Track->GetControlRig())
+		{
+			if (TSharedPtr<IControlRigObjectBinding> ObjectBinding = ControlRig->GetObjectBinding())
+			{
+				if (ObjectBinding.IsValid() && ObjectBinding->GetBoundObject() == SkeletalMeshComponent)
+				{
+					ExistingControlRigTracks.Add(Track);
+				}
+			}
+		}
+		return true;
+	});
+
+	return ExistingControlRigTracks;
+}
+
 static UMovieSceneControlRigParameterTrack* AddControlRig(TSharedPtr<ISequencer> SharedSequencer , UMovieSceneSequence* Sequence, const UClass* InClass, FGuid ObjectBinding, UControlRig* InExistingControlRig, bool bIsAdditiveControlRig)
 {
 	FSlateApplication::Get().DismissAllMenus();
@@ -1687,12 +1833,44 @@ static UMovieSceneControlRigParameterTrack* AddControlRig(TSharedPtr<ISequencer>
 
 	FScopedTransaction AddControlRigTrackTransaction(LOCTEXT("AddControlRigTrack", "Add Control Rig Track"));
 
+	TArray<UMovieSceneControlRigParameterTrack*> ExistingRigTracks;
+	
+	if (USkeletalMeshComponent* SkeletalMeshComponent = AcquireSkeletalMeshFromObjectGuid(ObjectBinding, SharedSequencer))
+	{
+		ExistingRigTracks = GetExistingControlRigTracksForSkeletalMeshComponent(OwnerMovieScene, SkeletalMeshComponent);
+	}
+	
 	UMovieSceneControlRigParameterTrack* Track = Cast<UMovieSceneControlRigParameterTrack>(OwnerMovieScene->AddTrack(UMovieSceneControlRigParameterTrack::StaticClass(), ObjectBinding));
 	if (Track)
 	{
+		TArray<FName> ExistingObjectNames;
+		for (UMovieSceneControlRigParameterTrack* RigTrack : ExistingRigTracks)
+		{
+			if (UControlRig* Rig = RigTrack->GetControlRig())
+			{
+				if (Rig->GetClass() == InClass)
+				{
+					ExistingObjectNames.Add(RigTrack->GetTrackName());
+				}
+			}
+		}
+		
 		FString ObjectName = InClass->GetName(); //GetDisplayNameText().ToString();
 		ObjectName.RemoveFromEnd(TEXT("_C"));
 
+		{
+			FName UniqueObjectName = *ObjectName;
+			int32 UniqueSuffix = 1;
+			while(ExistingObjectNames.Contains(UniqueObjectName))
+			{
+				UniqueObjectName = *(ObjectName + TEXT("_") + FString::FromInt(UniqueSuffix));
+				UniqueSuffix++;
+			}
+
+			ObjectName = UniqueObjectName.ToString();
+		}
+		
+		
 		bool bSequencerOwnsControlRig = false;
 		UControlRig* ControlRig = InExistingControlRig;
 		if (ControlRig == nullptr)
@@ -1815,6 +1993,27 @@ void FControlRigParameterTrackEditor::AddControlRig(const UClass* InClass, UObje
 	}
 }
 
+void FControlRigParameterTrackEditor::AddControlRig(const FAssetData& InAsset, UObject* BoundActor, FGuid ObjectBinding)
+{
+	if (UControlRigBlueprint* ControlRigBlueprint = Cast<UControlRigBlueprint>(InAsset.GetAsset()))
+	{
+		AddControlRig(ControlRigBlueprint->GetRigVMBlueprintGeneratedClass(), BoundActor, ObjectBinding);
+	}
+}
+
+void FControlRigParameterTrackEditor::AddControlRig(const TArray<FAssetData>& InAssets, UObject* BoundActor, FGuid ObjectBinding)
+{
+	if (InAssets.Num() > 0)
+	{
+		AddControlRig(InAssets[0], BoundActor, ObjectBinding);	
+	}
+}
+
+void FControlRigParameterTrackEditor::AddFKControlRig(UObject* BoundActor, FGuid ObjectBinding)
+{
+	AddControlRig(UFKControlRig::StaticClass(), BoundActor, ObjectBinding);
+}
+
 void FControlRigParameterTrackEditor::AddControlRig(UClass* InClass, UObject* BoundActor, FGuid ObjectBinding)
 {
 	if (InClass == UFKControlRig::StaticClass())
diff --git a/Engine/Plugins/Animation/ControlRig/Source/ControlRigEditor/Private/Sequencer/ControlRigParameterTrackEditor.h b/Engine/Plugins/Animation/ControlRig/Source/ControlRigEditor/Private/Sequencer/ControlRigParameterTrackEditor.h
--- a/Engine/Plugins/Animation/ControlRig/Source/ControlRigEditor/Private/Sequencer/ControlRigParameterTrackEditor.h
+++ b/Engine/Plugins/Animation/ControlRig/Source/ControlRigEditor/Private/Sequencer/ControlRigParameterTrackEditor.h
@@ -14,6 +14,7 @@
 #include "KeyframeTrackEditor.h"
 #include "Sections/MovieScene3DTransformSection.h"
 #include "AcquiredResources.h"
+#include "ContentBrowserDelegates.h"
 #include "MovieSceneToolHelpers.h"
 #include "MovieSceneToolsModule.h"
 #include "Engine/EngineTypes.h"
@@ -90,7 +91,6 @@ public:
 private:
 
 	void HandleAddTrackSubMenu(FMenuBuilder& MenuBuilder, TArray<FGuid> ObjectBindings, UMovieSceneTrack* Track);
-	void HandleAddControlRigSubMenu(FMenuBuilder& MenuBuilder, TArray<FGuid> ObjectBindings, UMovieSceneTrack* Track);
 
 	void ToggleIsAdditiveControlRig();
 	bool IsToggleIsAdditiveControlRig();
@@ -105,7 +105,12 @@ private:
 	/** Control Rig Picked */
 	void AddControlRig(UClass* InClass, UObject* BoundActor, FGuid ObjectBinding);
 	void AddControlRig(const UClass* InClass, UObject* BoundActor, FGuid ObjectBinding, UControlRig* InExistingControlRig);
+	void AddControlRig(const FAssetData& InAsset, UObject* BoundActor, FGuid ObjectBinding);
+	void AddControlRig(const TArray<FAssetData>& InAssets, UObject* BoundActor, FGuid ObjectBinding);
+	void AddFKControlRig(UObject* BoundActor, FGuid ObjectBinding);
 	void AddControlRigFromComponent(FGuid InGuid);
+
+	bool IsControlRigAllowed(const FAssetData& AssetData, TArray<UClass*> ExistingRigs, USkeleton* Skeleton);
 	
 	/** Delegate for Selection Changed Event */
 	void OnSelectionChanged(TArray<UMovieSceneTrack*> InTracks);
@@ -210,7 +215,6 @@ private:
 
 	/** Helper functions to iterate over UMovieSceneControlRigParameterTracks in the currently focussed MovieScene*/
 	void IterateTracks(TFunctionRef<bool(UMovieSceneControlRigParameterTrack*)> Callback) const;
-	void IterateTracksInMovieScene(UMovieScene& MovieScene, TFunctionRef<bool(UMovieSceneControlRigParameterTrack*)> Callback) const; 
 private:
 	/** Command Bindings added by the Transform Track Editor to Sequencer and curve editor. */
 	TSharedPtr<FUICommandList> CommandBindings;
@@ -292,6 +296,8 @@ private:
 	/** Whether or not we should check for Animatable Controls when filtering*/
 	bool bFilterAssetByAnimatableControls;
 
+	FRefreshAssetViewDelegate RefreshControlRigPickerDelegate;
+
 	/** Handle to help updating selection on tick tick to avoid too many flooded selections*/
 	FTimerHandle UpdateSelectionTimerHandle;
 
diff --git a/Engine/Plugins/Runtime/RigVM/Source/RigVMDeveloper/Public/RigVMBlueprint.h b/Engine/Plugins/Runtime/RigVM/Source/RigVMDeveloper/Public/RigVMBlueprint.h
--- a/Engine/Plugins/Runtime/RigVM/Source/RigVMDeveloper/Public/RigVMBlueprint.h
+++ b/Engine/Plugins/Runtime/RigVM/Source/RigVMDeveloper/Public/RigVMBlueprint.h
@@ -861,6 +861,7 @@ private:
 	friend class SRigVMGraphNode;
 	friend struct FGuardSkipDirtyBlueprintStatus;
 	friend class SRigModuleAssetBrowser;
+	friend class FFrontendFilter_ControlRigFilterByAssetTag;
 };
 
 class RIGVMDEVELOPER_API FRigVMBlueprintCompileScope
